[Chaos] Perf improvements to UClusterUnionComponent::SyncClusterUnionFromProxy

- Attempt to pre-allocate the sublevel map used to map ClusterUnion data to Component/BoneIdTransform data, using the size of the existing map for a determined component. Also made it a TSorted map. TSortedMap is backed by an Array and does not do hashing. It search cost is higher for large number of elements though, but I didn't see a negative impact on my tests, and it seems to be faster for this use case.
- Updated the Component removal handling to avoid needing allocate a local array  and do a search in the components map multiple times.

diff --git a/Engine/Source/Runtime/Engine/Classes/PhysicsEngine/ClusterUnionComponent.h b/Engine/Source/Runtime/Engine/Classes/PhysicsEngine/ClusterUnionComponent.h
--- a/Engine/Source/Runtime/Engine/Classes/PhysicsEngine/ClusterUnionComponent.h
+++ b/Engine/Source/Runtime/Engine/Classes/PhysicsEngine/ClusterUnionComponent.h
@@ -252,6 +252,11 @@ public:
 			ComponentKey = InComponent;
 			ComponentPtr = InComponent;
 		}
+		
+		bool operator<(const FMappedComponentKey& Other) const
+		{
+			return ComponentKey < Other.ComponentKey;
+		}
 
 		bool operator==(const FMappedComponentKey& Other) const
 		{
@@ -285,6 +290,9 @@ public:
 		FTransform ChildToParentTransform;
 	};
 
+	// There are several instances where we have 1 single elements. The inline allocator should help in these cases without a huge impact when we go over
+	using FLocalBonesToTransformMap = TSortedMap<int32, FMappedBoneData, TInlineAllocator<1>>;
+
 	// SyncClusterUnionFromProxy will examine the make up of the cluster union (particles, child to parent, etc.) and do whatever is needed on the GT in terms of bookkeeping.
 	ENGINE_API void SyncClusterUnionFromProxy();
 
@@ -436,8 +444,8 @@ private:
 
 	// These functions only get called when the physics thread syncs to the game thread thereby enforcing a physics thread authoritative view of
 	// what particles are currently contained within the cluster union.
-	ENGINE_API void HandleAddOrModifiedClusteredComponent(const FMappedComponentKey& ChangedComponentData, const TMap<int32, FMappedBoneData>& PerBoneChildToParent);
-	ENGINE_API void HandleRemovedClusteredComponent(TObjectKey<UPrimitiveComponent> ChangedComponent);
+	ENGINE_API void HandleAddOrModifiedClusteredComponent(const FMappedComponentKey& ChangedComponentData, const FLocalBonesToTransformMap& PerBoneChildToParent);
+	ENGINE_API void HandleRemovedClusteredComponent(TObjectKey<UPrimitiveComponent> RemovedComponent, const FClusteredComponentData& ComponentData);
 
 	ENGINE_API TArray<UPrimitiveComponent*> GetAllCurrentChildComponents() const;
 	ENGINE_API TArray<AActor*> GetAllCurrentActors() const;
diff --git a/Engine/Source/Runtime/Engine/Private/PhysicsEngine/ClusterUnionComponent.cpp b/Engine/Source/Runtime/Engine/Private/PhysicsEngine/ClusterUnionComponent.cpp
--- a/Engine/Source/Runtime/Engine/Private/PhysicsEngine/ClusterUnionComponent.cpp
+++ b/Engine/Source/Runtime/Engine/Private/PhysicsEngine/ClusterUnionComponent.cpp
@@ -28,6 +28,12 @@ namespace
 	bool bUseLocalRoleForAuthorityCheck = true;
 	FAutoConsoleVariableRef CVarUseLocalRoleForAuthorityCheck(TEXT("ClusterUnion.UseLocalRoleForAuthorityCheck"), bUseLocalRoleForAuthorityCheck, TEXT("If true, we will only check this component's owner local role to determine authority"));
 
+	bool bPreAllocateLocalBoneDataMap = true;
+	FAutoConsoleVariableRef CVarPreAllocateLocalBoneDataMap(TEXT("ClusterUnion.PreAllocateLocalBoneDataMap"), bPreAllocateLocalBoneDataMap, TEXT("If true, it will reserve an expected size for the local map used to cache updated bone data"));
+
+	float LocalBoneDataMapGrowFactor = 1.2f;
+	FAutoConsoleVariableRef CVarLocalBoneDataMapGrowFactor(TEXT("ClusterUnion.LocalBoneDataMapGrowFactor"), LocalBoneDataMapGrowFactor, TEXT("Grow factor to apply to the size of bone data array of pre-existing component when preallocating the local bones data map"));
+
 	template<typename PayloadType>
 	struct TClusterUnionAABBTreeStorageTraits
 	{
@@ -733,13 +739,12 @@ void UClusterUnionComponent::OnDestroyPhysicsState()
 
 	// We need to make sure we *immediately* disconnect on the GT side since there's no guarantee the normal flow
 	// will happen once we've destroyed things.
-	TArray<TObjectKey<UPrimitiveComponent>> ComponentsToRemove;
-	PerComponentData.GetKeys(ComponentsToRemove);
-
-	for (const TObjectKey<UPrimitiveComponent>& Component : ComponentsToRemove)
+	for (TPair<TObjectKey<UPrimitiveComponent>, FClusteredComponentData>& ComponentData : PerComponentData)
 	{
-		HandleRemovedClusteredComponent(Component);
+		HandleRemovedClusteredComponent(ComponentData.Key, ComponentData.Value);
 	}
+	
+	PerComponentData.Reset();
 
 	if (FPhysScene_Chaos* Scene = GetChaosScene())
 	{
@@ -900,7 +905,12 @@ void UClusterUnionComponent::SyncClusterUnionFromProxy()
 	// Note that at the UClusterUnionComponent level we really only want to be dealing with components.
 	// Hence why we need to modify each of the particles that we synced from the game thread into a
 	// component + bone id combination for identification. 
-	TMap<FMappedComponentKey, TMap<int32, FMappedBoneData>> MappedData;
+	TMap<FMappedComponentKey, FLocalBonesToTransformMap> MappedData;
+
+	if (bPreAllocateLocalBoneDataMap)
+	{
+		MappedData.Reserve(PerComponentData.Num());
+	}
 
 	{
 		FLockedReadPhysicsObjectExternalInterface Interface = FPhysicsObjectExternalInterface::LockRead(GetWorld()->GetPhysicsScene());
@@ -909,11 +919,24 @@ void UClusterUnionComponent::SyncClusterUnionFromProxy()
 			// Using the scene's proxy to component mapping let's us detect a component physics state was destroyed.
 			if (UPrimitiveComponent* Component = Scene->GetOwningComponent<UPrimitiveComponent>(ChildData.Proxy))
 			{
-				FMappedComponentKey ComponentKey(Component);
-				TMap<int32, FMappedBoneData>& BoneIDToBoneData = MappedData.FindOrAdd(ComponentKey);
-	
+				FMappedComponentKey WrappedComponentKey(Component);
+				int32 PrevMappedDataNum = MappedData.Num();
+				FLocalBonesToTransformMap& BoneIDToBoneData = MappedData.FindOrAdd(WrappedComponentKey);
+
+				// If we added new entry, and it is for an existing component, preallocate the local map using the
+				// existing size as a base
+				if (bPreAllocateLocalBoneDataMap && PrevMappedDataNum < MappedData.Num())
+				{
+					if (FClusteredComponentData* ComponentData = PerComponentData.Find(WrappedComponentKey.ComponentKey))
+					{
+						if (!ComponentData->BonesData.IsEmpty())
+						{
+							BoneIDToBoneData.Reserve(ComponentData->BonesData.Num() * LocalBoneDataMapGrowFactor);
+						}
+					}
+				}
+
 				Chaos::FPhysicsObjectHandle Handle = Component->GetPhysicsObjectById(ChildData.BoneId);
-				
 				Chaos::FPBDRigidParticle* Particle = Interface->GetRigidParticle(Handle);
 
 				if (FullData.bDidSyncGeometry)
@@ -948,26 +971,22 @@ void UClusterUnionComponent::SyncClusterUnionFromProxy()
 	// We need to handle any additions, deletions, and modifications to any child in the cluster union here.
 	// If a component lives in MappedData but not in PerComponentData, new component!
 	// If a component lives in both, then it's a modified component.
-	for (const TPair<FMappedComponentKey, TMap<int32, FMappedBoneData>>& Kvp : MappedData)
+	for (const TPair<FMappedComponentKey, FLocalBonesToTransformMap>& Kvp : MappedData)
 	{
 		HandleAddOrModifiedClusteredComponent(Kvp.Key, Kvp.Value);
 	}
 
 	// If a component lives in PerComponentData but not in MappedData, deleted component!
-	TArray<TObjectKey<UPrimitiveComponent>> ComponentsToRemove;
-	for (const TPair<TObjectKey<UPrimitiveComponent>, FClusteredComponentData>& Kvp : PerComponentData)
+	for (TMap<TObjectKey<UPrimitiveComponent>, FClusteredComponentData>::TIterator RemoveComponentDataIterator = PerComponentData.CreateIterator(); RemoveComponentDataIterator; ++RemoveComponentDataIterator)
 	{
 		// Build a mapped key with a null component ptr as we don't have it at this point and we don't need either
-		FMappedComponentKey MappedComponentKey(Kvp.Key, nullptr);
+		FMappedComponentKey MappedComponentKey(RemoveComponentDataIterator.Key(), nullptr);
 		if (!MappedData.Contains(MappedComponentKey))
 		{
-			ComponentsToRemove.Add(Kvp.Key);
-		}
-	}
+			HandleRemovedClusteredComponent(MappedComponentKey.ComponentKey, RemoveComponentDataIterator.Value());
 
-	for (TObjectKey<UPrimitiveComponent> Component : ComponentsToRemove)
-	{
-		HandleRemovedClusteredComponent(Component);
+			RemoveComponentDataIterator.RemoveCurrent();
+		}
 	}
 
 	const FBoxSphereBounds OldBounds = CachedLocalBounds;
@@ -988,7 +1007,7 @@ void UClusterUnionComponent::SyncClusterUnionFromProxy()
 	}
 }
 
-void UClusterUnionComponent::HandleAddOrModifiedClusteredComponent(const FMappedComponentKey& ChangedComponentData, const TMap<int32, FMappedBoneData>& PerBoneChildToParent)
+void UClusterUnionComponent::HandleAddOrModifiedClusteredComponent(const FMappedComponentKey& ChangedComponentData, const FLocalBonesToTransformMap& PerBoneChildToParent)
 {
 	if (!ChangedComponentData.ComponentPtr || !ChangedComponentData.ComponentPtr->HasValidPhysicsState() || !ChangedComponentData.ComponentPtr->GetWorld())
 	{
@@ -1144,64 +1163,62 @@ static UClusterUnionReplicatedProxyComponent* GetReplicatedProxyFromComponent(co
 	return nullptr;
 }
 
-void UClusterUnionComponent::HandleRemovedClusteredComponent(TObjectKey<UPrimitiveComponent> RemovedComponent)
+void UClusterUnionComponent::HandleRemovedClusteredComponent(TObjectKey<UPrimitiveComponent> RemovedComponent, const FClusteredComponentData& ComponentData)
 {
-	if (FClusteredComponentData* Data = PerComponentData.Find(RemovedComponent))
+	if (AccelerationStructure)
 	{
-		if (AccelerationStructure)
+		for (const FClusterUnionBoneData& BoneData : ComponentData.BonesData)
 		{
-			for (const FClusterUnionBoneData& BoneData : Data->BonesData)
-			{
-				FExternalSpatialAccelerationPayload Handle;
-				Handle.Initialize(RemovedComponent, BoneData.ID, BoneData.ParticleID);
+			FExternalSpatialAccelerationPayload Handle;
+			Handle.Initialize(RemovedComponent, BoneData.ID, BoneData.ParticleID);
 
-				if (ensure(Handle.IsValid()))
-				{
-					AccelerationStructure->RemoveElement(Handle);
-				}
+			if (ensure(Handle.IsValid()))
+			{
+				AccelerationStructure->RemoveElement(Handle);
 			}
 		}
+	}
 
-		if (IsAuthority())
+	UPrimitiveComponent* RemovedComponentPtr = RemovedComponent.ResolveObjectPtr();
+
+	if (IsAuthority())
+	{
+		if (UClusterUnionReplicatedProxyComponent* ProxyComponent = ComponentData.ReplicatedProxyComponent.Get())
 		{
-			if (UClusterUnionReplicatedProxyComponent* ProxyComponent = Data->ReplicatedProxyComponent.Get())
-			{
-				ProxyComponent->DestroyComponent();
-			}
+			ProxyComponent->DestroyComponent();
 		}
-		else
+	}
+	else
+	{
+		// On the client, we need to reset the state of the replicated proxy has it may outlive the cluster union
+		// and may still be in existence when the cluster union is created again 
+		// Note: we need to get the replicated component from the actor on the client because the PerComponentData do not set the replicated proxy pointer (unlike the server )
+		if (UClusterUnionReplicatedProxyComponent* ProxyComponent = GetReplicatedProxyFromComponent(this, RemovedComponentPtr))
 		{
-			// On the client, we need to reset the state of the replicated proxy has it may outlive the cluster union
-			// and may still be in existence when the cluster union is created again 
-			// Note: we need to get the replciated component from the actor on the client because the PerComponentData do not set the replciated proxy pointer (unlike the server )
-			if (UClusterUnionReplicatedProxyComponent* ProxyComponent = GetReplicatedProxyFromComponent(this, RemovedComponent.ResolveObjectPtr()))
-			{
-				ProxyComponent->ResetTransientState();
-			}
+			ProxyComponent->ResetTransientState();
 		}
+	}
 
-		if (AActor* Owner = Data->Owner.Get())
+	if (AActor* Owner = ComponentData.Owner.Get())
+	{
+		if (FClusteredActorData* ActorData = ActorToComponents.Find(Owner))
 		{
-			if (FClusteredActorData* ActorData = ActorToComponents.Find(Owner))
-			{
-				ActorData->Components.Remove(RemovedComponent);
+			ActorData->Components.Remove(RemovedComponent);
 
-				if (ActorData->Components.IsEmpty())
+			if (ActorData->Components.IsEmpty())
+			{
+				if (IsAuthority())
 				{
-					if (IsAuthority())
-					{
-						Owner->SetReplicatingMovement(ActorData->bWasReplicatingMovement);
-					}
-					ActorToComponents.Remove(Owner);
+					Owner->SetReplicatingMovement(ActorData->bWasReplicatingMovement);
 				}
+				ActorToComponents.Remove(Owner);
 			}
 		}
+	}
 
-		if (UPrimitiveComponent* ChangedComponent = RemovedComponent.ResolveObjectPtr())
-		{
-			BroadcastComponentRemovedEvents(ChangedComponent, Data->BonesData);
-		}
-		PerComponentData.Remove(RemovedComponent);
+	if (RemovedComponentPtr)
+	{
+		BroadcastComponentRemovedEvents(RemovedComponentPtr, ComponentData.BonesData);
 	}
 	PendingComponentsToAdd.Remove(RemovedComponent);
 	PendingComponentSync.Remove(RemovedComponent);
