Fix ClusterUnon/GCs offset issue when a client going far away and coming back

diff --git a/Engine/Source/Runtime/Engine/Private/PhysicsEngine/ClusterUnionReplicatedProxyComponent.cpp b/Engine/Source/Runtime/Engine/Private/PhysicsEngine/ClusterUnionReplicatedProxyComponent.cpp
--- a/Engine/Source/Runtime/Engine/Private/PhysicsEngine/ClusterUnionReplicatedProxyComponent.cpp
+++ b/Engine/Source/Runtime/Engine/Private/PhysicsEngine/ClusterUnionReplicatedProxyComponent.cpp
@@ -174,7 +174,7 @@ void UClusterUnionReplicatedProxyComponent::PostRepNotifies()
 		bNetUpdateParticleBoneIds = false;
 	}
 
-	if (bNetUpdateParticleChildToParents && ParticleBoneIds.Num() == ParticleChildToParents.Num())
+	if ((bNetUpdateParticleChildToParents || bIsInitialReplication) && ParticleBoneIds.Num() == ParticleChildToParents.Num())
 	{
 		// This particular bit can't happen until *after* we add the component to the cluster union. There's an additional deferral
 		// in AddComponentToCluster that we have to wait for.
