[Automation] Added TestNotEqualInsensitive, a negative version of TestEqualInsensitive

#rnx

diff --git a/Engine/Source/Runtime/Core/Private/Misc/AutomationTest.cpp b/Engine/Source/Runtime/Core/Private/Misc/AutomationTest.cpp
--- a/Engine/Source/Runtime/Core/Private/Misc/AutomationTest.cpp
+++ b/Engine/Source/Runtime/Core/Private/Misc/AutomationTest.cpp
@@ -1756,6 +1756,16 @@ bool FAutomationTestBase::TestEqualInsensitive(const TCHAR* What, const TCHAR* A
 	return true;
 }
 
+bool FAutomationTestBase::TestNotEqualInsensitive(const TCHAR* What, const TCHAR* Actual, const TCHAR* Expected)
+{
+	if (FCString::Stricmp(Actual, Expected) == 0)
+	{
+		AddError(FString::Printf(TEXT("Expected '%s' to differ from \"%s\", but it was \"%s\"."), What, Expected, Actual), 1);
+		return false;
+	}
+	return true;
+}
+
 bool FAutomationTestBase::TestNearlyEqual(const TCHAR* What, const float Actual, const float Expected, float Tolerance)
 {
 	return TestEqual(What, Actual, Expected, Tolerance);
diff --git a/Engine/Source/Runtime/Core/Public/Misc/AutomationTest.h b/Engine/Source/Runtime/Core/Public/Misc/AutomationTest.h
--- a/Engine/Source/Runtime/Core/Public/Misc/AutomationTest.h
+++ b/Engine/Source/Runtime/Core/Public/Misc/AutomationTest.h
@@ -1817,6 +1817,7 @@ public:
 	CORE_API bool TestEqual(const TCHAR* What, const FLinearColor Actual, const FLinearColor Expected);
 	CORE_API bool TestEqual(const TCHAR* What, const TCHAR* Actual, const TCHAR* Expected);
 	CORE_API bool TestEqualInsensitive(const TCHAR* What, const TCHAR* Actual, const TCHAR* Expected);
+	CORE_API bool TestNotEqualInsensitive(const TCHAR* What, const TCHAR* Actual, const TCHAR* Expected);
 
 	bool TestEqual(const FString& What, const int32 Actual, const int32 Expected)
 	{
@@ -3910,6 +3911,18 @@ public: \
 		return false;\
 	}
 
+#define UTEST_NOT_EQUAL_INSENSITIVE(What, Actual, Expected)\
+	if (!TestNotEqualInsensitive(What, Actual, Expected))\
+	{\
+		return false;\
+	}
+
+#define UTEST_NOT_EQUAL_INSENSITIVE_EXPR(Actual, Expected)\
+	if (!TestNotEqualInsensitive(TEXT(#Actual), Actual, Expected))\
+	{\
+		return false;\
+	}
+
 #define UTEST_NOT_EQUAL(What, Actual, Expected)\
 	if (!TestNotEqual(What, Actual, Expected))\
 	{\
diff --git a/Engine/Source/Runtime/Engine/Private/Tests/EngineAutomationTests.cpp b/Engine/Source/Runtime/Engine/Private/Tests/EngineAutomationTests.cpp
--- a/Engine/Source/Runtime/Engine/Private/Tests/EngineAutomationTests.cpp
+++ b/Engine/Source/Runtime/Engine/Private/Tests/EngineAutomationTests.cpp
@@ -851,6 +851,7 @@ protected:
 	static const float ExpectedFloatValueOutOfTolerance;
 	static const FString ActualFStringValue;
 	static const FString ExpectedFStringValueLowerCase;
+	static const FString UnexpectedFStringValueLowerCase;
 };
 
 const float FAutomationUTestMacrosExpr::PositiveToleranceFloat(1.e-4f);
@@ -860,6 +861,7 @@ const float FAutomationUTestMacrosExpr::WrongFloatValue(ActualFloatValue + 1.f);
 const float FAutomationUTestMacrosExpr::ExpectedFloatValueOutOfTolerance(ActualFloatValue + PositiveToleranceFloat);
 const FString FAutomationUTestMacrosExpr::ActualFStringValue(TEXT("EQUALS"));
 const FString FAutomationUTestMacrosExpr::ExpectedFStringValueLowerCase(TEXT("equals"));
+const FString FAutomationUTestMacrosExpr::UnexpectedFStringValueLowerCase(TEXT("not-equals"));
 
 IMPLEMENT_CUSTOM_SIMPLE_AUTOMATION_TEST(FAutomationEqualEXPR, FAutomationUTestMacrosExpr, "TestFramework.Validation.UTestEqual", EAutomationTestFlags::ApplicationContextMask | EAutomationTestFlags::EngineFilter);
 bool FAutomationEqualEXPR::RunTest(const FString& Parameters)
@@ -870,6 +872,7 @@ bool FAutomationEqualEXPR::RunTest(const FString& Parameters)
 	UTEST_EQUAL_TOLERANCE_EXPR(ActualFloatValue, ExpectedFloatValueOutOfTolerance, PositiveToleranceFloat);
 	UTEST_NOT_EQUAL_EXPR(ActualFloatValue, WrongFloatValue);
 	UTEST_EQUAL_INSENSITIVE_EXPR(*ActualFStringValue, *ExpectedFStringValueLowerCase);
+	UTEST_NOT_EQUAL_INSENSITIVE_EXPR(*ActualFStringValue, *UnexpectedFStringValueLowerCase);
 
 	return true;
 }
