[NetPhysics] Don't use physics ReplicationCache for the default physics replication.

diff --git a/Engine/Source/Runtime/Engine/Private/ActorReplication.cpp b/Engine/Source/Runtime/Engine/Private/ActorReplication.cpp
--- a/Engine/Source/Runtime/Engine/Private/ActorReplication.cpp
+++ b/Engine/Source/Runtime/Engine/Private/ActorReplication.cpp
@@ -400,16 +400,20 @@ void AActor::GatherCurrentMovement()
 			const bool bPrevRepPhysics = ReplicatedMovement.bRepPhysics;
 #endif // UE_WITH_IRIS
 
+			const bool bShouldUsePhysicsReplicationCache = GetPhysicsReplicationMode() != EPhysicsReplicationMode::Default;
 			bool bFoundInCache = false;
 
 			UWorld* World = GetWorld();
 			int ServerFrame = 0; 
-			if (FPhysScene_Chaos* Scene = static_cast<FPhysScene_Chaos*>(World->GetPhysicsScene()))
+			if (bShouldUsePhysicsReplicationCache)
 			{
-				if (const FRigidBodyState* FoundState = Scene->GetStateFromReplicationCache(RootPrimComp, ServerFrame)) 
+				if (FPhysScene_Chaos* Scene = static_cast<FPhysScene_Chaos*>(World->GetPhysicsScene()))
 				{
-					ReplicatedMovement.FillFrom(*FoundState, this, Scene->ReplicationCache.ServerFrame);
-					bFoundInCache = true;
+					if (const FRigidBodyState* FoundState = Scene->GetStateFromReplicationCache(RootPrimComp, ServerFrame))
+					{
+						ReplicatedMovement.FillFrom(*FoundState, this, Scene->ReplicationCache.ServerFrame);
+						bFoundInCache = true;
+					}
 				}
 			}
 
