World Partition Runtime HashSet:
  - Added support for always loaded HLOD setups.
  - Added mechanism to register automatic converters for runtime hash classes.
  - Added converter implementation for SpatialHash to HashSet using LHGrid.

diff --git a/Engine/Source/Editor/WorldPartitionEditor/Private/WorldPartition/Customizations/WorldPartitionDetailsCustomization.cpp b/Engine/Source/Editor/WorldPartitionEditor/Private/WorldPartition/Customizations/WorldPartitionDetailsCustomization.cpp
--- a/Engine/Source/Editor/WorldPartitionEditor/Private/WorldPartition/Customizations/WorldPartitionDetailsCustomization.cpp
+++ b/Engine/Source/Editor/WorldPartitionEditor/Private/WorldPartition/Customizations/WorldPartitionDetailsCustomization.cpp
@@ -183,8 +183,7 @@ void FWorldPartitionDetails::HandleWorldPartitionRuntimeHashClassChanged(const U
 	if (WorldPartition.IsValid() && (!WorldPartition->RuntimeHash || WorldPartition->RuntimeHash->GetClass() != RuntimeHashClass))
 	{
 		WorldPartition->Modify();
-		WorldPartition->RuntimeHash = NewObject<UWorldPartitionRuntimeHash>(WorldPartition.Get(), RuntimeHashClass, NAME_None, RF_Transactional);
-		WorldPartition->RuntimeHash->SetDefaultValues();
+		WorldPartition->RuntimeHash = UWorldPartitionRuntimeHash::ConvertWorldPartitionHash(WorldPartition->RuntimeHash, RuntimeHashClass);
 
 		DetailBuilder->ForceRefreshDetails();
 	}
diff --git a/Engine/Source/Runtime/Engine/Private/WorldPartition/RuntimeHashSet/WorldPartitionRuntimeHashSet.cpp b/Engine/Source/Runtime/Engine/Private/WorldPartition/RuntimeHashSet/WorldPartitionRuntimeHashSet.cpp
--- a/Engine/Source/Runtime/Engine/Private/WorldPartition/RuntimeHashSet/WorldPartitionRuntimeHashSet.cpp
+++ b/Engine/Source/Runtime/Engine/Private/WorldPartition/RuntimeHashSet/WorldPartitionRuntimeHashSet.cpp
@@ -71,7 +71,20 @@ void URuntimeHashSetExternalStreamingObject::AddReferencedObjects(UObject* InThi
 
 UWorldPartitionRuntimeHashSet::UWorldPartitionRuntimeHashSet(const FObjectInitializer& ObjectInitializer)
 	: Super(ObjectInitializer)
-{}
+{
+#if WITH_EDITOR
+	if (HasAnyFlags(RF_ClassDefaultObject))
+	{
+		if (UClass* RuntimeSpatialHashClass = FindObject<UClass>(nullptr, TEXT("/Script/Engine.WorldPartitionRuntimeSpatialHash")))
+		{
+			RegisterWorldPartitionRuntimeHashConverter(RuntimeSpatialHashClass, GetClass(), [](const UWorldPartitionRuntimeHash* SrcHash) -> UWorldPartitionRuntimeHash*
+			{
+				return CreateFrom(SrcHash);
+			});
+		}
+	}
+#endif
+}
 
 void UWorldPartitionRuntimeHashSet::PostLoad()
 {
@@ -540,6 +553,7 @@ void UWorldPartitionRuntimeHashSet::PostEditChangeChainProperty(FPropertyChanged
 				FRuntimePartitionHLODSetup& RuntimePartitionHLODSetup = RuntimePartitionDesc.HLODSetups[HLODSetupsIndex];
 				URuntimePartition* ParentRuntimePartition = HLODSetupsIndex ? RuntimePartitionDesc.HLODSetups[HLODSetupsIndex - 1].PartitionLayer : RuntimePartitionDesc.MainLayer;
 				RuntimePartitionHLODSetup.Name = *FString::Printf(TEXT("HLOD_%d"), HLODSetupsIndex);
+				RuntimePartitionHLODSetup.bIsSpatiallyLoaded = true;
 				RuntimePartitionHLODSetup.PartitionLayer = ParentRuntimePartition->CreateHLODRuntimePartition(HLODSetupsIndex);
 				RuntimePartitionHLODSetup.PartitionLayer->Name = RuntimePartitionHLODSetup.Name;
 			}
@@ -580,6 +594,32 @@ void UWorldPartitionRuntimeHashSet::PostEditChangeChainProperty(FPropertyChanged
 			}
 		}
 	}
+	else if (PropertyName == GET_MEMBER_NAME_CHECKED(FRuntimePartitionHLODSetup, bIsSpatiallyLoaded))
+	{
+		int32 RuntimePartitionIndex = PropertyChangedEvent.GetArrayIndex(NAME_RuntimePartitions.ToString());
+		if (RuntimePartitions.IsValidIndex(RuntimePartitionIndex))
+		{
+			FRuntimePartitionDesc& RuntimePartitionDesc = RuntimePartitions[RuntimePartitionIndex];
+
+			int32 HLODSetupsIndex = PropertyChangedEvent.GetArrayIndex(NAME_HLODSetups.ToString());
+			if (RuntimePartitionDesc.HLODSetups.IsValidIndex(HLODSetupsIndex))
+			{
+				FRuntimePartitionHLODSetup& RuntimePartitionHLODSetup = RuntimePartitionDesc.HLODSetups[HLODSetupsIndex];
+
+				if (RuntimePartitionHLODSetup.bIsSpatiallyLoaded)
+				{
+					URuntimePartition* ParentRuntimePartition = HLODSetupsIndex ? RuntimePartitionDesc.HLODSetups[HLODSetupsIndex - 1].PartitionLayer : RuntimePartitionDesc.MainLayer;
+					RuntimePartitionHLODSetup.PartitionLayer = ParentRuntimePartition->CreateHLODRuntimePartition(HLODSetupsIndex);
+					RuntimePartitionHLODSetup.PartitionLayer->Name = RuntimePartitionHLODSetup.Name;
+				}
+				else
+				{
+					RuntimePartitionHLODSetup.PartitionLayer = NewObject<URuntimePartitionPersistent>(this, NAME_None);
+					RuntimePartitionHLODSetup.PartitionLayer->Name = RuntimePartitionHLODSetup.Name;
+				}
+			}
+		}
+	}
 }
 
 UWorldPartitionRuntimeHashSet::FCellUniqueId UWorldPartitionRuntimeHashSet::GetCellUniqueId(const URuntimePartition::FCellDescInstance& InCellDescInstance) const
diff --git a/Engine/Source/Runtime/Engine/Private/WorldPartition/RuntimeHashSet/WorldPartitionRuntimeHashSetConversions.cpp b/Engine/Source/Runtime/Engine/Private/WorldPartition/RuntimeHashSet/WorldPartitionRuntimeHashSetConversions.cpp
new file mode 100644
--- /dev/null
+++ b/Engine/Source/Runtime/Engine/Private/WorldPartition/RuntimeHashSet/WorldPartitionRuntimeHashSetConversions.cpp
@@ -0,0 +1,104 @@
+// Copyright Epic Games, Inc. All Rights Reserved.
+
+#include "WorldPartition/RuntimeHashSet/WorldPartitionRuntimeHashSet.h"
+#include "WorldPartition/RuntimeHashSet/RuntimePartitionLHGrid.h"
+#include "WorldPartition/RuntimeHashSet/RuntimePartitionPersistent.h"
+#include "WorldPartition/WorldPartitionRuntimeSpatialHash.h"
+#include "WorldPartition/HLOD/HLODLayer.h"
+
+#if WITH_EDITOR
+UWorldPartitionRuntimeHashSet* UWorldPartitionRuntimeHashSet::CreateFrom(const UWorldPartitionRuntimeHash* SrcHash)
+{
+	const UWorldPartitionRuntimeSpatialHash* SpatialHash = CastChecked<UWorldPartitionRuntimeSpatialHash>(SrcHash);
+	UWorldPartitionRuntimeHashSet* HashSet = NewObject<UWorldPartitionRuntimeHashSet>(SrcHash->GetOuter(), NAME_None, RF_Transactional);
+
+	UWorldPartition* WorldPartition = SrcHash->GetTypedOuter<UWorldPartition>();
+	check(WorldPartition);
+
+	UWorldPartition::FGenerateStreamingParams Params = UWorldPartition::FGenerateStreamingParams()
+		.SetActorDescContainer(WorldPartition->GetActorDescContainer());
+
+	UWorldPartition::FGenerateStreamingContext Context;
+	TUniquePtr<IStreamingGenerationContext> StreamingGenerationContext = WorldPartition->GenerateStreamingGenerationContext(Params, Context);
+	if (StreamingGenerationContext.IsValid())
+	{
+		// Gather all HLOD layers into their corresponding grids
+		TMap<FName, TMap<FName, TSet<const UHLODLayer*>>> GridHLODLayersMap;
+		const UHLODLayer* DefaultHLODLayer = WorldPartition->GetDefaultHLODLayer();
+
+		StreamingGenerationContext->ForEachActorSetInstance([&GridHLODLayersMap, SpatialHash, DefaultHLODLayer](const IStreamingGenerationContext::FActorSetInstance& ActorSetInstance)
+		{
+			ActorSetInstance.ForEachActor([&GridHLODLayersMap, SpatialHash, DefaultHLODLayer, &ActorSetInstance](const FGuid& ActorGuid)
+			{
+				const FWorldPartitionActorDescView& ActorDescView = ActorSetInstance.ContainerInstance->ActorDescViewMap->FindByGuidChecked(ActorGuid);
+
+				if (const UHLODLayer* HLODLayer = ActorDescView.GetHLODLayer().IsValid() ? Cast<UHLODLayer>(ActorDescView.GetHLODLayer().TryLoad()) : DefaultHLODLayer)
+				{
+					uint32 HLODLIndex = 0;
+					while (HLODLayer)
+					{
+						const FName GridName = ActorDescView.GetRuntimeGrid().IsNone() ? SpatialHash->Grids[0].GridName : ActorDescView.GetRuntimeGrid();
+						const FName HLODGridName = HLODLayer->GetRuntimeGrid(HLODLIndex++);
+						GridHLODLayersMap.FindOrAdd(GridName).FindOrAdd(HLODGridName).Add(HLODLayer);
+						HLODLayer = HLODLayer->GetParentLayer();
+					}
+				}
+			});
+		});
+
+		for (const FSpatialHashRuntimeGrid& Grid : SpatialHash->Grids)
+		{
+			if (!Grid.HLODLayer)
+			{
+				FRuntimePartitionDesc& RuntimePartitionDesc = HashSet->RuntimePartitions.AddDefaulted_GetRef();
+				RuntimePartitionDesc.Class = URuntimePartitionLHGrid::StaticClass();
+				RuntimePartitionDesc.Name = Grid.GridName;
+
+				URuntimePartitionLHGrid* LHGrid = NewObject<URuntimePartitionLHGrid>(HashSet, NAME_None);
+				LHGrid->Name = RuntimePartitionDesc.Name;
+				LHGrid->CellSize = Grid.CellSize;
+				LHGrid->bBlockOnSlowStreaming = Grid.bBlockOnSlowStreaming;
+				LHGrid->bClientOnlyVisible = Grid.bClientOnlyVisible;
+				LHGrid->Priority = Grid.Priority;
+				LHGrid->LoadingRange = Grid.LoadingRange;
+				LHGrid->HLODIndex = INDEX_NONE;
+
+				RuntimePartitionDesc.MainLayer = LHGrid;
+
+				for (const auto& [HLODGridName, HLODLayers] : GridHLODLayersMap.FindChecked(Grid.GridName))
+				{
+					const int32 HLODIndex = RuntimePartitionDesc.HLODSetups.Num();
+					FRuntimePartitionHLODSetup& HLODSetup = RuntimePartitionDesc.HLODSetups.AddDefaulted_GetRef();
+
+					HLODSetup.Name = HLODGridName.IsNone() ? NAME_PersistentLevel : HLODGridName;
+					HLODSetup.bIsSpatiallyLoaded = !HLODGridName.IsNone();
+					HLODSetup.HLODLayers = HLODLayers.Array();
+
+					if (HLODSetup.bIsSpatiallyLoaded)
+					{
+						URuntimePartitionLHGrid* HLODLHGrid = NewObject<URuntimePartitionLHGrid>(HashSet, NAME_None);
+						HLODLHGrid->CellSize = HLODSetup.HLODLayers[0]->GetCellSize();
+						HLODLHGrid->LoadingRange = HLODSetup.HLODLayers[0]->GetLoadingRange();
+						HLODSetup.PartitionLayer = HLODLHGrid;
+					}
+					else
+					{
+						HLODSetup.PartitionLayer = NewObject<URuntimePartitionPersistent>(HashSet, NAME_None);;
+					}
+
+					HLODSetup.PartitionLayer->Name = HLODSetup.Name;
+					HLODSetup.PartitionLayer->bBlockOnSlowStreaming = false;
+					HLODSetup.PartitionLayer->bClientOnlyVisible = true;
+					HLODSetup.PartitionLayer->Priority = 0;
+					HLODSetup.PartitionLayer->LoadingRange = 0;
+					HLODSetup.PartitionLayer->HLODIndex = HLODIndex;
+				}
+			}
+		}
+
+		WorldPartition->FlushStreaming();
+	}
+
+	return HashSet;
+}
+#endif
\ No newline at end of file
diff --git a/Engine/Source/Runtime/Engine/Private/WorldPartition/WorldPartitionRuntimeHash.cpp b/Engine/Source/Runtime/Engine/Private/WorldPartition/WorldPartitionRuntimeHash.cpp
--- a/Engine/Source/Runtime/Engine/Private/WorldPartition/WorldPartitionRuntimeHash.cpp
+++ b/Engine/Source/Runtime/Engine/Private/WorldPartition/WorldPartitionRuntimeHash.cpp
@@ -82,6 +82,8 @@ void URuntimeHashExternalStreamingObjectBase::PopulateGeneratorPackageForCook()
 		LevelStreamingDynamic->SetFlags(RF_Transient);
 	});
 }
+
+TMap<TPair<const UClass*, const UClass*>, UWorldPartitionRuntimeHash::FRuntimeHashConvertFunc> UWorldPartitionRuntimeHash::WorldPartitionRuntimeHashConverters;
 #endif
 
 UWorldPartitionRuntimeHash::UWorldPartitionRuntimeHash(const FObjectInitializer& ObjectInitializer)
@@ -496,6 +498,37 @@ bool UWorldPartitionRuntimeHash::RemoveExternalStreamingObject(URuntimeHashExter
 	return true;
 }
 
+#if WITH_EDITOR
+void UWorldPartitionRuntimeHash::RegisterWorldPartitionRuntimeHashConverter(const UClass* InSrcClass, const UClass* InDstClass, FRuntimeHashConvertFunc&& InConverter)
+{
+	check(InSrcClass->IsChildOf<UWorldPartitionRuntimeHash>());
+	check(InDstClass->IsChildOf<UWorldPartitionRuntimeHash>());
+	WorldPartitionRuntimeHashConverters.Add({ InSrcClass, InDstClass }, MoveTemp(InConverter));
+}
+
+UWorldPartitionRuntimeHash* UWorldPartitionRuntimeHash::ConvertWorldPartitionHash(const UWorldPartitionRuntimeHash* InSrcHash, const UClass* InDstClass)
+{
+	check(InSrcHash);
+	check(InDstClass);
+	check(InDstClass->IsChildOf<UWorldPartitionRuntimeHash>());
+	check(!InDstClass->HasAnyClassFlags(CLASS_Abstract));
+
+	const UClass* CurrentSrcClass = InSrcHash->GetClass();
+	while(CurrentSrcClass != UWorldPartitionRuntimeHash::StaticClass())
+	{
+		if (FRuntimeHashConvertFunc* Converter = WorldPartitionRuntimeHashConverters.Find({ CurrentSrcClass, InDstClass }))
+		{
+			return (*Converter)(InSrcHash);
+		}
+		CurrentSrcClass = CurrentSrcClass->GetSuperClass();
+	}
+
+	UWorldPartitionRuntimeHash* RuntimeHash = NewObject<UWorldPartitionRuntimeHash>(InSrcHash->GetOuter(), InDstClass, NAME_None, RF_Transactional);
+	RuntimeHash->SetDefaultValues();
+	return RuntimeHash;
+}
+#endif
+
 void UWorldPartitionRuntimeHash::FStreamingSourceCells::AddCell(const UWorldPartitionRuntimeCell* Cell, const FWorldPartitionStreamingSource& Source, const FSphericalSector& SourceShape)
 {
 	Cell->AppendStreamingSourceInfo(Source, SourceShape);
diff --git a/Engine/Source/Runtime/Engine/Private/WorldPartition/WorldPartitionStreamingGeneration.cpp b/Engine/Source/Runtime/Engine/Private/WorldPartition/WorldPartitionStreamingGeneration.cpp
--- a/Engine/Source/Runtime/Engine/Private/WorldPartition/WorldPartitionStreamingGeneration.cpp
+++ b/Engine/Source/Runtime/Engine/Private/WorldPartition/WorldPartitionStreamingGeneration.cpp
@@ -1451,6 +1451,47 @@ bool UWorldPartition::GenerateContainerStreaming(const FGenerateStreamingParams&
 	return false;
 }
 
+class FStreamingGenerationContextCopy : public FStreamingGenerationContextProxy
+{
+public:
+	FStreamingGenerationContextCopy(const FWorldPartitionStreamingGenerator::FWorldPartitionStreamingGeneratorParams& InParams)
+		: FStreamingGenerationContextProxy(nullptr)
+		, StreamingGenerator(InParams)
+	{}
+
+	void SetSourceContext(const IStreamingGenerationContext* InStreamingGenerationContext)
+	{
+		SourceContext = InStreamingGenerationContext;
+	}
+
+	FWorldPartitionStreamingGenerator StreamingGenerator;
+};
+
+TUniquePtr<IStreamingGenerationContext> UWorldPartition::GenerateStreamingGenerationContext(const FGenerateStreamingParams& InParams, FGenerateStreamingContext& InContext)
+{
+	TRACE_CPUPROFILER_EVENT_SCOPE(UWorldPartition::GenerateStreamingGenerationContext);
+
+	FStreamingGenerationLogErrorHandler LogErrorHandler;
+	FWorldPartitionStreamingGenerator::FWorldPartitionStreamingGeneratorParams StreamingGeneratorParams = FWorldPartitionStreamingGenerator::FWorldPartitionStreamingGeneratorParams()
+		.SetWorldPartitionContext(this)
+		.SetIsValidGrid([this](FName GridName) { return RuntimeHash->IsValidGrid(GridName); })
+		.SetIsValidHLODLayer([this](FName GridName, const FSoftObjectPath& HLODLayerPath) { return RuntimeHash->IsValidHLODLayer(GridName, HLODLayerPath); })
+		.SetErrorHandler(&LogErrorHandler)
+		.SetEnableStreaming(IsStreamingEnabled())
+		.SetCreateContainerResolver(FEditorPathHelper::IsEnabled());
+
+	TUniquePtr<FStreamingGenerationContextCopy> StreamingGenerationContextCopy = MakeUnique<FStreamingGenerationContextCopy>(StreamingGeneratorParams);
+	
+	StreamingGenerationContextCopy->StreamingGenerator.PreparationPhase(InParams.ActorDescCollection);
+	
+	const IStreamingGenerationContext* StreamingGenerationContext = StreamingGenerationContextCopy->StreamingGenerator.GetStreamingGenerationContext(InParams.ActorDescCollection);	
+	check(StreamingGenerationContext);
+
+	StreamingGenerationContextCopy->SetSourceContext(StreamingGenerationContext);
+	
+	return MoveTemp(StreamingGenerationContextCopy);
+}
+
 void UWorldPartition::FlushStreaming()
 {
 	RuntimeHash->FlushStreaming();
diff --git a/Engine/Source/Runtime/Engine/Public/WorldPartition/RuntimeHashSet/RuntimePartition.h b/Engine/Source/Runtime/Engine/Public/WorldPartition/RuntimeHashSet/RuntimePartition.h
--- a/Engine/Source/Runtime/Engine/Public/WorldPartition/RuntimeHashSet/RuntimePartition.h
+++ b/Engine/Source/Runtime/Engine/Public/WorldPartition/RuntimeHashSet/RuntimePartition.h
@@ -74,13 +74,13 @@ public:
 	UPROPERTY()
 	FName Name;
 
-	UPROPERTY(EditAnywhere, Category = RuntimeSettings, Meta = (EditCondition = "HLODSetup != INDEX_NONE", EditConditionHides, HideEditConditionToggle))
+	UPROPERTY(EditAnywhere, Category = RuntimeSettings, Meta = (EditCondition = "HLODIndex != INDEX_NONE", EditConditionHides, HideEditConditionToggle))
 	bool bBlockOnSlowStreaming;
 
-	UPROPERTY(EditAnywhere, Category = RuntimeSettings, Meta = (EditCondition = "HLODSetup != INDEX_NONE", EditConditionHides, HideEditConditionToggle))
+	UPROPERTY(EditAnywhere, Category = RuntimeSettings, Meta = (EditCondition = "HLODIndex != INDEX_NONE", EditConditionHides, HideEditConditionToggle))
 	bool bClientOnlyVisible;
 
-	UPROPERTY(EditAnywhere, Category = RuntimeSettings, Meta = (EditCondition = "HLODSetup != INDEX_NONE", EditConditionHides, HideEditConditionToggle))
+	UPROPERTY(EditAnywhere, Category = RuntimeSettings, Meta = (EditCondition = "HLODIndex != INDEX_NONE", EditConditionHides, HideEditConditionToggle))
 	int32 Priority;
 
 	UPROPERTY(EditAnywhere, Category = RuntimeSettings)
diff --git a/Engine/Source/Runtime/Engine/Public/WorldPartition/RuntimeHashSet/RuntimePartitionLHGrid.h b/Engine/Source/Runtime/Engine/Public/WorldPartition/RuntimeHashSet/RuntimePartitionLHGrid.h
--- a/Engine/Source/Runtime/Engine/Public/WorldPartition/RuntimeHashSet/RuntimePartitionLHGrid.h
+++ b/Engine/Source/Runtime/Engine/Public/WorldPartition/RuntimeHashSet/RuntimePartitionLHGrid.h
@@ -10,6 +10,8 @@ class URuntimePartitionLHGrid : public URuntimePartition
 {
 	GENERATED_BODY()
 
+	friend class UWorldPartitionRuntimeHashSet;
+
 public:
 #if WITH_EDITOR
 	//~ Begin UObject Interface.
diff --git a/Engine/Source/Runtime/Engine/Public/WorldPartition/RuntimeHashSet/WorldPartitionRuntimeHashSet.h b/Engine/Source/Runtime/Engine/Public/WorldPartition/RuntimeHashSet/WorldPartitionRuntimeHashSet.h
--- a/Engine/Source/Runtime/Engine/Public/WorldPartition/RuntimeHashSet/WorldPartitionRuntimeHashSet.h
+++ b/Engine/Source/Runtime/Engine/Public/WorldPartition/RuntimeHashSet/WorldPartitionRuntimeHashSet.h
@@ -27,7 +27,11 @@ struct FRuntimePartitionHLODSetup
 	UPROPERTY(EditAnywhere, Category = RuntimeSettings)
 	TArray<TObjectPtr<const UHLODLayer>> HLODLayers;
 
-	UPROPERTY(VisibleAnywhere, Category = RuntimeSettings, Instanced, Meta = (NoResetToDefault, TitleProperty = "Name"))
+	/** whether this HLOD setup is spatially loaded or not */
+	UPROPERTY(EditAnywhere, Category = RuntimeSettings)
+	bool bIsSpatiallyLoaded;
+
+	UPROPERTY(VisibleAnywhere, Category = RuntimeSettings, Instanced, Meta = (EditCondition = "bIsSpatiallyLoaded", HideEditConditionToggle, NoResetToDefault, TitleProperty = "Name"))
 	TObjectPtr<URuntimePartition> PartitionLayer;
 };
 
@@ -141,6 +145,9 @@ public:
 
 	// Helpers
 	static ENGINE_API bool ParseGridName(FName GridName, TArray<FName>& MainPartitionTokens, TArray<FName>& HLODPartitionTokens);
+
+	// Conversions
+	static ENGINE_API UWorldPartitionRuntimeHashSet* CreateFrom(const UWorldPartitionRuntimeHash* SrcHash);
 #endif
 
 	// External streaming object interface
@@ -184,4 +191,4 @@ public:
 
 	UPROPERTY()
 	TArray<FRuntimePartitionStreamingData> RuntimeStreamingData;
-};
\ No newline at end of file
+};
diff --git a/Engine/Source/Runtime/Engine/Public/WorldPartition/WorldPartition.h b/Engine/Source/Runtime/Engine/Public/WorldPartition/WorldPartition.h
--- a/Engine/Source/Runtime/Engine/Public/WorldPartition/WorldPartition.h
+++ b/Engine/Source/Runtime/Engine/Public/WorldPartition/WorldPartition.h
@@ -45,7 +45,7 @@ class ULevel;
 class FAutoConsoleVariableRef;
 class FWorldPartitionDraw2DContext;
 class FContentBundleEditor;
-
+class IStreamingGenerationContext;
 struct IWorldPartitionStreamingSourceProvider;
 
 enum class EWorldPartitionRuntimeCellState : uint8;
@@ -264,6 +264,7 @@ public:
 
 	ENGINE_API bool GenerateStreaming(const FGenerateStreamingParams& InParams, FGenerateStreamingContext& InContext);
 	ENGINE_API bool GenerateContainerStreaming(const FGenerateStreamingParams& InParams, FGenerateStreamingContext& InContext);
+	ENGINE_API TUniquePtr<IStreamingGenerationContext> GenerateStreamingGenerationContext(const FGenerateStreamingParams& InParams, FGenerateStreamingContext& InContext);
 
 	ENGINE_API void FlushStreaming();
 	ENGINE_API URuntimeHashExternalStreamingObjectBase* FlushStreamingToExternalStreamingObject(const FString& ExternalStreamingObjectName);
diff --git a/Engine/Source/Runtime/Engine/Public/WorldPartition/WorldPartitionRuntimeHash.h b/Engine/Source/Runtime/Engine/Public/WorldPartition/WorldPartitionRuntimeHash.h
--- a/Engine/Source/Runtime/Engine/Public/WorldPartition/WorldPartitionRuntimeHash.h
+++ b/Engine/Source/Runtime/Engine/Public/WorldPartition/WorldPartitionRuntimeHash.h
@@ -199,4 +199,14 @@ protected:
 #endif
 
 	TSet<TWeakObjectPtr<URuntimeHashExternalStreamingObjectBase>> InjectedExternalStreamingObjects;
+
+#if WITH_EDITOR
+private:
+	using FRuntimeHashConvertFunc = TFunction<UWorldPartitionRuntimeHash*(const UWorldPartitionRuntimeHash*)>;
+	static TMap<TPair<const UClass*, const UClass*>, FRuntimeHashConvertFunc> WorldPartitionRuntimeHashConverters;
+
+public:
+	static ENGINE_API void RegisterWorldPartitionRuntimeHashConverter(const UClass* InSrcClass, const UClass* InDstClass, FRuntimeHashConvertFunc&& InConverter);
+	static ENGINE_API UWorldPartitionRuntimeHash* ConvertWorldPartitionHash(const UWorldPartitionRuntimeHash* InSrcHash, const UClass* InDstClass);
+#endif
 };
diff --git a/Engine/Source/Runtime/Engine/Public/WorldPartition/WorldPartitionRuntimeSpatialHash.h b/Engine/Source/Runtime/Engine/Public/WorldPartition/WorldPartitionRuntimeSpatialHash.h
--- a/Engine/Source/Runtime/Engine/Public/WorldPartition/WorldPartitionRuntimeSpatialHash.h
+++ b/Engine/Source/Runtime/Engine/Public/WorldPartition/WorldPartitionRuntimeSpatialHash.h
@@ -379,4 +379,5 @@ private:
 	ENGINE_API TArray<const FSpatialHashStreamingGrid*> GetFilteredStreamingGrids() const;
 
 	friend class UWorldPartitionSubsystem;
+	friend class UWorldPartitionRuntimeHashSet;
 };
diff --git a/Engine/Source/Runtime/Engine/Public/WorldPartition/WorldPartitionStreamingGenerationContext.h b/Engine/Source/Runtime/Engine/Public/WorldPartition/WorldPartitionStreamingGenerationContext.h
--- a/Engine/Source/Runtime/Engine/Public/WorldPartition/WorldPartitionStreamingGenerationContext.h
+++ b/Engine/Source/Runtime/Engine/Public/WorldPartition/WorldPartitionStreamingGenerationContext.h
@@ -133,7 +133,7 @@ public:
 		SourceContext->ForEachActorSetContainer(Func);
 	}
 
-private:
+protected:
 	const IStreamingGenerationContext* SourceContext;
 };
 #endif
\ No newline at end of file
