fix Texture Editor viewer not gamma correcting correctly
stop setting bIgnoreGammaConversions ; it is deprecated and only causes problems
bIgnoreGammaConversions is only used by the Canvas/Slate/SimpleElementShader path, not by Materials
bIgnoreGammaConversions was being turned on for non-SRGB non-HDR textures
this caused them to be gamma corrected wrong
this showed up in the Texture Editor viewer and thumbnails
just removing this variable fixes it

#rnx

diff --git a/Engine/Plugins/Runtime/nDisplay/Source/DisplayCluster/Private/Render/Viewport/RenderTarget/DisplayClusterRenderTargetResource.cpp b/Engine/Plugins/Runtime/nDisplay/Source/DisplayCluster/Private/Render/Viewport/RenderTarget/DisplayClusterRenderTargetResource.cpp
--- a/Engine/Plugins/Runtime/nDisplay/Source/DisplayCluster/Private/Render/Viewport/RenderTarget/DisplayClusterRenderTargetResource.cpp
+++ b/Engine/Plugins/Runtime/nDisplay/Source/DisplayCluster/Private/Render/Viewport/RenderTarget/DisplayClusterRenderTargetResource.cpp
@@ -14,6 +14,9 @@ FDisplayClusterViewportTextureResource::FDisplayClusterViewportTextureResource(c
 {
 	bSRGB = EnumHasAnyFlags(GetResourceSettings().GetResourceFlags(), EDisplayClusterViewportResourceSettingsFlags::ShouldUseSRGB);
 	bGreyScaleFormat = false;
+	
+	// I believe this setting of bIgnoreGammaConversions was probably a bad copy-paste and is probably not intentional
+	//	bIgnoreGammaConversions is deprecated.
 	bIgnoreGammaConversions = true;
 }
 
diff --git a/Engine/Plugins/VirtualProduction/TextureShare/Source/TextureShare/Private/Resources/TextureShareResource.cpp b/Engine/Plugins/VirtualProduction/TextureShare/Source/TextureShare/Private/Resources/TextureShareResource.cpp
--- a/Engine/Plugins/VirtualProduction/TextureShare/Source/TextureShare/Private/Resources/TextureShareResource.cpp
+++ b/Engine/Plugins/VirtualProduction/TextureShare/Source/TextureShare/Private/Resources/TextureShareResource.cpp
@@ -25,6 +25,9 @@ FTextureShareResource::FTextureShareResource(const TSharedRef<ITextureShareCoreO
 	bSRGB = InResourceSettings.bShouldUseSRGB;
 
 	bGreyScaleFormat = false;
+	
+	// I believe this setting of bIgnoreGammaConversions was probably a bad copy-paste and is probably not intentional
+	//	bIgnoreGammaConversions is deprecated.
 	bIgnoreGammaConversions = true;
 }
 
diff --git a/Engine/Source/Runtime/Engine/Private/Rendering/StreamableTextureResource.cpp b/Engine/Source/Runtime/Engine/Private/Rendering/StreamableTextureResource.cpp
--- a/Engine/Source/Runtime/Engine/Private/Rendering/StreamableTextureResource.cpp
+++ b/Engine/Source/Runtime/Engine/Private/Rendering/StreamableTextureResource.cpp
@@ -93,8 +93,6 @@ FStreamableTextureResource::FStreamableTextureResource(UTexture* InOwner, const
 	, LODGroup(InOwner->LODGroup)
 	, PixelFormat(InPlatformData->PixelFormat)
 {
-	// HDR images are stored in linear but still require gamma correction to display correctly.
-	bIgnoreGammaConversions = !InOwner->SRGB && !IsHDR(PixelFormat);
 	bSRGB = InOwner->SRGB;
 	bGreyScaleFormat = UE::TextureDefines::ShouldUseGreyScaleEditorVisualization( InOwner->CompressionSettings );
 
diff --git a/Engine/Source/Runtime/Engine/Private/Texture2D.cpp b/Engine/Source/Runtime/Engine/Private/Texture2D.cpp
--- a/Engine/Source/Runtime/Engine/Private/Texture2D.cpp
+++ b/Engine/Source/Runtime/Engine/Private/Texture2D.cpp
@@ -1696,8 +1696,6 @@ void FVirtualTexture2DResource::InitializeEditorResources(IVirtualTexture* InVir
 		TextureRHI->SetName(TextureName);
 		RHIBindDebugLabelName(TextureRHI, *Name);
 		RHIUpdateTextureReference(TextureReferenceRHI, TextureRHI);
-
-		bIgnoreGammaConversions = !bSRGB && !IsHDR(PixelFormat);
 	}
 }
 #endif // WITH_EDITOR
diff --git a/Engine/Source/Runtime/Engine/Private/UserInterface/Canvas.cpp b/Engine/Source/Runtime/Engine/Private/UserInterface/Canvas.cpp
--- a/Engine/Source/Runtime/Engine/Private/UserInterface/Canvas.cpp
+++ b/Engine/Source/Runtime/Engine/Private/UserInterface/Canvas.cpp
@@ -436,6 +436,8 @@ bool FCanvasBatchedElementRenderItem::Render_RenderThread(FCanvasRenderContext&
 			// current render target set for the canvas
 			const FRenderTarget* CanvasRenderTarget = Canvas->GetRenderTarget();
 			float Gamma = 1.0f / CanvasRenderTarget->GetDisplayGamma();
+
+			// bIgnoreGammaConversions is deprecated; remove this
 			if (LocalData->Texture && LocalData->Texture->bIgnoreGammaConversions)
 			{
 				Gamma = 1.0f;
@@ -474,6 +476,8 @@ bool FCanvasBatchedElementRenderItem::Render_GameThread(const FCanvas* Canvas, F
 		// current render target set for the canvas
 		const FRenderTarget* CanvasRenderTarget = Canvas->GetRenderTarget();
 		float Gamma = 1.0f / CanvasRenderTarget->GetDisplayGamma();
+		
+		// bIgnoreGammaConversions is deprecated; remove this
 		if ( Data->Texture && Data->Texture->bIgnoreGammaConversions )
 		{
 			Gamma = 1.0f;
diff --git a/Engine/Source/Runtime/RenderCore/Public/RenderResource.h b/Engine/Source/Runtime/RenderCore/Public/RenderResource.h
--- a/Engine/Source/Runtime/RenderCore/Public/RenderResource.h
+++ b/Engine/Source/Runtime/RenderCore/Public/RenderResource.h
@@ -381,6 +381,8 @@ public:
 	 * true if the texture is in the same gamma space as the intended rendertarget (e.g. screenshots).
 	 * The texture will have sRGB==false and bIgnoreGammaConversions==true, causing a non-sRGB texture lookup
 	 * and no gamma-correction in the shader.
+	 *
+	 * I believe this comment is no longer correct and bIgnoreGammaConversions is now deprecated and should not be used.
 	 */
 	bool				bIgnoreGammaConversions = false;
 
