[shaders] bump DDC request priority to "highest" in cooker so that DDC jobs can interleave with preprocessing work, allowing shader compilation to kick off earlier

diff --git a/Engine/Source/Runtime/Engine/Private/ShaderCompiler/ShaderCompiler.cpp b/Engine/Source/Runtime/Engine/Private/ShaderCompiler/ShaderCompiler.cpp
--- a/Engine/Source/Runtime/Engine/Private/ShaderCompiler/ShaderCompiler.cpp
+++ b/Engine/Source/Runtime/Engine/Private/ShaderCompiler/ShaderCompiler.cpp
@@ -10361,11 +10361,18 @@ FShaderJobCacheRef FShaderJobCache::FindOrAdd(const FJobInputHash& Hash, EShader
 			static const bool PerShaderDDCAsync = CVarShaderCompilerPerShaderDDCAsync.GetValueOnAnyThread();
 			if (PerShaderDDCAsync && FGenericPlatformProcess::SupportsMultithreading())
 			{
-				switch (JobPriority)
+				if (IsRunningCookCommandlet())
 				{
-				case EShaderCompileJobPriority::Low:		DerivedDataPriority = UE::DerivedData::EPriority::Low;		break;
-				case EShaderCompileJobPriority::Normal:		DerivedDataPriority = UE::DerivedData::EPriority::Normal;	break;
-				default:									DerivedDataPriority = UE::DerivedData::EPriority::Highest;	break;
+					DerivedDataPriority = UE::DerivedData::EPriority::Highest;
+				}
+				else
+				{
+					switch (JobPriority)
+					{
+					case EShaderCompileJobPriority::Low:		DerivedDataPriority = UE::DerivedData::EPriority::Low;		break;
+					case EShaderCompileJobPriority::Normal:		DerivedDataPriority = UE::DerivedData::EPriority::Normal;	break;
+					default:									DerivedDataPriority = UE::DerivedData::EPriority::Highest;	break;
+					}
 				}
 				InoutRequestOwner = MakePimpl<UE::DerivedData::FRequestOwner>(DerivedDataPriority);
 				RequestOwner = InoutRequestOwner.Get();
