Strip the PrimaryAssetId out of _Generated_ World assets that represent runtime world partition cells, this stops them from ending up in the cooked asset registry and wasting memory and query time

diff --git a/Engine/Source/Runtime/Engine/Private/World.cpp b/Engine/Source/Runtime/Engine/Private/World.cpp
--- a/Engine/Source/Runtime/Engine/Private/World.cpp
+++ b/Engine/Source/Runtime/Engine/Private/World.cpp
@@ -9194,8 +9194,10 @@ bool UWorld::ResolveSubobject(const TCHAR* SubObjectPath, UObject*& OutObject, b
 FPrimaryAssetId UWorld::GetPrimaryAssetId() const
 {
 	UPackage* Package = GetOutermost();
+	const IWorldPartitionCell* WorldPartitionCell = PersistentLevel ? PersistentLevel->GetWorldPartitionRuntimeCell() : nullptr;
 
-	if (!Package->HasAnyPackageFlags(PKG_PlayInEditor))
+	// PIE and world partition runtime levels are temporary and do not represent a primary asset
+	if (!Package->HasAnyPackageFlags(PKG_PlayInEditor) && !WorldPartitionCell)
 	{
 		// Return Map:/path/to/map
 		return FPrimaryAssetId(UAssetManager::MapType, Package->GetFName());
