Optimize net.RemapStableSubobjects and add fast arrays to the test case.
Store unmapped NetGUIDs of stably-named subobjects persistently instead of iterating through object replicators to find them every time.

diff --git a/Engine/Source/Runtime/CoreUObject/Public/UObject/CoreNet.h b/Engine/Source/Runtime/CoreUObject/Public/UObject/CoreNet.h
--- a/Engine/Source/Runtime/CoreUObject/Public/UObject/CoreNet.h
+++ b/Engine/Source/Runtime/CoreUObject/Public/UObject/CoreNet.h
@@ -218,6 +218,9 @@ class UPackageMap : public UObject
 	const TSet< FNetworkGUID > &	GetTrackedUnmappedGuids() const { return TrackedUnmappedNetGuids; }
 	const TSet< FNetworkGUID > &	GetTrackedDynamicMappedGuids() const { return TrackedMappedDynamicNetGuids; }
 
+	virtual void AddUnmappedNetGUIDReference(FNetworkGUID UnmappedGUID) {}
+	virtual void RemoveUnmappedNetGUIDReference(FNetworkGUID UnmappedGUID) {}
+
 	// For sync load debugging with LogNetSyncLoads
 	virtual void			ResetTrackedSyncLoadedGuids() {}
 	virtual void			ReportSyncLoadsForProperty(const FProperty* Property, const UObject* Object) {}
diff --git a/Engine/Source/Runtime/Engine/Classes/Engine/PackageMapClient.h b/Engine/Source/Runtime/Engine/Classes/Engine/PackageMapClient.h
--- a/Engine/Source/Runtime/Engine/Classes/Engine/PackageMapClient.h
+++ b/Engine/Source/Runtime/Engine/Classes/Engine/PackageMapClient.h
@@ -146,6 +146,31 @@ private:
 	TObjectPtr<UObject> Object;
 };
 
+namespace UE::Net::Private
+{
+	/** Simple class to manage a reference-counted array of FNetworkGUIDs */
+	class FRefCountedNetGUIDArray
+	{
+	public:
+		FRefCountedNetGUIDArray() {}
+		FRefCountedNetGUIDArray(const FRefCountedNetGUIDArray&) = delete;
+		FRefCountedNetGUIDArray& operator=(const FRefCountedNetGUIDArray&) = delete;
+
+		FRefCountedNetGUIDArray(FRefCountedNetGUIDArray&& Other) = default;
+		FRefCountedNetGUIDArray& operator=(FRefCountedNetGUIDArray&& Other) = default;
+		~FRefCountedNetGUIDArray() = default;
+
+		void Add(FNetworkGUID NetGUID);
+		void RemoveSwap(FNetworkGUID NetGUID);
+
+		const TArray<FNetworkGUID>& GetNetGUIDs() const { return NetGUIDs; }
+
+	private:
+		TArray<FNetworkGUID> NetGUIDs;
+		TArray<int32> RefCounts;
+	};
+}
+
 class FNetGUIDCache
 {
 public:
@@ -228,6 +253,9 @@ public:
 
 	ENGINE_API void ResetReplayDirtyTracking();
 
+	ENGINE_API const TArray<FNetworkGUID>* FindUnmappedStablyNamedGuidsWithOuter(FNetworkGUID OuterGUID) const;
+	ENGINE_API const void RemoveUnmappedStablyNamedGuidsWithOuter(FNetworkGUID OuterGUID) { UnmappedStablyNamedGuids_OuterToInner.Remove(OuterGUID); }
+
 	TMap< FNetworkGUID, FNetGuidCacheObject >		ObjectLookup;
 	TMap< TWeakObjectPtr< UObject >, FNetworkGUID >	NetGUIDLookup;
 
@@ -330,6 +358,12 @@ private:
 	 */
 	TMap<FNetworkGUID, TWeakPtr<FQueuedBunchObjectReference>> QueuedBunchObjectReferences;
 
+	/**
+	 * Map of outer GUIDs to their unmapped, stably named subobjects.
+	 * So that the inner objects can be remapped when their outer GUID is imported again.
+	 */
+	TMap<FNetworkGUID, UE::Net::Private::FRefCountedNetGUIDArray> UnmappedStablyNamedGuids_OuterToInner;
+
 #if CSV_PROFILER
 public:
 
@@ -507,6 +541,9 @@ public:
 		return FullGuidCachePath;
 	}
 
+	virtual void AddUnmappedNetGUIDReference(FNetworkGUID UnmappedGUID) override;
+	virtual void RemoveUnmappedNetGUIDReference(FNetworkGUID MappedGUID) override;
+
 protected:
 
 	/** Functions to help with exporting/importing net field export info */
diff --git a/Engine/Source/Runtime/Engine/Private/DataChannel.cpp b/Engine/Source/Runtime/Engine/Private/DataChannel.cpp
--- a/Engine/Source/Runtime/Engine/Private/DataChannel.cpp
+++ b/Engine/Source/Runtime/Engine/Private/DataChannel.cpp
@@ -15,6 +15,7 @@
 #include "Misc/MemStack.h"
 #include "Misc/ScopeExit.h"
 #include "Net/Core/Trace/Private/NetTraceInternal.h"
+#include "Net/Core/Misc/GuidReferences.h"
 #include "Net/Core/NetCoreModule.h"
 #include "UObject/UObjectIterator.h"
 #include "EngineStats.h"
@@ -112,7 +113,6 @@ namespace UE::Net
 {
 	extern int32 FilterGuidRemapping;
 	extern bool bDiscardTornOffActorRPCs;
-	extern bool bRemapStableSubobjects;
 
 	static float QueuedBunchTimeoutSeconds = 30.0f;
 	static FAutoConsoleVariableRef CVarQueuedBunchTimeoutSeconds(
@@ -2408,11 +2408,11 @@ void UActorChannel::DestroyActorAndComponents()
 		}
 
 		// Also unmap any stably-named subobjects we didn't create
-		if (UE::Net::bRemapStableSubobjects)
+		if (UE::Net::Private::bRemapStableSubobjects)
 		{
 			TArray<UObject*> Inners;
 			GetObjectsWithOuter(Actor, Inners);
-			for (UObject* Inner : Inners)
+			for (const UObject* Inner : Inners)
 			{
 				if (Inner->IsNameStableForNetworking())
 				{
diff --git a/Engine/Source/Runtime/Engine/Private/NetDriver.cpp b/Engine/Source/Runtime/Engine/Private/NetDriver.cpp
--- a/Engine/Source/Runtime/Engine/Private/NetDriver.cpp
+++ b/Engine/Source/Runtime/Engine/Private/NetDriver.cpp
@@ -33,6 +33,7 @@
 #include "NetworkingDistanceConstants.h"
 #include "Engine/ChildConnection.h"
 #include "Net/Core/Trace/NetTrace.h"
+#include "Net/Core/Misc/GuidReferences.h"
 #include "Net/Core/PropertyConditions/PropertyConditions.h"
 #include "Net/DataChannel.h"
 #include "GameFramework/PlayerState.h"
@@ -929,9 +930,6 @@ namespace UE
 	{
 		int32 FilterGuidRemapping = 1;
 		static FAutoConsoleVariableRef CVarFilterGuidRemapping(TEXT("net.FilterGuidRemapping"), FilterGuidRemapping, TEXT("Remove destroyed and parent guids from unmapped list"));
-
-		bool bRemapStableSubobjects = true;
-		static FAutoConsoleVariableRef CVarNetRemapStableSubobjects(TEXT("net.RemapStableSubobjects"), bRemapStableSubobjects, TEXT("If enabled, attempts to remap stable subobjects when net.OptimizedRemapping is also enabled."));
 	};
 };
 
@@ -1082,33 +1080,22 @@ void UNetDriver::TickFlush(float DeltaSeconds)
 
 					if (GuidCache->GetObjectFromNetGUID(NetworkGuid, false) != nullptr)
 					{
-						if (UE::Net::bRemapStableSubobjects)
+						if (UE::Net::Private::bRemapStableSubobjects)
 						{
 							QUICK_SCOPE_CYCLE_COUNTER(STAT_NetRemapStableSubobjects);
 
-							// Stably-named/net addressable subobjects are created by user code on clients and not by replication.
-							// Normally this optimized remapping path relies on adding NetGuids of spawned actors/objects to the
-							// GuidCache->ImportedNetGuids list, but for net addressable subobjects there's no hook to do this.
-							// In order to actually remap them so references will be valid, this code will find any potentially
-							// unmapped addressable subobjects of an owning object that is in the import list (like a replicated actor).
-							// It looks through the unmapped replicators' reference list for any guids that have the imported
-							// guid in their outer chain.
-							TSet<FNetworkGUID>& PendingGuidsRef = PendingOuterNetGuidsRef.FindOrAdd(NetworkGuid);
-							
-							for (FObjectReplicator* Replicator : UnmappedReplicators)
+							// Import any unmapped, stably-named guids that are inners of the GUID that just mapped.
+							// These are tracked separately from the normal ImportedNetGuids since they are often
+							// default subobjects created in constructors and there's no other hook to import them.
+							const TArray<FNetworkGUID>* Inners = GuidCache->FindUnmappedStablyNamedGuidsWithOuter(NetworkGuid);
+
+							if (Inners)
 							{
-								for (const FNetworkGUID ReferencedGuid : Replicator->ReferencedGuids)
-								{
-									const FNetGuidCacheObject* const CacheObjectPtr = GuidCache->GetCacheObject(ReferencedGuid);
-
-									if (CacheObjectPtr && !CacheObjectPtr->PathName.IsNone())
-									{
-										if (UE::Net::Private::IsGuidInOuterChain(*GuidCache, CacheObjectPtr, NetworkGuid))
-										{
-											PendingGuidsRef.Add(ReferencedGuid);
-										}
-									}
-								}
+								TSet<FNetworkGUID>& PendingGuidsRef = PendingOuterNetGuidsRef.FindOrAdd(NetworkGuid);
+								PendingGuidsRef.Append(*Inners);
+									
+								// Now that they're on the pending import list, they will stay there until mapped. Can remove from the outer-to-inner map.
+								GuidCache->RemoveUnmappedStablyNamedGuidsWithOuter(NetworkGuid);
 							}
 						}
 
diff --git a/Engine/Source/Runtime/Engine/Private/PackageMapClient.cpp b/Engine/Source/Runtime/Engine/Private/PackageMapClient.cpp
--- a/Engine/Source/Runtime/Engine/Private/PackageMapClient.cpp
+++ b/Engine/Source/Runtime/Engine/Private/PackageMapClient.cpp
@@ -15,6 +15,7 @@
 #include "ProfilingDebugging/ScopedTimers.h"
 #include "GameFramework/GameStateBase.h"
 #include "HAL/LowLevelMemStats.h"
+#include "Net/Core/Misc/GuidReferences.h"
 #include "Net/Core/Trace/NetTrace.h"
 #include "Serialization/MemoryReader.h"
 #include "Net/NetworkGranularMemoryLogging.h"
@@ -83,6 +84,44 @@ namespace UE
 	};
 };
 
+namespace UE::Net::Private
+{
+	void FRefCountedNetGUIDArray::Add(FNetworkGUID NetGUID)
+	{
+		const int32 FoundIndex = NetGUIDs.IndexOfByKey(NetGUID);
+
+		if (RefCounts.IsValidIndex(FoundIndex))
+		{
+			++RefCounts[FoundIndex];
+		}
+		else
+		{
+			NetGUIDs.Add(NetGUID);
+			RefCounts.Add(1);
+
+			ensureMsgf(NetGUIDs.Num() == RefCounts.Num(), TEXT("FRefCountedNetGUIDArray::Add: arrays out of sync"));
+		}
+	}
+
+	void FRefCountedNetGUIDArray::RemoveSwap(FNetworkGUID NetGUID)
+	{
+		const int32 FoundIndex = NetGUIDs.IndexOfByKey(NetGUID);
+
+		if (RefCounts.IsValidIndex(FoundIndex))
+		{
+			--RefCounts[FoundIndex];
+
+			ensureMsgf(RefCounts[FoundIndex] >= 0, TEXT("FRefCountedNetGUIDArray::RemoveSwap: invalid RefCount %d at index %d"), RefCounts[FoundIndex], FoundIndex);
+
+			if (RefCounts[FoundIndex] == 0)
+			{
+				NetGUIDs.RemoveAtSwap(FoundIndex);
+				RefCounts.RemoveAtSwap(FoundIndex);
+			}
+		}
+	}
+}
+
 static TAutoConsoleVariable<int32> CVarAllowAsyncLoading(
 	TEXT("net.AllowAsyncLoading"),
 	0,
@@ -2569,6 +2608,66 @@ void UPackageMapClient::Serialize(FArchive& Ar)
 	}
 }
 
+const TArray<FNetworkGUID>* FNetGUIDCache::FindUnmappedStablyNamedGuidsWithOuter(FNetworkGUID OuterGUID) const
+{
+	using namespace UE::Net::Private;
+
+	const FRefCountedNetGUIDArray* Found = UnmappedStablyNamedGuids_OuterToInner.Find(OuterGUID);
+	if (Found)
+	{
+		return &Found->GetNetGUIDs();
+	}
+
+	return nullptr;
+}
+
+void UPackageMapClient::AddUnmappedNetGUIDReference(FNetworkGUID UnmappedGUID)
+{
+	using namespace UE::Net::Private;
+
+	if (bRemapStableSubobjects && GuidCache)
+	{
+		// For any new unmapped guids that represent stably-named inner objects, keep track of them
+		// so that when the NetDriver updates unmapped objects, if an outer GUID is imported, we can also import its
+		// stably-named inners. These are usually subobjects created in the constructor and don't get imported via any other path.
+		const FNetGuidCacheObject* CacheObject = GuidCache->GetCacheObject(UnmappedGUID);
+		if (CacheObject && CacheObject->OuterGUID.IsValid() && !CacheObject->PathName.IsNone())
+		{
+			FRefCountedNetGUIDArray& Inners = GuidCache->UnmappedStablyNamedGuids_OuterToInner.FindOrAdd(CacheObject->OuterGUID);
+			Inners.Add(UnmappedGUID);
+
+			UE_LOG(LogNetPackageMap, VeryVerbose, TEXT("Adding unmapped stably-named inner object NetGUID to tracking map: %s. With outer: %s"), *GuidCache->Describe(UnmappedGUID), *GuidCache->Describe(CacheObject->OuterGUID));
+		}
+	}
+}
+
+void UPackageMapClient::RemoveUnmappedNetGUIDReference(FNetworkGUID NetGUID)
+{
+	using namespace UE::Net::Private;
+
+	if (bRemapStableSubobjects && GuidCache)
+	{
+		// When a GUID reference is no longer tracked, if we were tracking it as a stably-named inner object,
+		// do the bookkeeping here. Decrement the refcount and remove it when there are no more references.
+		const FNetGuidCacheObject* CacheObject = GuidCache->GetCacheObject(NetGUID);
+		if (CacheObject)
+		{
+			FRefCountedNetGUIDArray* FoundInners = GuidCache->UnmappedStablyNamedGuids_OuterToInner.Find(CacheObject->OuterGUID);
+			if (FoundInners)
+			{
+				UE_LOG(LogNetPackageMap, VeryVerbose, TEXT("Removing stably-named inner GUID from tracking map: %s. With outer: %s"), *GuidCache->Describe(NetGUID), *GuidCache->Describe(CacheObject->OuterGUID));
+				
+				FoundInners->RemoveSwap(NetGUID);
+
+				if (FoundInners->GetNetGUIDs().Num() == 0)
+				{
+					GuidCache->RemoveUnmappedStablyNamedGuidsWithOuter(CacheObject->OuterGUID);
+				}
+			}
+		}
+	}
+}
+
 //----------------------------------------------------------------------------------------
 //	FNetGUIDCache
 //----------------------------------------------------------------------------------------
diff --git a/Engine/Source/Runtime/Engine/Private/RepLayout.cpp b/Engine/Source/Runtime/Engine/Private/RepLayout.cpp
--- a/Engine/Source/Runtime/Engine/Private/RepLayout.cpp
+++ b/Engine/Source/Runtime/Engine/Private/RepLayout.cpp
@@ -3395,7 +3395,7 @@ static bool ReceivePropertyHelper(
 				check(GuidReferences->ParentIndex == Cmd.ParentIndex);
 
 				// If we're already tracking the guids, re-copy lists only if they've changed
-				if (!NetworkGuidSetsAreSame(GuidReferences->UnmappedGUIDs, TrackedUnmappedGuids))
+				if (!NetworkGuidSetsAreSame(GuidReferences->GetUnmappedGUIDs(), TrackedUnmappedGuids))
 				{
 					bOutGuidsChanged = true;
 				}
@@ -3408,7 +3408,7 @@ static bool ReceivePropertyHelper(
 			if (GuidReferences == nullptr || bOutGuidsChanged)
 			{
 				// First time tracking these guids (or guids changed), so add (or replace) new entry
-				GuidReferencesMap->Add(AbsOffset, FGuidReferences(Bunch, Mark, TrackedUnmappedGuids, TrackedDynamicMappedGuids, Cmd.ParentIndex, CmdIndex));
+				GuidReferencesMap->Emplace(AbsOffset, FGuidReferences(Bunch, Mark, TrackedUnmappedGuids, TrackedDynamicMappedGuids, Cmd.ParentIndex, CmdIndex, Bunch.PackageMap));
 				bOutGuidsChanged = true;
 			}
 			else if (UE::Net::Private::bAlwaysUpdateGuidReferenceMapForNetSerializeObjectStruct && Cmd.Type == ERepLayoutCmdType::NetSerializeStructWithObjectReferences)
@@ -3416,7 +3416,7 @@ static bool ReceivePropertyHelper(
 				// If this is a NetSerialize struct with object references, there may be other properties "wrapped up" with this GUID reference.
 				// In this case, the entry in the map should be always be updated, so there isn't outdated data in the entry that also gets
 				// applied when the Guid possibly goes unmapped and then mapped later.
-				GuidReferencesMap->Add(AbsOffset, FGuidReferences(Bunch, Mark, TrackedUnmappedGuids, TrackedDynamicMappedGuids, Cmd.ParentIndex, CmdIndex));
+				GuidReferencesMap->Emplace(AbsOffset, FGuidReferences(Bunch, Mark, TrackedUnmappedGuids, TrackedDynamicMappedGuids, Cmd.ParentIndex, CmdIndex, Bunch.PackageMap));
 			}
 		}
 		else
@@ -3450,7 +3450,8 @@ static FGuidReferencesMap* PrepReceivedArray(
 	FRepShadowDataBuffer* OutShadowBaseData,
 	FRepObjectDataBuffer* OutBaseData,
 	TArray<FProperty*>* RepNotifies,
-	bool& bOutShadowDataCopied)
+	bool& bOutShadowDataCopied,
+	UPackageMap* PackageMap)
 {
 	FGuidReferences* NewGuidReferencesArray = nullptr;
 
@@ -3461,11 +3462,7 @@ static FGuidReferencesMap* PrepReceivedArray(
 
 		if (NewGuidReferencesArray == nullptr)
 		{
-			NewGuidReferencesArray = &ParentGuidReferences->FindOrAdd(AbsOffset);
-
-			NewGuidReferencesArray->Array = new FGuidReferencesMap;
-			NewGuidReferencesArray->ParentIndex = Cmd.ParentIndex;
-			NewGuidReferencesArray->CmdIndex = CmdIndex;
+			NewGuidReferencesArray = &ParentGuidReferences->Emplace(AbsOffset, FGuidReferences(new FGuidReferencesMap, Cmd.ParentIndex, CmdIndex, PackageMap));
 		}
 
 		check(NewGuidReferencesArray != nullptr);
@@ -3668,7 +3665,8 @@ static bool ReceiveProperties_r(FReceivePropertiesSharedParams& Params, FReceive
 					&ShadowArrayBuffer,
 					&ObjectArrayBuffer,
 					StackParams.RepNotifies,
-					ArrayStackParams.bShadowDataCopied);
+					ArrayStackParams.bShadowDataCopied,
+					Params.Bunch.PackageMap);
 
 				// Read the next array handle.
 				ReadPropertyHandle(Params);
@@ -4053,8 +4051,8 @@ bool FRepLayout::ReceiveProperties_BackwardsCompatible_r(
 				&LocalShadowData,
 				&LocalData,
 				ShadowData ? &RepState->RepNotifies : nullptr,
-				bShadowDataCopied
-				);
+				bShadowDataCopied,
+				TempReader.PackageMap);
 
 			// Read until we read all array elements
 			while (true)
@@ -4187,7 +4185,7 @@ void FRepLayout::GatherGuidReferences_r(
 
 		OutTrackedGuidMemoryBytes += GuidReferences.Buffer.Num();
 
-		OutReferencedGuids.Append(GuidReferences.UnmappedGUIDs);
+		OutReferencedGuids.Append(GuidReferences.GetUnmappedGUIDs());
 		OutReferencedGuids.Append(GuidReferences.MappedDynamicGUIDs);
 	}
 }
@@ -4250,7 +4248,7 @@ bool FRepLayout::MoveMappedObjectToUnmapped_r(FGuidReferencesMap* GuidReferences
 		if (GuidReferences.MappedDynamicGUIDs.Contains(GUID))
 		{
 			GuidReferences.MappedDynamicGUIDs.Remove(GUID);
-			GuidReferences.UnmappedGUIDs.Add(GUID);
+			GuidReferences.AddUnmappedGUID(GUID);
 			bFoundGUID = true;
 
 #if WITH_PUSH_MODEL
@@ -4374,40 +4372,7 @@ void FRepLayout::UpdateUnmappedObjects_r(
 			continue;
 		}
 
-		bool bMappedSomeGUIDs = false;
-
-		for (auto UnmappedIt = GuidReferences.UnmappedGUIDs.CreateIterator(); UnmappedIt; ++UnmappedIt)
-		{
-			const FNetworkGUID& GUID = *UnmappedIt;
-
-			if (Connection->PackageMap->IsGUIDBroken(GUID, false))
-			{
-				UE_LOG(LogRep, Warning, TEXT("UpdateUnmappedObjects_r: Broken GUID. NetGuid: %s"), *GUID.ToString());
-				UnmappedIt.RemoveCurrent();
-				continue;
-			}
-
-			UObject* Object = Connection->PackageMap->GetObjectFromNetGUID(GUID, false);
-
-			if (Object != nullptr)
-			{
-				UE_LOG(LogRep, VeryVerbose, TEXT("UpdateUnmappedObjects_r: REMOVED unmapped property: Offset: %i, Guid: %s, PropName: %s, ObjName: %s"), AbsOffset, *GUID.ToString(), *Cmd.Property->GetName(), *Object->GetName());
-
-				if (GUID.IsDynamic())
-				{
-					// If this guid is dynamic, move it to the dynamic guids list
-					GuidReferences.MappedDynamicGUIDs.Add(GUID);
-				}
-
-				// Remove from unmapped guids list
-				UnmappedIt.RemoveCurrent();
-				bMappedSomeGUIDs = true;
-
-#if WITH_PUSH_MODEL
-				FNetPrivatePushIdHelper::MarkPropertyDirty(OriginalObject, GuidReferences.ParentIndex);
-#endif
-			}
-		}
+		bool bMappedSomeGUIDs = GuidReferences.UpdateUnmappedGUIDs(Connection->PackageMap, OriginalObject, Cmd.Property, AbsOffset);
 
 		// If we resolved some guids, re-deserialize the data which will hook up the object pointer with the property
 		if (bMappedSomeGUIDs)
@@ -4456,11 +4421,11 @@ void FRepLayout::UpdateUnmappedObjects_r(
 		}
 
 		// If we still have more unmapped guids, we need to keep processing this entry
-		if (GuidReferences.UnmappedGUIDs.Num() > 0)
+		if (GuidReferences.GetUnmappedGUIDs().Num() > 0)
 		{
 			bOutHasMoreUnmapped = true;
 		}
-		else if (GuidReferences.UnmappedGUIDs.Num() == 0 && GuidReferences.MappedDynamicGUIDs.Num() == 0)
+		else if (GuidReferences.GetUnmappedGUIDs().Num() == 0 && GuidReferences.MappedDynamicGUIDs.Num() == 0)
 		{
 			It.RemoveCurrent();
 		}
diff --git a/Engine/Source/Runtime/Net/Core/Private/Net/Core/Misc/GuidReferences.cpp b/Engine/Source/Runtime/Net/Core/Private/Net/Core/Misc/GuidReferences.cpp
--- a/Engine/Source/Runtime/Net/Core/Private/Net/Core/Misc/GuidReferences.cpp
+++ b/Engine/Source/Runtime/Net/Core/Private/Net/Core/Misc/GuidReferences.cpp
@@ -1,12 +1,194 @@
 // Copyright Epic Games, Inc. All Rights Reserved.
 
 #include "Net/Core/Misc/GuidReferences.h"
+#include "HAL/IConsoleManager.h"
+#include "Net/Core/Misc/NetCoreLog.h"
+#include "Net/Core/PushModel/PushModel.h"
+#include "UObject/CoreNet.h"
+#include "UObject/UnrealType.h"
 
-FGuidReferences::~FGuidReferences()
+namespace UE::Net::Private
+{
+	bool bRemapStableSubobjects = true;
+	static FAutoConsoleVariableRef CVarNetRemapStableSubobjects(TEXT("net.RemapStableSubobjects"), bRemapStableSubobjects, TEXT("If enabled, attempts to remap stable subobjects when net.OptimizedRemapping is also enabled."));
+}
+
+FGuidReferences::FGuidReferences(
+	FBitReader&					InReader,
+	FBitReaderMark&				InMark,
+	const TSet<FNetworkGUID>&	InUnmappedGUIDs,
+	const TSet<FNetworkGUID>&	InMappedDynamicGUIDs,
+	const int32					InParentIndex,
+	const int32					InCmdIndex,
+	UPackageMap*				InPackageMap)
+	: ParentIndex(InParentIndex)
+	, CmdIndex(InCmdIndex)
+	, UnmappedGUIDs(InUnmappedGUIDs)
+	, MappedDynamicGUIDs(InMappedDynamicGUIDs)
+	, Array(nullptr)
+	, PackageMap(InPackageMap)
 {
-	if (Array != NULL)
+	NumBufferBits = UE_PTRDIFF_TO_INT32(InReader.GetPosBits() - InMark.GetPos());
+	InMark.Copy(InReader, Buffer);
+
+	TrackAllUnmappedGUIDs();
+}
+
+FGuidReferences::FGuidReferences(const FGuidReferences& Other)
+	: ParentIndex(Other.ParentIndex)
+	, CmdIndex(Other.CmdIndex)
+	, NumBufferBits(Other.NumBufferBits)
+	, UnmappedGUIDs(Other.UnmappedGUIDs)
+	, MappedDynamicGUIDs(Other.MappedDynamicGUIDs)
+	, Buffer(Other.Buffer)
+	, Array(nullptr)
+	, PackageMap(Other.PackageMap)
+{
+	if (Other.Array)
+	{
+		Array = new FGuidReferencesMap(*Other.Array);
+	}
+
+	TrackAllUnmappedGUIDs();
+}
+
+FGuidReferences& FGuidReferences::operator=(FGuidReferences&& Other)
+{
+	if (this != &Other)
 	{
 		delete Array;
-		Array = NULL;
+
+		ParentIndex = Other.ParentIndex;
+		CmdIndex = Other.CmdIndex;
+		NumBufferBits = Other.NumBufferBits;
+		UnmappedGUIDs = MoveTemp(Other.UnmappedGUIDs);
+		MappedDynamicGUIDs = MoveTemp(Other.MappedDynamicGUIDs);
+		Buffer = MoveTemp(Other.Buffer);
+		Array = Other.Array;
+		PackageMap = Other.PackageMap;
+
+		Other.Array = nullptr;
+		Other.NumBufferBits = 0;
+		Other.PackageMap.Reset();
+	}
+
+	return *this;
+}
+
+FGuidReferences::FGuidReferences(FGuidReferences&& Other) 
+	: ParentIndex(Other.ParentIndex)
+	, CmdIndex(Other.CmdIndex)
+	, NumBufferBits(Other.NumBufferBits)
+	, UnmappedGUIDs(MoveTemp(Other.UnmappedGUIDs))
+	, MappedDynamicGUIDs(MoveTemp(Other.MappedDynamicGUIDs))
+	, Buffer(MoveTemp(Other.Buffer))
+	, Array(Other.Array)
+	, PackageMap(Other.PackageMap)
+{
+	Other.Array = nullptr;
+	Other.NumBufferBits = 0;
+	Other.PackageMap.Reset();
+}
+
+FGuidReferences::~FGuidReferences()
+{
+	delete Array;
+
+	if (UE::Net::Private::bRemapStableSubobjects)
+	{
+		if (UPackageMap* LocalPackageMap = PackageMap.Get())
+		{
+			for (const FNetworkGUID NetGUID : UnmappedGUIDs)
+			{
+				LocalPackageMap->RemoveUnmappedNetGUIDReference(NetGUID);
+			}
+		}
+	}
+}
+
+void FGuidReferences::AddUnmappedGUID(FNetworkGUID InGUID)
+{
+	bool bIsAlreadyInSet = false;
+	UnmappedGUIDs.Add(InGUID, &bIsAlreadyInSet);
+
+	if (!bIsAlreadyInSet && UE::Net::Private::bRemapStableSubobjects)
+	{
+		if (UPackageMap* LocalPackageMap = PackageMap.Get())
+		{
+			PackageMap->AddUnmappedNetGUIDReference(InGUID);
+		}
+	}
+}
+
+void FGuidReferences::RemoveUnmappedGUID(FNetworkGUID InGUID)
+{
+	int32 NumRemoved = UnmappedGUIDs.Remove(InGUID);
+
+	if (NumRemoved > 0 && UE::Net::Private::bRemapStableSubobjects)
+	{
+		ensureMsgf(NumRemoved == 1, TEXT("FGuidReferences::RemoveUnmappedGUID: removed more than one GUID from UnmappedGUIDs set."));
+
+		if (UPackageMap* LocalPackageMap = PackageMap.Get())
+		{
+			PackageMap->RemoveUnmappedNetGUIDReference(InGUID);
+		}
+	}
+}
+
+bool FGuidReferences::UpdateUnmappedGUIDs(UPackageMap* InPackageMap, UObject* OriginalObject, const FProperty* Property, int32 AbsOffset)
+{
+	checkf(InPackageMap != nullptr, TEXT("FGuidReferences::UpdateUnmappedGUIDs: invalid PackageMap"));
+
+	bool bMappedSomeGUIDs = false;
+
+	for (auto UnmappedIt = UnmappedGUIDs.CreateIterator(); UnmappedIt; ++UnmappedIt)
+	{
+		const FNetworkGUID& GUID = *UnmappedIt;
+
+		if (InPackageMap->IsGUIDBroken(GUID, false))
+		{
+			UE_LOG(LogNetCore, Warning, TEXT("FGuidReferences::UpdateUnmappedGUIDs: Broken GUID. NetGuid: %s"), *GUID.ToString());
+			InPackageMap->RemoveUnmappedNetGUIDReference(GUID);
+			UnmappedIt.RemoveCurrent();
+			continue;
+		}
+
+		UObject* Object = InPackageMap->GetObjectFromNetGUID(GUID, false);
+
+		if (Object != nullptr)
+		{
+			UE_LOG(LogNetCore, VeryVerbose, TEXT("FGuidReferences::UpdateUnmappedGUIDs: REMOVED unmapped property: Offset: %i, Guid: %s, PropName: %s, ObjName: %s"), AbsOffset, *GUID.ToString(), *Property->GetName(), *Object->GetName());
+
+			if (GUID.IsDynamic())
+			{
+				// If this guid is dynamic, move it to the dynamic guids list
+				MappedDynamicGUIDs.Add(GUID);
+			}
+
+			// Remove from unmapped guids list
+			InPackageMap->RemoveUnmappedNetGUIDReference(GUID);
+			UnmappedIt.RemoveCurrent();
+			bMappedSomeGUIDs = true;
+
+#if WITH_PUSH_MODEL
+			MARK_PROPERTY_DIRTY_UNSAFE(OriginalObject, ParentIndex);
+#endif
+		}
+	}
+
+	return bMappedSomeGUIDs;
+}
+
+void FGuidReferences::TrackAllUnmappedGUIDs() const
+{
+	if (UE::Net::Private::bRemapStableSubobjects)
+	{
+		if (UPackageMap* LocalPackageMap = PackageMap.Get())
+		{
+			for (const FNetworkGUID UnmappedGUID : UnmappedGUIDs)
+			{
+				LocalPackageMap->AddUnmappedNetGUIDReference(UnmappedGUID);
+			}
+		}
 	}
 }
\ No newline at end of file
diff --git a/Engine/Source/Runtime/Net/Core/Public/Net/Core/Misc/GuidReferences.h b/Engine/Source/Runtime/Net/Core/Public/Net/Core/Misc/GuidReferences.h
--- a/Engine/Source/Runtime/Net/Core/Public/Net/Core/Misc/GuidReferences.h
+++ b/Engine/Source/Runtime/Net/Core/Public/Net/Core/Misc/GuidReferences.h
@@ -10,12 +10,20 @@
 #include "Misc/NetworkGuid.h"
 #include "Serialization/BitReader.h"
 #include "Templates/Tuple.h"
+#include "UObject/WeakObjectPtr.h"
 
 class FArchive;
 class FGuidReferences;
+class FProperty;
+class UPackageMap;
 
 using FGuidReferencesMap = TMap<int32, FGuidReferences>;
 
+namespace UE::Net::Private
+{
+	extern NETCORE_API bool bRemapStableSubobjects;
+}
+
 /**
  * References to Objects (including Actors, Components, etc.) are replicated as NetGUIDs, since
  * the literal memory pointers will be different across game instances. In these cases, actual
@@ -40,58 +48,37 @@ using FGuidReferencesMap = TMap<int32, FGuidReferences>;
 class FGuidReferences
 {
 public:
-	FGuidReferences():
-		NumBufferBits(0),
-		Array(nullptr)
-	{}
-
-	FGuidReferences(
+	NETCORE_API FGuidReferences(
 		FBitReader&					InReader,
 		FBitReaderMark&				InMark,
 		const TSet<FNetworkGUID>&	InUnmappedGUIDs,
 		const TSet<FNetworkGUID>&	InMappedDynamicGUIDs,
 		const int32					InParentIndex,
-		const int32					InCmdIndex
-	):
-		ParentIndex(InParentIndex),
-		CmdIndex(InCmdIndex),
-		UnmappedGUIDs(InUnmappedGUIDs),
-		MappedDynamicGUIDs(InMappedDynamicGUIDs),
-		Array(nullptr)
-	{
-		NumBufferBits = UE_PTRDIFF_TO_INT32(InReader.GetPosBits() - InMark.GetPos());
-		InMark.Copy(InReader, Buffer);
-	}
+		const int32					InCmdIndex,
+		UPackageMap*				InPackageMap
+	);
 
 	FGuidReferences(
 		FGuidReferencesMap*	InArray,
 		const int32			InParentIndex,
-		const int32			InCmdIndex
+		const int32			InCmdIndex,
+		UPackageMap*		InPackageMap
 	):
 		ParentIndex(InParentIndex),
 		CmdIndex(InCmdIndex),
 		NumBufferBits(0),
-		Array(InArray)
+		Array(InArray),
+		PackageMap(InPackageMap)
 	{}
 
+
 	NETCORE_API ~FGuidReferences();
 
 	FGuidReferences& operator=(const FGuidReferences& Other) = delete;
+	NETCORE_API FGuidReferences(const FGuidReferences& Other);
 
-	FGuidReferences(const FGuidReferences& Other)
-		: ParentIndex(Other.ParentIndex)
-		, CmdIndex(Other.CmdIndex)
-		, NumBufferBits(Other.NumBufferBits)
-		, UnmappedGUIDs(Other.UnmappedGUIDs)
-		, MappedDynamicGUIDs(Other.MappedDynamicGUIDs)
-		, Buffer(Other.Buffer)
-		, Array(nullptr)
-	{
-		if (Other.Array)
-		{
-			Array = new FGuidReferencesMap(*Other.Array);
-		}
-	}
+	NETCORE_API FGuidReferences& operator=(FGuidReferences&&);
+	NETCORE_API FGuidReferences(FGuidReferences&& Other);
 
 	void CountBytes(FArchive& Ar) const
 	{
@@ -109,6 +96,18 @@ public:
 		}
 	}
 
+	/** Get the set of GUIDs for objects that haven't been loaded / created yet. */
+	const TSet<FNetworkGUID>& GetUnmappedGUIDs() const { return UnmappedGUIDs; }
+
+	/** Starts tracking InGUID if not already and calls UPackageMap::AddUnmappedNetGUIDReference if bRemapStableSubobjects is enabled. */
+	NETCORE_API void AddUnmappedGUID(FNetworkGUID InGUID);
+
+	/** Stops tracking a GUID and calls UPackageMap::RemoveUnmappedNetGUIDReference if bRemapStableSubobjects is enabled. */
+	NETCORE_API void RemoveUnmappedGUID(FNetworkGUID InGUID);
+
+	/** Called by FRepLayout to update GUID tracking when updating unmapped objects. */
+	NETCORE_API bool UpdateUnmappedGUIDs(UPackageMap* InPackageMap, UObject* OriginalObject, const FProperty* Property, int32 AbsOffset);
+
 	/** The Property Command index of the top level property that references the GUID. */
 	int32 ParentIndex;
 
@@ -117,9 +116,11 @@ public:
 
 	int32 NumBufferBits;
 
+private:
 	/** GUIDs for objects that haven't been loaded / created yet. */
 	TSet<FNetworkGUID> UnmappedGUIDs;
 
+public:
 	/** GUIDs for dynamically spawned objects that have already been created. */
 	TSet<FNetworkGUID> MappedDynamicGUIDs;
 
@@ -131,4 +132,11 @@ public:
 	 * then this will be a valid FGuidReferencesMap containing the nested FGuidReferences.
 	 */
 	FGuidReferencesMap* Array;
+
+private:
+	/** Stored PackageMap reference to update with tracked UnmappedGUIDs. */
+	TWeakObjectPtr<UPackageMap> PackageMap;
+
+	/** Calls UPackageMap::AddUnmappedNetGUIDReference for all GUIDs in UnmappedGUIDs, if bRemapStableSubobjects si enabled. */
+	void TrackAllUnmappedGUIDs() const;
 };
