Add missing null check

diff --git a/Engine/Source/Runtime/Engine/Private/UnrealEngine.cpp b/Engine/Source/Runtime/Engine/Private/UnrealEngine.cpp
--- a/Engine/Source/Runtime/Engine/Private/UnrealEngine.cpp
+++ b/Engine/Source/Runtime/Engine/Private/UnrealEngine.cpp
@@ -1881,7 +1881,7 @@ void UEngine::PerformGarbageCollectionAndCleanupActors()
 		for (FWorldContext& Context : WorldList)
 		{
 			UWorld* World = Context.World();
-			if (World->IsGameWorld())
+			if (World != nullptr && World->IsGameWorld())
 			{
 				bForcePurge = false;
 				break;
