[StaticMeshComponent] Added navigation update when associated mesh gets compiled. This will invalidate any pending add to the octree and dirty tiles since the component is no longer relevant until compilation completes.

diff --git a/Engine/Source/Runtime/Engine/Classes/Components/StaticMeshComponent.h b/Engine/Source/Runtime/Engine/Classes/Components/StaticMeshComponent.h
--- a/Engine/Source/Runtime/Engine/Classes/Components/StaticMeshComponent.h
+++ b/Engine/Source/Runtime/Engine/Classes/Components/StaticMeshComponent.h
@@ -131,6 +131,7 @@ private:
 #if WITH_EDITOR
 	/** Used to track down when StaticMesh has been modified for notification purpose */
 	class UStaticMesh* KnownStaticMesh = nullptr;
+	FDelegateHandle PreMeshBuildDelegateHandle;
 #endif
 	ENGINE_API void NotifyIfStaticMeshChanged();
 
diff --git a/Engine/Source/Runtime/Engine/Private/Components/StaticMeshComponent.cpp b/Engine/Source/Runtime/Engine/Private/Components/StaticMeshComponent.cpp
--- a/Engine/Source/Runtime/Engine/Private/Components/StaticMeshComponent.cpp
+++ b/Engine/Source/Runtime/Engine/Private/Components/StaticMeshComponent.cpp
@@ -678,13 +678,33 @@ void UStaticMeshComponent::NotifyIfStaticMeshChanged()
 #if WITH_EDITOR
 	if (KnownStaticMesh != StaticMesh)
 	{
+		// Remove delegate from our previous mesh 
+		if (KnownStaticMesh)
+		{
+			KnownStaticMesh->OnPreMeshBuild().Remove(PreMeshBuildDelegateHandle);
+			PreMeshBuildDelegateHandle.Reset();
+		}
+
 		KnownStaticMesh = StaticMesh;
+
+		// Add delegate on our new mesh
+		if (KnownStaticMesh)
+		{
+			PreMeshBuildDelegateHandle = KnownStaticMesh->OnPreMeshBuild().AddWeakLambda(this, [this](UStaticMesh*)
+				{
+					// Our associated mesh is compiling so we are no longer navigation relevant.
+					// This will invalidate any pending add to the octree and dirty tiles until compilation completes.
+					bNavigationRelevant = IsNavigationRelevant();
+					FNavigationSystem::UpdateComponentData(*this);
+				});
+		}
+
 		FObjectCacheEventSink::NotifyStaticMeshChanged_Concurrent(GetStaticMeshComponentInterface());
 
 		// Update this component streaming data.
 		IStreamingManager::Get().NotifyPrimitiveUpdated(this);
 	}
-#endif
+#endif // WITH_EDITOR
 }
 
 #if WITH_EDITOR
@@ -1844,7 +1864,12 @@ void UStaticMeshComponent::BeginDestroy()
 #if WITH_EDITOR
 	// The object cache needs to be notified when we're getting destroyed
 	FObjectCacheEventSink::NotifyStaticMeshChanged_Concurrent(GetStaticMeshComponentInterface());
-#endif
+	
+	if (PreMeshBuildDelegateHandle.IsValid())
+	{
+		KnownStaticMesh->OnPreMeshBuild().Remove(PreMeshBuildDelegateHandle);
+	}
+#endif // WITH_EDITOR
 }
 
 void UStaticMeshComponent::ExportCustomProperties(FOutputDevice& Out, uint32 Indent)
