Fix MetaSounds not playing within SoundCues with certain sound cue quality settings

diff --git a/Engine/Source/Runtime/Engine/Classes/Sound/SoundNode.h b/Engine/Source/Runtime/Engine/Classes/Sound/SoundNode.h
--- a/Engine/Source/Runtime/Engine/Classes/Sound/SoundNode.h
+++ b/Engine/Source/Runtime/Engine/Classes/Sound/SoundNode.h
@@ -235,5 +235,13 @@ public:
 	 * This should only be called by USoundNodeQualityLevel::PostLoad when au.CullSoundWaveHardReferences is 1.
 	 */
 	ENGINE_API virtual void RemoveSoundWaveOnChildWavePlayers();
+
+	/**
+	 * Load wave assets for any wave player sound nodes 
+	 * childed off of this node, taking into account quality nodes
+	 *
+	 * @param bRecurse when true, this will cause all children of child nodes to be loaded as well.
+	 */
+	ENGINE_API virtual void LoadChildWavePlayerAssets(bool bAddToRoot, bool bRecurse);
 };
 
diff --git a/Engine/Source/Runtime/Engine/Classes/Sound/SoundNodeAssetReferencer.h b/Engine/Source/Runtime/Engine/Classes/Sound/SoundNodeAssetReferencer.h
--- a/Engine/Source/Runtime/Engine/Classes/Sound/SoundNodeAssetReferencer.h
+++ b/Engine/Source/Runtime/Engine/Classes/Sound/SoundNodeAssetReferencer.h
@@ -19,6 +19,7 @@ class USoundNodeAssetReferencer : public USoundNode
 public:
 	ENGINE_API virtual void LoadAsset(bool bAddToRoot = false) PURE_VIRTUAL(USoundNodeAssetReferencer::LoadAsset,);
 	ENGINE_API virtual void ClearAssetReferences() PURE_VIRTUAL(USoundNodeAssetReferencer::ClearAssetReferences, );
+	ENGINE_API virtual bool ContainsProceduralSoundReference() const { return false; }
 
 	ENGINE_API bool ShouldHardReferenceAsset(const class ITargetPlatform*) const;
 
diff --git a/Engine/Source/Runtime/Engine/Classes/Sound/SoundNodeQualityLevel.h b/Engine/Source/Runtime/Engine/Classes/Sound/SoundNodeQualityLevel.h
--- a/Engine/Source/Runtime/Engine/Classes/Sound/SoundNodeQualityLevel.h
+++ b/Engine/Source/Runtime/Engine/Classes/Sound/SoundNodeQualityLevel.h
@@ -36,6 +36,7 @@ public:
 	virtual void PrimeChildWavePlayers(bool bRecurse) override;
 	virtual void RetainChildWavePlayers(bool bRecurse) override;
 	virtual void ReleaseRetainerOnChildWavePlayers(bool bRecurse) override;
+	void LoadChildWavePlayers(bool bAddToRoot, bool bRecurse);
 
 #if WITH_EDITOR
 	virtual FText GetInputPinName(int32 PinIndex) const override;
diff --git a/Engine/Source/Runtime/Engine/Classes/Sound/SoundNodeWavePlayer.h b/Engine/Source/Runtime/Engine/Classes/Sound/SoundNodeWavePlayer.h
--- a/Engine/Source/Runtime/Engine/Classes/Sound/SoundNodeWavePlayer.h
+++ b/Engine/Source/Runtime/Engine/Classes/Sound/SoundNodeWavePlayer.h
@@ -73,6 +73,7 @@ public:
 	//~ Begin USoundNodeAssetReferencer Interface
 	virtual void LoadAsset(bool bAddToRoot = false) override;
 	virtual void ClearAssetReferences() override;
+	virtual bool ContainsProceduralSoundReference() const override;
 	//~ End USoundNode Interface
 
 	// If this returns true, this wave player currently has an async load for the USoundWave in flight.
diff --git a/Engine/Source/Runtime/Engine/Private/SoundCue.cpp b/Engine/Source/Runtime/Engine/Private/SoundCue.cpp
--- a/Engine/Source/Runtime/Engine/Private/SoundCue.cpp
+++ b/Engine/Source/Runtime/Engine/Private/SoundCue.cpp
@@ -303,6 +303,34 @@ void USoundCue::EvaluateNodes(bool bAddToRoot)
 		}
 	};
 
+	TFunction<void(USoundNode*)> LoadProceduralAssets = [&](USoundNode* SoundNode)
+	{
+		if (SoundNode == nullptr)
+		{
+			return;
+		}
+
+		if (USoundNodeAssetReferencer* AssetReferencerNode = Cast<USoundNodeAssetReferencer>(SoundNode))
+		{
+			if (AssetReferencerNode->ContainsProceduralSoundReference())
+			{
+				AssetReferencerNode->ConditionalPostLoad();
+				AssetReferencerNode->LoadAsset(bAddToRoot);
+			}
+		}
+		else if (USoundNodeQualityLevel* QualityLevelNode = Cast<USoundNodeQualityLevel>(SoundNode))
+		{
+			QualityLevelNode->LoadChildWavePlayers(bAddToRoot, /*bRecurse=*/true);
+		}
+		else
+		{
+			for (USoundNode* ChildNode : SoundNode->ChildNodes)
+			{
+				LoadProceduralAssets(ChildNode);
+			}
+		}
+	};
+
 	// Only Evaluate nodes if we haven't been cooked, as cooked builds will hard-ref all SoundAssetReferences.	
 	UE_CLOG(CookedQualityIndex == INDEX_NONE, LogAudio, Verbose, TEXT("'%s', DOING EvaluateNodes as we are *NOT* cooked"), *GetName());
 	UE_CLOG(CookedQualityIndex != INDEX_NONE, LogAudio, Verbose, TEXT("'%s', SKIPPING EvaluateNodes as we *ARE* cooked"), *GetName());
@@ -311,6 +339,12 @@ void USoundCue::EvaluateNodes(bool bAddToRoot)
 	{		
 		EvaluateNodes_Internal(FirstNode);
 	}
+	else
+	{
+		// We need to load procedural assets (MetaSounds) to initialize their resources
+		// before playing (which EvaluateNodes_Internal does in the other case)
+		LoadProceduralAssets(FirstNode);
+	}
 }
 
 void USoundCue::CacheQualityLevel()
diff --git a/Engine/Source/Runtime/Engine/Private/SoundNode.cpp b/Engine/Source/Runtime/Engine/Private/SoundNode.cpp
--- a/Engine/Source/Runtime/Engine/Private/SoundNode.cpp
+++ b/Engine/Source/Runtime/Engine/Private/SoundNode.cpp
@@ -7,6 +7,7 @@
 #include "Sound/SoundCue.h"
 #include "Misc/App.h"
 #include "Sound/SoundNodeWavePlayer.h"
+#include "Sound/SoundNodeQualityLevel.h"
 #include "AudioCompressionSettingsUtils.h"
 
 #include UE_INLINE_GENERATED_CPP_BY_NAME(SoundNode)
@@ -243,6 +244,31 @@ void USoundNode::RemoveSoundWaveOnChildWavePlayers()
 	}
 }
 
+void USoundNode::LoadChildWavePlayerAssets(bool bAddToRoot, bool bRecurse)
+{
+	// Search child nodes for wave players, then load their sound wave asset.
+	for (USoundNode* ChildNode : ChildNodes)
+	{
+		if (ChildNode)
+		{
+			if (bRecurse)
+			{
+				ChildNode->LoadChildWavePlayerAssets(bAddToRoot, bRecurse);
+			}
+
+			if (USoundNodeWavePlayer* WavePlayer = Cast<USoundNodeWavePlayer>(ChildNode))
+			{
+				WavePlayer->LoadAsset(bAddToRoot);
+			}
+			else if (USoundNodeQualityLevel* QualityNode = Cast<USoundNodeQualityLevel>(ChildNode))
+			{
+				// Take into account quality nodes by only loading wave players for the relevant quality level
+				QualityNode->LoadChildWavePlayers(bAddToRoot, bRecurse);
+			}
+		}
+	}
+}
+
 void USoundNode::BeginDestroy()
 {
 	Super::BeginDestroy();
diff --git a/Engine/Source/Runtime/Engine/Private/SoundNodeQualityLevel.cpp b/Engine/Source/Runtime/Engine/Private/SoundNodeQualityLevel.cpp
--- a/Engine/Source/Runtime/Engine/Private/SoundNodeQualityLevel.cpp
+++ b/Engine/Source/Runtime/Engine/Private/SoundNodeQualityLevel.cpp
@@ -76,6 +76,11 @@ void USoundNodeQualityLevel::ReleaseRetainerOnChildWavePlayers(bool bRecurse)
 	ForCurrentQualityLevel([bRecurse](USoundNode* Node) { Node->ReleaseRetainerOnChildWavePlayers(bRecurse); });
 }
 
+void USoundNodeQualityLevel::LoadChildWavePlayers(bool bAddToRoot, bool bRecurse)
+{
+	ForCurrentQualityLevel([bAddToRoot, bRecurse](USoundNode* Node) { Node->LoadChildWavePlayerAssets(bAddToRoot, bRecurse); });
+}
+
 int32 USoundNodeQualityLevel::GetMaxChildNodes() const
 {
 	return GetDefault<UAudioSettings>()->QualityLevels.Num();
diff --git a/Engine/Source/Runtime/Engine/Private/SoundNodeWavePlayer.cpp b/Engine/Source/Runtime/Engine/Private/SoundNodeWavePlayer.cpp
--- a/Engine/Source/Runtime/Engine/Private/SoundNodeWavePlayer.cpp
+++ b/Engine/Source/Runtime/Engine/Private/SoundNodeWavePlayer.cpp
@@ -32,6 +32,15 @@ void USoundNodeWavePlayer::Serialize(FArchive& Ar)
 	}
 }
 
+bool USoundNodeWavePlayer::ContainsProceduralSoundReference() const
+{
+	if (SoundWave)
+	{
+		return SoundWave->IsA<USoundWaveProcedural>();
+	}
+	return false;
+}
+
 void USoundNodeWavePlayer::LoadAsset(bool bAddToRoot)
 {
 	if (IsAsyncLoadingMultithreaded())
