WorkingColorSpace: Replace recently added convenience sRGB-to-working colorspace function with native expression to minimize material instructions when used on parameters.

Note however that with the new HLSL generator, EOperation::Dot is currently not implemented making this particular path suboptimal.

diff --git a/Engine/Source/Runtime/ColorManagement/Private/ColorSpace.cpp b/Engine/Source/Runtime/ColorManagement/Private/ColorSpace.cpp
--- a/Engine/Source/Runtime/ColorManagement/Private/ColorSpace.cpp
+++ b/Engine/Source/Runtime/ColorManagement/Private/ColorSpace.cpp
@@ -301,6 +301,13 @@ FMatrix44d FColorSpaceTransform::CalcChromaticAdaptionMatrix(FVector3d SourceXYZ
 	return XyzToRgb * ScaleMat * RgbToXyz;
 }
 
+FColorSpaceTransform FColorSpaceTransform::GetSRGBToWorkingColorSpace()
+{
+	static FColorSpaceTransform CachedTransform = FColorSpaceTransform(FColorSpace(EColorSpace::sRGB), FColorSpace::GetWorking());
+	
+	return CachedTransform;
+}
+
 static FMatrix44d CalcColorSpaceTransformMatrix(const FColorSpace& Src, const FColorSpace& Dst, EChromaticAdaptationMethod Method)
 {
 	if (Method == UE::Color::EChromaticAdaptationMethod::None)
diff --git a/Engine/Source/Runtime/ColorManagement/Public/ColorSpace.h b/Engine/Source/Runtime/ColorManagement/Public/ColorSpace.h
--- a/Engine/Source/Runtime/ColorManagement/Public/ColorSpace.h
+++ b/Engine/Source/Runtime/ColorManagement/Public/ColorSpace.h
@@ -388,6 +388,11 @@ struct FColorSpaceTransform : FMatrix44d
 	* @param Method Adaptation method (None, Bradford, CAT02).
 	*/
 	COLORMANAGEMENT_API static FMatrix44d CalcChromaticAdaptionMatrix(FVector3d SourceXYZ, FVector3d TargetXYZ, EChromaticAdaptationMethod Method = EChromaticAdaptationMethod::Bradford);
+
+	/**
+	* Convenience function to get a (statically cached) conversion from sRGB/Rec709 to the working color space.
+	*/
+	COLORMANAGEMENT_API static FColorSpaceTransform GetSRGBToWorkingColorSpace();
 };
 
 } }  // end namespace UE::Color
diff --git a/Engine/Source/Runtime/Engine/Classes/Materials/MaterialExpressionSRGBColorToWorkingColorSpace.h b/Engine/Source/Runtime/Engine/Classes/Materials/MaterialExpressionSRGBColorToWorkingColorSpace.h
new file mode 100644
--- /dev/null
+++ b/Engine/Source/Runtime/Engine/Classes/Materials/MaterialExpressionSRGBColorToWorkingColorSpace.h
@@ -0,0 +1,31 @@
+// Copyright Epic Games, Inc. All Rights Reserved.
+
+
+#pragma once
+
+#include "CoreMinimal.h"
+#include "UObject/ObjectMacros.h"
+#include "MaterialExpressionIO.h"
+#include "Materials/MaterialExpression.h"
+#include "MaterialExpressionSRGBColorToWorkingColorSpace.generated.h"
+
+UCLASS(MinimalAPI, collapsecategories, hidecategories=Object)
+class UMaterialExpressionSRGBColorToWorkingColorSpace : public UMaterialExpression
+{
+	GENERATED_UCLASS_BODY()
+
+	UPROPERTY()
+	FExpressionInput Input;
+
+	//~ Begin UMaterialExpression Interface
+#if WITH_EDITOR
+	virtual int32 Compile(class FMaterialCompiler* Compiler, int32 OutputIndex) override;
+	virtual void GetCaption(TArray<FString>& OutCaptions) const override;
+	virtual FText GetKeywords() const override {return FText::FromString(TEXT("srgb workingcolorspace"));}
+	virtual bool GenerateHLSLExpression(FMaterialHLSLGenerator& Generator, UE::HLSLTree::FScope& Scope, int32 OutputIndex, UE::HLSLTree::FExpression const*& OutExpression) const override;
+#endif
+	//~ End UMaterialExpression Interface
+};
+
+
+
diff --git a/Engine/Source/Runtime/Engine/Private/Materials/MaterialExpressionHLSL.cpp b/Engine/Source/Runtime/Engine/Private/Materials/MaterialExpressionHLSL.cpp
--- a/Engine/Source/Runtime/Engine/Private/Materials/MaterialExpressionHLSL.cpp
+++ b/Engine/Source/Runtime/Engine/Private/Materials/MaterialExpressionHLSL.cpp
@@ -197,6 +197,7 @@
 #include "Materials/MaterialExpressionSpeedTree.h"
 #include "Materials/MaterialExpressionSphereMask.h"
 #include "Materials/MaterialExpressionSquareRoot.h"
+#include "Materials/MaterialExpressionSRGBColorToWorkingColorSpace.h"
 #include "Materials/MaterialExpressionStaticBool.h"
 #include "Materials/MaterialExpressionStaticComponentMaskParameter.h"
 #include "Materials/MaterialExpressionStaticSwitch.h"
@@ -238,6 +239,7 @@
 #include "MaterialShared.h"
 #include "Misc/MemStackUtility.h"
 #include "RenderUtils.h"
+#include "ColorSpace.h"
 
 bool UMaterialExpression::GenerateHLSLExpression(FMaterialHLSLGenerator& Generator, UE::HLSLTree::FScope& Scope, int32 OutputIndex, UE::HLSLTree::FExpression const*& OutExpression) const
 {
@@ -2378,6 +2380,36 @@ bool UMaterialExpressionSquareRoot::GenerateHLSLExpression(FMaterialHLSLGenerato
 	return true;
 }
 
+bool UMaterialExpressionSRGBColorToWorkingColorSpace::GenerateHLSLExpression(FMaterialHLSLGenerator& Generator, UE::HLSLTree::FScope& Scope, int32 OutputIndex, UE::HLSLTree::FExpression const*& OutExpression) const
+{
+	using namespace UE::HLSLTree;
+	const FExpression* InputExpression = Input.AcquireHLSLExpression(Generator, Scope);
+	if (!InputExpression)
+	{
+		return false;
+	}
+
+	if (!UE::Color::FColorSpace::GetWorking().IsSRGB())
+	{
+		UE::HLSLTree::FTree& Tree = Generator.GetTree();
+		const UE::Color::FColorSpaceTransform& Transform = UE::Color::FColorSpaceTransform::GetSRGBToWorkingColorSpace();
+
+		// TODO: Replace by matrix or dot operations if possible. Currently there is no way to create a matrix constant or apply a dot product, resulting in higher instruction counts.
+		const FExpression* RRR = Tree.NewBinaryOp(EOperation::Mul, InputExpression, Generator.NewConstant(UE::Shader::FValue((float)Transform.M[0][0], (float)Transform.M[1][0], (float)Transform.M[2][0])));
+		const FExpression* GGG = Tree.NewBinaryOp(EOperation::Mul, InputExpression, Generator.NewConstant(UE::Shader::FValue((float)Transform.M[0][1], (float)Transform.M[1][1], (float)Transform.M[2][1])));
+		const FExpression* BBB = Tree.NewBinaryOp(EOperation::Mul, InputExpression, Generator.NewConstant(UE::Shader::FValue((float)Transform.M[0][2], (float)Transform.M[1][2], (float)Transform.M[2][2])));
+		const FExpression* R = Tree.NewUnaryOp(EOperation::Sum, RRR);
+		const FExpression* G = Tree.NewUnaryOp(EOperation::Sum, GGG);
+		const FExpression* B = Tree.NewUnaryOp(EOperation::Sum, BBB);
+		const FExpression* ColorExpression = Tree.NewExpression<FExpressionAppend>(Tree.NewExpression<FExpressionAppend>(R, G), B);
+
+		// We preserve the original alpha if possible
+		OutExpression = Tree.NewExpression<FExpressionAppend>(ColorExpression, Generator.NewSwizzle(FSwizzleParameters(3), InputExpression));
+	}
+
+	return true;
+}
+
 bool UMaterialExpressionExponential::GenerateHLSLExpression(FMaterialHLSLGenerator& Generator, UE::HLSLTree::FScope& Scope, int32 OutputIndex, UE::HLSLTree::FExpression const*& OutExpression) const
 {
 	using namespace UE::HLSLTree;
diff --git a/Engine/Source/Runtime/Engine/Private/Materials/MaterialExpressions.cpp b/Engine/Source/Runtime/Engine/Private/Materials/MaterialExpressions.cpp
--- a/Engine/Source/Runtime/Engine/Private/Materials/MaterialExpressions.cpp
+++ b/Engine/Source/Runtime/Engine/Private/Materials/MaterialExpressions.cpp
@@ -41,6 +41,7 @@
 #include "VT/RuntimeVirtualTexture.h"
 #include "SparseVolumeTexture/SparseVolumeTexture.h"
 #include "ProfilingDebugging/LoadTimeTracker.h"
+#include "ColorSpace.h"
 
 #include "Materials/MaterialExpressionAbs.h"
 #include "Materials/MaterialExpressionAbsorptionMediumMaterialOutput.h"
@@ -214,6 +215,7 @@
 #include "Materials/MaterialExpressionSphereMask.h"
 #include "Materials/MaterialExpressionSphericalParticleOpacity.h"
 #include "Materials/MaterialExpressionSquareRoot.h"
+#include "Materials/MaterialExpressionSRGBColorToWorkingColorSpace.h"
 #include "Materials/MaterialExpressionStaticBool.h"
 #include "Materials/MaterialExpressionStaticSwitch.h"
 #include "Materials/MaterialExpressionStep.h"
@@ -11032,6 +11034,62 @@ void UMaterialExpressionSquareRoot::GetCaption(TArray<FString>& OutCaptions) con
 }
 #endif // WITH_EDITOR
 
+UMaterialExpressionSRGBColorToWorkingColorSpace::UMaterialExpressionSRGBColorToWorkingColorSpace(const FObjectInitializer& ObjectInitializer)
+	: Super(ObjectInitializer)
+{
+#if WITH_EDITORONLY_DATA
+	// Structure to hold one-time initialization
+	struct FConstructorStatics
+	{
+		FText NAME_Color;
+		FConstructorStatics()
+			: NAME_Color(LOCTEXT("Color", "Color"))
+		{
+		}
+	};
+	static FConstructorStatics ConstructorStatics;
+
+	MenuCategories.Add(ConstructorStatics.NAME_Color);
+#endif
+}
+
+#if WITH_EDITOR
+int32 UMaterialExpressionSRGBColorToWorkingColorSpace::Compile(class FMaterialCompiler* Compiler, int32 OutputIndex)
+{
+	if (!Input.GetTracedInput().Expression)
+	{
+		return Compiler->Errorf(TEXT("Missing sRGBColorToWorkingColorSpace input"));
+	}
+
+	int32 Color = Input.Compile(Compiler);
+	int32 Result = Color;
+
+	if (!UE::Color::FColorSpace::GetWorking().IsSRGB())
+	{
+		const UE::Color::FColorSpaceTransform& Transform = UE::Color::FColorSpaceTransform::GetSRGBToWorkingColorSpace();
+
+		const int32 R = Compiler->Dot(Color, Compiler->Constant3((float)Transform.M[0][0], (float)Transform.M[1][0], (float)Transform.M[2][0]));
+		const int32 G = Compiler->Dot(Color, Compiler->Constant3((float)Transform.M[0][1], (float)Transform.M[1][1], (float)Transform.M[2][1]));
+		const int32 B = Compiler->Dot(Color, Compiler->Constant3((float)Transform.M[0][2], (float)Transform.M[1][2], (float)Transform.M[2][2]));
+		Result = Compiler->AppendVector(Compiler->AppendVector(R, G), B);
+
+		EMaterialValueType VectorType = Compiler->GetParameterType(Color);
+		if (VectorType & MCT_Float4 || VectorType == MCT_LWCVector4)
+		{
+			// We preserve the original alpha when applicable
+			Result = Compiler->AppendVector(Result, Compiler->ComponentMask(Color, false, false, false, true));
+		}
+	}
+
+	return Result;
+}
+
+void UMaterialExpressionSRGBColorToWorkingColorSpace::GetCaption(TArray<FString>& OutCaptions) const
+{
+	OutCaptions.Add(TEXT("sRGBColorToWorkingColorSpace"));
+}
+#endif // WITH_EDITOR
+
 ///////////////////////////////////////////////////////////////////////////////
 // UMaterialExpressionPixelDepth
 ///////////////////////////////////////////////////////////////////////////////
diff --git a/Engine/Source/Runtime/Engine/Private/ShaderCompiler/ShaderCompiler.cpp b/Engine/Source/Runtime/Engine/Private/ShaderCompiler/ShaderCompiler.cpp
--- a/Engine/Source/Runtime/Engine/Private/ShaderCompiler/ShaderCompiler.cpp
+++ b/Engine/Source/Runtime/Engine/Private/ShaderCompiler/ShaderCompiler.cpp
@@ -8288,7 +8288,7 @@ void GlobalBeginCompileShader(
 					FromXYZ.M[0][1], FromXYZ.M[1][1], FromXYZ.M[2][1],
 					FromXYZ.M[0][2], FromXYZ.M[1][2], FromXYZ.M[2][2]));
 
-			const FColorSpaceTransform FromSRGB(FColorSpace(EColorSpace::sRGB), WorkingColorSpace);
+			const FColorSpaceTransform& FromSRGB = FColorSpaceTransform::GetSRGBToWorkingColorSpace();
 			SET_SHADER_DEFINE(Input.Environment,
 				SRGB_TO_WORKING_COLOR_SPACE_MAT,
 				FString::Printf(MatrixFormat,
diff --git a/Engine/Source/Runtime/Engine/Public/Rendering/SkyAtmosphereCommonData.cpp b/Engine/Source/Runtime/Engine/Public/Rendering/SkyAtmosphereCommonData.cpp
--- a/Engine/Source/Runtime/Engine/Public/Rendering/SkyAtmosphereCommonData.cpp
+++ b/Engine/Source/Runtime/Engine/Public/Rendering/SkyAtmosphereCommonData.cpp
@@ -63,7 +63,7 @@ FAtmosphereSetup::FAtmosphereSetup(const USkyAtmosphereComponent& SkyAtmosphereC
 				FMath::Exp(-CoeffSRGB.B));
 
 			// Convert transmittance color from sRGB to working color space.
-			Transmittance = FColorSpaceTransform(FColorSpace(EColorSpace::sRGB), WorkingColorSpace).Apply(Transmittance);
+			Transmittance = FColorSpaceTransform::GetSRGBToWorkingColorSpace().Apply(Transmittance);
 
 			// New we have a transmittance in working color space, convert it back to coefficients for this working color space.
 			return FLinearColor(
