Fix warning V1047: Lifetime of the lambda is greater than lifetime of the local variable

diff --git a/Engine/Source/Runtime/Engine/Private/HierarchicalInstancedStaticMesh.cpp b/Engine/Source/Runtime/Engine/Private/HierarchicalInstancedStaticMesh.cpp
--- a/Engine/Source/Runtime/Engine/Private/HierarchicalInstancedStaticMesh.cpp
+++ b/Engine/Source/Runtime/Engine/Private/HierarchicalInstancedStaticMesh.cpp
@@ -2596,7 +2596,7 @@ void UHierarchicalInstancedStaticMeshComponent::BuildComponentInstanceData(FInst
 	OutData.NumSourceInstances = PerInstanceSMData.Num();
 	OutData.NumCustomDataFloats = NumCustomDataFloats;
 
-	OutData.BuildChangeSet = [&](FISMInstanceUpdateChangeSet &ChangeSet)
+	OutData.BuildChangeSet = [this](FISMInstanceUpdateChangeSet &ChangeSet)
 	{
 		BuildInstanceDataDeltaChangeSetCommon(ChangeSet);
 		ChangeSet.SetInstanceTransforms(MakeStridedView(PerInstanceSMData, &FInstancedStaticMeshInstanceData::Transform), -TranslatedInstanceSpaceOrigin);
diff --git a/Engine/Source/Runtime/Engine/Private/InstancedStaticMesh.cpp b/Engine/Source/Runtime/Engine/Private/InstancedStaticMesh.cpp
--- a/Engine/Source/Runtime/Engine/Private/InstancedStaticMesh.cpp
+++ b/Engine/Source/Runtime/Engine/Private/InstancedStaticMesh.cpp
@@ -2419,7 +2419,7 @@ void UInstancedStaticMeshComponent::BuildComponentInstanceData(FInstanceUpdateCo
 	OutData.NumCustomDataFloats = NumCustomDataFloats;
 
 	// Function that only gets called if we actually need to flush any changes.
-	OutData.BuildChangeSet = [&](FISMInstanceUpdateChangeSet &ChangeSet)
+	OutData.BuildChangeSet = [this](FISMInstanceUpdateChangeSet &ChangeSet)
 	{
 		BuildInstanceDataDeltaChangeSetCommon(ChangeSet);
 
