Fixed synchronous read back functions for cube textures : added GetShaderResourceTexture() to FRenderTarget that returns the RT's RHI texture resource by default but can be overridden to return another one, for classes where the shader resource RHI may differ from the render target resource RHI (FTextureRenderTargetResource and its child classes : e.g. FTextureRenderTargetCubeResource)

#tests editor

diff --git a/Engine/Source/Runtime/Engine/Private/TextureRenderTarget.cpp b/Engine/Source/Runtime/Engine/Private/TextureRenderTarget.cpp
--- a/Engine/Source/Runtime/Engine/Private/TextureRenderTarget.cpp
+++ b/Engine/Source/Runtime/Engine/Private/TextureRenderTarget.cpp
@@ -422,6 +422,14 @@ bool FTextureRenderTargetResource::IsSupportedFormat( EPixelFormat Format )
 	}
 }
 
+const FTextureRHIRef& FTextureRenderTargetResource::GetShaderResourceTexture() const
+{
+	// Override GetShaderResourceTexture() because FTextureRenderTargetResource is both a FRenderTarget (RenderTargetTextureRHI) and a FTexture (TextureRHI)
+	//  but in some implementations (e.g. cubemaps), those are different. We need to return the one that's used as a shader resource here :
+	return TextureRHI;
+}
+
+
 /** 
 * Render target resource should be sampled in linear color space
 *
diff --git a/Engine/Source/Runtime/Engine/Private/TextureRenderTargetVolume.cpp b/Engine/Source/Runtime/Engine/Private/TextureRenderTargetVolume.cpp
--- a/Engine/Source/Runtime/Engine/Private/TextureRenderTargetVolume.cpp
+++ b/Engine/Source/Runtime/Engine/Private/TextureRenderTargetVolume.cpp
@@ -241,9 +241,6 @@ void FTextureRenderTargetVolumeResource::InitRHI(FRHICommandListBase& RHICmdList
 
 		RHIUpdateTextureReference(Owner->TextureReference.TextureReferenceRHI,TextureRHI);
 
-		// Can't set this as it's a texture 2D
-		//RenderTargetTextureRHI = VolumeSurfaceRHI;
-
 		AddToDeferredUpdateList(true);
 	}
 
diff --git a/Engine/Source/Runtime/Engine/Private/UnrealClient.cpp b/Engine/Source/Runtime/Engine/Private/UnrealClient.cpp
--- a/Engine/Source/Runtime/Engine/Private/UnrealClient.cpp
+++ b/Engine/Source/Runtime/Engine/Private/UnrealClient.cpp
@@ -62,7 +62,7 @@ bool FRenderTarget::ReadPixels(TArray< FColor >& OutImageData, FReadSurfaceDataF
 	ENQUEUE_RENDER_COMMAND(ReadSurfaceCommand)(
 		[RenderTarget_RT = this, SrcRect_RT = InSrcRect, OutData_RT = &OutImageData, Flags_RT = InFlags](FRHICommandListImmediate& RHICmdList)
 		{
-			RHICmdList.ReadSurfaceData(RenderTarget_RT->GetRenderTargetTexture(), SrcRect_RT, *OutData_RT, Flags_RT);
+			RHICmdList.ReadSurfaceData(RenderTarget_RT->GetShaderResourceTexture(), SrcRect_RT, *OutData_RT, Flags_RT);
 		});
 	FlushRenderingCommands();
 
@@ -102,11 +102,11 @@ bool FRenderTarget::ReadFloat16Pixels(TArray<FFloat16Color>& OutImageData, FRead
 	ENQUEUE_RENDER_COMMAND(ReadSurfaceFloatCommand)(
 		[RenderTarget_RT = this, SrcRect_RT = InSrcRect, OutData_RT = &OutImageData, Flags_RT = InFlags](FRHICommandListImmediate& RHICmdList)
 	{
-		RHICmdList.ReadSurfaceFloatData(RenderTarget_RT->GetRenderTargetTexture(), SrcRect_RT, *OutData_RT, Flags_RT);
+		RHICmdList.ReadSurfaceFloatData(RenderTarget_RT->GetShaderResourceTexture(), SrcRect_RT, *OutData_RT, Flags_RT);
 	});
 	FlushRenderingCommands();
 
-	return OutImageData.Num() > 0;;
+	return OutImageData.Num() > 0;
 }
 
 bool FRenderTarget::ReadLinearColorPixels(TArray<FLinearColor> &OutImageData, FReadSurfaceDataFlags InFlags, FIntRect InSrcRect)
@@ -122,7 +122,7 @@ bool FRenderTarget::ReadLinearColorPixels(TArray<FLinearColor> &OutImageData, FR
 	ENQUEUE_RENDER_COMMAND(ReadSurfaceCommand)(
 		[RenderTarget_RT = this, SrcRect_RT = InSrcRect, OutData_RT = &OutImageData, Flags_RT = InFlags](FRHICommandListImmediate& RHICmdList)
 		{
-			RHICmdList.ReadSurfaceData(RenderTarget_RT->GetRenderTargetTexture(), SrcRect_RT, *OutData_RT, Flags_RT);
+			RHICmdList.ReadSurfaceData(RenderTarget_RT->GetShaderResourceTexture(), SrcRect_RT, *OutData_RT, Flags_RT);
 		});
 	FlushRenderingCommands();
 
@@ -143,9 +143,6 @@ bool FRenderTarget::ReadLinearColorPixelsPtr(FLinearColor* OutImageBytes, FReadS
 	return bResult;
 }
 
-/** 
-* @return display gamma expected for rendering to this render target 
-*/
 float FRenderTarget::GetDisplayGamma() const
 {
 	if (GEngine == NULL)
@@ -163,10 +160,6 @@ float FRenderTarget::GetDisplayGamma() const
 	}
 }
 
-/**
-* Accessor for the surface RHI when setting this render target
-* @return render target surface RHI resource
-*/
 const FTextureRHIRef& FRenderTarget::GetRenderTargetTexture() const
 {
 	return RenderTargetTextureRHI;
@@ -182,6 +175,11 @@ FUnorderedAccessViewRHIRef FRenderTarget::GetRenderTargetUAV() const
 	return FUnorderedAccessViewRHIRef();
 }
 
+const FTextureRHIRef& FRenderTarget::GetShaderResourceTexture() const
+{
+	return RenderTargetTextureRHI;
+}
+
 void FScreenshotRequest::RequestScreenshot(bool bInShowUI)
 {
 	// empty string means we'll later pick the name
diff --git a/Engine/Source/Runtime/Engine/Public/TextureResource.h b/Engine/Source/Runtime/Engine/Public/TextureResource.h
--- a/Engine/Source/Runtime/Engine/Public/TextureResource.h
+++ b/Engine/Source/Runtime/Engine/Public/TextureResource.h
@@ -362,6 +362,9 @@ public:
 	 */
 	ENGINE_API static bool IsSupportedFormat( EPixelFormat Format );
 
+	// FRenderTarget implementation
+	virtual const FTextureRHIRef& GetShaderResourceTexture() const override;
+
 	// FTextureRenderTargetResource interface
 	
 	virtual class FTextureRenderTarget2DResource* GetTextureRenderTarget2DResource() { return nullptr; }
diff --git a/Engine/Source/Runtime/Engine/Public/UnrealClient.h b/Engine/Source/Runtime/Engine/Public/UnrealClient.h
--- a/Engine/Source/Runtime/Engine/Public/UnrealClient.h
+++ b/Engine/Source/Runtime/Engine/Public/UnrealClient.h
@@ -59,6 +59,12 @@ public:
 	 */
 	ENGINE_API virtual FRDGTextureRef GetRenderTargetTexture(FRDGBuilder& GraphBuilder) const;
 
+	/**
+	* Accessor for the surface RHI to use when reading the reading target (may differ from GetRenderTargetTexture() for some implementations like cubemaps)
+	* @return surface RHI resource
+	*/
+	ENGINE_API virtual const FTextureRHIRef& GetShaderResourceTexture() const;
+
 	// Properties.
 	virtual FIntPoint GetSizeXY() const = 0;
 
