Fix warning V674: The 'bUseTextureAsInput?-1.0f:0.0f' literal of the 'float' type is assigned to a variable of the 'int' type. Consider inspecting the '= bUseTextureAsInput ? - 1.0f : 0.0f' expression.
#rnx

diff --git a/Engine/Source/Runtime/Engine/Private/Materials/MaterialExpressionHLSL.cpp b/Engine/Source/Runtime/Engine/Private/Materials/MaterialExpressionHLSL.cpp
--- a/Engine/Source/Runtime/Engine/Private/Materials/MaterialExpressionHLSL.cpp
+++ b/Engine/Source/Runtime/Engine/Private/Materials/MaterialExpressionHLSL.cpp
@@ -4894,7 +4894,7 @@ bool UMaterialExpressionNeuralNetworkInput::GenerateHLSLExpression(FMaterialHLSL
 		else
 		{
 			const FExpression* ViewportUV = Tree.NewExpression<Material::FExpressionExternalInput>(Material::EExternalInput::ViewportUV);
-			int32 BatchIndex = bUseTextureAsInput ? -1.0f : 0.0f;
+			float BatchIndex = bUseTextureAsInput ? -1.0f : 0.0f;
 			OutExpression = Tree.NewAppend(Generator.GetTree().NewConstant(FVector2f(BatchIndex, 0.0f)), ViewportUV);
 		}
 	}
