Fix warning V502: Perhaps the '?:' operator works in a different way than it was expected. The '?:' operator has a lower priority than the '!=' operator.
#rnx

diff --git a/Engine/Source/Runtime/Engine/Private/DeviceProfiles/DeviceProfileManager.cpp b/Engine/Source/Runtime/Engine/Private/DeviceProfiles/DeviceProfileManager.cpp
--- a/Engine/Source/Runtime/Engine/Private/DeviceProfiles/DeviceProfileManager.cpp
+++ b/Engine/Source/Runtime/Engine/Private/DeviceProfiles/DeviceProfileManager.cpp
@@ -862,7 +862,7 @@ UDeviceProfile* UDeviceProfileManager::FindProfile(const FString& ProfileName, b
 
 	if ( bCreateProfileOnFail && FoundProfile == nullptr )
 	{
-		FString PlatformName = OptionalPlatformName != NAME_None ? OptionalPlatformName.ToString() : FString(FPlatformProperties::IniPlatformName());
+		FString PlatformName = (OptionalPlatformName != NAME_None) ? OptionalPlatformName.ToString() : FString(FPlatformProperties::IniPlatformName());
 		FoundProfile = CreateProfile(ProfileName, PlatformName);
 	}
 	return FoundProfile;
