UWorldPartitionRuntimeHashSet:
- Fixed uinitialized member FRuntimePartitionHLODSetup::bIsSpatiallyLoaded.
- Fixed bad loading range when converting from UWorldPartitionRuntimeSpatialHash.

diff --git a/Engine/Source/Runtime/Engine/Public/WorldPartition/RuntimeHashSet/WorldPartitionRuntimeHashSet.h b/Engine/Source/Runtime/Engine/Public/WorldPartition/RuntimeHashSet/WorldPartitionRuntimeHashSet.h
--- a/Engine/Source/Runtime/Engine/Public/WorldPartition/RuntimeHashSet/WorldPartitionRuntimeHashSet.h
+++ b/Engine/Source/Runtime/Engine/Public/WorldPartition/RuntimeHashSet/WorldPartitionRuntimeHashSet.h
@@ -29,7 +29,7 @@ struct FRuntimePartitionHLODSetup
 
 	/** whether this HLOD setup is spatially loaded or not */
 	UPROPERTY(EditAnywhere, Category = RuntimeSettings)
-	bool bIsSpatiallyLoaded;
+	bool bIsSpatiallyLoaded = true;
 
 	UPROPERTY(VisibleAnywhere, Category = RuntimeSettings, Instanced, Meta = (EditCondition = "bIsSpatiallyLoaded", HideEditConditionToggle, NoResetToDefault, TitleProperty = "Name"))
 	TObjectPtr<URuntimePartition> PartitionLayer;
