Fix warning V595: The pointer was utilized before it was verified against nullptr.

diff --git a/Engine/Plugins/Experimental/PCG/Source/PCG/Private/PCGComponent.cpp b/Engine/Plugins/Experimental/PCG/Source/PCG/Private/PCGComponent.cpp
--- a/Engine/Plugins/Experimental/PCG/Source/PCG/Private/PCGComponent.cpp
+++ b/Engine/Plugins/Experimental/PCG/Source/PCG/Private/PCGComponent.cpp
@@ -656,10 +656,9 @@ AActor* UPCGComponent::ClearPCGLink(UClass* TemplateActor)
 		{
 			Cleanup(true);
 		}
-		else
+		else if (NewActor)
 		{
-			checkSlow(World);
-			World->DestroyActor(NewActor);
+			NewActor->Destroy();
 			NewActor = nullptr;
 		}
 	}
diff --git a/Engine/Plugins/Runtime/RigVM/Source/RigVMEditor/Private/Editor/RigVMEditor.cpp b/Engine/Plugins/Runtime/RigVM/Source/RigVMEditor/Private/Editor/RigVMEditor.cpp
--- a/Engine/Plugins/Runtime/RigVM/Source/RigVMEditor/Private/Editor/RigVMEditor.cpp
+++ b/Engine/Plugins/Runtime/RigVM/Source/RigVMEditor/Private/Editor/RigVMEditor.cpp
@@ -3645,7 +3645,7 @@ void FRigVMEditor::HandleSetObjectBeingDebugged(UObject* InObject)
 	if (DebuggedHost == nullptr)
 	{
 		// fall back to our default control rig (which still can be nullptr)
-		if (GetRigVMBlueprint() != nullptr && GetBlueprintObj() && !bIsSettingObjectBeingDebugged)
+		if (GetRigVMBlueprint() != nullptr && !bIsSettingObjectBeingDebugged)
 		{
 			TGuardValue<bool> GuardSettingObjectBeingDebugged(bIsSettingObjectBeingDebugged, true);
 			GetBlueprintObj()->SetObjectBeingDebugged(GetRigVMHost());
diff --git a/Engine/Source/Editor/LevelInstanceEditor/Private/PackedLevelActorUtils.cpp b/Engine/Source/Editor/LevelInstanceEditor/Private/PackedLevelActorUtils.cpp
--- a/Engine/Source/Editor/LevelInstanceEditor/Private/PackedLevelActorUtils.cpp
+++ b/Engine/Source/Editor/LevelInstanceEditor/Private/PackedLevelActorUtils.cpp
@@ -135,10 +135,7 @@ void FPackedLevelActorUtils::PackAllLoadedActors()
 		return;
 	}
 
-	if (GEditor)
-	{
-		GEditor->SelectNone(true, true);
-	}
+	GEditor->SelectNone(true, true);
 
 	FScopedSlowTask SlowTask(Count, (LOCTEXT("TaskPackLevels", "Packing Levels")));
 	SlowTask.MakeDialog();
diff --git a/Engine/Source/Runtime/Engine/Private/GameViewportClient.cpp b/Engine/Source/Runtime/Engine/Private/GameViewportClient.cpp
--- a/Engine/Source/Runtime/Engine/Private/GameViewportClient.cpp
+++ b/Engine/Source/Runtime/Engine/Private/GameViewportClient.cpp
@@ -1347,7 +1347,7 @@ void UGameViewportClient::Draw(FViewport* InViewport, FCanvas* SceneCanvas)
 
 	// Create temp debug canvas object
 	FIntPoint DebugCanvasSize = InViewport->GetSizeXY();
-	if (bStereoRendering && GEngine->XRSystem.IsValid() && GEngine->XRSystem->GetHMDDevice())
+	if (bStereoRendering && GEngine && GEngine->XRSystem.IsValid() && GEngine->XRSystem->GetHMDDevice())
 	{
 		DebugCanvasSize = GEngine->XRSystem->GetHMDDevice()->GetIdealDebugCanvasRenderTargetSize();
 	}
diff --git a/Engine/Source/Runtime/Engine/Private/LevelInstance/LevelInstanceSubsystem.cpp b/Engine/Source/Runtime/Engine/Private/LevelInstance/LevelInstanceSubsystem.cpp
--- a/Engine/Source/Runtime/Engine/Private/LevelInstance/LevelInstanceSubsystem.cpp
+++ b/Engine/Source/Runtime/Engine/Private/LevelInstance/LevelInstanceSubsystem.cpp
@@ -740,7 +740,7 @@ bool ULevelInstanceSubsystem::CheckForLoop(const ILevelInstanceInterface* LevelI
 				if (LoopInfo)
 				{
 					TSoftObjectPtr<UWorld> CurrentAsset = (CurrentLevelInstance == LevelInstance)? WorldAsset : (CurrentLevelInstance)? CurrentLevelInstance->GetWorldAsset() : nullptr;
-					FText LevelInstanceName = (CurrentActor)? FText::FromString(CurrentActor->GetPathName()) : FText::FromString("Invalid Actor cannot retrieve LevelInstance name.");
+					FText LevelInstanceName = FText::FromString(CurrentActor->GetPathName());
 					FText Description = FText::Format(LOCTEXT("LevelInstanceLoopLink", "-> Actor: {0} loads"), LevelInstanceName);
 					LoopInfo->Emplace(Description, CurrentAsset);
 				}
diff --git a/Engine/Source/Runtime/Engine/Private/LocalPlayer.cpp b/Engine/Source/Runtime/Engine/Private/LocalPlayer.cpp
--- a/Engine/Source/Runtime/Engine/Private/LocalPlayer.cpp
+++ b/Engine/Source/Runtime/Engine/Private/LocalPlayer.cpp
@@ -895,9 +895,9 @@ FSceneView* ULocalPlayer::CalcSceneView( class FSceneViewFamily* ViewFamily,
 		View->OverridePostProcessSettings(ViewInfo.PostProcessSettings, ViewInfo.PostProcessBlendWeight);
 
 		// Camera overrides (cameras blending in, camera modifiers, etc)
-		if (PlayerController->PlayerCameraManager &&
-				CameraAnimPPSettings && CameraAnimPPBlendWeights && CameraAnimPPBlendOrders)
+		if (PlayerController->PlayerCameraManager)
 		{
+			checkSlow(CameraAnimPPBlendWeights && CameraAnimPPBlendOrders && CameraAnimPPSettings);
 			for (int32 PPIdx = 0; PPIdx < CameraAnimPPBlendWeights->Num(); ++PPIdx)
 			{
 				if ((*CameraAnimPPBlendOrders)[PPIdx] == VTBlendOrder_Override)
