World Partition: detect replaced actors and update dirty list.

diff --git a/Engine/Source/Runtime/Engine/Classes/GameFramework/Actor.h b/Engine/Source/Runtime/Engine/Classes/GameFramework/Actor.h
--- a/Engine/Source/Runtime/Engine/Classes/GameFramework/Actor.h
+++ b/Engine/Source/Runtime/Engine/Classes/GameFramework/Actor.h
@@ -2182,7 +2182,7 @@ public:
 	/** Used to check if Actor is the main actor of a package (currently Child Actors are not) */
 	ENGINE_API bool IsMainPackageActor() const;
 
-	static ENGINE_API AActor* FindActorInPackage(UPackage* InPackage);
+	static ENGINE_API AActor* FindActorInPackage(UPackage* InPackage, bool bEvenIfPendingKill = true);
 
 #if WITH_EDITOR
 	ENGINE_API virtual bool Modify(bool bAlwaysMarkDirty = true) override;
diff --git a/Engine/Source/Runtime/Engine/Private/Actor.cpp b/Engine/Source/Runtime/Engine/Private/Actor.cpp
--- a/Engine/Source/Runtime/Engine/Private/Actor.cpp
+++ b/Engine/Source/Runtime/Engine/Private/Actor.cpp
@@ -1313,7 +1313,7 @@ bool AActor::IsMainPackageActor() const
 	return IsPackageExternal() && ParentComponent.IsExplicitlyNull();
 }
 
-AActor* AActor::FindActorInPackage(UPackage* InPackage)
+AActor* AActor::FindActorInPackage(UPackage* InPackage, bool bEvenIfPendingKill)
 {
 	AActor* Actor = nullptr;
 	ForEachObjectWithPackage(InPackage, [&Actor](UObject* Object)
@@ -1326,7 +1326,7 @@ AActor* AActor::FindActorInPackage(UPackage* InPackage)
 			}
 		}
 		return !Actor;
-	}, false);
+	}, false, bEvenIfPendingKill ? RF_NoFlags : (RF_InternalPendingKill | RF_InternalGarbage));
 	return Actor;
 }
 
diff --git a/Engine/Source/Runtime/Engine/Private/WorldPartition/WorldPartition.cpp b/Engine/Source/Runtime/Engine/Private/WorldPartition/WorldPartition.cpp
--- a/Engine/Source/Runtime/Engine/Private/WorldPartition/WorldPartition.cpp
+++ b/Engine/Source/Runtime/Engine/Private/WorldPartition/WorldPartition.cpp
@@ -1625,7 +1625,7 @@ void UWorldPartition::Tick(float DeltaSeconds)
 
 	for (TMap<FGuid, FDirtyActor>::TIterator DirtyActorIt(DirtyActors); DirtyActorIt; ++DirtyActorIt)
 	{
-		const FDirtyActor& DirtyActor = DirtyActorIt.Value();
+		FDirtyActor& DirtyActor = DirtyActorIt.Value();
 		if (DirtyActor.WorldPartitionRef.IsSet())
 		{
 			const FWorldPartitionReference& Ref = DirtyActor.WorldPartitionRef.GetValue();
@@ -1655,10 +1655,23 @@ void UWorldPartition::Tick(float DeltaSeconds)
 					DirtyActorIt.RemoveCurrent();
 				}
 			}
-			// In this case, we know that the DirtyActor is not in the transaction buffer anymore and this is fine removing it.
-			else if ((!DirtyActor.ActorPtr.IsValid(true)))
+			else
 			{
-				DirtyActorIt.RemoveCurrent();
+				if (DirtyActor.ActorPtr.IsValid(true))
+				{
+					// If the actor is pending kill, check if there's another version of it in the same package and replace the pointer with the new version
+					// (since UEditorEngine::ReplaceActors don't always broadcast FCoreUObjectDelegates::OnObjectsReplaced).
+					if (AActor* NewActor = AActor::FindActorInPackage(DirtyActor.ActorPtr.Get(true)->GetPackage(), false); NewActor && NewActor != DirtyActor.ActorPtr.Get(true))
+					{
+						DirtyActor.ActorPtr = NewActor;
+					}
+				}
+				else
+				{
+					// In this case, we know that the dirty actor is not in the transaction buffer anymore and this is fine removing it (actor was added, then
+					// deleted and gargage collection happened).
+					DirtyActorIt.RemoveCurrent();
+				}
 			}
 		}
 	}
