Update compatible skeleton mappings when a ref skeleton is updated.

diff --git a/Engine/Source/Runtime/Engine/Private/Animation/ReferenceSkeleton.cpp b/Engine/Source/Runtime/Engine/Private/Animation/ReferenceSkeleton.cpp
--- a/Engine/Source/Runtime/Engine/Private/Animation/ReferenceSkeleton.cpp
+++ b/Engine/Source/Runtime/Engine/Private/Animation/ReferenceSkeleton.cpp
@@ -2,6 +2,7 @@
 
 #include "ReferenceSkeleton.h"
 #include "Animation/Skeleton.h"
+#include "Animation/SkeletonRemappingRegistry.h"
 #include "EngineLogs.h"
 #include "Engine/SkeletalMesh.h"
 
@@ -500,6 +501,9 @@ void FReferenceSkeleton::RebuildRefSkeleton(const USkeleton* Skeleton, bool bReb
 			}
 		}
 	}
+
+	// Full rebuild of all compatible with this and with ones we are compatible with.
+	UE::Anim::FSkeletonRemappingRegistry::Get().RefreshMappings(Skeleton);
 }
 
 void FReferenceSkeleton::RemoveDuplicateBones(const UObject* Requester, TArray<FBoneIndexType> & DuplicateBones)
diff --git a/Engine/Source/Runtime/Engine/Private/Animation/SkeletonRemappingRegistry.cpp b/Engine/Source/Runtime/Engine/Private/Animation/SkeletonRemappingRegistry.cpp
--- a/Engine/Source/Runtime/Engine/Private/Animation/SkeletonRemappingRegistry.cpp
+++ b/Engine/Source/Runtime/Engine/Private/Animation/SkeletonRemappingRegistry.cpp
@@ -99,7 +99,7 @@ const FSkeletonRemapping& FSkeletonRemappingRegistry::GetRemapping(const USkelet
 	}
 }
 
-void FSkeletonRemappingRegistry::RefreshMappings(USkeleton* InSkeleton)
+void FSkeletonRemappingRegistry::RefreshMappings(const USkeleton* InSkeleton)
 {
 	TArray<TSharedPtr<FSkeletonRemapping>> ExistingMappings;
 	{
diff --git a/Engine/Source/Runtime/Engine/Public/Animation/SkeletonRemappingRegistry.h b/Engine/Source/Runtime/Engine/Public/Animation/SkeletonRemappingRegistry.h
--- a/Engine/Source/Runtime/Engine/Public/Animation/SkeletonRemappingRegistry.h
+++ b/Engine/Source/Runtime/Engine/Public/Animation/SkeletonRemappingRegistry.h
@@ -29,7 +29,7 @@ public:
 	ENGINE_API const FSkeletonRemapping& GetRemapping(const USkeleton* InSourceSkeleton, const USkeleton* InTargetSkeleton);
 
 	// Refresh any existing mappings that the supplied skeleton is used with
-	ENGINE_API void RefreshMappings(USkeleton* InSkeleton);
+	ENGINE_API void RefreshMappings(const USkeleton* InSkeleton);
 
 private:
 	friend FSkeletonRemappingRegistryPrivate;
