Add ability to steal memory from the stream cache via "memory counted features" that get added to the streaming cache.

This effectively "uses" streaming cache memory so that actual streaming audio is constrained to not go above the max usage.

Add ability to adjust the amount of memory the Pitch Shifter features use

#tests PIE

#localization none

diff --git a/Engine/Source/Runtime/Engine/Private/AudioStreamingCache.cpp b/Engine/Source/Runtime/Engine/Private/AudioStreamingCache.cpp
--- a/Engine/Source/Runtime/Engine/Private/AudioStreamingCache.cpp
+++ b/Engine/Source/Runtime/Engine/Private/AudioStreamingCache.cpp
@@ -439,7 +439,25 @@ void FCachedAudioStreamingManager::RemoveForceInlineSoundWave(const FSoundWavePr
 	// remove the sound wave from the first cache
 	if (ensure(CacheArray.Num() > 0))
 	{
-		CacheArray[0].RemoveForecInlineSoundWave(SoundWave);
+		CacheArray[0].RemoveForceInlineSoundWave(SoundWave);
+	}
+}
+
+void FCachedAudioStreamingManager::AddMemoryCountedFeature(const FAudioStreamingMemoryCountedFeature& Feature)
+{
+	// add memory count to the first cache
+	if (ensure(CacheArray.Num() > 0))
+	{
+		CacheArray[0].AddMemoryCountedFeature(Feature);
+	}
+}
+
+void FCachedAudioStreamingManager::RemoveMemoryCountedFeature(const FAudioStreamingMemoryCountedFeature& Feature)
+{
+	// remove memory count from the first cache
+	if (ensure(CacheArray.Num() > 0))
+	{
+		CacheArray[0].RemoveMemoryCountedFeature(Feature);
 	}
 }
 
@@ -720,6 +738,7 @@ FAudioChunkCache::FAudioChunkCache(uint32 InMaxChunkSize, uint32 NumChunks, uint
 	, MemoryCounterBytes(0)
 	, MemoryLimitBytes(InMemoryLimitInBytes)
 	, ForceInlineMemoryCounterBytes(0)
+	, FeatureMemoryCounterBytes(0)
 	, bLogCacheMisses(false)
 {
 	check(NumChunks > 0);
@@ -797,7 +816,7 @@ uint64 FAudioChunkCache::AddOrTouchChunk(const FChunkKey& InKey, const TSharedPt
 		{
 			int32 ChunkDataSize = Chunk->AudioDataSize;
 
-			const uint64 MemoryUsageBytes = MemoryCounterBytes + ForceInlineMemoryCounterBytes + ChunkDataSize;
+			const uint64 MemoryUsageBytes = GetCurrentMemoryUsageBytes() + ChunkDataSize;
 			if (TrimCacheWhenOverBudgetCVar != 0 && MemoryUsageBytes > MemoryLimitBytes)
 			{
 				uint64 MemoryToTrim = 0;
@@ -1006,7 +1025,7 @@ void FAudioChunkCache::AddForceInlineSoundWave(const FSoundWaveProxyPtr& SoundWa
 	FByteBulkData* Data = SoundWave->GetCompressedData(Format);
 	ForceInlineMemoryCounterBytes += Data ? Data->GetBulkDataSize() : 0;
 
-	const uint64 MemoryUsageBytes = MemoryCounterBytes + ForceInlineMemoryCounterBytes;
+	const uint64 MemoryUsageBytes = GetCurrentMemoryUsageBytes();
 	if (TrimCacheWhenOverBudgetCVar != 0 && MemoryUsageBytes > MemoryLimitBytes)
 	{
 		uint64 MemoryToTrim = 0;
@@ -1023,7 +1042,7 @@ void FAudioChunkCache::AddForceInlineSoundWave(const FSoundWaveProxyPtr& SoundWa
 	}
 }
 
-void FAudioChunkCache::RemoveForecInlineSoundWave(const FSoundWaveProxyPtr& SoundWave)
+void FAudioChunkCache::RemoveForceInlineSoundWave(const FSoundWaveProxyPtr& SoundWave)
 {
 	check(SoundWave.IsValid());
 	ensureMsgf(SoundWave->GetLoadingBehavior() == ESoundWaveLoadingBehavior::ForceInline,
@@ -1039,6 +1058,36 @@ void FAudioChunkCache::RemoveForecInlineSoundWave(const FSoundWaveProxyPtr& Soun
 	ForceInlineMemoryCounterBytes -= Data ? Data->GetBulkDataSize() : 0;
 }
 
+void FAudioChunkCache::AddMemoryCountedFeature(const FAudioStreamingMemoryCountedFeature& Feature)
+{
+	UE_LOG(LogAudioStreamCaching, Log, TEXT("Adding Memory Counted Feature (%s) Memory Usage: %d bytes"), *Feature.GetFeatureName().ToString(), (int32)Feature.GetMemoryUseInBytes());
+	const uint64 OldMemoryCount = FeatureMemoryCounterBytes.AddExchange(Feature.GetMemoryUseInBytes());
+	UE_LOG(LogAudioStreamCaching, Log, TEXT("Total Memory Usage for all features: %d -> %d bytes"), (int32)OldMemoryCount, (int32)FeatureMemoryCounterBytes.Load());
+	const uint64 MemoryUsageBytes = GetCurrentMemoryUsageBytes();
+	if (TrimCacheWhenOverBudgetCVar != 0 && MemoryUsageBytes > MemoryLimitBytes)
+	{
+		uint64 MemoryToTrim = 0;
+		if (MemoryLimitTrimPercentageCVar > 0.0f)
+		{
+			MemoryToTrim = MemoryLimitBytes * FMath::Min(MemoryLimitTrimPercentageCVar, 1.0f);
+		}
+		else
+		{
+			MemoryToTrim = MemoryUsageBytes - MemoryLimitBytes;
+		}
+
+		TrimMemory(MemoryToTrim, true);
+	}
+}
+
+void FAudioChunkCache::RemoveMemoryCountedFeature(const FAudioStreamingMemoryCountedFeature& Feature)
+{
+	UE_LOG(LogAudioStreamCaching, Log, TEXT("Removing Memory Counted Feature (%s) Memory Usage: %d"), *Feature.GetFeatureName().ToString(), (int32)Feature.GetMemoryUseInBytes());
+	check(FeatureMemoryCounterBytes.Load() <= Feature.GetMemoryUseInBytes());
+	const uint32 OldMemoryCount = FeatureMemoryCounterBytes.SubExchange(Feature.GetMemoryUseInBytes());
+	UE_LOG(LogAudioStreamCaching, Log, TEXT("Total Memory Usage for all features: %d -> %d"), (int32)OldMemoryCount, (int32)FeatureMemoryCounterBytes.Load());
+}
+
 uint64 FAudioChunkCache::TrimMemory(uint64 BytesToFree, bool bInAllowRetainedChunkTrimming)
 {
 	FScopeLock ScopeLock(&CacheMutationCriticalSection);
@@ -2187,18 +2236,23 @@ FString FAudioChunkCache::DebugPrint()
 	// Num bytes in use should include Force Inline data!
 	NumBytesCounter += ForceInlineMemoryCounterBytes;
 
+	// Num bytes should include feature data!
+	NumBytesCounter += FeatureMemoryCounterBytes;
+
 	// Convert to megabytes and print the total size:
 	const double NumMegabytesInUse = (double)NumBytesCounter / (1024 * 1024);
 	const double NumMegabytesForceInline = (double)ForceInlineMemoryCounterBytes / (1024 * 1024);
+	const double NumMegabytesExternalFeatures = (double)FeatureMemoryCounterBytes / (1024 * 1024);
 	const double NumMegabytesRetained = (double)NumBytesRetained / (1024 * 1024);
 
 	const double MaxCacheSizeMB = ((double)MemoryLimitBytes) / (1024 * 1024);
 	const double PercentageOfCacheRetained = NumMegabytesRetained / MaxCacheSizeMB;
 	const double PercentageOfCacheForceInlined = NumMegabytesForceInline / MaxCacheSizeMB;
+	const double PercentageOfCacheExternalFeatures = NumMegabytesExternalFeatures / MaxCacheSizeMB;
 
-	FString CacheMemoryHeader = *FString::Printf(TEXT("Force Inline:\t, Retaining:\t, Loaded:\t, Max Potential Usage:\t, \n"));
-	FString CacheMemoryUsage = *FString::Printf(TEXT("%.4f Megabytes (%.3f of total capacity)\t %.4f Megabytes (%.3f of total capacity)\t,  %.4f Megabytes (%lu bytes)\t, %.4f Megabytes\t, \n"), 
-		NumMegabytesForceInline, PercentageOfCacheForceInlined, NumMegabytesRetained, PercentageOfCacheRetained, NumMegabytesInUse, MemoryCounterBytes.Load(), MaxCacheSizeMB);
+	FString CacheMemoryHeader = *FString::Printf(TEXT("External Features:\t, Force Inline:\t, Retaining:\t, Loaded:\t, Max Potential Usage:\t, \n"));
+	FString CacheMemoryUsage = *FString::Printf(TEXT("%.4f Megabytes (%.3f%% of total capacity)\t %.4f Megabytes (%.3f%% of total capacity)\t %.4f Megabytes (%.3f%% of total capacity)\t,  %.4f Megabytes (%lu bytes)\t, %.4f Megabytes\t, \n"), 
+		NumMegabytesExternalFeatures, PercentageOfCacheExternalFeatures, NumMegabytesForceInline, PercentageOfCacheForceInlined, NumMegabytesRetained, PercentageOfCacheRetained, NumMegabytesInUse, MemoryCounterBytes.Load(), MaxCacheSizeMB);
 	OutputString += CacheMemoryHeader + CacheMemoryUsage + TEXT("\n");
 
 	// Second Pass: We're going to list the actual chunks in the cache.
@@ -2290,6 +2344,7 @@ static const FLinearColor ColorTrimmed = FLinearColor::Red;
 static const FLinearColor ColorCacheMiss = ColorLOD;
 static const FLinearColor ColorOther = FLinearColor::Gray;
 static const FLinearColor ColorForceInline(255 / ColorMax, 0, 255 / ColorMax); // Magenta
+static const FLinearColor ColorExternalFeatures(255 / ColorMax, 100 / ColorMax, 0x00 / ColorMax); // Orange
 
 
 
@@ -2417,7 +2472,9 @@ TPair<int, int> FAudioChunkCache::DebugDisplay(UWorld* World, FViewport* Viewpor
 	}
 
 	uint64 ForceInlineBytes = ForceInlineMemoryCounterBytes.Load();
+	uint64 ExternalFeaturesBytes = FeatureMemoryCounterBytes.Load();
 	NumBytesCounter += ForceInlineBytes;
+	NumBytesCounter += ExternalFeaturesBytes;
 	// Convert to megabytes and print the total size:
 	const double NumMegabytesInUse = (double)NumBytesCounter / (1024 * 1024);
 	const double MaxCacheSizeMB = ((double)MemoryLimitBytes) / (1024 * 1024);
@@ -2426,6 +2483,9 @@ TPair<int, int> FAudioChunkCache::DebugDisplay(UWorld* World, FViewport* Viewpor
 	const double NumMegabytesForceInline = (double)ForceInlineBytes / (1024 * 1024);
 	const float PercentageForceInline = NumBytesCounter > 0 ? (double)ForceInlineBytes / NumBytesCounter : 0;
 
+	const double NumMegabytesExternalFeatures = (double)ExternalFeaturesBytes / (1024 * 1024);
+	const float PercentageExternalFeatures = NumBytesCounter > 0 ? (double)ExternalFeaturesBytes / NumBytesCounter : 0;
+
 	FString CacheMemoryUsage = *FString::Printf(TEXT("Using: %.4f Megabytes (%lu bytes). Max Potential Usage: %.4f Megabytes."), 
 		NumMegabytesInUse, MemoryCounterBytes.Load() + ForceInlineBytes, MaxCacheSizeMB);
 
@@ -2448,30 +2508,30 @@ TPair<int, int> FAudioChunkCache::DebugDisplay(UWorld* World, FViewport* Viewpor
 		+ NumLoadInProgress
 		+ NumOther;
 
-
-	if (FMath::IsNearlyEqual(PercentageForceInline, 1.0))
+	float PercentageExtra = PercentageForceInline + PercentageExternalFeatures;
+	if (FMath::IsNearlyEqual(PercentageExtra, 1.0))
 	{
-		// if the ForceInline Percentage is basically 1, then just set the "number of chunks" to a really big number.
+		// if the Percentage is basically 1, then just set the "number of chunks" to a really big number.
 		// so everything else just gets zeroed out
 		NumChunks = UE_BIG_NUMBER;
 	}
-	else if (PercentageForceInline > 0.0)
+	else if (PercentageExtra > 0.0)
 	{
-		// calculate the NumForceInline based on the percentage of memory used.
-		int32 NumForceInline = NumChunks * (PercentageForceInline / (1 - PercentageForceInline));
+		// calculate the NumExtra based on the percentage of memory used.
+		int32 NumExtra = NumChunks * (PercentageExtra / (1 - PercentageExtra));
 
 		// derivation:
-		// NumChunks + NumForceInline = TotalNumChunks
-		// NumForceInline = TotalNumChunks * PercentageForceInline
+		// NumChunks + NumExtra = TotalNumChunks
+		// NumExtra = TotalNumChunks * PercentageExtra
 		// 
-		// NumChunks = TotalNumChunks * (1 - PercentageForceInline)
-		// TotalNumChunks = NumChunks / (1 - PercentageForceInline)
+		// NumChunks = TotalNumChunks * (1 - PercentageExtra)
+		// TotalNumChunks = NumChunks / (1 - PercentageExtra)
 		// 
 		// - using substitution with the above
-		//  NumForceInline = NumChunks * PercentageForceInline / (1 - PercentageForceInline)
+		//  NumExtra = NumChunks * PercentageExtra / (1 - PercentageExtra)
 
-		// Add the newly calculated ForceInline to the mix 
-		NumChunks += NumForceInline;
+		// Add the newly calculated Extra "chunks" to the mix 
+		NumChunks += NumExtra;
 	}
 
 	if (NumChunks == 0)
@@ -2512,6 +2572,7 @@ TPair<int, int> FAudioChunkCache::DebugDisplay(UWorld* World, FViewport* Viewpor
 	const int32 BarWidthLoadInProgress = PercentageLoadInProgress * BarWidth;
 	const int32 BarWidthOther = PercentageOther * BarWidth;
 	const int32 BarWidthForceInline = PercentageForceInline * BarWidth;
+	const int32 BarWidthExternalFeatures = PercentageExternalFeatures * BarWidth;
 
 
 	// Draw color key
@@ -2540,6 +2601,10 @@ TPair<int, int> FAudioChunkCache::DebugDisplay(UWorld* World, FViewport* Viewpor
 
 	TempString = *FString::Printf(TEXT("Force Inline: %.2f %% (%.2f MB)"), 100.f * PercentageForceInline, NumMegabytesForceInline);
 	Canvas->DrawShadowedString(X, Y, *TempString, UEngine::GetSmallFont(), ColorForceInline);
+	Y += 15;
+
+	TempString = *FString::Printf(TEXT("External Features: %.2f %% (%.2f MB)"), 100.f * PercentageExternalFeatures, NumMegabytesExternalFeatures);
+	Canvas->DrawShadowedString(X, Y, *TempString, UEngine::GetSmallFont(), ColorExternalFeatures);
 	Y += 25;
 
 	TempString = *FString::Printf(TEXT("Other: %.2f %%"), 100.f * PercentageOther);
@@ -2612,17 +2677,27 @@ TPair<int, int> FAudioChunkCache::DebugDisplay(UWorld* World, FViewport* Viewpor
 	// (other)
 	Canvas->DrawTile(CurrHorzOffset, CurrVertOffset, BarWidthOther, BarHeight, 0, 0, 0, 0, ColorOther);
 	CurrHorzOffset += BarWidthOther;
-
-	if (BarWidthForceInline > 0)
+	
+	if (BarWidthForceInline > 0 || BarWidthExternalFeatures > 0)
 	{
-		// (|| divider between cache and force inline)
+		// (|| divider between cache and chunk memory usage && force inline + External features)
 		const int32 DividerWidth = 5;
 		Canvas->DrawTile(CurrHorzOffset, CurrVertOffset, DividerWidth, BarHeight, 0, 0, 0, 0, FLinearColor::Black);
 		CurrHorzOffset += DividerWidth;
 
-		// (Force Inline)
-		Canvas->DrawTile(CurrHorzOffset, CurrVertOffset, BarWidthForceInline - DividerWidth, BarHeight, 0, 0, 0, 0, ColorForceInline);
-		CurrHorzOffset += BarWidthForceInline;	
+		if (BarWidthForceInline > 0)
+		{
+			// (Force Inline)
+			Canvas->DrawTile(CurrHorzOffset, CurrVertOffset, BarWidthForceInline - DividerWidth, BarHeight, 0, 0, 0, 0, ColorForceInline);
+			CurrHorzOffset += BarWidthForceInline;
+		}
+
+		if (BarWidthExternalFeatures > 0)
+		{
+			// (Force Inline)
+			Canvas->DrawTile(CurrHorzOffset, CurrVertOffset, BarWidthExternalFeatures - DividerWidth, BarHeight, 0, 0, 0, 0, ColorExternalFeatures);
+			CurrHorzOffset += BarWidthExternalFeatures;
+		}
 	}
 
 	Y = (CurrVertOffset + 24);
diff --git a/Engine/Source/Runtime/Engine/Public/AudioStreamingCache.h b/Engine/Source/Runtime/Engine/Public/AudioStreamingCache.h
--- a/Engine/Source/Runtime/Engine/Public/AudioStreamingCache.h
+++ b/Engine/Source/Runtime/Engine/Public/AudioStreamingCache.h
@@ -25,6 +25,8 @@ AudioStreaming.h: Definitions of classes used for audio streaming.
 
 ENGINE_API DECLARE_LOG_CATEGORY_EXTERN(LogAudioStreamCaching, Display, All);
 
+class FAudioStreamingMemoryCountedFeature;
+
 // Basic fixed-size LRU cache for retaining chunks of compressed audio data.
 class FAudioChunkCache
 {
@@ -127,7 +129,11 @@ public:
 
 	void AddForceInlineSoundWave(const FSoundWaveProxyPtr&);
 
-	void RemoveForecInlineSoundWave(const FSoundWaveProxyPtr&);
+	void RemoveForceInlineSoundWave(const FSoundWaveProxyPtr&);
+
+	void AddMemoryCountedFeature(const FAudioStreamingMemoryCountedFeature& Feature);
+
+	void RemoveMemoryCountedFeature(const FAudioStreamingMemoryCountedFeature& Feature);
 
 	// This function will reclaim memory by freeing as many chunks as needed to free BytesToFree.
 	// returns the amount of bytes we were actually able to free.
@@ -402,6 +408,8 @@ private:
 
 	TAtomic<uint64> ForceInlineMemoryCounterBytes;
 
+	TAtomic<uint64> FeatureMemoryCounterBytes;
+
 	// Number of async load operations we have currently in flight.
 	FThreadSafeCounter NumberOfLoadsInFlight;
 
@@ -428,6 +436,8 @@ private:
 	// This is set to true when BeginLoggingCacheMisses is called. 
 	bool bLogCacheMisses;
 
+	uint64 GetCurrentMemoryUsageBytes() const { return MemoryCounterBytes + ForceInlineMemoryCounterBytes + FeatureMemoryCounterBytes; }
+
 	// Returns cached element if it exists in our cache, nullptr otherwise.
 	// If the index of the element is already known, it can be used here to avoid searching the cache.
 	FCacheElement* FindElementForKey(const FChunkKey& InKey);
@@ -518,6 +528,8 @@ public:
 	virtual void RemoveStreamingSoundWave(const FSoundWaveProxyPtr& SoundWave) override;
 	virtual void AddForceInlineSoundWave(const FSoundWaveProxyPtr& SoundWave) override;
 	virtual void RemoveForceInlineSoundWave(const FSoundWaveProxyPtr& SoundWave) override;
+	virtual void AddMemoryCountedFeature(const FAudioStreamingMemoryCountedFeature& Feature) override;
+	virtual void RemoveMemoryCountedFeature(const FAudioStreamingMemoryCountedFeature& Feature) override;
 	virtual void AddDecoder(ICompressedAudioInfo* CompressedAudioInfo) override;
 	virtual void RemoveDecoder(ICompressedAudioInfo* CompressedAudioInfo) override;
 	virtual bool IsManagedStreamingSoundWave(const FSoundWaveProxyPtr&  SoundWave) const override;
diff --git a/Engine/Source/Runtime/Engine/Public/ContentStreaming.h b/Engine/Source/Runtime/Engine/Public/ContentStreaming.h
--- a/Engine/Source/Runtime/Engine/Public/ContentStreaming.h
+++ b/Engine/Source/Runtime/Engine/Public/ContentStreaming.h
@@ -47,6 +47,7 @@ class UStaticMesh;
 class USkeletalMesh;
 class UStreamableRenderAsset;
 class FSoundSource;
+class FAudioStreamingMemoryCountedFeature;
 struct FWaveInstance;
 struct FRenderAssetStreamingManager;
 
@@ -526,6 +527,10 @@ struct IAudioStreamingManager : public IStreamingManager
 	/** Removes the memory usage of the force inline sound from the streaming cache budget */
 	virtual void RemoveForceInlineSoundWave(const FSoundWaveProxyPtr& SoundWave) { };
 
+	virtual void AddMemoryCountedFeature(const FAudioStreamingMemoryCountedFeature& Feature) { };
+
+	virtual void RemoveMemoryCountedFeature(const FAudioStreamingMemoryCountedFeature& Feature) { };
+
 	/** Adds the decoder to the streaming manager to prevent stream chunks from getting reaped from underneath it */
 	virtual void AddDecoder(ICompressedAudioInfo* CompressedAudioInfo) = 0;
 
@@ -853,3 +858,45 @@ protected:
 #endif
 };
 
+class ENGINE_API FAudioStreamingMemoryCountedFeature : public TSharedFromThis<FAudioStreamingMemoryCountedFeature>
+{
+public:
+	FAudioStreamingMemoryCountedFeature(FName InFeatureName, uint64 InMemoryUseInBytes)
+	: FeatureName(InFeatureName)
+	, MemoryUseInBytes(InMemoryUseInBytes)
+	{
+		if (MemoryUseInBytes != 0)
+		{
+			IStreamingManager::Get().GetAudioStreamingManager().AddMemoryCountedFeature(*this);
+		}
+	}
+	
+	~FAudioStreamingMemoryCountedFeature()
+	{
+		if (MemoryUseInBytes != 0)
+		{
+			IStreamingManager::Get().GetAudioStreamingManager().RemoveMemoryCountedFeature(*this);
+		}
+	}
+	
+	uint64 GetMemoryUseInBytes() const { return MemoryUseInBytes; } 
+	FName GetFeatureName() const { return FeatureName; }
+
+	void ResetMemoryUseInBytes(uint64 InMemoryUseInBytes)
+	{
+		if (MemoryUseInBytes != 0)
+		{
+			IStreamingManager::Get().GetAudioStreamingManager().RemoveMemoryCountedFeature(*this);
+		}
+		
+		MemoryUseInBytes = InMemoryUseInBytes;
+
+		if (MemoryUseInBytes != 0)
+		{
+			IStreamingManager::Get().GetAudioStreamingManager().AddMemoryCountedFeature(*this);
+		}
+	}
+private:
+	const FName FeatureName;
+	uint64 MemoryUseInBytes;
+};
\ No newline at end of file
