Fix warning V674: The 'bUseTextureAsInput?-1.0f:0.0f' literal of the 'float' type is assigned to a variable of the 'int' type. Consider inspecting the '= bUseTextureAsInput ? - 1.0f : 0.0f' expression.
#rnx

diff --git a/Engine/Source/Runtime/Engine/Private/Materials/MaterialExpressions.cpp b/Engine/Source/Runtime/Engine/Private/Materials/MaterialExpressions.cpp
--- a/Engine/Source/Runtime/Engine/Private/Materials/MaterialExpressions.cpp
+++ b/Engine/Source/Runtime/Engine/Private/Materials/MaterialExpressions.cpp
@@ -27970,7 +27970,7 @@ int32 UMaterialExpressionNeuralNetworkInput::Compile(class FMaterialCompiler* Co
 		else
 		{
 			int32 ViewportUV = Compiler->GetViewportUV();
-			int32 BatchIndex = bUseTextureAsInput ? -1.0f : 0.0f;
+			float BatchIndex = bUseTextureAsInput ? -1.0f : 0.0f;
 				
 			CodeInput = Compiler->AppendVector(Compiler->Constant2(BatchIndex, 0.0f), ViewportUV);
 		}
