Fix for potential crash when retrieving anim attribute values

diff --git a/Engine/Source/Runtime/Engine/Private/Components/SkeletalMeshComponent.cpp b/Engine/Source/Runtime/Engine/Private/Components/SkeletalMeshComponent.cpp
--- a/Engine/Source/Runtime/Engine/Private/Components/SkeletalMeshComponent.cpp
+++ b/Engine/Source/Runtime/Engine/Private/Components/SkeletalMeshComponent.cpp
@@ -4595,7 +4595,7 @@ bool USkeletalMeshComponent::FindAttributeChecked(const FName& BoneName, const F
 
 		const CustomAttributeType* AttributePtr = Attributes.Find<CustomAttributeType>(UE::Anim::FAttributeId(AttributeName, FCompactPoseBoneIndex(BoneIndex)));
 
-		if (AttributePtr == nullptr && LookupType != ECustomBoneAttributeLookup::BoneOnly)
+		if (AttributePtr == nullptr && BoneIndex != INDEX_NONE && LookupType != ECustomBoneAttributeLookup::BoneOnly)
 		{
 			if (LookupType == ECustomBoneAttributeLookup::ImmediateParent)
 			{
