Fix warning V590: Consider inspecting this expression. The expression is excessive or contains a misprint.
#rnx

diff --git a/Engine/Source/Runtime/Engine/Private/WorldPartition/DataLayer/WorldDataLayers.cpp b/Engine/Source/Runtime/Engine/Private/WorldPartition/DataLayer/WorldDataLayers.cpp
--- a/Engine/Source/Runtime/Engine/Private/WorldPartition/DataLayer/WorldDataLayers.cpp
+++ b/Engine/Source/Runtime/Engine/Private/WorldPartition/DataLayer/WorldDataLayers.cpp
@@ -197,7 +197,7 @@ void AWorldDataLayers::SetDataLayerRuntimeState(const UDataLayerInstance* InData
 	}
 	else if (bDataLayerServerOnly)
 	{
-		if (NetMode != NM_Standalone && NetMode == NM_Client)
+		if (NetMode == NM_Client)
 		{
 			UE_LOG(LogWorldPartition, Log, TEXT("Server Only Data Layer state change '%s' was ignored: %s -> %s"),
 				*InDataLayerInstance->GetDataLayerShortName(),
