Fix various PVS SA issues
#rnx
#rnx

diff --git a/Engine/Source/Runtime/Engine/Private/GameViewportClient.cpp b/Engine/Source/Runtime/Engine/Private/GameViewportClient.cpp
--- a/Engine/Source/Runtime/Engine/Private/GameViewportClient.cpp
+++ b/Engine/Source/Runtime/Engine/Private/GameViewportClient.cpp
@@ -1347,7 +1347,7 @@ void UGameViewportClient::Draw(FViewport* InViewport, FCanvas* SceneCanvas)
 
 	// Create temp debug canvas object
 	FIntPoint DebugCanvasSize = InViewport->GetSizeXY();
-	if (bStereoRendering && GEngine && GEngine->XRSystem.IsValid() && GEngine->XRSystem->GetHMDDevice())
+	if (bStereoRendering && GEngine->XRSystem.IsValid() && GEngine->XRSystem->GetHMDDevice())
 	{
 		DebugCanvasSize = GEngine->XRSystem->GetHMDDevice()->GetIdealDebugCanvasRenderTargetSize();
 	}
diff --git a/Engine/Source/Runtime/Engine/Private/PhysicsEngine/PhysicsReplication.cpp b/Engine/Source/Runtime/Engine/Private/PhysicsEngine/PhysicsReplication.cpp
--- a/Engine/Source/Runtime/Engine/Private/PhysicsEngine/PhysicsReplication.cpp
+++ b/Engine/Source/Runtime/Engine/Private/PhysicsEngine/PhysicsReplication.cpp
@@ -1087,7 +1087,7 @@ bool FPhysicsReplicationAsync::DefaultReplication(Chaos::FPBDRigidParticleHandle
 
 	const FString ObjectName
 #if CHAOS_DEBUG_NAME
-		= Handle->DebugName() ? *Handle->DebugName() : FString(TEXT(""));
+		= Handle && Handle->DebugName() ? *Handle->DebugName() : FString(TEXT(""));
 #else
 		= FString(TEXT(""));
 #endif
diff --git a/Engine/Source/Runtime/Engine/Private/PrecomputedVolumetricLightmap.cpp b/Engine/Source/Runtime/Engine/Private/PrecomputedVolumetricLightmap.cpp
--- a/Engine/Source/Runtime/Engine/Private/PrecomputedVolumetricLightmap.cpp
+++ b/Engine/Source/Runtime/Engine/Private/PrecomputedVolumetricLightmap.cpp
@@ -719,7 +719,8 @@ void FPrecomputedVolumetricLightmap::RemoveFromScene(FSceneInterface* Scene)
 		// Certain paths in the editor (namely, ReloadPackages and ForceDelete) will GC the registry before the UWorld destruction (which destructs FScene)
 		ensureMsgf(
 			SourceRegistry.IsValid() // either SourceRegistry is valid
-			|| (!Scene->GetWorld() || !IsValidChecked(Scene->GetWorld()) || Scene->GetWorld()->IsUnreachable()) // or the world we're in is going away (usually during shutdown)
+			|| Scene == nullptr // or there is no scene
+			|| !Scene->GetWorld() || !IsValidChecked(Scene->GetWorld()) || Scene->GetWorld()->IsUnreachable() // or the world we're in is going away (usually during shutdown)
 			, TEXT("UMapBuildDataRegistry is garbage collected before an FPrecomputedVolumetricLightmap is removed from the scene. Is there a missing ReleaseRenderingResources() call?"));
 
 		// While that can be explained as missing ReleaseRenderingResources() calls, this fail-safe guard is added here
diff --git a/Engine/Source/Runtime/RHI/Private/RHICommandList.cpp b/Engine/Source/Runtime/RHI/Private/RHICommandList.cpp
--- a/Engine/Source/Runtime/RHI/Private/RHICommandList.cpp
+++ b/Engine/Source/Runtime/RHI/Private/RHICommandList.cpp
@@ -724,7 +724,7 @@ void FRHICommandListImmediate::ExecuteAndReset(bool bFlushResources)
 		{
 			if (Context)
 			{
-				IRHIPlatformCommandList* CommandList = GDynamicRHI->RHIFinalizeContext(Context);
+				IRHIPlatformCommandList* CommandList = GDynamicRHI ? GDynamicRHI->RHIFinalizeContext(Context) : nullptr;
 				if (CommandList)
 				{
 					CommandLists.Add(CommandList);
