Add a translated world space options to several material expressions

diff --git a/Engine/Shaders/Private/DistanceField/GlobalDistanceFieldShared.ush b/Engine/Shaders/Private/DistanceField/GlobalDistanceFieldShared.ush
--- a/Engine/Shaders/Private/DistanceField/GlobalDistanceFieldShared.ush
+++ b/Engine/Shaders/Private/DistanceField/GlobalDistanceFieldShared.ush
@@ -355,10 +355,8 @@ float3 GetDistanceFieldGradientGlobal(float3 TranslatedWorldPosition)
 	return Gradient;
 }
 
-float CalculateDistanceFieldApproxAO(FLWCVector3 WorldPosition, float3 WorldNormal, uint NumSteps, float StepDistance, float StepScale, float DistanceBias, float MaxDistance)
+float CalculateDistanceFieldApproxAO(float3 TranslatedWorldPosition, float3 WorldNormal, uint NumSteps, float StepDistance, float StepScale, float DistanceBias, float MaxDistance)
 {
-	const float3 TranslatedWorldPosition = LWCToFloat(LWCAdd(WorldPosition, ResolvedView.PreViewTranslation));
-
 	// determine the min clipmap that covers all the sampling points (sphere centered on TranslatedWorldPosition with radius = MaxDistance)
 	int MinClipmapIndex = -1;
 
@@ -408,7 +406,7 @@ float3 GetDistanceFieldGradientGlobal(float3 TranslatedWorldPosition)
 	return float3(0, 0, .001f);
 }
 
-float CalculateDistanceFieldApproxAO(FLWCVector3 WorldPosition, float3 WorldNormal, uint NumSteps, float StepDistance, float StepScale, float DistanceBias, float MaxDistance)
+float CalculateDistanceFieldApproxAO(float3 TranslatedWorldPosition, float3 WorldNormal, uint NumSteps, float StepDistance, float StepScale, float DistanceBias, float MaxDistance)
 {
 	return 1.0f;
 }
@@ -429,3 +427,8 @@ float3 GetDistanceFieldGradientGlobal(FLWCVector3 WorldPosition)
 	return GetDistanceFieldGradientGlobal(TranslatedWorldPosition);
 }
 
+float CalculateDistanceFieldApproxAO(FLWCVector3 WorldPosition, float3 WorldNormal, uint NumSteps, float StepDistance, float StepScale, float DistanceBias, float MaxDistance)
+{
+	const float3 TranslatedWorldPosition = LWCToFloat(LWCAdd(WorldPosition, ResolvedView.PreViewTranslation));
+	return CalculateDistanceFieldApproxAO(TranslatedWorldPosition, WorldNormal, NumSteps, StepDistance, StepScale, DistanceBias, MaxDistance);
+}
diff --git a/Engine/Shaders/Private/MaterialTemplate.ush b/Engine/Shaders/Private/MaterialTemplate.ush
--- a/Engine/Shaders/Private/MaterialTemplate.ush
+++ b/Engine/Shaders/Private/MaterialTemplate.ush
@@ -1734,10 +1734,10 @@ float3 MaterialExpressionAtmosphericLightColor(FMaterialPixelParameters Paramete
 	return MaterialExpressionSkyAtmosphereLightIlluminanceOnGround(Parameters, 0);
 }
 
-float3 MaterialExpressionSkyAtmosphereLightIlluminance(FMaterialPixelParameters Parameters, FLWCVector3 WorldPosition, uint LightIndex)
+float3 MaterialExpressionSkyAtmosphereLightIlluminance(FMaterialPixelParameters Parameters, float3 TranslatedWorldPosition, uint LightIndex)
 {
 #if MATERIAL_SKY_ATMOSPHERE && PROJECT_SUPPORT_SKY_ATMOSPHERE
-	const float3 PlanetCenterToTranslatedWorldPos = (LWCToFloat(LWCAdd(WorldPosition, ResolvedView.PreViewTranslation)) - ResolvedView.SkyPlanetTranslatedWorldCenterAndViewHeight.xyz) * CM_TO_SKY_UNIT;
+	const float3 PlanetCenterToTranslatedWorldPos = (TranslatedWorldPosition - ResolvedView.SkyPlanetTranslatedWorldCenterAndViewHeight.xyz) * CM_TO_SKY_UNIT;
 
 	// GetAtmosphereTransmittance does a shadow test against the virtual planet.
 	const float3 TransmittanceToLight = GetAtmosphereTransmittance(
@@ -1750,6 +1750,12 @@ float3 MaterialExpressionSkyAtmosphereLightIlluminance(FMaterialPixelParameters
 #endif
 }
 
+float3 MaterialExpressionSkyAtmosphereLightIlluminance(FMaterialPixelParameters Parameters, FLWCVector3 WorldPosition, uint LightIndex)
+{
+	const float3 TranslatedWorldPosition = LWCToFloat(LWCAdd(WorldPosition, ResolvedView.PreViewTranslation));
+	return MaterialExpressionSkyAtmosphereLightIlluminance(Parameters, TranslatedWorldPosition, LightIndex);
+}
+
 #if MATERIAL_SKY_ATMOSPHERE && PROJECT_SUPPORT_SKY_ATMOSPHERE
 	#define DEFINE_SKYATMLIGHTDIRECTION(MaterialParamType) float3 MaterialExpressionSkyAtmosphereLightDirection(MaterialParamType Parameters, uint LightIndex) {return ResolvedView.AtmosphereLightDirection[LightIndex].xyz;}
 #else
@@ -1806,11 +1812,11 @@ float3 MaterialExpressionSkyAtmosphereViewLuminance(FMaterialPixelParameters Par
 #endif
 }
 
-float4 MaterialExpressionSkyAtmosphereAerialPerspective(FMaterialPixelParameters Parameters, FLWCVector3 WorldPosition)
+float4 MaterialExpressionSkyAtmosphereAerialPerspective(FMaterialPixelParameters Parameters, float3 TranslatedWorldPosition)
 {
 #if MATERIAL_SKY_ATMOSPHERE && PROJECT_SUPPORT_SKY_ATMOSPHERE
-	const float3 TranslatedWorldPosition = LWCToFloat(LWCAdd(WorldPosition, ResolvedView.PreViewTranslation)) * CM_TO_SKY_UNIT;
-	const float3 SkyCameraTranslatedWorldOrigin = ResolvedView.SkyCameraTranslatedWorldOrigin.xyz*CM_TO_SKY_UNIT;
+	const float3 SkyTranslatedWorldPosition = TranslatedWorldPosition * CM_TO_SKY_UNIT;
+	const float3 SkyCameraTranslatedWorldOrigin = ResolvedView.SkyCameraTranslatedWorldOrigin.xyz * CM_TO_SKY_UNIT;
 
 	// NDCPosition is not computed using WorldPosition because it could result in position outside the frustum, 
 	// distorted uvs and bad visuals with artefact. Only the distance computation can actually be benefit from the surface position specified here.
@@ -1819,7 +1825,7 @@ float4 MaterialExpressionSkyAtmosphereAerialPerspective(FMaterialPixelParameters
 	const float NearFadeOutRangeInvDepthKm = 1.0 / 0.00001f; // 1 centimeter fade region
 	float4 AerialPerspective = GetAerialPerspectiveLuminanceTransmittance(
 		ResolvedView.RealTimeReflectionCapture, ResolvedView.SkyAtmosphereCameraAerialPerspectiveVolumeSizeAndInvSize,
-		NDCPosition, TranslatedWorldPosition - SkyCameraTranslatedWorldOrigin,
+		NDCPosition, SkyTranslatedWorldPosition - SkyCameraTranslatedWorldOrigin,
 		View.CameraAerialPerspectiveVolume, View.CameraAerialPerspectiveVolumeSampler,
 		ResolvedView.SkyAtmosphereCameraAerialPerspectiveVolumeDepthResolutionInv,
 		ResolvedView.SkyAtmosphereCameraAerialPerspectiveVolumeDepthResolution,
@@ -1834,6 +1840,12 @@ float4 MaterialExpressionSkyAtmosphereAerialPerspective(FMaterialPixelParameters
 #endif
 }
 
+float4 MaterialExpressionSkyAtmosphereAerialPerspective(FMaterialPixelParameters Parameters, FLWCVector3 WorldPosition)
+{
+	const float3 TranslatedWorldPosition = LWCToFloat(LWCAdd(WorldPosition, ResolvedView.PreViewTranslation));
+	return MaterialExpressionSkyAtmosphereAerialPerspective(Parameters, TranslatedWorldPosition);
+}
+
 float3 MaterialExpressionSkyAtmosphereDistantLightScatteredLuminance(FMaterialPixelParameters Parameters)
 {
 #if MATERIAL_SKY_ATMOSPHERE && PROJECT_SUPPORT_SKY_ATMOSPHERE
diff --git a/Engine/Shaders/Private/VirtualTextureCommon.ush b/Engine/Shaders/Private/VirtualTextureCommon.ush
--- a/Engine/Shaders/Private/VirtualTextureCommon.ush
+++ b/Engine/Shaders/Private/VirtualTextureCommon.ush
@@ -794,6 +794,13 @@ float2 VirtualTextureWorldToUV(in FLWCVector3 WorldPos, in FLWCVector3 O, in flo
 	return float2(dot(P, U), dot(P, V));
 }
 
+/**	Helper function to convert translated world space to virtual texture UV space (requires O to be in translated space). */
+float2 VirtualTextureWorldToUV(in float3 TranslatedWorldPos, in float3 O, in float3 U, in float3 V)
+{
+	float3 P = TranslatedWorldPos - O;
+	return float2(dot(P, U), dot(P, V));
+}
+
 /**	Derivative-aware helper function to convert world space to virtual texture UV space. */
 FloatDeriv2 VirtualTextureWorldToUVDeriv(in FLWCVector3Deriv WorldPos, in FLWCVector3 O, in float3 U, in float3 V)
 {
@@ -804,6 +811,16 @@ FloatDeriv2 VirtualTextureWorldToUVDeriv(in FLWCVector3Deriv WorldPos, in FLWCVe
 	return Result;
 }
 
+/**	Derivative-aware helper function to convert translated world space to virtual texture UV space (requires O to be in translated space). */
+FloatDeriv2 VirtualTextureWorldToUVDeriv(in FloatDeriv3 WorldPos, in float3 O, in float3 U, in float3 V)
+{
+	FloatDeriv2 Result;
+	Result.Value = VirtualTextureWorldToUV(WorldPos.Value, O, U, V);
+	Result.Ddx = float2(dot(WorldPos.Ddx, U), dot(WorldPos.Ddx, V));
+	Result.Ddy = float2(dot(WorldPos.Ddy, U), dot(WorldPos.Ddy, V));
+	return Result;
+}
+
 /** Unpack color from YCoCg stored in BC3. */
 float3 VirtualTextureUnpackBaseColorYCoCg(in float4 PackedValue)
 {
diff --git a/Engine/Source/Developer/MaterialUtilities/Private/MaterialUtilities.cpp b/Engine/Source/Developer/MaterialUtilities/Private/MaterialUtilities.cpp
--- a/Engine/Source/Developer/MaterialUtilities/Private/MaterialUtilities.cpp
+++ b/Engine/Source/Developer/MaterialUtilities/Private/MaterialUtilities.cpp
@@ -370,7 +370,7 @@ struct FExportMaterialCompiler : public FProxyMaterialCompiler
 		return Compiler->ReflectionVector();
 	}
 
-	virtual int32 AtmosphericFogColor(int32 WorldPosition) override
+	virtual int32 AtmosphericFogColor(int32 WorldPosition, EPositionOrigin PositionOrigin) override
 	{
 		return INDEX_NONE;
 	}
diff --git a/Engine/Source/Runtime/Engine/Classes/Materials/MaterialExpression.h b/Engine/Source/Runtime/Engine/Classes/Materials/MaterialExpression.h
--- a/Engine/Source/Runtime/Engine/Classes/Materials/MaterialExpression.h
+++ b/Engine/Source/Runtime/Engine/Classes/Materials/MaterialExpression.h
@@ -618,6 +618,6 @@ enum class EPositionOrigin : uint8
 	/** Absolute world position, i.e. relative to (0,0,0) */
 	Absolute UMETA(DisplayName="Absolute World Position"),
 
-	/** Camera relative world position */
+	/** Camera relative world position, i.e. translated world space */
 	CameraRelative UMETA(DisplayName="Camera Relative World Position")
 };
diff --git a/Engine/Source/Runtime/Engine/Classes/Materials/MaterialExpressionAtmosphericFogColor.h b/Engine/Source/Runtime/Engine/Classes/Materials/MaterialExpressionAtmosphericFogColor.h
--- a/Engine/Source/Runtime/Engine/Classes/Materials/MaterialExpressionAtmosphericFogColor.h
+++ b/Engine/Source/Runtime/Engine/Classes/Materials/MaterialExpressionAtmosphericFogColor.h
@@ -16,8 +16,14 @@ class UMaterialExpressionAtmosphericFogColor : public UMaterialExpression
 	UPROPERTY()
 	FExpressionInput WorldPosition;
 
+	/** Defines the reference space for the WorldPosition input. */
+	UPROPERTY(EditAnywhere, Category = MaterialExpressionAtmosphericFogColor)
+	EPositionOrigin WorldPositionOriginType = EPositionOrigin::Absolute;
+
 	//~ Begin UMaterialExpression Interface
 #if WITH_EDITOR
+	virtual FName GetInputName(int32 InputIndex) const override;
+	virtual void PostEditChangeProperty(FPropertyChangedEvent& PropertyChangedEvent) override;
 	virtual int32 Compile(class FMaterialCompiler* Compiler, int32 OutputIndex) override;
 	virtual void GetCaption(TArray<FString>& OutCaptions) const override;
 	virtual bool GenerateHLSLExpression(FMaterialHLSLGenerator& Generator, UE::HLSLTree::FScope& Scope, int32 OutputIndex, UE::HLSLTree::FExpression const*& OutExpression) const override;
diff --git a/Engine/Source/Runtime/Engine/Classes/Materials/MaterialExpressionDistanceFieldApproxAO.h b/Engine/Source/Runtime/Engine/Classes/Materials/MaterialExpressionDistanceFieldApproxAO.h
--- a/Engine/Source/Runtime/Engine/Classes/Materials/MaterialExpressionDistanceFieldApproxAO.h
+++ b/Engine/Source/Runtime/Engine/Classes/Materials/MaterialExpressionDistanceFieldApproxAO.h
@@ -17,6 +17,10 @@ class UMaterialExpressionDistanceFieldApproxAO : public UMaterialExpression
 	UPROPERTY(meta = (RequiredInput = "false", ToolTip = "Defaults to current world position if not specified"))
 	FExpressionInput Position;
 
+	/** Defines the reference space for the Position input. */
+	UPROPERTY(EditAnywhere, Category = MaterialExpressionDistanceFieldApproxAO)
+	EPositionOrigin WorldPositionOriginType = EPositionOrigin::Absolute;
+
 	UPROPERTY(meta = (RequiredInput = "false", ToolTip = "Defaults to current world normal if not specified"))
 	FExpressionInput Normal;
 
@@ -42,6 +46,8 @@ class UMaterialExpressionDistanceFieldApproxAO : public UMaterialExpression
 
 	//~ Begin UMaterialExpression Interface
 #if WITH_EDITOR
+	virtual FName GetInputName(int32 InputIndex) const override;
+	virtual void PostEditChangeProperty(FPropertyChangedEvent& PropertyChangedEvent) override;
 	virtual int32 Compile(class FMaterialCompiler* Compiler, int32 OutputIndex) override;
 	virtual void GetCaption(TArray<FString>& OutCaptions) const override;
 
diff --git a/Engine/Source/Runtime/Engine/Classes/Materials/MaterialExpressionDistanceFieldGradient.h b/Engine/Source/Runtime/Engine/Classes/Materials/MaterialExpressionDistanceFieldGradient.h
--- a/Engine/Source/Runtime/Engine/Classes/Materials/MaterialExpressionDistanceFieldGradient.h
+++ b/Engine/Source/Runtime/Engine/Classes/Materials/MaterialExpressionDistanceFieldGradient.h
@@ -17,8 +17,14 @@ class UMaterialExpressionDistanceFieldGradient : public UMaterialExpression
 	UPROPERTY(meta = (RequiredInput = "false", ToolTip = "Defaults to current world position if not specified"))
 	FExpressionInput Position;
 
+	/** Defines the reference space for the Position input. */
+	UPROPERTY(EditAnywhere, Category = MaterialExpressionDistanceFieldGradient)
+	EPositionOrigin WorldPositionOriginType = EPositionOrigin::Absolute;
+
 	//~ Begin UMaterialExpression Interface
 #if WITH_EDITOR
+	virtual FName GetInputName(int32 InputIndex) const override;
+	virtual void PostEditChangeProperty(FPropertyChangedEvent& PropertyChangedEvent) override;
 	virtual int32 Compile(class FMaterialCompiler* Compiler, int32 OutputIndex) override;
 	virtual void GetCaption(TArray<FString>& OutCaptions) const override;
 
diff --git a/Engine/Source/Runtime/Engine/Classes/Materials/MaterialExpressionDistanceToNearestSurface.h b/Engine/Source/Runtime/Engine/Classes/Materials/MaterialExpressionDistanceToNearestSurface.h
--- a/Engine/Source/Runtime/Engine/Classes/Materials/MaterialExpressionDistanceToNearestSurface.h
+++ b/Engine/Source/Runtime/Engine/Classes/Materials/MaterialExpressionDistanceToNearestSurface.h
@@ -17,8 +17,14 @@ class UMaterialExpressionDistanceToNearestSurface : public UMaterialExpression
 	UPROPERTY(meta = (RequiredInput = "false", ToolTip = "Defaults to current world position if not specified"))
 	FExpressionInput Position;
 
+	/** Defines the reference space for the Position input. */
+	UPROPERTY(EditAnywhere, Category = MaterialExpressionDistanceToNearestSurface)
+	EPositionOrigin WorldPositionOriginType = EPositionOrigin::Absolute;
+
 	//~ Begin UMaterialExpression Interface
 #if WITH_EDITOR
+	virtual FName GetInputName(int32 InputIndex) const override;
+	virtual void PostEditChangeProperty(FPropertyChangedEvent& PropertyChangedEvent) override;
 	virtual int32 Compile(class FMaterialCompiler* Compiler, int32 OutputIndex) override;
 	virtual void GetCaption(TArray<FString>& OutCaptions) const override;
 
diff --git a/Engine/Source/Runtime/Engine/Classes/Materials/MaterialExpressionNoise.h b/Engine/Source/Runtime/Engine/Classes/Materials/MaterialExpressionNoise.h
--- a/Engine/Source/Runtime/Engine/Classes/Materials/MaterialExpressionNoise.h
+++ b/Engine/Source/Runtime/Engine/Classes/Materials/MaterialExpressionNoise.h
@@ -64,6 +64,10 @@ class UMaterialExpressionNoise : public UMaterialExpression
 	UPROPERTY()
 	FExpressionInput Position;
 
+	/** Defines the reference space for the Position input. */
+	UPROPERTY(EditAnywhere, Category = MaterialExpressionNoise)
+	EPositionOrigin WorldPositionOriginType = EPositionOrigin::Absolute;
+
 	/** scalar, to clamp the Levels at pixel level, can be computed like this: max(length(ddx(Position)), length(ddy(Position)) */
 	UPROPERTY()
 	FExpressionInput FilterWidth;
@@ -109,6 +113,8 @@ class UMaterialExpressionNoise : public UMaterialExpression
 	//~ Begin UMaterialExpression Interface
 #if WITH_EDITOR
 	virtual bool CanEditChange(const FProperty* InProperty) const override;
+	virtual FName GetInputName(int32 InputIndex) const override;
+	virtual void PostEditChangeProperty(FPropertyChangedEvent& PropertyChangedEvent) override;
 	virtual int32 Compile(class FMaterialCompiler* Compiler, int32 OutputIndex) override;
 	virtual void GetCaption(TArray<FString>& OutCaptions) const override;
 	virtual bool GenerateHLSLExpression(FMaterialHLSLGenerator& Generator, UE::HLSLTree::FScope& Scope, int32 OutputIndex, UE::HLSLTree::FExpression const*& OutExpression) const override;
diff --git a/Engine/Source/Runtime/Engine/Classes/Materials/MaterialExpressionRuntimeVirtualTextureSample.h b/Engine/Source/Runtime/Engine/Classes/Materials/MaterialExpressionRuntimeVirtualTextureSample.h
--- a/Engine/Source/Runtime/Engine/Classes/Materials/MaterialExpressionRuntimeVirtualTextureSample.h
+++ b/Engine/Source/Runtime/Engine/Classes/Materials/MaterialExpressionRuntimeVirtualTextureSample.h
@@ -98,6 +98,10 @@ class UMaterialExpressionRuntimeVirtualTextureSample : public UMaterialExpressio
 	UPROPERTY(EditAnywhere, Category = VirtualTexture)
 	bool bEnableFeedback = true;
 
+	/** Defines the reference space for the WorldPosition input. */
+	UPROPERTY(EditAnywhere, Category = TextureSample)
+	EPositionOrigin WorldPositionOriginType = EPositionOrigin::Absolute;
+
 	/** Defines how the mip level is calculated for the virtual texture lookup. */
 	UPROPERTY(EditAnywhere, Category = TextureSample)
 	TEnumAsByte<enum ERuntimeVirtualTextureMipValueMode> MipValueMode = RVTMVM_None;
@@ -119,6 +123,7 @@ protected:
 
 #if WITH_EDITOR
 	ENGINE_API virtual void PostLoad() override;
+	ENGINE_API virtual FName GetInputName(int32 InputIndex) const override;
 	ENGINE_API virtual int32 Compile(class FMaterialCompiler* Compiler, int32 OutputIndex) override;
 	ENGINE_API virtual void GetCaption(TArray<FString>& OutCaptions) const override;
 	ENGINE_API virtual bool CanEditChange(const FProperty* InProperty) const override;
diff --git a/Engine/Source/Runtime/Engine/Classes/Materials/MaterialExpressionSamplePhysicsField.h b/Engine/Source/Runtime/Engine/Classes/Materials/MaterialExpressionSamplePhysicsField.h
--- a/Engine/Source/Runtime/Engine/Classes/Materials/MaterialExpressionSamplePhysicsField.h
+++ b/Engine/Source/Runtime/Engine/Classes/Materials/MaterialExpressionSamplePhysicsField.h
@@ -21,12 +21,18 @@ class UMaterialExpressionSamplePhysicsVectorField : public UMaterialExpression
 	UPROPERTY(meta = (RequiredInput = "false", ToolTip = "Defaults to current world position if not specified"))
 	FExpressionInput WorldPosition;
 
+	/** Defines the reference space for the WorldPosition input. */
+	UPROPERTY(EditAnywhere, Category = MaterialExpressionSamplePhysicsVectorField)
+	EPositionOrigin WorldPositionOriginType = EPositionOrigin::Absolute;
+
 	/** Target Type to be accessed */
 	UPROPERTY(EditAnywhere, Category = "Field Setup", meta = (DisplayName = "Target Type", ShowAsInputPin = "Advanced"))
 	TEnumAsByte<EFieldVectorType> FieldTarget;
 
 	//~ Begin UMaterialExpression Interface
 #if WITH_EDITOR
+	virtual FName GetInputName(int32 InputIndex) const override;
+	virtual void PostEditChangeProperty(FPropertyChangedEvent& PropertyChangedEvent) override;
 	virtual int32 Compile(class FMaterialCompiler* Compiler, int32 OutputIndex) override;
 	virtual void GetCaption(TArray<FString>& OutCaptions) const override;
 
@@ -43,12 +49,18 @@ class UMaterialExpressionSamplePhysicsScalarField : public UMaterialExpression
 	UPROPERTY(meta = (RequiredInput = "false", ToolTip = "Defaults to current world position if not specified"))
 	FExpressionInput WorldPosition;
 
+	/** Defines the reference space for the WorldPosition input. */
+	UPROPERTY(EditAnywhere, Category = MaterialExpressionSamplePhysicsScalarField)
+	EPositionOrigin WorldPositionOriginType = EPositionOrigin::Absolute;
+
 	/** Target Type to be accessed */
 	UPROPERTY(EditAnywhere, Category = "Field Setup", meta = (DisplayName = "Target Type", ShowAsInputPin = "Advanced"))
 	TEnumAsByte<EFieldScalarType> FieldTarget;
 
 	//~ Begin UMaterialExpression Interface
 #if WITH_EDITOR
+	virtual FName GetInputName(int32 InputIndex) const override;
+	virtual void PostEditChangeProperty(FPropertyChangedEvent& PropertyChangedEvent) override;
 	virtual int32 Compile(class FMaterialCompiler* Compiler, int32 OutputIndex) override;
 	virtual void GetCaption(TArray<FString>& OutCaptions) const override;
 
@@ -65,12 +77,18 @@ class UMaterialExpressionSamplePhysicsIntegerField : public UMaterialExpression
 	UPROPERTY(meta = (RequiredInput = "false", ToolTip = "Defaults to current world position if not specified"))
 	FExpressionInput WorldPosition;
 
+	/** Defines the reference space for the WorldPosition input. */
+	UPROPERTY(EditAnywhere, Category = MaterialExpressionSamplePhysicsIntegerField)
+	EPositionOrigin WorldPositionOriginType = EPositionOrigin::Absolute;
+
 	/** Target Type to be accessed */
 	UPROPERTY(EditAnywhere, Category = "Field Setup", meta = (DisplayName = "Target Type", ShowAsInputPin = "Advanced"))
 	TEnumAsByte<EFieldIntegerType> FieldTarget;
 
 	//~ Begin UMaterialExpression Interface
 #if WITH_EDITOR
+	virtual FName GetInputName(int32 InputIndex) const override;
+	virtual void PostEditChangeProperty(FPropertyChangedEvent& PropertyChangedEvent) override;
 	virtual int32 Compile(class FMaterialCompiler* Compiler, int32 OutputIndex) override;
 	virtual void GetCaption(TArray<FString>& OutCaptions) const override;
 
diff --git a/Engine/Source/Runtime/Engine/Classes/Materials/MaterialExpressionSkyAtmosphereLightIlluminance.h b/Engine/Source/Runtime/Engine/Classes/Materials/MaterialExpressionSkyAtmosphereLightIlluminance.h
--- a/Engine/Source/Runtime/Engine/Classes/Materials/MaterialExpressionSkyAtmosphereLightIlluminance.h
+++ b/Engine/Source/Runtime/Engine/Classes/Materials/MaterialExpressionSkyAtmosphereLightIlluminance.h
@@ -21,8 +21,14 @@ class UMaterialExpressionSkyAtmosphereLightIlluminance : public UMaterialExpress
 	UPROPERTY()
 	FExpressionInput WorldPosition;
 
+	/** Defines the reference space for the WorldPosition input. */
+	UPROPERTY(EditAnywhere, Category = MaterialExpressionSkyAtmosphereLightIlluminance)
+	EPositionOrigin WorldPositionOriginType = EPositionOrigin::Absolute;
+
 	//~ Begin UMaterialExpression Interface
 #if WITH_EDITOR
+	virtual FName GetInputName(int32 InputIndex) const override;
+	virtual void PostEditChangeProperty(FPropertyChangedEvent& PropertyChangedEvent) override;
 	virtual int32 Compile(class FMaterialCompiler* Compiler, int32 OutputIndex) override;
 	virtual void GetCaption(TArray<FString>& OutCaptions) const override;
 	virtual void GetExpressionToolTip(TArray<FString>& OutToolTip) override;
@@ -91,8 +97,14 @@ class UMaterialExpressionSkyAtmosphereAerialPerspective : public UMaterialExpres
 	UPROPERTY()
 	FExpressionInput WorldPosition;
 
+	/** Defines the reference space for the WorldPosition input. */
+	UPROPERTY(EditAnywhere, Category = MaterialExpressionSkyAtmosphereAerialPerspective)
+	EPositionOrigin WorldPositionOriginType = EPositionOrigin::Absolute;
+
 	//~ Begin UMaterialExpression Interface
 #if WITH_EDITOR
+	virtual FName GetInputName(int32 InputIndex) const override;
+	virtual void PostEditChangeProperty(FPropertyChangedEvent& PropertyChangedEvent) override;
 	virtual int32 Compile(class FMaterialCompiler* Compiler, int32 OutputIndex) override;
 	virtual void GetCaption(TArray<FString>& OutCaptions) const override;
 
diff --git a/Engine/Source/Runtime/Engine/Classes/Materials/MaterialExpressionVectorNoise.h b/Engine/Source/Runtime/Engine/Classes/Materials/MaterialExpressionVectorNoise.h
--- a/Engine/Source/Runtime/Engine/Classes/Materials/MaterialExpressionVectorNoise.h
+++ b/Engine/Source/Runtime/Engine/Classes/Materials/MaterialExpressionVectorNoise.h
@@ -57,6 +57,10 @@ class UMaterialExpressionVectorNoise : public UMaterialExpression
 	UPROPERTY()
 	FExpressionInput Position;
 
+	/** Defines the reference space for the Position input. */
+	UPROPERTY(EditAnywhere, Category = MaterialExpressionVectorNoise)
+	EPositionOrigin WorldPositionOriginType = EPositionOrigin::Absolute;
+
 	/** Noise function, affects performance and look */
 	UPROPERTY(EditAnywhere, Category = MaterialExpressionVectorNoise, meta = (DisplayName = "Function", ShowAsInputPin = "Advanced"))
 	TEnumAsByte<enum EVectorNoiseFunction> NoiseFunction;
@@ -78,6 +82,8 @@ class UMaterialExpressionVectorNoise : public UMaterialExpression
 	//~ Begin UMaterialExpression Interface
 #if WITH_EDITOR
 	virtual bool CanEditChange(const FProperty* InProperty) const override;
+	virtual FName GetInputName(int32 InputIndex) const override;
+	virtual void PostEditChangeProperty(FPropertyChangedEvent& PropertyChangedEvent) override;
 	virtual int32 Compile(class FMaterialCompiler* Compiler, int32 OutputIndex) override;
 	virtual void GetCaption(TArray<FString>& OutCaptions) const override;
 
diff --git a/Engine/Source/Runtime/Engine/Private/Materials/HLSLMaterialTranslator.cpp b/Engine/Source/Runtime/Engine/Private/Materials/HLSLMaterialTranslator.cpp
--- a/Engine/Source/Runtime/Engine/Private/Materials/HLSLMaterialTranslator.cpp
+++ b/Engine/Source/Runtime/Engine/Private/Materials/HLSLMaterialTranslator.cpp
@@ -3453,6 +3453,21 @@ void FHLSLMaterialTranslator::AddLWCFuncUsage(ELWCFunctionKind Kind, const uint3
 	MaterialCompilationOutput.EstimatedLWCFuncUsages[(int)Kind] += Count;
 }
 
+FString FHLSLMaterialTranslator::GetWorldPositionOrDefault(int32 WorldPosition, EPositionOrigin PositionOrigin)
+{
+	FString WorldPosCode;
+	if (PositionOrigin == EPositionOrigin::Absolute)
+	{
+		WorldPosCode = WorldPosition == INDEX_NONE ? FString(TEXT("Parameters.AbsoluteWorldPosition")) : CoerceParameter(WorldPosition, MCT_LWCVector3);
+	}
+	else if (PositionOrigin == EPositionOrigin::CameraRelative)
+	{
+		WorldPosCode = WorldPosition == INDEX_NONE ? FString(TEXT("GetTranslatedWorldPosition(Parameters)")) : CoerceParameter(WorldPosition, MCT_Float3);
+	}
+	else { check(0); }
+	return WorldPosCode;
+}
+
 /** Creates a unique symbol name and adds it to the symbol list. */
 FString FHLSLMaterialTranslator::CreateSymbolName(const TCHAR* SymbolNameHint)
 {
@@ -7798,7 +7813,7 @@ int32 FHLSLMaterialTranslator::VirtualTextureUniform(FName ParameterName, int32
 	return AddUniformExpression(new FMaterialUniformExpressionRuntimeVirtualTextureUniform(ParameterInfo, TextureIndex, VectorIndex), GetMaterialValueType(Type), TEXT(""));
 }
 
-int32 FHLSLMaterialTranslator::VirtualTextureWorldToUV(int32 WorldPositionIndex, int32 P0, int32 P1, int32 P2)
+int32 FHLSLMaterialTranslator::VirtualTextureWorldToUV(int32 WorldPositionIndex, int32 P0, int32 P1, int32 P2, EPositionOrigin PositionOrigin)
 {
 	if (!UseVirtualTexturing(Platform) || WorldPositionIndex == INDEX_NONE || P0 == INDEX_NONE || P1 == INDEX_NONE || P2 == INDEX_NONE)
 	{
@@ -7807,12 +7822,19 @@ int32 FHLSLMaterialTranslator::VirtualTextureWorldToUV(int32 WorldPositionIndex,
 
 	const EDerivativeStatus WorldPositionDerivStatus = GetDerivativeStatus(WorldPositionIndex);
 
-	FString CodeFinite = FString::Printf(TEXT("VirtualTextureWorldToUV(%s, %s, %s, %s)"), *CoerceParameter(WorldPositionIndex, MCT_LWCVector3), *GetParameterCode(P0), *GetParameterCode(P1), *GetParameterCode(P2));
+	EMaterialValueType WorldPositionType = (EMaterialValueType)0;
+	switch (PositionOrigin)
+	{
+	case EPositionOrigin::Absolute: { WorldPositionType = MCT_LWCVector3; break; }
+	case EPositionOrigin::CameraRelative: { WorldPositionType = MCT_Float3; break; } //P0 is camera-relative
+	default: { checkNoEntry(); }
+	}
+	FString CodeFinite = FString::Printf(TEXT("VirtualTextureWorldToUV(%s, %s, %s, %s)"), *CoerceParameter(WorldPositionIndex, WorldPositionType), *GetParameterCode(P0), *GetParameterCode(P1), *GetParameterCode(P2));
 	if (IsAnalyticDerivEnabled() && IsDerivativeValid(WorldPositionDerivStatus))
 	{
 		if (WorldPositionDerivStatus == EDerivativeStatus::Valid)
 		{
-			WorldPositionIndex = ValidCast(WorldPositionIndex, MCT_LWCVector3);
+			WorldPositionIndex = ValidCast(WorldPositionIndex, WorldPositionType);
 			FString WorldPositionDeriv = GetParameterCodeDeriv(WorldPositionIndex, CompiledPDV_Analytic);
 			FString CodeAnalytic = FString::Printf(TEXT("VirtualTextureWorldToUVDeriv(%s, %s, %s, %s)"), *WorldPositionDeriv, *GetParameterCode(P0), *GetParameterCode(P1), *GetParameterCode(P2));
 			return AddCodeChunkInnerDeriv(*CodeFinite, *CodeAnalytic, MCT_Float2, false, EDerivativeStatus::Valid);
@@ -10624,7 +10646,7 @@ int32 FHLSLMaterialTranslator::TemporalSobol(int32 Index, int32 Seed)
 		*GetParameterCode(Seed));
 }
 
-int32 FHLSLMaterialTranslator::Noise(int32 Position, float Scale, int32 Quality, uint8 NoiseFunction, bool bTurbulence, int32 Levels, float OutputMin, float OutputMax, float LevelScale, int32 FilterWidth, bool bTiling, uint32 RepeatSize)
+int32 FHLSLMaterialTranslator::Noise(int32 Position, EPositionOrigin PositionOrigin, float Scale, int32 Quality, uint8 NoiseFunction, bool bTurbulence, int32 Levels, float OutputMin, float OutputMax, float LevelScale, int32 FilterWidth, bool bTiling, uint32 RepeatSize)
 {
 	if(Position == INDEX_NONE || FilterWidth == INDEX_NONE)
 	{
@@ -10652,6 +10674,12 @@ int32 FHLSLMaterialTranslator::Noise(int32 Position, float Scale, int32 Quality,
 	int32 TilingConst = Constant(bTiling);
 	int32 RepeatSizeConst = Constant(RepeatSize);
 
+	if (PositionOrigin == EPositionOrigin::CameraRelative)
+	{
+		//LWC_TODO: add support for translated world positions in the corresponding HLSL function
+		Position = TransformPosition(MCB_TranslatedWorld, MCB_World, Position);
+	}
+
 	const EMaterialValueType PositionType = GetParameterType(Position);
 	if (IsLWCType(PositionType))
 	{
@@ -10677,13 +10705,19 @@ int32 FHLSLMaterialTranslator::Noise(int32 Position, float Scale, int32 Quality,
 		*GetParameterCode(RepeatSizeConst));
 }
 
-int32 FHLSLMaterialTranslator::VectorNoise(int32 Position, int32 Quality, uint8 NoiseFunction, bool bTiling, uint32 TileSize)
+int32 FHLSLMaterialTranslator::VectorNoise(int32 Position, EPositionOrigin PositionOrigin, int32 Quality, uint8 NoiseFunction, bool bTiling, uint32 TileSize)
 {
 	if (Position == INDEX_NONE)
 	{
 		return INDEX_NONE;
 	}
 
+	if (PositionOrigin == EPositionOrigin::CameraRelative)
+	{
+		//LWC_TODO: add support for translated world positions in the corresponding HLSL function
+		Position = TransformPosition(MCB_TranslatedWorld, MCB_World, Position);
+	}
+
 	int32 QualityConst = Constant(Quality);
 	int32 NoiseFunctionConst = Constant(NoiseFunction);
 	int32 TilingConst = Constant(bTiling);
@@ -10810,7 +10844,7 @@ int32 FHLSLMaterialTranslator::GetHairColorFromMelanin(int32 Melanin, int32 Redn
 	return AddCodeChunk(MCT_Float3, TEXT("MaterialExpressionGetHairColorFromMelanin(%s, %s, %s)"), *GetParameterCode(Melanin), *GetParameterCode(Redness), *GetParameterCode(DyeColor));
 }
 
-int32 FHLSLMaterialTranslator::DistanceToNearestSurface(int32 PositionArg)
+int32 FHLSLMaterialTranslator::DistanceToNearestSurface(int32 PositionArg, EPositionOrigin PositionOrigin)
 {
 	if (ErrorUnlessPlatformSupports(FDataDrivenShaderPlatformInfo::GetSupportsDistanceFields, TEXT("DistanceField")) == INDEX_NONE)
 	{
@@ -10824,10 +10858,10 @@ int32 FHLSLMaterialTranslator::DistanceToNearestSurface(int32 PositionArg)
 
 	MaterialCompilationOutput.bUsesGlobalDistanceField = true;
 
-	return AddCodeChunk(MCT_Float, TEXT("GetDistanceToNearestSurfaceGlobal(%s)"), *CoerceParameter(PositionArg, MCT_LWCVector3));
+	return AddCodeChunk(MCT_Float, TEXT("GetDistanceToNearestSurfaceGlobal(%s)"), *GetWorldPositionOrDefault(PositionArg, PositionOrigin));
 }
 
-int32 FHLSLMaterialTranslator::DistanceFieldGradient(int32 PositionArg)
+int32 FHLSLMaterialTranslator::DistanceFieldGradient(int32 PositionArg, EPositionOrigin PositionOrigin)
 {
 	if (ErrorUnlessPlatformSupports(FDataDrivenShaderPlatformInfo::GetSupportsDistanceFields, TEXT("DistanceField")) == INDEX_NONE)
 	{
@@ -10841,10 +10875,10 @@ int32 FHLSLMaterialTranslator::DistanceFieldGradient(int32 PositionArg)
 
 	MaterialCompilationOutput.bUsesGlobalDistanceField = true;
 
-	return AddCodeChunk(MCT_Float3, TEXT("GetDistanceFieldGradientGlobal(%s)"), *CoerceParameter(PositionArg, MCT_LWCVector3));
+	return AddCodeChunk(MCT_Float3, TEXT("GetDistanceFieldGradientGlobal(%s)"), *GetWorldPositionOrDefault(PositionArg, PositionOrigin));
 }
 
-int32 FHLSLMaterialTranslator::DistanceFieldApproxAO(int32 PositionArg, int32 NormalArg, int32 BaseDistanceArg, int32 RadiusArg, uint32 NumSteps, float StepScale)
+int32 FHLSLMaterialTranslator::DistanceFieldApproxAO(int32 PositionArg, EPositionOrigin PositionOrigin, int32 NormalArg, int32 BaseDistanceArg, int32 RadiusArg, uint32 NumSteps, float StepScale)
 {
 	if (ErrorUnlessPlatformSupports(FDataDrivenShaderPlatformInfo::GetSupportsDistanceFields, TEXT("DistanceField")) == INDEX_NONE)
 	{
@@ -10899,7 +10933,7 @@ int32 FHLSLMaterialTranslator::DistanceFieldApproxAO(int32 PositionArg, int32 No
 
 	return AddCodeChunk(MCT_Float,
 		TEXT("CalculateDistanceFieldApproxAO(%s, %s, %s, %s, %s, %s, %s)"),
-		*CoerceParameter(PositionArg, MCT_LWCVector3),
+		*GetWorldPositionOrDefault(PositionArg, PositionOrigin),
 		*CoerceParameter(NormalArg, MCT_Float3),
 		*GetParameterCode(NumStepsConst),
 		*CoerceParameter(StepDistance, MCT_Float),
@@ -10908,7 +10942,7 @@ int32 FHLSLMaterialTranslator::DistanceFieldApproxAO(int32 PositionArg, int32 No
 		*CoerceParameter(MaxDistance, MCT_Float));
 }
 
-int32 FHLSLMaterialTranslator::SamplePhysicsField(int32 PositionArg, const int32 OutputType, const int32 TargetIndex)
+int32 FHLSLMaterialTranslator::SamplePhysicsField(int32 PositionArg, EPositionOrigin PositionOrigin, const int32 OutputType, const int32 TargetIndex)
 {
 	if (ErrorUnlessFeatureLevelSupported(ERHIFeatureLevel::SM5) == INDEX_NONE)
 	{
@@ -10920,25 +10954,29 @@ int32 FHLSLMaterialTranslator::SamplePhysicsField(int32 PositionArg, const int32
 		return INDEX_NONE;
 	}
 
+	if (TargetIndex == INDEX_NONE)
+	{
+		return INDEX_NONE;
+	}
+
+	if (PositionOrigin == EPositionOrigin::CameraRelative)
+	{
+		// LWC_TODO: support translated-world coordinates
+		PositionArg = TransformPosition(MCB_TranslatedWorld, MCB_World, PositionArg);
+	}
+
 	// LWC_TODO: LWC aware physics field
-	if (TargetIndex != INDEX_NONE)
+	if (OutputType == EFieldOutputType::Field_Output_Vector)
 	{
-		if (OutputType == EFieldOutputType::Field_Output_Vector)
-		{
-			return AddCodeChunk(MCT_Float3, TEXT("MatPhysicsField_SamplePhysicsVectorField(%s,%d)"), *CoerceParameter(PositionArg, MCT_Float3), static_cast<uint8>(TargetIndex));
-		}
-		else if (OutputType == EFieldOutputType::Field_Output_Scalar)
-		{
-			return AddCodeChunk(MCT_Float, TEXT("MatPhysicsField_SamplePhysicsScalarField(%s,%d)"), *CoerceParameter(PositionArg, MCT_Float3), static_cast<uint8>(TargetIndex));
-		}
-		else if (OutputType == EFieldOutputType::Field_Output_Integer)
-		{
-			return AddCodeChunk(MCT_Float, TEXT("MatPhysicsField_SamplePhysicsIntegerField(%s,%d)"), *CoerceParameter(PositionArg, MCT_Float3), static_cast<uint8>(TargetIndex));
-		}
-		else
-		{
-			return INDEX_NONE;
-		}
+		return AddCodeChunk(MCT_Float3, TEXT("MatPhysicsField_SamplePhysicsVectorField(%s,%d)"), *CoerceParameter(PositionArg, MCT_Float3), static_cast<uint8>(TargetIndex));
+	}
+	else if (OutputType == EFieldOutputType::Field_Output_Scalar)
+	{
+		return AddCodeChunk(MCT_Float, TEXT("MatPhysicsField_SamplePhysicsScalarField(%s,%d)"), *CoerceParameter(PositionArg, MCT_Float3), static_cast<uint8>(TargetIndex));
+	}
+	else if (OutputType == EFieldOutputType::Field_Output_Integer)
+	{
+		return AddCodeChunk(MCT_Float, TEXT("MatPhysicsField_SamplePhysicsIntegerField(%s,%d)"), *CoerceParameter(PositionArg, MCT_Float3), static_cast<uint8>(TargetIndex));
 	}
 	else
 	{
@@ -10946,9 +10984,9 @@ int32 FHLSLMaterialTranslator::SamplePhysicsField(int32 PositionArg, const int32
 	}
 }
 
-int32 FHLSLMaterialTranslator::AtmosphericFogColor( int32 WorldPosition )
+int32 FHLSLMaterialTranslator::AtmosphericFogColor( int32 WorldPosition, EPositionOrigin PositionOrigin )
 {
-	return SkyAtmosphereAerialPerspective(WorldPosition);
+	return SkyAtmosphereAerialPerspective(WorldPosition, PositionOrigin);
 }
 
 int32 FHLSLMaterialTranslator::AtmosphericLightVector()
@@ -10961,10 +10999,10 @@ int32 FHLSLMaterialTranslator::AtmosphericLightColor()
 	return AddCodeChunk(MCT_Float3, TEXT("MaterialExpressionAtmosphericLightColor(Parameters)"));
 }
 
-int32 FHLSLMaterialTranslator::SkyAtmosphereLightIlluminance(int32 WorldPosition, int32 LightIndex)
+int32 FHLSLMaterialTranslator::SkyAtmosphereLightIlluminance(int32 WorldPosition, EPositionOrigin PositionOrigin, int32 LightIndex)
 {
 	bUsesSkyAtmosphere = true;
-	FString WorldPosCode = WorldPosition == INDEX_NONE ? FString(TEXT("Parameters.AbsoluteWorldPosition")) : CoerceParameter(WorldPosition, MCT_LWCVector3);
+	FString WorldPosCode = GetWorldPositionOrDefault(WorldPosition, PositionOrigin);
 	return AddCodeChunk(MCT_Float3, TEXT("MaterialExpressionSkyAtmosphereLightIlluminance(Parameters, %s, %d)"), *WorldPosCode, LightIndex);
 }
 
@@ -10993,10 +11031,10 @@ int32 FHLSLMaterialTranslator::SkyAtmosphereViewLuminance()
 	return AddCodeChunk(MCT_Float3, TEXT("MaterialExpressionSkyAtmosphereViewLuminance(Parameters)"));
 }
 
-int32 FHLSLMaterialTranslator::SkyAtmosphereAerialPerspective(int32 WorldPosition)
+int32 FHLSLMaterialTranslator::SkyAtmosphereAerialPerspective(int32 WorldPosition, EPositionOrigin PositionOrigin)
 {
 	bUsesSkyAtmosphere = true;
-	FString WorldPosCode = WorldPosition == INDEX_NONE ? FString(TEXT("Parameters.AbsoluteWorldPosition")) : CoerceParameter(WorldPosition, MCT_LWCVector3);
+	FString WorldPosCode = GetWorldPositionOrDefault(WorldPosition, PositionOrigin);
 	return AddCodeChunk(MCT_Float4, TEXT("MaterialExpressionSkyAtmosphereAerialPerspective(Parameters, %s)"), *WorldPosCode);
 }
 
diff --git a/Engine/Source/Runtime/Engine/Private/Materials/HLSLMaterialTranslator.h b/Engine/Source/Runtime/Engine/Private/Materials/HLSLMaterialTranslator.h
--- a/Engine/Source/Runtime/Engine/Private/Materials/HLSLMaterialTranslator.h
+++ b/Engine/Source/Runtime/Engine/Private/Materials/HLSLMaterialTranslator.h
@@ -659,6 +659,9 @@ protected:
 	void AddEstimatedTextureSample(const uint32 Count = 1);
 	void AddLWCFuncUsage(ELWCFunctionKind Kind, const uint32 Count = 1);
 
+	/** If InWorldPosition is set, coerce it to the correct type, else get the position from the material params */
+	FString GetWorldPositionOrDefault(int32 InWorldPosition, EPositionOrigin PositionOrigin);
+
 	/** Creates a unique symbol name and adds it to the symbol list. */
 	FString CreateSymbolName(const TCHAR* SymbolNameHint);
 
@@ -1020,7 +1023,7 @@ protected:
 	virtual int32 VirtualTextureParameter(FName ParameterName, URuntimeVirtualTexture* DefaultValue, int32 TextureLayerIndex, int32 PageTableLayerIndex, int32& TextureReferenceIndex, EMaterialSamplerType SamplerType) override;
 	virtual int32 VirtualTextureUniform(int32 TextureIndex, int32 VectorIndex, UE::Shader::EValueType Type) override;
 	virtual int32 VirtualTextureUniform(FName ParameterName, int32 TextureIndex, int32 VectorIndex, UE::Shader::EValueType Type) override;
-	virtual int32 VirtualTextureWorldToUV(int32 WorldPositionIndex, int32 P0, int32 P1, int32 P2) override;
+	virtual int32 VirtualTextureWorldToUV(int32 WorldPositionIndex, int32 P0, int32 P1, int32 P2, EPositionOrigin PositionOrigin) override;
 	virtual int32 VirtualTextureUnpack(int32 CodeIndex0, int32 CodeIndex1, int32 CodeIndex2, int32 P0, EVirtualTextureUnpackType UnpackType) override;
 
 	virtual int32 ExternalTexture(const FGuid& ExternalTextureGuid) override;
@@ -1122,8 +1125,8 @@ protected:
 	virtual int32 DepthOfFieldFunction(int32 Depth, int32 FunctionValueIndex) override;
 	virtual int32 Sobol(int32 Cell, int32 Index, int32 Seed) override;
 	virtual int32 TemporalSobol(int32 Index, int32 Seed) override;
-	virtual int32 Noise(int32 Position, float Scale, int32 Quality, uint8 NoiseFunction, bool bTurbulence, int32 Levels, float OutputMin, float OutputMax, float LevelScale, int32 FilterWidth, bool bTiling, uint32 RepeatSize) override;
-	virtual int32 VectorNoise(int32 Position, int32 Quality, uint8 NoiseFunction, bool bTiling, uint32 TileSize) override;
+	virtual int32 Noise(int32 Position, EPositionOrigin PositionOrigin, float Scale, int32 Quality, uint8 NoiseFunction, bool bTurbulence, int32 Levels, float OutputMin, float OutputMax, float LevelScale, int32 FilterWidth, bool bTiling, uint32 RepeatSize) override;
+	virtual int32 VectorNoise(int32 Position, EPositionOrigin PositionOrigin, int32 Quality, uint8 NoiseFunction, bool bTiling, uint32 TileSize) override;
 
 	virtual int32 BlackBody(int32 Temp) override;
 
@@ -1142,20 +1145,20 @@ protected:
 	virtual int32 GetHairAtlasUVs() override;
 	virtual int32 GetHairGroupIndex() override;
 	virtual int32 GetHairColorFromMelanin(int32 Melanin, int32 Redness, int32 DyeColor) override;
-	virtual int32 DistanceToNearestSurface(int32 PositionArg) override;
-	virtual int32 DistanceFieldGradient(int32 PositionArg) override;
-	virtual int32 DistanceFieldApproxAO(int32 PositionArg, int32 NormalArg, int32 BaseDistanceArg, int32 RadiusArg, uint32 NumSteps, float StepScale) override;
-	virtual int32 SamplePhysicsField(int32 PositionArg, const int32 OutputType, const int32 TargetIndex) override;
-	virtual int32 AtmosphericFogColor(int32 WorldPosition) override;
+	virtual int32 DistanceToNearestSurface(int32 PositionArg, EPositionOrigin PositionOrigin) override;
+	virtual int32 DistanceFieldGradient(int32 PositionArg, EPositionOrigin PositionOrigin) override;
+	virtual int32 DistanceFieldApproxAO(int32 PositionArg, EPositionOrigin PositionOrigin, int32 NormalArg, int32 BaseDistanceArg, int32 RadiusArg, uint32 NumSteps, float StepScale) override;
+	virtual int32 SamplePhysicsField(int32 PositionArg, EPositionOrigin PositionOrigin, const int32 OutputType, const int32 TargetIndex) override;
+	virtual int32 AtmosphericFogColor(int32 WorldPosition, EPositionOrigin PositionOrigin) override;
 	virtual int32 AtmosphericLightVector() override;
 	virtual int32 AtmosphericLightColor() override;
 
-	virtual int32 SkyAtmosphereLightIlluminance(int32 WorldPosition, int32 LightIndex) override;
+	virtual int32 SkyAtmosphereLightIlluminance(int32 WorldPosition, EPositionOrigin PositionOrigin, int32 LightIndex) override;
 	virtual int32 SkyAtmosphereLightIlluminanceOnGround(int32 LightIndex) override;
 	virtual int32 SkyAtmosphereLightDirection(int32 LightIndex) override;
 	virtual int32 SkyAtmosphereLightDiskLuminance(int32 LightIndex, int32 OverrideAtmosphereLightDiscCosHalfApexAngle) override;
 	virtual int32 SkyAtmosphereViewLuminance() override;
-	virtual int32 SkyAtmosphereAerialPerspective(int32 WorldPosition) override;
+	virtual int32 SkyAtmosphereAerialPerspective(int32 WorldPosition, EPositionOrigin PositionOrigin) override;
 	virtual int32 SkyAtmosphereDistantLightScatteredLuminance() override;
 
 	virtual int32 SkyLightEnvMapSample(int32 DirectionCodeChunk, int32 RoughnessCodeChunk) override;
diff --git a/Engine/Source/Runtime/Engine/Private/Materials/MaterialExpressions.cpp b/Engine/Source/Runtime/Engine/Private/Materials/MaterialExpressions.cpp
--- a/Engine/Source/Runtime/Engine/Private/Materials/MaterialExpressions.cpp
+++ b/Engine/Source/Runtime/Engine/Private/Materials/MaterialExpressions.cpp
@@ -674,6 +674,28 @@ int32 CompileSubstrateBlendFunction(FMaterialCompiler* Compiler, const int32 A,
 	return INDEX_NONE;
 }
 
+EWorldPositionIncludedOffsets GetWorldPositionTypeWithOrigin(EPositionOrigin PositionOrigin, bool IncludeOffsets = true)
+{
+	switch (PositionOrigin)
+	{
+	case EPositionOrigin::Absolute: return IncludeOffsets ? WPT_Default : WPT_ExcludeAllShaderOffsets;
+	case EPositionOrigin::CameraRelative: return IncludeOffsets ? WPT_CameraRelative : WPT_CameraRelativeNoOffsets;
+	default: checkNoEntry();
+	}
+	return WPT_Default;
+}
+
+FName GetWorldPositionInputName(EPositionOrigin PositionOrigin)
+{
+	switch (PositionOrigin)
+	{
+	case EPositionOrigin::CameraRelative: return TEXT("Translated World Position");
+	case EPositionOrigin::Absolute: return TEXT("World Position");
+	default: checkNoEntry();
+	}
+	return FName();
+}
+
 void FMaterialExpressionCollection::AddExpression(UMaterialExpression* InExpression)
 {
 	Expressions.AddUnique(InExpression);
@@ -3167,6 +3189,16 @@ UObject* UMaterialExpressionRuntimeVirtualTextureSample::GetReferencedTexture()
 
 #if WITH_EDITOR
 
+FName UMaterialExpressionRuntimeVirtualTextureSample::GetInputName(int32 InputIndex) const
+{
+	if (CachedInputs[InputIndex] == &WorldPosition)
+	{
+		return GetWorldPositionInputName(WorldPositionOriginType);
+	}
+
+	return Super::GetInputName(InputIndex);
+}
+
 void UMaterialExpressionRuntimeVirtualTextureSample::PostLoad()
 {
 	Super::PostLoad();
@@ -3201,7 +3233,8 @@ bool UMaterialExpressionRuntimeVirtualTextureSample::CanEditChange(const FProper
 void UMaterialExpressionRuntimeVirtualTextureSample::PostEditChangeProperty(FPropertyChangedEvent& PropertyChangedEvent)
 {
 	// Update MaterialType setting to match VirtualTexture
-	if (PropertyChangedEvent.Property && PropertyChangedEvent.Property->GetName() == TEXT("VirtualTexture"))
+	const FName PropertyName = PropertyChangedEvent.GetPropertyName();
+	if (PropertyName == GET_MEMBER_NAME_STRING_CHECKED(ThisClass, VirtualTexture))
 	{
 		if (VirtualTexture != nullptr)
 		{
@@ -3209,6 +3242,13 @@ void UMaterialExpressionRuntimeVirtualTextureSample::PostEditChangeProperty(FPro
 			FEditorSupportDelegates::ForcePropertyWindowRebuild.Broadcast(this);
 		}
 	}
+	else if (PropertyName == GET_MEMBER_NAME_STRING_CHECKED(ThisClass, WorldPositionOriginType))
+	{
+		if (GraphNode)
+		{
+			GraphNode->ReconstructNode();
+		}
+	}
 
 	Super::PostEditChangeProperty(PropertyChangedEvent);
 }
@@ -3440,16 +3480,34 @@ int32 UMaterialExpressionRuntimeVirtualTextureSample::Compile(class FMaterialCom
 		}
 		else
 		{
-			WorldPositionIndex = Compiler->WorldPosition(WPT_Default);
+			WorldPositionIndex = Compiler->WorldPosition(GetWorldPositionTypeWithOrigin(WorldPositionOriginType));
 			ensure(WorldPositionIndex != INDEX_NONE);
 		}
 		
 		if (WorldPositionIndex != INDEX_NONE)
 		{
-			const int32 P0 = Uniforms[ERuntimeVirtualTextureShaderUniform_WorldToUVTransform0];
-			const int32 P1 = Uniforms[ERuntimeVirtualTextureShaderUniform_WorldToUVTransform1];
-			const int32 P2 = Uniforms[ERuntimeVirtualTextureShaderUniform_WorldToUVTransform2];
-			CoordinateIndex = Compiler->VirtualTextureWorldToUV(WorldPositionIndex, P0, P1, P2);
+			if (WorldPositionOriginType == EPositionOrigin::Absolute)
+			{
+				const int32 P0 = Uniforms[ERuntimeVirtualTextureShaderUniform_WorldToUVTransform0];
+				const int32 P1 = Uniforms[ERuntimeVirtualTextureShaderUniform_WorldToUVTransform1];
+				const int32 P2 = Uniforms[ERuntimeVirtualTextureShaderUniform_WorldToUVTransform2];
+				CoordinateIndex = Compiler->VirtualTextureWorldToUV(WorldPositionIndex, P0, P1, P2, EPositionOrigin::Absolute);
+			}
+			else if (WorldPositionOriginType == EPositionOrigin::CameraRelative)
+			{
+				//TODO: optimize by calculating translated world to VT directly.
+				//This requires some more work as the transform is currently fed in through a preshader variable, which is cached.
+				const int32 AbsWorldPosIndex = Compiler->TransformPosition(EMaterialCommonBasis::MCB_TranslatedWorld, EMaterialCommonBasis::MCB_World, WorldPositionIndex);
+
+				const int32 P0 = Uniforms[ERuntimeVirtualTextureShaderUniform_WorldToUVTransform0];
+				const int32 P1 = Uniforms[ERuntimeVirtualTextureShaderUniform_WorldToUVTransform1];
+				const int32 P2 = Uniforms[ERuntimeVirtualTextureShaderUniform_WorldToUVTransform2];
+				CoordinateIndex = Compiler->VirtualTextureWorldToUV(AbsWorldPosIndex, P0, P1, P2, EPositionOrigin::Absolute);
+			}
+			else
+			{
+				checkNoEntry();
+			}
 		}
 	}
 	
@@ -19586,6 +19644,30 @@ bool UMaterialExpressionNoise::CanEditChange(const FProperty* InProperty) const
 	return bIsEditable;
 }
 
+FName UMaterialExpressionNoise::GetInputName(int32 InputIndex) const
+{
+	if (CachedInputs[InputIndex] == &Position)
+	{
+		return GetWorldPositionInputName(WorldPositionOriginType);
+	}
+
+	return Super::GetInputName(InputIndex);
+}
+
+void UMaterialExpressionNoise::PostEditChangeProperty(FPropertyChangedEvent& PropertyChangedEvent)
+{
+	const FName PropertyName = PropertyChangedEvent.GetPropertyName();
+	if (PropertyName == GET_MEMBER_NAME_STRING_CHECKED(ThisClass, WorldPositionOriginType))
+	{
+		if (GraphNode)
+		{
+			GraphNode->ReconstructNode();
+		}
+	}
+
+	Super::PostEditChangeProperty(PropertyChangedEvent);
+}
+
 int32 UMaterialExpressionNoise::Compile(class FMaterialCompiler* Compiler, int32 OutputIndex)
 {
 	int32 PositionInput;
@@ -19596,7 +19678,7 @@ int32 UMaterialExpressionNoise::Compile(class FMaterialCompiler* Compiler, int32
 	}
 	else
 	{
-		PositionInput = Compiler->WorldPosition(WPT_Default);
+		PositionInput = Compiler->WorldPosition(GetWorldPositionTypeWithOrigin(WorldPositionOriginType));
 	}
 
 	int32 FilterWidthInput;
@@ -19610,7 +19692,7 @@ int32 UMaterialExpressionNoise::Compile(class FMaterialCompiler* Compiler, int32
 		FilterWidthInput = Compiler->Constant(0);
 	}
 
-	return Compiler->Noise(PositionInput, Scale, Quality, NoiseFunction, bTurbulence, Levels, OutputMin, OutputMax, LevelScale, FilterWidthInput, bTiling, RepeatSize);
+	return Compiler->Noise(PositionInput, WorldPositionOriginType, Scale, Quality, NoiseFunction, bTurbulence, Levels, OutputMin, OutputMax, LevelScale, FilterWidthInput, bTiling, RepeatSize);
 }
 
 void UMaterialExpressionNoise::GetCaption(TArray<FString>& OutCaptions) const
@@ -19673,6 +19755,30 @@ bool UMaterialExpressionVectorNoise::CanEditChange(const FProperty* InProperty)
 	return bIsEditable;
 }
 
+FName UMaterialExpressionVectorNoise::GetInputName(int32 InputIndex) const
+{
+	if (CachedInputs[InputIndex] == &Position)
+	{
+		return GetWorldPositionInputName(WorldPositionOriginType);
+	}
+
+	return Super::GetInputName(InputIndex);
+}
+
+void UMaterialExpressionVectorNoise::PostEditChangeProperty(FPropertyChangedEvent& PropertyChangedEvent)
+{
+	const FName PropertyName = PropertyChangedEvent.GetPropertyName();
+	if (PropertyName == GET_MEMBER_NAME_STRING_CHECKED(ThisClass, WorldPositionOriginType))
+	{
+		if (GraphNode)
+		{
+			GraphNode->ReconstructNode();
+		}
+	}
+
+	Super::PostEditChangeProperty(PropertyChangedEvent);
+}
+
 int32 UMaterialExpressionVectorNoise::Compile(class FMaterialCompiler* Compiler, int32 OutputIndex)
 {
 	int32 PositionInput;
@@ -19683,10 +19789,10 @@ int32 UMaterialExpressionVectorNoise::Compile(class FMaterialCompiler* Compiler,
 	}
 	else
 	{
-		PositionInput = Compiler->WorldPosition(WPT_Default);
+		PositionInput = Compiler->WorldPosition(GetWorldPositionTypeWithOrigin(WorldPositionOriginType));
 	}
 
-	return Compiler->VectorNoise(PositionInput, Quality, NoiseFunction, bTiling, TileSize);
+	return Compiler->VectorNoise(PositionInput, WorldPositionOriginType, Quality, NoiseFunction, bTiling, TileSize);
 }
 
 void UMaterialExpressionVectorNoise::GetCaption(TArray<FString>& OutCaptions) const
@@ -19767,6 +19873,30 @@ UMaterialExpressionDistanceToNearestSurface::UMaterialExpressionDistanceToNeares
 }
 
 #if WITH_EDITOR
+FName UMaterialExpressionDistanceToNearestSurface::GetInputName(int32 InputIndex) const
+{
+	if (CachedInputs[InputIndex] == &Position)
+	{
+		return GetWorldPositionInputName(WorldPositionOriginType);
+	}
+
+	return Super::GetInputName(InputIndex);
+}
+
+void UMaterialExpressionDistanceToNearestSurface::PostEditChangeProperty(FPropertyChangedEvent& PropertyChangedEvent)
+{
+	const FName PropertyName = PropertyChangedEvent.GetPropertyName();
+	if (PropertyName == GET_MEMBER_NAME_STRING_CHECKED(ThisClass, WorldPositionOriginType))
+	{
+		if (GraphNode)
+		{
+			GraphNode->ReconstructNode();
+		}
+	}
+
+	Super::PostEditChangeProperty(PropertyChangedEvent);
+}
+
 int32 UMaterialExpressionDistanceToNearestSurface::Compile(class FMaterialCompiler* Compiler, int32 OutputIndex)
 {
 	int32 PositionArg = INDEX_NONE;
@@ -19777,10 +19907,10 @@ int32 UMaterialExpressionDistanceToNearestSurface::Compile(class FMaterialCompil
 	}
 	else 
 	{
-		PositionArg = Compiler->WorldPosition(WPT_Default);
+		PositionArg = Compiler->WorldPosition(GetWorldPositionTypeWithOrigin(WorldPositionOriginType));
 	}
 
-	return Compiler->DistanceToNearestSurface(PositionArg);
+	return Compiler->DistanceToNearestSurface(PositionArg, WorldPositionOriginType);
 }
 
 void UMaterialExpressionDistanceToNearestSurface::GetCaption(TArray<FString>& OutCaptions) const
@@ -19812,6 +19942,30 @@ UMaterialExpressionDistanceFieldGradient::UMaterialExpressionDistanceFieldGradie
 }
 
 #if WITH_EDITOR
+FName UMaterialExpressionDistanceFieldGradient::GetInputName(int32 InputIndex) const
+{
+	if (CachedInputs[InputIndex] == &Position)
+	{
+		return GetWorldPositionInputName(WorldPositionOriginType);
+	}
+
+	return Super::GetInputName(InputIndex);
+}
+
+void UMaterialExpressionDistanceFieldGradient::PostEditChangeProperty(FPropertyChangedEvent& PropertyChangedEvent)
+{
+	const FName PropertyName = PropertyChangedEvent.GetPropertyName();
+	if (PropertyName == GET_MEMBER_NAME_STRING_CHECKED(ThisClass, WorldPositionOriginType))
+	{
+		if (GraphNode)
+		{
+			GraphNode->ReconstructNode();
+		}
+	}
+
+	Super::PostEditChangeProperty(PropertyChangedEvent);
+}
+
 int32 UMaterialExpressionDistanceFieldGradient::Compile(class FMaterialCompiler* Compiler, int32 OutputIndex)
 {
 	int32 PositionArg = INDEX_NONE;
@@ -19822,10 +19976,10 @@ int32 UMaterialExpressionDistanceFieldGradient::Compile(class FMaterialCompiler*
 	}
 	else 
 	{
-		PositionArg = Compiler->WorldPosition(WPT_Default);
+		PositionArg = Compiler->WorldPosition(GetWorldPositionTypeWithOrigin(WorldPositionOriginType));
 	}
 
-	return Compiler->DistanceFieldGradient(PositionArg);
+	return Compiler->DistanceFieldGradient(PositionArg, WorldPositionOriginType);
 }
 
 void UMaterialExpressionDistanceFieldGradient::GetCaption(TArray<FString>& OutCaptions) const
@@ -19862,6 +20016,30 @@ UMaterialExpressionDistanceFieldApproxAO::UMaterialExpressionDistanceFieldApprox
 }
 
 #if WITH_EDITOR
+FName UMaterialExpressionDistanceFieldApproxAO::GetInputName(int32 InputIndex) const
+{
+	if (CachedInputs[InputIndex] == &Position)
+	{
+		return GetWorldPositionInputName(WorldPositionOriginType);
+	}
+
+	return Super::GetInputName(InputIndex);
+}
+
+void UMaterialExpressionDistanceFieldApproxAO::PostEditChangeProperty(FPropertyChangedEvent& PropertyChangedEvent)
+{
+	const FName PropertyName = PropertyChangedEvent.GetPropertyName();
+	if (PropertyName == GET_MEMBER_NAME_STRING_CHECKED(ThisClass, WorldPositionOriginType))
+	{
+		if (GraphNode)
+		{
+			GraphNode->ReconstructNode();
+		}
+	}
+
+	Super::PostEditChangeProperty(PropertyChangedEvent);
+}
+
 int32 UMaterialExpressionDistanceFieldApproxAO::Compile(class FMaterialCompiler* Compiler, int32 OutputIndex)
 {
 	int32 PositionArg = INDEX_NONE;
@@ -19872,7 +20050,7 @@ int32 UMaterialExpressionDistanceFieldApproxAO::Compile(class FMaterialCompiler*
 	}
 	else
 	{
-		PositionArg = Compiler->WorldPosition(WPT_Default);
+		PositionArg = Compiler->WorldPosition(GetWorldPositionTypeWithOrigin(WorldPositionOriginType));
 	}
 
 	int32 NormalArg = INDEX_NONE;
@@ -19889,7 +20067,7 @@ int32 UMaterialExpressionDistanceFieldApproxAO::Compile(class FMaterialCompiler*
 	int32 BaseDistanceArg = BaseDistance.GetTracedInput().Expression ? BaseDistance.Compile(Compiler) : Compiler->Constant(BaseDistanceDefault);
 	int32 RadiusArg = Radius.GetTracedInput().Expression ? Radius.Compile(Compiler) : Compiler->Constant(RadiusDefault);
 
-	return Compiler->DistanceFieldApproxAO(PositionArg, NormalArg, BaseDistanceArg, RadiusArg, NumSteps, StepScaleDefault);
+	return Compiler->DistanceFieldApproxAO(PositionArg, WorldPositionOriginType, NormalArg, BaseDistanceArg, RadiusArg, NumSteps, StepScaleDefault);
 }
 
 void UMaterialExpressionDistanceFieldApproxAO::GetCaption(TArray<FString>& OutCaptions) const
@@ -19924,6 +20102,30 @@ UMaterialExpressionSamplePhysicsVectorField::UMaterialExpressionSamplePhysicsVec
 }
 
 #if WITH_EDITOR
+FName UMaterialExpressionSamplePhysicsVectorField::GetInputName(int32 InputIndex) const
+{
+	if (CachedInputs[InputIndex] == &WorldPosition)
+	{
+		return GetWorldPositionInputName(WorldPositionOriginType);
+	}
+
+	return Super::GetInputName(InputIndex);
+}
+
+void UMaterialExpressionSamplePhysicsVectorField::PostEditChangeProperty(FPropertyChangedEvent& PropertyChangedEvent)
+{
+	const FName PropertyName = PropertyChangedEvent.GetPropertyName();
+	if (PropertyName == GET_MEMBER_NAME_STRING_CHECKED(ThisClass, WorldPositionOriginType))
+	{
+		if (GraphNode)
+		{
+			GraphNode->ReconstructNode();
+		}
+	}
+
+	Super::PostEditChangeProperty(PropertyChangedEvent);
+}
+
 int32 UMaterialExpressionSamplePhysicsVectorField::Compile(class FMaterialCompiler* Compiler, int32 OutputIndex)
 {
 	int32 PositionArg = INDEX_NONE;
@@ -19934,10 +20136,10 @@ int32 UMaterialExpressionSamplePhysicsVectorField::Compile(class FMaterialCompil
 	}
 	else
 	{
-		PositionArg = Compiler->WorldPosition(WPT_Default);
+		PositionArg = Compiler->WorldPosition(GetWorldPositionTypeWithOrigin(WorldPositionOriginType));
 	}
 
-	return Compiler->SamplePhysicsField(PositionArg, EFieldOutputType::Field_Output_Vector, static_cast<uint8>(FieldTarget));
+	return Compiler->SamplePhysicsField(PositionArg, WorldPositionOriginType, EFieldOutputType::Field_Output_Vector, static_cast<uint8>(FieldTarget));
 }
 
 void UMaterialExpressionSamplePhysicsVectorField::GetCaption(TArray<FString>& OutCaptions) const
@@ -19972,6 +20174,30 @@ UMaterialExpressionSamplePhysicsScalarField::UMaterialExpressionSamplePhysicsSca
 }
 
 #if WITH_EDITOR
+FName UMaterialExpressionSamplePhysicsScalarField::GetInputName(int32 InputIndex) const
+{
+	if (CachedInputs[InputIndex] == &WorldPosition)
+	{
+		return GetWorldPositionInputName(WorldPositionOriginType);
+	}
+
+	return Super::GetInputName(InputIndex);
+}
+
+void UMaterialExpressionSamplePhysicsScalarField::PostEditChangeProperty(FPropertyChangedEvent& PropertyChangedEvent)
+{
+	const FName PropertyName = PropertyChangedEvent.GetPropertyName();
+	if (PropertyName == GET_MEMBER_NAME_STRING_CHECKED(ThisClass, WorldPositionOriginType))
+	{
+		if (GraphNode)
+		{
+			GraphNode->ReconstructNode();
+		}
+	}
+
+	Super::PostEditChangeProperty(PropertyChangedEvent);
+}
+
 int32 UMaterialExpressionSamplePhysicsScalarField::Compile(class FMaterialCompiler* Compiler, int32 OutputIndex)
 {
 	int32 PositionArg = INDEX_NONE;
@@ -19982,10 +20208,10 @@ int32 UMaterialExpressionSamplePhysicsScalarField::Compile(class FMaterialCompil
 	}
 	else
 	{
-		PositionArg = Compiler->WorldPosition(WPT_Default);
+		PositionArg = Compiler->WorldPosition(GetWorldPositionTypeWithOrigin(WorldPositionOriginType));
 	}
 
-	return Compiler->SamplePhysicsField(PositionArg, EFieldOutputType::Field_Output_Scalar, static_cast<uint8>(FieldTarget));
+	return Compiler->SamplePhysicsField(PositionArg, WorldPositionOriginType, EFieldOutputType::Field_Output_Scalar, static_cast<uint8>(FieldTarget));
 }
 
 void UMaterialExpressionSamplePhysicsScalarField::GetCaption(TArray<FString>& OutCaptions) const
@@ -20020,7 +20246,33 @@ UMaterialExpressionSamplePhysicsIntegerField::UMaterialExpressionSamplePhysicsIn
 #endif
 }
 
+
+
 #if WITH_EDITOR
+FName UMaterialExpressionSamplePhysicsIntegerField::GetInputName(int32 InputIndex) const
+{
+	if (CachedInputs[InputIndex] == &WorldPosition)
+	{
+		return GetWorldPositionInputName(WorldPositionOriginType);
+	}
+
+	return Super::GetInputName(InputIndex);
+}
+
+void UMaterialExpressionSamplePhysicsIntegerField::PostEditChangeProperty(FPropertyChangedEvent& PropertyChangedEvent)
+{
+	const FName PropertyName = PropertyChangedEvent.GetPropertyName();
+	if (PropertyName == GET_MEMBER_NAME_STRING_CHECKED(ThisClass, WorldPositionOriginType))
+	{
+		if (GraphNode)
+		{
+			GraphNode->ReconstructNode();
+		}
+	}
+
+	Super::PostEditChangeProperty(PropertyChangedEvent);
+}
+
 int32 UMaterialExpressionSamplePhysicsIntegerField::Compile(class FMaterialCompiler* Compiler, int32 OutputIndex)
 {
 	int32 PositionArg = INDEX_NONE;
@@ -20031,10 +20283,10 @@ int32 UMaterialExpressionSamplePhysicsIntegerField::Compile(class FMaterialCompi
 	}
 	else
 	{
-		PositionArg = Compiler->WorldPosition(WPT_Default);
+		PositionArg = Compiler->WorldPosition(GetWorldPositionTypeWithOrigin(WorldPositionOriginType));
 	}
 
-	return Compiler->SamplePhysicsField(PositionArg, EFieldOutputType::Field_Output_Integer, static_cast<uint8>(FieldTarget));
+	return Compiler->SamplePhysicsField(PositionArg, WorldPositionOriginType, EFieldOutputType::Field_Output_Integer, static_cast<uint8>(FieldTarget));
 }
 
 void UMaterialExpressionSamplePhysicsIntegerField::GetCaption(TArray<FString>& OutCaptions) const
@@ -21236,6 +21488,30 @@ UMaterialExpressionAtmosphericFogColor::UMaterialExpressionAtmosphericFogColor(c
 }
 
 #if WITH_EDITOR
+FName UMaterialExpressionAtmosphericFogColor::GetInputName(int32 InputIndex) const
+{
+	if (CachedInputs[InputIndex] == &WorldPosition)
+	{
+		return GetWorldPositionInputName(WorldPositionOriginType);
+	}
+
+	return Super::GetInputName(InputIndex);
+}
+
+void UMaterialExpressionAtmosphericFogColor::PostEditChangeProperty(FPropertyChangedEvent& PropertyChangedEvent)
+{
+	const FName PropertyName = PropertyChangedEvent.GetPropertyName();
+	if (PropertyName == GET_MEMBER_NAME_STRING_CHECKED(ThisClass, WorldPositionOriginType))
+	{
+		if (GraphNode)
+		{
+			GraphNode->ReconstructNode();
+		}
+	}
+
+	Super::PostEditChangeProperty(PropertyChangedEvent);
+}
+
 int32 UMaterialExpressionAtmosphericFogColor::Compile(class FMaterialCompiler* Compiler, int32 OutputIndex)
 {
 	int32 WorldPositionInput = INDEX_NONE;
@@ -21245,7 +21521,7 @@ int32 UMaterialExpressionAtmosphericFogColor::Compile(class FMaterialCompiler* C
 		WorldPositionInput = WorldPosition.Compile(Compiler);
 	}
 
-	return Compiler->AtmosphericFogColor( WorldPositionInput );
+	return Compiler->AtmosphericFogColor( WorldPositionInput, WorldPositionOriginType );
 }
 
 void UMaterialExpressionAtmosphericFogColor::GetCaption(TArray<FString>& OutCaptions) const
@@ -21815,6 +22091,30 @@ UMaterialExpressionSkyAtmosphereLightIlluminance::UMaterialExpressionSkyAtmosphe
 }
 
 #if WITH_EDITOR
+FName UMaterialExpressionSkyAtmosphereLightIlluminance::GetInputName(int32 InputIndex) const
+{
+	if (CachedInputs[InputIndex] == &WorldPosition)
+	{
+		return GetWorldPositionInputName(WorldPositionOriginType);
+	}
+
+	return Super::GetInputName(InputIndex);
+}
+
+void UMaterialExpressionSkyAtmosphereLightIlluminance::PostEditChangeProperty(FPropertyChangedEvent& PropertyChangedEvent)
+{
+	const FName PropertyName = PropertyChangedEvent.GetPropertyName();
+	if (PropertyName == GET_MEMBER_NAME_STRING_CHECKED(ThisClass, WorldPositionOriginType))
+	{
+		if (GraphNode)
+		{
+			GraphNode->ReconstructNode();
+		}
+	}
+
+	Super::PostEditChangeProperty(PropertyChangedEvent);
+}
+
 int32 UMaterialExpressionSkyAtmosphereLightIlluminance::Compile(class FMaterialCompiler* Compiler, int32 OutputIndex)
 {
 	int32 WorldPositionInput;
@@ -21824,9 +22124,9 @@ int32 UMaterialExpressionSkyAtmosphereLightIlluminance::Compile(class FMaterialC
 	}
 	else
 	{
-		WorldPositionInput = Compiler->WorldPosition(WPT_Default);
+		WorldPositionInput = Compiler->WorldPosition(GetWorldPositionTypeWithOrigin(WorldPositionOriginType));
 	}
-	return Compiler->SkyAtmosphereLightIlluminance(WorldPositionInput, LightIndex);
+	return Compiler->SkyAtmosphereLightIlluminance(WorldPositionInput, WorldPositionOriginType, LightIndex);
 }
 
 void UMaterialExpressionSkyAtmosphereLightIlluminance::GetCaption(TArray<FString>& OutCaptions) const
@@ -22012,6 +22312,30 @@ UMaterialExpressionSkyAtmosphereAerialPerspective::UMaterialExpressionSkyAtmosph
 }
 
 #if WITH_EDITOR
+FName UMaterialExpressionSkyAtmosphereAerialPerspective::GetInputName(int32 InputIndex) const
+{
+	if (CachedInputs[InputIndex] == &WorldPosition)
+	{
+		return GetWorldPositionInputName(WorldPositionOriginType);
+	}
+
+	return Super::GetInputName(InputIndex);
+}
+
+void UMaterialExpressionSkyAtmosphereAerialPerspective::PostEditChangeProperty(FPropertyChangedEvent& PropertyChangedEvent)
+{
+	const FName PropertyName = PropertyChangedEvent.GetPropertyName();
+	if (PropertyName == GET_MEMBER_NAME_STRING_CHECKED(ThisClass, WorldPositionOriginType))
+	{
+		if (GraphNode)
+		{
+			GraphNode->ReconstructNode();
+		}
+	}
+
+	Super::PostEditChangeProperty(PropertyChangedEvent);
+}
+
 int32 UMaterialExpressionSkyAtmosphereAerialPerspective::Compile(class FMaterialCompiler* Compiler, int32 OutputIndex)
 {
 	int32 WorldPositionInput;
@@ -22021,9 +22345,9 @@ int32 UMaterialExpressionSkyAtmosphereAerialPerspective::Compile(class FMaterial
 	}
 	else
 	{
-		WorldPositionInput = Compiler->WorldPosition(WPT_Default);
+		WorldPositionInput = Compiler->WorldPosition(GetWorldPositionTypeWithOrigin(WorldPositionOriginType));
 	}
-	return Compiler->SkyAtmosphereAerialPerspective(WorldPositionInput);
+	return Compiler->SkyAtmosphereAerialPerspective(WorldPositionInput, WorldPositionOriginType);
 }
 
 void UMaterialExpressionSkyAtmosphereAerialPerspective::GetCaption(TArray<FString>& OutCaptions) const
diff --git a/Engine/Source/Runtime/Engine/Public/MaterialCompiler.h b/Engine/Source/Runtime/Engine/Public/MaterialCompiler.h
--- a/Engine/Source/Runtime/Engine/Public/MaterialCompiler.h
+++ b/Engine/Source/Runtime/Engine/Public/MaterialCompiler.h
@@ -234,6 +234,7 @@ public:
 	virtual int32 GetPixelPosition() = 0;
 	virtual int32 WorldPosition(EWorldPositionIncludedOffsets WorldPositionIncludedOffsets) = 0;
 	virtual int32 ObjectWorldPosition(EPositionOrigin OriginType) = 0;
+	UE_DEPRECATED(5.4, "Use ObjectWorldPosition(EPositionOrigin) instead")
 	int32 ObjectWorldPosition() { return ObjectWorldPosition(EPositionOrigin::Absolute); }
 	virtual int32 ObjectRadius() = 0;
 	virtual int32 ObjectBounds() = 0;
@@ -241,12 +242,14 @@ public:
 	virtual int32 PreSkinnedLocalBounds(int32 OutputIndex) = 0;
 	virtual int32 DistanceCullFade() = 0;
 	virtual int32 ActorWorldPosition(EPositionOrigin OriginType) = 0;
+	UE_DEPRECATED(5.4, "Use ActorWorldPosition(EPositionOrigin) instead")
 	int32 ActorWorldPosition() { return ActorWorldPosition(EPositionOrigin::Absolute); }
 	virtual int32 ParticleMacroUV() = 0;
 	virtual int32 ParticleSubUV(int32 TextureIndex, EMaterialSamplerType SamplerType, int32 MipValue0Index, int32 MipValue1Index, ETextureMipValueMode MipValueMode, bool bBlend) = 0;
 	virtual int32 ParticleSubUVProperty(int32 PropertyIndex) = 0;
 	virtual int32 ParticleColor() = 0;
 	virtual int32 ParticlePosition(EPositionOrigin OriginType) = 0;
+	UE_DEPRECATED(5.4, "Use ParticlePosition(EPositionOrigin) instead")
 	int32 ParticlePosition() { return ParticlePosition(EPositionOrigin::Absolute); }
 	virtual int32 ParticleRadius() = 0;
 	virtual int32 SphericalParticleOpacity(int32 Density) = 0;
@@ -278,7 +281,9 @@ public:
 	virtual int32 VirtualTextureParameter(FName ParameterName, URuntimeVirtualTexture* DefaultValue, int32 TextureLayerIndex, int32 PageTableLayerIndex, int32& TextureReferenceIndex, EMaterialSamplerType SamplerType) = 0;
 	virtual int32 VirtualTextureUniform(int32 TextureIndex, int32 VectorIndex, UE::Shader::EValueType Type) = 0;
 	virtual int32 VirtualTextureUniform(FName ParameterName, int32 TextureIndex, int32 VectorIndex, UE::Shader::EValueType Type) = 0;
-	virtual int32 VirtualTextureWorldToUV(int32 WorldPositionIndex, int32 P0, int32 P1, int32 P2) = 0;
+	virtual int32 VirtualTextureWorldToUV(int32 WorldPositionIndex, int32 P0, int32 P1, int32 P2, EPositionOrigin PositionOrigin) = 0;
+	UE_DEPRECATED(5.4, "Use VirtualTextureWorldToUV(int32, int32, int32, int32, EPositionOrigin) instead")
+	int32 VirtualTextureWorldToUV(int32 WorldPositionIndex, int32 P0, int32 P1, int32 P2) { return VirtualTextureWorldToUV(WorldPositionIndex, P0, P1, P2, EPositionOrigin::Absolute); }
 	virtual int32 VirtualTextureUnpack(int32 CodeIndex0, int32 CodeIndex1, int32 CodeIndex2, int32 P0, EVirtualTextureUnpackType UnpackType) = 0;
 
 	virtual int32 ExternalTexture(const FGuid& ExternalTextureGuid) = 0;
@@ -486,15 +491,50 @@ public:
 	virtual int32 AntialiasedTextureMask(int32 Tex, int32 UV, float Threshold, uint8 Channel) = 0;
 	virtual int32 Sobol(int32 Cell, int32 Index, int32 Seed) = 0;
 	virtual int32 TemporalSobol(int32 Index, int32 Seed) = 0;
-	virtual int32 Noise(int32 Position, float Scale, int32 Quality, uint8 NoiseFunction, bool bTurbulence, int32 Levels, float OutputMin, float OutputMax, float LevelScale, int32 FilterWidth, bool bTiling, uint32 RepeatSize) = 0;
-	virtual int32 VectorNoise(int32 Position, int32 Quality, uint8 NoiseFunction, bool bTiling, uint32 RepeatSize) = 0;
+	virtual int32 Noise(int32 Position, EPositionOrigin PositionOrigin, float Scale, int32 Quality, uint8 NoiseFunction, bool bTurbulence, int32 Levels, float OutputMin, float OutputMax, float LevelScale, int32 FilterWidth, bool bTiling, uint32 RepeatSize) = 0;
+	UE_DEPRECATED(5.4, "Use Noise(int32, EPositionOrigin, float, int32, uint8, bool, int32, float, float, float, int32, bool, uint32) instead")
+	int32 Noise(int32 Position, float Scale, int32 Quality, uint8 NoiseFunction, bool bTurbulence, int32 Levels, float OutputMin, float OutputMax, float LevelScale, int32 FilterWidth, bool bTiling, uint32 RepeatSize)
+	{
+		return Noise(Position, EPositionOrigin::Absolute, Scale, Quality, NoiseFunction, bTurbulence, Levels, OutputMin, OutputMax, LevelScale, FilterWidth, bTiling, RepeatSize);
+	}
+	virtual int32 VectorNoise(int32 Position, EPositionOrigin PositionOrigin, int32 Quality, uint8 NoiseFunction, bool bTiling, uint32 RepeatSize) = 0;
+	UE_DEPRECATED(5.4, "Use VectorNoise(int32, EPositionOrigin, int32, int8, bool, uint32) instead")
+	int32 VectorNoise(int32 Position, int32 Quality, uint8 NoiseFunction, bool bTiling, uint32 RepeatSize)
+	{
+		return VectorNoise(Position, EPositionOrigin::Absolute, Quality, NoiseFunction, bTiling, RepeatSize);
+	}
 	virtual int32 BlackBody( int32 Temp ) = 0;
-	virtual int32 DistanceToNearestSurface(int32 PositionArg) = 0;
-	virtual int32 DistanceFieldGradient(int32 PositionArg) = 0;
-	virtual int32 DistanceFieldApproxAO(int32 PositionArg, int32 NormalArg, int32 BaseDistanceArg, int32 RadiusArg, uint32 NumSteps, float StepScale) = 0;
-	virtual int32 SamplePhysicsField(int32 PositionArg, const int32 OutputType, const int32 TargetIndex) = 0;
+	virtual int32 DistanceToNearestSurface(int32 PositionArg, EPositionOrigin PositionOrigin) = 0;
+	UE_DEPRECATED(5.4, "Use DistanceToNearestSurface(int32, EPositionOrigin) instead")
+	int32 DistanceToNearestSurface(int32 PositionArg)
+	{
+		return DistanceToNearestSurface(PositionArg, EPositionOrigin::Absolute);
+	}
+	virtual int32 DistanceFieldGradient(int32 PositionArg, EPositionOrigin PositionOrigin) = 0;
+	UE_DEPRECATED(5.4, "Use DistanceFieldGradient(int32, EPositionOrigin) instead")
+	int32 DistanceFieldGradient(int32 PositionArg)
+	{
+		return DistanceFieldGradient(PositionArg, EPositionOrigin::Absolute);
+	}
+	virtual int32 DistanceFieldApproxAO(int32 PositionArg, EPositionOrigin PositionOrigin, int32 NormalArg, int32 BaseDistanceArg, int32 RadiusArg, uint32 NumSteps, float StepScale) = 0;
+	UE_DEPRECATED(5.4, "Use DistanceFieldApproxAO(int32, EPositionOrigin, int32, int32, int32, uint32, float) instead")
+	int32 DistanceFieldApproxAO(int32 PositionArg, int32 NormalArg, int32 BaseDistanceArg, int32 RadiusArg, uint32 NumSteps, float StepScale)
+	{
+		return DistanceFieldApproxAO(PositionArg, EPositionOrigin::Absolute, NormalArg, BaseDistanceArg, RadiusArg, NumSteps, StepScale);
+	}
+	virtual int32 SamplePhysicsField(int32 PositionArg, EPositionOrigin PositionOrigin, const int32 OutputType, const int32 TargetIndex) = 0;
+	UE_DEPRECATED(5.4, "Use SamplePhysicsField(int32, EPositionOrigin, int32, int32) instead")
+	int32 SamplePhysicsField(int32 PositionArg, const int32 OutputType, const int32 TargetIndex)
+	{
+		return SamplePhysicsField(PositionArg, EPositionOrigin::Absolute, OutputType, TargetIndex);
+	}
 	virtual int32 DepthOfFieldFunction(int32 Depth, int32 FunctionValueIndex) = 0;
-	virtual int32 AtmosphericFogColor(int32 WorldPosition) = 0;
+	virtual int32 AtmosphericFogColor(int32 WorldPosition, EPositionOrigin PositionOrigin) = 0;
+	UE_DEPRECATED(5.4, "Use AtmosphericFogColor(int32, EPositionOrigin) instead")
+	int32 AtmosphericFogColor(int32 WorldPosition)
+	{
+		return AtmosphericFogColor(WorldPosition, EPositionOrigin::Absolute);
+	}
 	virtual int32 RotateScaleOffsetTexCoords(int32 TexCoordCodeIndex, int32 RotationScale, int32 Offset) = 0;
 	virtual int32 SpeedTree(int32 GeometryArg, int32 WindArg, int32 LODArg, float BillboardThreshold, bool bAccurateWindVelocities, bool bExtraBend, int32 ExtraBendArg) = 0;
 	virtual int32 EyeAdaptation() = 0;
@@ -502,12 +542,22 @@ public:
 	virtual int32 AtmosphericLightVector() = 0;
 	virtual int32 AtmosphericLightColor() = 0;
 
-	virtual int32 SkyAtmosphereLightIlluminance(int32 WorldPosition, int32 LightIndex) = 0;
+	virtual int32 SkyAtmosphereLightIlluminance(int32 WorldPosition, EPositionOrigin PositionOrigin, int32 LightIndex) = 0;
+	UE_DEPRECATED(5.4, "Use SkyAtmosphereLightIlluminance(int32, EPositionOrigin, int32) instead")
+	int32 SkyAtmosphereLightIlluminance(int32 WorldPosition, int32 LightIndex)
+	{
+		return SkyAtmosphereLightIlluminance(WorldPosition, EPositionOrigin::Absolute, LightIndex);
+	}
 	virtual int32 SkyAtmosphereLightIlluminanceOnGround(int32 LightIndex) = 0;
 	virtual int32 SkyAtmosphereLightDirection(int32 LightIndex) = 0;
 	virtual int32 SkyAtmosphereLightDiskLuminance(int32 LightIndex, int32 OverrideAtmosphereLightDiscCosHalfApexAngle) = 0;
 	virtual int32 SkyAtmosphereViewLuminance() = 0;
-	virtual int32 SkyAtmosphereAerialPerspective(int32 WorldPosition) = 0;
+	virtual int32 SkyAtmosphereAerialPerspective(int32 WorldPosition, EPositionOrigin PositionOrigin) = 0;
+	UE_DEPRECATED(5.4, "Use SkyAtmosphereAerialPerspective(int32, EPositionOrigin) instead")
+	int32 SkyAtmosphereAerialPerspective(int32 WorldPosition)
+	{
+		return SkyAtmosphereAerialPerspective(WorldPosition, EPositionOrigin::Absolute);
+	}
 	virtual int32 SkyAtmosphereDistantLightScatteredLuminance() = 0;
 
 	virtual int32 SkyLightEnvMapSample(int32 DirectionCodeChunk, int32 RoughnessCodeChunk) = 0;
@@ -839,7 +889,7 @@ public:
 	}
 	virtual int32 VirtualTextureUniform(int32 TextureIndex, int32 VectorIndex, UE::Shader::EValueType Type) override { return Compiler->VirtualTextureUniform(TextureIndex, VectorIndex, Type); }
 	virtual int32 VirtualTextureUniform(FName ParameterName, int32 TextureIndex, int32 VectorIndex, UE::Shader::EValueType Type) override { return Compiler->VirtualTextureUniform(ParameterName, TextureIndex, VectorIndex, Type); }
-	virtual int32 VirtualTextureWorldToUV(int32 WorldPositionIndex, int32 P0, int32 P1, int32 P2) override { return Compiler->VirtualTextureWorldToUV(WorldPositionIndex, P0, P1, P2); }
+	virtual int32 VirtualTextureWorldToUV(int32 WorldPositionIndex, int32 P0, int32 P1, int32 P2, EPositionOrigin PositionOrigin) override { return Compiler->VirtualTextureWorldToUV(WorldPositionIndex, P0, P1, P2, PositionOrigin); }
 	virtual int32 VirtualTextureUnpack(int32 CodeIndex0, int32 CodeIndex1, int32 CodeIndex2, int32 P0, EVirtualTextureUnpackType UnpackType) override { return Compiler->VirtualTextureUnpack(CodeIndex0, CodeIndex1, CodeIndex2, P0, UnpackType); }
 
 	virtual int32 ExternalTexture(const FGuid& ExternalTextureGuid) override { return Compiler->ExternalTexture(ExternalTextureGuid); }
@@ -958,19 +1008,19 @@ public:
 	}
 	virtual int32 Sobol(int32 Cell, int32 Index, int32 Seed) override { return Compiler->Sobol(Cell, Index, Seed); }
 	virtual int32 TemporalSobol(int32 Index, int32 Seed) override { return Compiler->TemporalSobol(Index, Seed); }
-	virtual int32 Noise(int32 Position, float Scale, int32 Quality, uint8 NoiseFunction, bool bTurbulence, int32 Levels, float OutputMin, float OutputMax, float LevelScale, int32 FilterWidth, bool bTiling, uint32 TileSize) override
+	virtual int32 Noise(int32 Position, EPositionOrigin PositionOrigin, float Scale, int32 Quality, uint8 NoiseFunction, bool bTurbulence, int32 Levels, float OutputMin, float OutputMax, float LevelScale, int32 FilterWidth, bool bTiling, uint32 TileSize) override
 	{
-		return Compiler->Noise(Position, Scale, Quality, NoiseFunction, bTurbulence, Levels, OutputMin, OutputMax, LevelScale, FilterWidth, bTiling, TileSize);
+		return Compiler->Noise(Position, PositionOrigin, Scale, Quality, NoiseFunction, bTurbulence, Levels, OutputMin, OutputMax, LevelScale, FilterWidth, bTiling, TileSize);
 	}
-	virtual int32 VectorNoise(int32 Position, int32 Quality, uint8 NoiseFunction, bool bTiling, uint32 TileSize) override
+	virtual int32 VectorNoise(int32 Position, EPositionOrigin PositionOrigin, int32 Quality, uint8 NoiseFunction, bool bTiling, uint32 TileSize) override
 	{
-		return Compiler->VectorNoise(Position, Quality, NoiseFunction, bTiling, TileSize);
+		return Compiler->VectorNoise(Position, PositionOrigin, Quality, NoiseFunction, bTiling, TileSize);
 	}
 	virtual int32 BlackBody(int32 Temp) override { return Compiler->BlackBody(Temp); }
-	virtual int32 DistanceToNearestSurface(int32 PositionArg) override { return Compiler->DistanceToNearestSurface(PositionArg); }
-	virtual int32 DistanceFieldGradient(int32 PositionArg) override { return Compiler->DistanceFieldGradient(PositionArg); }
-	virtual int32 DistanceFieldApproxAO(int32 PositionArg, int32 NormalArg, int32 BaseDistanceArg, int32 RadiusArg, uint32 NumSteps, float StepScale) override { return Compiler->DistanceFieldApproxAO(PositionArg, NormalArg, BaseDistanceArg, RadiusArg, NumSteps, StepScale); }
-	virtual int32 SamplePhysicsField(int32 PositionArg, const int32 OutputType, const int32 TargetIndex)  override { return Compiler->SamplePhysicsField(PositionArg, OutputType, TargetIndex); }
+	virtual int32 DistanceToNearestSurface(int32 PositionArg, EPositionOrigin PositionOrigin) override { return Compiler->DistanceToNearestSurface(PositionArg, PositionOrigin); }
+	virtual int32 DistanceFieldGradient(int32 PositionArg, EPositionOrigin PositionOrigin) override { return Compiler->DistanceFieldGradient(PositionArg, PositionOrigin); }
+	virtual int32 DistanceFieldApproxAO(int32 PositionArg, EPositionOrigin PositionOrigin, int32 NormalArg, int32 BaseDistanceArg, int32 RadiusArg, uint32 NumSteps, float StepScale) override { return Compiler->DistanceFieldApproxAO(PositionArg, PositionOrigin, NormalArg, BaseDistanceArg, RadiusArg, NumSteps, StepScale); }
+	virtual int32 SamplePhysicsField(int32 PositionArg, EPositionOrigin PositionOrigin, const int32 OutputType, const int32 TargetIndex)  override { return Compiler->SamplePhysicsField(PositionArg, PositionOrigin, OutputType, TargetIndex); }
 	virtual int32 PerInstanceRandom() override { return Compiler->PerInstanceRandom(); }
 	virtual int32 PerInstanceFadeAmount() override { return Compiler->PerInstanceFadeAmount(); }
 	virtual int32 PerInstanceCustomData(int32 DataIndex, int32 DefaultValueIndex) override { return Compiler->PerInstanceCustomData(DataIndex, DefaultValueIndex); }
@@ -1006,9 +1056,9 @@ public:
 		return Compiler->SpeedTree(GeometryArg, WindArg, LODArg, BillboardThreshold, bAccurateWindVelocities, bExtraBend, ExtraBendArg);
 	}
 
-	virtual int32 AtmosphericFogColor(int32 WorldPosition) override
+	virtual int32 AtmosphericFogColor(int32 WorldPosition, EPositionOrigin PositionOrigin) override
 	{
-		return Compiler->AtmosphericFogColor(WorldPosition);
+		return Compiler->AtmosphericFogColor(WorldPosition, PositionOrigin);
 	}
 
 	virtual int32 AtmosphericLightVector() override
@@ -1021,9 +1071,9 @@ public:
 		return Compiler->AtmosphericLightColor();
 	}
 
-	virtual int32 SkyAtmosphereLightIlluminance(int32 WorldPosition, int32 LightIndex) override
+	virtual int32 SkyAtmosphereLightIlluminance(int32 WorldPosition, EPositionOrigin PositionOrigin, int32 LightIndex) override
 	{
-		return Compiler->SkyAtmosphereLightIlluminance(WorldPosition, LightIndex);
+		return Compiler->SkyAtmosphereLightIlluminance(WorldPosition, PositionOrigin, LightIndex);
 	}
 
 	virtual int32 SkyAtmosphereLightIlluminanceOnGround(int32 LightIndex) override
@@ -1046,9 +1096,9 @@ public:
 		return Compiler->SkyAtmosphereViewLuminance();
 	}
 
-	virtual int32 SkyAtmosphereAerialPerspective(int32 WorldPosition) override
+	virtual int32 SkyAtmosphereAerialPerspective(int32 WorldPosition, EPositionOrigin PositionOrigin) override
 	{
-		return Compiler->SkyAtmosphereAerialPerspective(WorldPosition);
+		return Compiler->SkyAtmosphereAerialPerspective(WorldPosition, PositionOrigin);
 	}
 
 	virtual int32 SkyAtmosphereDistantLightScatteredLuminance() override
