AudioStreamingCache, missed a couple of places where I didn't add all the memory together

Specifically in ShouldAddNewChunk, which checks whether a new chunk of audio should be added.

This change fixes the potential for new chunks to be added when they shouldn't be when memory counted features are added.

#tests PIE
#localization none

diff --git a/Engine/Source/Runtime/Engine/Private/AudioStreamingCache.cpp b/Engine/Source/Runtime/Engine/Private/AudioStreamingCache.cpp
--- a/Engine/Source/Runtime/Engine/Private/AudioStreamingCache.cpp
+++ b/Engine/Source/Runtime/Engine/Private/AudioStreamingCache.cpp
@@ -1476,7 +1476,7 @@ void FAudioChunkCache::TouchElement(FCacheElement* InElement)
 
 bool FAudioChunkCache::ShouldAddNewChunk() const
 {
-	return (ChunksInUse < CachePool.Num()) && (MemoryCounterBytes.Load() + ForceInlineMemoryCounterBytes.Load() < MemoryLimitBytes);
+	return (ChunksInUse < CachePool.Num()) && (GetCurrentMemoryUsageBytes() < MemoryLimitBytes);
 }
 
 FAudioChunkCache::FCacheElement* FAudioChunkCache::InsertChunk(const FChunkKey& InKey, const TSharedPtr<FSoundWaveData>& InSoundWavePtr)
@@ -2487,7 +2487,7 @@ TPair<int, int> FAudioChunkCache::DebugDisplay(UWorld* World, FViewport* Viewpor
 	const float PercentageExternalFeatures = NumBytesCounter > 0 ? (double)ExternalFeaturesBytes / NumBytesCounter : 0;
 
 	FString CacheMemoryUsage = *FString::Printf(TEXT("Using: %.4f Megabytes (%lu bytes). Max Potential Usage: %.4f Megabytes."), 
-		NumMegabytesInUse, MemoryCounterBytes.Load() + ForceInlineBytes, MaxCacheSizeMB);
+		NumMegabytesInUse, GetCurrentMemoryUsageBytes(), MaxCacheSizeMB);
 
 	// We're going to align this horizontally with the number of elements right above it.
 	Canvas->DrawShadowedString(X, Y, *CacheMemoryUsage, UEngine::GetMediumFont(), FLinearColor::White);
@@ -2694,7 +2694,7 @@ TPair<int, int> FAudioChunkCache::DebugDisplay(UWorld* World, FViewport* Viewpor
 
 		if (BarWidthExternalFeatures > 0)
 		{
-			// (Force Inline)
+			// (External Features)
 			Canvas->DrawTile(CurrHorzOffset, CurrVertOffset, BarWidthExternalFeatures - DividerWidth, BarHeight, 0, 0, 0, 0, ColorExternalFeatures);
 			CurrHorzOffset += BarWidthExternalFeatures;
 		}
