World Partition: fix actor descriptors for actors attached to child actor components; they were keeping the parent actor guid instead of the owning actor of the child component, resulting in validation errors.
Note: this requires a resave of affected actors to take effect.

#prelight none

diff --git a/Engine/Source/Runtime/Engine/Private/WorldPartition/WorldPartitionActorDesc.cpp b/Engine/Source/Runtime/Engine/Private/WorldPartition/WorldPartitionActorDesc.cpp
--- a/Engine/Source/Runtime/Engine/Private/WorldPartition/WorldPartitionActorDesc.cpp
+++ b/Engine/Source/Runtime/Engine/Private/WorldPartition/WorldPartitionActorDesc.cpp
@@ -188,6 +188,10 @@ void FWorldPartitionActorDesc::Init(const AActor* InActor)
 		const AActor* AttachParentActor = InActor->GetAttachParentActor();
 		if (AttachParentActor)
 		{
+			while (AttachParentActor->GetParentActor())
+			{
+				AttachParentActor = AttachParentActor->GetParentActor();
+			}
 			ParentActor = AttachParentActor->GetActorGuid();
 		}
 
