PR #10950: Add material name to material usage flag logging so it's actually useful

diff --git a/Engine/Source/Runtime/Engine/Private/Materials/Material.cpp b/Engine/Source/Runtime/Engine/Private/Materials/Material.cpp
--- a/Engine/Source/Runtime/Engine/Private/Materials/Material.cpp
+++ b/Engine/Source/Runtime/Engine/Private/Materials/Material.cpp
@@ -1710,7 +1710,7 @@ bool UMaterial::CheckMaterialUsage_Concurrent(EMaterialUsage Usage) const
 						Material->CheckMaterialUsage(Usage);
 					}
 				};
-				UE_LOG(LogMaterial, Log, TEXT("Had to pass SMU back to game thread. Please ensure correct material usage flags."));
+				UE_LOG(LogMaterial, Log, TEXT("Had to pass SMU back to game thread. Please ensure correct material usage flags on %s"), *GetPathNameSafe(this));
 
 				TSharedRef<FCallSMU, ESPMode::ThreadSafe> CallSMU = MakeShareable(new FCallSMU(const_cast<UMaterial*>(this), Usage));
 
diff --git a/Engine/Source/Runtime/Engine/Private/Materials/MaterialInstance.cpp b/Engine/Source/Runtime/Engine/Private/Materials/MaterialInstance.cpp
--- a/Engine/Source/Runtime/Engine/Private/Materials/MaterialInstance.cpp
+++ b/Engine/Source/Runtime/Engine/Private/Materials/MaterialInstance.cpp
@@ -1530,7 +1530,7 @@ bool UMaterialInstance::CheckMaterialUsage_Concurrent(const EMaterialUsage Usage
 						Material->CheckMaterialUsage(Usage);
 					}
 				};
-				UE_LOG(LogMaterial, Log, TEXT("Had to pass SMU back to game thread. Please ensure correct material usage flags."));
+				UE_LOG(LogMaterial, Log, TEXT("Had to pass SMU back to game thread. Please ensure correct material usage flags on %s"), *GetPathNameSafe(this));
 
 				TSharedRef<FCallSMU, ESPMode::ThreadSafe> CallSMU = MakeShareable(new FCallSMU(const_cast<UMaterialInstance*>(this), Usage));
 				bUsageSetSuccessfully = false;
