Use checked member names in MaterialExpressions

diff --git a/Engine/Source/Runtime/Engine/Private/Materials/MaterialExpressions.cpp b/Engine/Source/Runtime/Engine/Private/Materials/MaterialExpressions.cpp
--- a/Engine/Source/Runtime/Engine/Private/Materials/MaterialExpressions.cpp
+++ b/Engine/Source/Runtime/Engine/Private/Materials/MaterialExpressions.cpp
@@ -2429,37 +2429,33 @@ void UMaterialExpressionTextureBase::PostEditChangeProperty(FPropertyChangedEven
 {
 	Super::PostEditChangeProperty( PropertyChangedEvent );
 
-	if (IsDefaultMeshpaintTexture && PropertyChangedEvent.Property)
+	if (IsDefaultMeshpaintTexture && PropertyChangedEvent.GetPropertyName() == GET_MEMBER_NAME_STRING_CHECKED(ThisClass, IsDefaultMeshpaintTexture))
 	{
-		const FName PropertyName = PropertyChangedEvent.Property->GetFName();
-		if (PropertyName == FName(TEXT("IsDefaultMeshpaintTexture")))
+		// Check for other defaulted textures in THIS material (does not search sub levels ie functions etc, as these are ignored in the texture painter). 
+		for (UMaterialExpression* Expression : this->Material->GetMaterial()->GetExpressions())
 		{
-			// Check for other defaulted textures in THIS material (does not search sub levels ie functions etc, as these are ignored in the texture painter). 
-			for (UMaterialExpression* Expression : this->Material->GetMaterial()->GetExpressions())
+			UMaterialExpressionTextureBase* TextureSample = Cast<UMaterialExpressionTextureBase>(Expression);
+			if (TextureSample != nullptr && TextureSample != this)
 			{
-				UMaterialExpressionTextureBase* TextureSample = Cast<UMaterialExpressionTextureBase>(Expression);
-				if (TextureSample != nullptr && TextureSample != this)
+				if(TextureSample->IsDefaultMeshpaintTexture)
 				{
-					if(TextureSample->IsDefaultMeshpaintTexture)
+					FText ErrorMessage = LOCTEXT("MeshPaintDefaultTextureErrorDefault","Only one texture can be set as the Mesh Paint Default Texture, disabling previous default");
+					if (TextureSample->Texture != nullptr)
 					{
-						FText ErrorMessage = LOCTEXT("MeshPaintDefaultTextureErrorDefault","Only one texture can be set as the Mesh Paint Default Texture, disabling previous default");
-						if (TextureSample->Texture != nullptr)
-						{
-							FFormatNamedArguments Args;
-							Args.Add( TEXT("TextureName"), FText::FromString( TextureSample->Texture->GetName() ) );
-							ErrorMessage = FText::Format(LOCTEXT("MeshPaintDefaultTextureErrorTextureKnown","Only one texture can be set as the Mesh Paint Default Texture, disabling {TextureName}"), Args );
-						}
+						FFormatNamedArguments Args;
+						Args.Add( TEXT("TextureName"), FText::FromString( TextureSample->Texture->GetName() ) );
+						ErrorMessage = FText::Format(LOCTEXT("MeshPaintDefaultTextureErrorTextureKnown","Only one texture can be set as the Mesh Paint Default Texture, disabling {TextureName}"), Args );
+					}
 										
-						// Launch notification to inform user of default change
-						FNotificationInfo Info( ErrorMessage );
-						Info.ExpireDuration = 5.0f;
-						Info.Image = FCoreStyle::Get().GetBrush(TEXT("MessageLog.Error"));
+					// Launch notification to inform user of default change
+					FNotificationInfo Info( ErrorMessage );
+					Info.ExpireDuration = 5.0f;
+					Info.Image = FCoreStyle::Get().GetBrush(TEXT("MessageLog.Error"));
 
-						FSlateNotificationManager::Get().AddNotification(Info);
+					FSlateNotificationManager::Get().AddNotification(Info);
 
-						// Reset the previous default to false;
-						TextureSample->IsDefaultMeshpaintTexture = false;
-					}
+					// Reset the previous default to false;
+					TextureSample->IsDefaultMeshpaintTexture = false;
 				}
 			}
 		}
@@ -2606,7 +2602,7 @@ bool UMaterialExpressionTextureSample::CanEditChange(const FProperty* InProperty
 
 void UMaterialExpressionTextureSample::PostEditChangeProperty(FPropertyChangedEvent& PropertyChangedEvent)
 {
-	if ( PropertyChangedEvent.Property && PropertyChangedEvent.Property->GetName() == TEXT("Texture") )
+	if ( PropertyChangedEvent.GetPropertyName() == GET_MEMBER_NAME_STRING_CHECKED(ThisClass, Texture) )
 	{
 		if ( Texture )
 		{
@@ -2615,15 +2611,11 @@ void UMaterialExpressionTextureSample::PostEditChangeProperty(FPropertyChangedEv
 		}
 	}
 
-	if (PropertyChangedEvent.MemberProperty)
+	if (PropertyChangedEvent.GetMemberPropertyName() == GET_MEMBER_NAME_CHECKED(ThisClass, MipValueMode))
 	{
-		const FName PropertyName = PropertyChangedEvent.MemberProperty->GetFName();
-		if (PropertyName == GET_MEMBER_NAME_CHECKED(UMaterialExpressionTextureSample, MipValueMode))
+		if (GraphNode)
 		{
-			if (GraphNode)
-			{
-				GraphNode->ReconstructNode();
-			}
+			GraphNode->ReconstructNode();
 		}
 	}
 	
@@ -3630,7 +3622,7 @@ bool UMaterialExpressionRuntimeVirtualTextureSampleParameter::SetParameterValue(
 		VirtualTexture = InValue;
 		if (EnumHasAnyFlags(Flags, EMaterialExpressionSetParameterValueFlags::SendPostEditChangeProperty))
 		{
-			SendPostEditChangeProperty(this, TEXT("VirtualTexture"));
+			SendPostEditChangeProperty(this, GET_MEMBER_NAME_STRING_CHECKED(ThisClass, VirtualTexture));
 		}
 		return true;
 	}
@@ -3886,7 +3878,7 @@ bool UMaterialExpressionTextureSampleParameter::SetParameterValue(FName InParame
 		Texture = InValue;
 		if (EnumHasAnyFlags(Flags, EMaterialExpressionSetParameterValueFlags::SendPostEditChangeProperty))
 		{
-			SendPostEditChangeProperty(this, TEXT("Texture"));
+			SendPostEditChangeProperty(this, GET_MEMBER_NAME_STRING_CHECKED(ThisClass, Texture));
 		}
 		return true;
 	}
@@ -3895,10 +3887,7 @@ bool UMaterialExpressionTextureSampleParameter::SetParameterValue(FName InParame
 
 void UMaterialExpressionTextureSampleParameter::PostEditChangeProperty(FPropertyChangedEvent& PropertyChangedEvent)
 {
-	FProperty* PropertyThatChanged = PropertyChangedEvent.MemberProperty;
-	const FString PropertyName = PropertyThatChanged ? PropertyThatChanged->GetName() : TEXT("");
-
-	if (PropertyName == GET_MEMBER_NAME_STRING_CHECKED(UMaterialExpressionTextureSampleParameter, ChannelNames))
+	if (PropertyChangedEvent.GetMemberPropertyName() == GET_MEMBER_NAME_STRING_CHECKED(ThisClass, ChannelNames))
 	{
 		ApplyChannelNames();
 
@@ -4064,7 +4053,7 @@ void UMaterialExpressionTextureObject::PostEditChangeProperty(FPropertyChangedEv
 {
 	Super::PostEditChangeProperty( PropertyChangedEvent );
 
-	if ( PropertyChangedEvent.Property && PropertyChangedEvent.Property->GetName() == TEXT("Texture") )
+	if ( PropertyChangedEvent.GetPropertyName() == GET_MEMBER_NAME_STRING_CHECKED(ThisClass, Texture) )
 	{
 		if ( Texture )
 		{
@@ -5861,10 +5850,10 @@ bool  UMaterialExpressionStaticComponentMaskParameter::SetParameterValue(FName I
 		}
 		if (EnumHasAnyFlags(Flags, EMaterialExpressionSetParameterValueFlags::SendPostEditChangeProperty))
 		{
-			SendPostEditChangeProperty(this, TEXT("DefaultR"));
-			SendPostEditChangeProperty(this, TEXT("DefaultG"));
-			SendPostEditChangeProperty(this, TEXT("DefaultB"));
-			SendPostEditChangeProperty(this, TEXT("DefaultA"));
+			SendPostEditChangeProperty(this, GET_MEMBER_NAME_STRING_CHECKED(ThisClass, DefaultR));
+			SendPostEditChangeProperty(this, GET_MEMBER_NAME_STRING_CHECKED(ThisClass, DefaultG));
+			SendPostEditChangeProperty(this, GET_MEMBER_NAME_STRING_CHECKED(ThisClass, DefaultB));
+			SendPostEditChangeProperty(this, GET_MEMBER_NAME_STRING_CHECKED(ThisClass, DefaultA));
 		}
 		return true;
 	}
@@ -8755,21 +8744,20 @@ bool UMaterialExpressionVectorParameter::SetParameterValue(FName InParameterName
 
 void UMaterialExpressionVectorParameter::PostEditChangeProperty(FPropertyChangedEvent& PropertyChangedEvent)
 {
-	FProperty* PropertyThatChanged = PropertyChangedEvent.MemberProperty;
-	const FString PropertyName = PropertyThatChanged ? PropertyThatChanged->GetName() : TEXT("");
+	const FName PropertyName = PropertyChangedEvent.GetMemberPropertyName();
 
-	if (PropertyName == GET_MEMBER_NAME_STRING_CHECKED(UMaterialExpressionVectorParameter, DefaultValue))
+	if (PropertyName == GET_MEMBER_NAME_STRING_CHECKED(ThisClass, DefaultValue))
 	{
 		// Callback into the editor
 		FEditorSupportDelegates::NumericParameterDefaultChanged.Broadcast(this, EMaterialParameterType::Vector, ParameterName, DefaultValue);
 	}
-	else if (PropertyName == GET_MEMBER_NAME_STRING_CHECKED(UMaterialExpressionVectorParameter, PrimitiveDataIndex))
+	else if (PropertyName == GET_MEMBER_NAME_STRING_CHECKED(ThisClass, PrimitiveDataIndex))
 	{
 		// Clamp value
 		const int32 PrimDataIndex = PrimitiveDataIndex;
 		PrimitiveDataIndex = (uint8)FMath::Clamp(PrimDataIndex, 0, FCustomPrimitiveData::NumCustomPrimitiveDataFloats-1);
 	}
-	else if (PropertyName == GET_MEMBER_NAME_STRING_CHECKED(UMaterialExpressionVectorParameter, ChannelNames)
+	else if (PropertyName == GET_MEMBER_NAME_STRING_CHECKED(ThisClass, ChannelNames)
 		&& !IsUsedAsChannelMask())
 	{
 		ApplyChannelNames();
@@ -8878,7 +8866,7 @@ bool UMaterialExpressionDoubleVectorParameter::SetParameterValue(FName InParamet
 		DefaultValue = InValue;
 		if (EnumHasAnyFlags(Flags, EMaterialExpressionSetParameterValueFlags::SendPostEditChangeProperty))
 		{
-			SendPostEditChangeProperty(this, TEXT("DefaultValue"));
+			SendPostEditChangeProperty(this, GET_MEMBER_NAME_STRING_CHECKED(ThisClass, DefaultValue));
 		}
 		return true;
 	}
@@ -8887,10 +8875,7 @@ bool UMaterialExpressionDoubleVectorParameter::SetParameterValue(FName InParamet
 
 void UMaterialExpressionDoubleVectorParameter::PostEditChangeProperty(FPropertyChangedEvent& PropertyChangedEvent)
 {
-	FProperty* PropertyThatChanged = PropertyChangedEvent.MemberProperty;
-	const FString PropertyName = PropertyThatChanged ? PropertyThatChanged->GetName() : TEXT("");
-
-	if (PropertyName == GET_MEMBER_NAME_STRING_CHECKED(UMaterialExpressionDoubleVectorParameter, DefaultValue))
+	if (PropertyChangedEvent.GetMemberPropertyName() == GET_MEMBER_NAME_STRING_CHECKED(ThisClass, DefaultValue))
 	{
 		// Callback into the editor
 		FEditorSupportDelegates::NumericParameterDefaultChanged.Broadcast(this, EMaterialParameterType::DoubleVector, ParameterName, DefaultValue);
@@ -8955,8 +8940,8 @@ bool UMaterialExpressionChannelMaskParameter::SetParameterValue(FName InParamete
 
 		if (EnumHasAnyFlags(Flags, EMaterialExpressionSetParameterValueFlags::SendPostEditChangeProperty))
 		{
-			SendPostEditChangeProperty(this, TEXT("DefaultValue"));
-			SendPostEditChangeProperty(this, TEXT("MaskChannel"));
+			SendPostEditChangeProperty(this, GET_MEMBER_NAME_STRING_CHECKED(ThisClass, DefaultValue));
+			SendPostEditChangeProperty(this, GET_MEMBER_NAME_STRING_CHECKED(ThisClass, MaskChannel));
 		}
 
 		return true;
@@ -8967,10 +8952,9 @@ bool UMaterialExpressionChannelMaskParameter::SetParameterValue(FName InParamete
 
 void UMaterialExpressionChannelMaskParameter::PostEditChangeProperty(FPropertyChangedEvent& PropertyChangedEvent)
 {
-	FProperty* PropertyThatChanged = PropertyChangedEvent.MemberProperty;
-	const FString PropertyName = PropertyThatChanged ? PropertyThatChanged->GetName() : TEXT("");
+	const FName PropertyName = PropertyChangedEvent.GetMemberPropertyName();
 
-	if (PropertyName == GET_MEMBER_NAME_STRING_CHECKED(UMaterialExpressionChannelMaskParameter, MaskChannel))
+	if (PropertyName == GET_MEMBER_NAME_STRING_CHECKED(ThisClass, MaskChannel))
 	{
 		// Update internal value
 		switch (MaskChannel)
@@ -8987,7 +8971,7 @@ void UMaterialExpressionChannelMaskParameter::PostEditChangeProperty(FPropertyCh
 
 		FEditorSupportDelegates::NumericParameterDefaultChanged.Broadcast(this, EMaterialParameterType::Vector, ParameterName, DefaultValue);
 	}
-	else if (PropertyName == GET_MEMBER_NAME_STRING_CHECKED(UMaterialExpressionVectorParameter, DefaultValue))
+	else if (PropertyName == GET_MEMBER_NAME_STRING_CHECKED(ThisClass, DefaultValue))
 	{
 		// If the vector parameter was updated, the enum needs to match and we assert the values are valid
 		if (DefaultValue.R > 0.0f)
@@ -9104,7 +9088,7 @@ bool UMaterialExpressionScalarParameter::SetParameterValue(FName InParameterName
 		DefaultValue = InValue;
 		if (EnumHasAnyFlags(Flags, EMaterialExpressionSetParameterValueFlags::SendPostEditChangeProperty))
 		{
-			SendPostEditChangeProperty(this, TEXT("DefaultValue"));
+			SendPostEditChangeProperty(this, GET_MEMBER_NAME_STRING_CHECKED(ThisClass, DefaultValue));
 		}
 		return true;
 	}
@@ -9113,15 +9097,14 @@ bool UMaterialExpressionScalarParameter::SetParameterValue(FName InParameterName
 
 void UMaterialExpressionScalarParameter::PostEditChangeProperty(FPropertyChangedEvent& PropertyChangedEvent)
 {
-	FProperty* PropertyThatChanged = PropertyChangedEvent.MemberProperty;
-	const FString PropertyName = PropertyThatChanged ? PropertyThatChanged->GetName() : TEXT("");
+	const FName PropertyName = PropertyChangedEvent.GetMemberPropertyName();
 
-	if (PropertyName == GET_MEMBER_NAME_STRING_CHECKED(UMaterialExpressionScalarParameter, DefaultValue))
+	if (PropertyName == GET_MEMBER_NAME_STRING_CHECKED(ThisClass, DefaultValue))
 	{
 		// Callback into the editor
 		FEditorSupportDelegates::NumericParameterDefaultChanged.Broadcast(this, EMaterialParameterType::Scalar, ParameterName, DefaultValue);
 	}
-	else if (PropertyName == GET_MEMBER_NAME_STRING_CHECKED(UMaterialExpressionScalarParameter, PrimitiveDataIndex))
+	else if (PropertyName == GET_MEMBER_NAME_STRING_CHECKED(ThisClass, PrimitiveDataIndex))
 	{
 		// Clamp value
 		const int32 PrimDataIndex = PrimitiveDataIndex;
@@ -9344,7 +9327,7 @@ bool UMaterialExpressionStaticBoolParameter::SetParameterValue(FName InParameter
 		}
 		if (EnumHasAnyFlags(Flags, EMaterialExpressionSetParameterValueFlags::SendPostEditChangeProperty))
 		{
-			SendPostEditChangeProperty(this, TEXT("DefaultValue"));
+			SendPostEditChangeProperty(this, GET_MEMBER_NAME_STRING_CHECKED(ThisClass, DefaultValue));
 		}
 		return true;
 	}
@@ -10025,10 +10008,9 @@ bool UMaterialExpressionDataDrivenShaderPlatformInfoSwitch::IsResultMaterialAttr
 
 void UMaterialExpressionDataDrivenShaderPlatformInfoSwitch::PostEditChangeProperty(FPropertyChangedEvent& PropertyChangedEvent)
 {
-	if (PropertyChangedEvent.MemberProperty && GraphNode)
+	if (PropertyChangedEvent.GetMemberPropertyName() == GET_MEMBER_NAME_CHECKED(ThisClass, DDSPIPropertyNames))
 	{
-		const FName PropertyName = PropertyChangedEvent.MemberProperty->GetFName();
-		if (PropertyName == GET_MEMBER_NAME_CHECKED(UMaterialExpressionDataDrivenShaderPlatformInfoSwitch, DDSPIPropertyNames))
+		if (GraphNode)
 		{
 			bContainsInvalidProperty = IsDataDrivenShaderPlatformInfoSwitchValid(DDSPIPropertyNames, Material);
 			GraphNode->ReconstructNode();
@@ -10584,15 +10566,11 @@ void UMaterialExpressionDynamicParameter::PostEditChangeProperty(FPropertyChange
 {
 	Super::PostEditChangeProperty(PropertyChangedEvent);
 
-	if (PropertyChangedEvent.MemberProperty)
+	if (PropertyChangedEvent.GetMemberPropertyName() == GET_MEMBER_NAME_CHECKED(ThisClass, ParamNames))
 	{
-		const FName PropertyName = PropertyChangedEvent.MemberProperty->GetFName();
-		if (PropertyName == GET_MEMBER_NAME_CHECKED(UMaterialExpressionDynamicParameter, ParamNames))
+		if (GraphNode)
 		{
-			if (GraphNode)
-			{
-				GraphNode->ReconstructNode();
-			}
+			GraphNode->ReconstructNode();
 		}
 	}
 }
@@ -12302,64 +12280,62 @@ void UMaterialExpressionComment::PostEditChangeProperty(FPropertyChangedEvent& P
 {
 	Super::PostEditChangeProperty(PropertyChangedEvent);
 
-	if (PropertyChangedEvent.MemberProperty)
+	const FName PropertyName = PropertyChangedEvent.GetMemberPropertyName();
+
+	if (PropertyName == GET_MEMBER_NAME_CHECKED(ThisClass, Text))
 	{
-		const FName PropertyName = PropertyChangedEvent.MemberProperty->GetFName();
-		if (PropertyName == GET_MEMBER_NAME_CHECKED(UMaterialExpressionComment, Text))
+		if (GraphNode)
 		{
-			if (GraphNode)
-			{
-				GraphNode->Modify();
-				GraphNode->NodeComment = Text;
-			}
+			GraphNode->Modify();
+			GraphNode->NodeComment = Text;
 		}
-		else if (PropertyName == GET_MEMBER_NAME_CHECKED(UMaterialExpressionComment, CommentColor))
+	}
+	else if (PropertyName == GET_MEMBER_NAME_CHECKED(ThisClass, CommentColor))
+	{
+		if (GraphNode)
 		{
-			if (GraphNode)
-			{
-				GraphNode->Modify();
-				CastChecked<UMaterialGraphNode_Comment>(GraphNode)->CommentColor = CommentColor;
-			}
+			GraphNode->Modify();
+			CastChecked<UMaterialGraphNode_Comment>(GraphNode)->CommentColor = CommentColor;
 		}
-		else if (PropertyName == GET_MEMBER_NAME_CHECKED(UMaterialExpressionComment, FontSize))
+	}
+	else if (PropertyName == GET_MEMBER_NAME_CHECKED(ThisClass, FontSize))
+	{
+		if (GraphNode)
 		{
-			if (GraphNode)
-			{
-				GraphNode->Modify();
-				CastChecked<UMaterialGraphNode_Comment>(GraphNode)->FontSize = FontSize;
-			}
+			GraphNode->Modify();
+			CastChecked<UMaterialGraphNode_Comment>(GraphNode)->FontSize = FontSize;
 		}
-		else if (PropertyName == GET_MEMBER_NAME_CHECKED(UMaterialExpressionComment, bColorCommentBubble))
+	}
+	else if (PropertyName == GET_MEMBER_NAME_CHECKED(ThisClass, bColorCommentBubble))
+	{
+		if (GraphNode)
 		{
-			if (GraphNode)
-			{
-				GraphNode->Modify();
-				CastChecked<UMaterialGraphNode_Comment>(GraphNode)->bColorCommentBubble = bColorCommentBubble;
-			}
+			GraphNode->Modify();
+			CastChecked<UMaterialGraphNode_Comment>(GraphNode)->bColorCommentBubble = bColorCommentBubble;
 		}
-		else if (PropertyName == GET_MEMBER_NAME_CHECKED(UMaterialExpressionComment, bCommentBubbleVisible_InDetailsPanel))
+	}
+	else if (PropertyName == GET_MEMBER_NAME_CHECKED(ThisClass, bCommentBubbleVisible_InDetailsPanel))
+	{
+		if (GraphNode)
 		{
-			if (GraphNode)
-			{
-				GraphNode->Modify();
-				UMaterialGraphNode_Comment* CommentNode = CastChecked<UMaterialGraphNode_Comment>(GraphNode);
-				CommentNode->bCommentBubbleVisible_InDetailsPanel = bCommentBubbleVisible_InDetailsPanel;
-				CommentNode->bCommentBubbleVisible = bCommentBubbleVisible_InDetailsPanel;
-				CommentNode->bCommentBubblePinned = bCommentBubbleVisible_InDetailsPanel;
-			}
+			GraphNode->Modify();
+			UMaterialGraphNode_Comment* CommentNode = CastChecked<UMaterialGraphNode_Comment>(GraphNode);
+			CommentNode->bCommentBubbleVisible_InDetailsPanel = bCommentBubbleVisible_InDetailsPanel;
+			CommentNode->bCommentBubbleVisible = bCommentBubbleVisible_InDetailsPanel;
+			CommentNode->bCommentBubblePinned = bCommentBubbleVisible_InDetailsPanel;
 		}
-		else if (PropertyName == GET_MEMBER_NAME_CHECKED(UMaterialExpressionComment, bGroupMode))
+	}
+	else if (PropertyName == GET_MEMBER_NAME_CHECKED(ThisClass, bGroupMode))
+	{
+		if (GraphNode)
 		{
-			if (GraphNode)
-			{
-				GraphNode->Modify();
-				CastChecked<UMaterialGraphNode_Comment>(GraphNode)->MoveMode = bGroupMode ? ECommentBoxMode::GroupMovement : ECommentBoxMode::NoGroupMovement;
-			}
+			GraphNode->Modify();
+			CastChecked<UMaterialGraphNode_Comment>(GraphNode)->MoveMode = bGroupMode ? ECommentBoxMode::GroupMovement : ECommentBoxMode::NoGroupMovement;
 		}
-
-		// Don't need to update preview after changing comments
-		bNeedToUpdatePreview = false;
 	}
+
+	// Don't need to update preview after changing comments
+	bNeedToUpdatePreview = false;
 }
 
 bool UMaterialExpressionComment::Modify( bool bAlwaysMarkDirty/*=true*/ )
@@ -13112,8 +13088,8 @@ bool UMaterialExpressionFontSampleParameter::SetParameterValue(FName InParameter
 		FontTexturePage = InFontPage;
 		if (EnumHasAnyFlags(Flags, EMaterialExpressionSetParameterValueFlags::SendPostEditChangeProperty))
 		{
-			SendPostEditChangeProperty(this, TEXT("Font"));
-			SendPostEditChangeProperty(this, TEXT("FontTexturePage"));
+			SendPostEditChangeProperty(this, GET_MEMBER_NAME_STRING_CHECKED(ThisClass, Font));
+			SendPostEditChangeProperty(this, GET_MEMBER_NAME_STRING_CHECKED(ThisClass, FontTexturePage));
 		}
 		return true;
 	}
@@ -13884,11 +13860,11 @@ void UMaterialExpressionCustom::PostEditChangeProperty(FPropertyChangedEvent& Pr
 
 	RebuildOutputs();
 
-	if (PropertyChangedEvent.MemberProperty && GraphNode)
+	const FName PropertyName = PropertyChangedEvent.GetMemberPropertyName();
+	if (PropertyName == GET_MEMBER_NAME_CHECKED(ThisClass, Inputs) ||
+		PropertyName == GET_MEMBER_NAME_CHECKED(ThisClass, AdditionalOutputs))
 	{
-		const FName PropertyName = PropertyChangedEvent.MemberProperty->GetFName();
-		if (PropertyName == GET_MEMBER_NAME_CHECKED(UMaterialExpressionCustom, Inputs) ||
-			PropertyName == GET_MEMBER_NAME_CHECKED(UMaterialExpressionCustom, AdditionalOutputs))
+		if (GraphNode)
 		{
 			GraphNode->ReconstructNode();
 		}
@@ -14271,9 +14247,9 @@ FName UMaterialExpressionSwitch::GetInputName(int32 InputIndex) const
 	switch (InputIndex)
 	{
 	case 0:
-		return FName(TEXT("SwitchValue"));
+		return GET_MEMBER_NAME_STRING_CHECKED(ThisClass, SwitchValue);
 	case 1:
-		return FName(TEXT("Default"));
+		return GET_MEMBER_NAME_STRING_CHECKED(ThisClass, Default);
 	default:
 		return Inputs[InputIndex - 2].InputName;
 	}
@@ -14297,10 +14273,9 @@ void UMaterialExpressionSwitch::PostEditChangeProperty(FPropertyChangedEvent& Pr
 
 	RebuildOutputs();
 
-	if (PropertyChangedEvent.MemberProperty && GraphNode)
+	if (PropertyChangedEvent.GetMemberPropertyName() == GET_MEMBER_NAME_CHECKED(ThisClass, Inputs))
 	{
-		const FName PropertyName = PropertyChangedEvent.MemberProperty->GetFName();
-		if (PropertyName == GET_MEMBER_NAME_CHECKED(UMaterialExpressionCustom, Inputs))
+		if (GraphNode)
 		{
 			GraphNode->ReconstructNode();
 		}
@@ -16626,7 +16601,7 @@ void UMaterialExpressionMaterialFunctionCall::GetDependentFunctions(TArray<UMate
 #if WITH_EDITOR
 void UMaterialExpressionMaterialFunctionCall::PreEditChange(FProperty* PropertyAboutToChange)
 {
-	if (PropertyAboutToChange && PropertyAboutToChange->GetFName() == FName(TEXT("MaterialFunction")))
+	if (PropertyAboutToChange && PropertyAboutToChange->GetFName() == GET_MEMBER_NAME_STRING_CHECKED(ThisClass, MaterialFunction))
 	{
 		// Save off the previous MaterialFunction value
 		SavedMaterialFunction = MaterialFunction;
@@ -16636,9 +16611,7 @@ void UMaterialExpressionMaterialFunctionCall::PreEditChange(FProperty* PropertyA
 
 void UMaterialExpressionMaterialFunctionCall::PostEditChangeProperty(FPropertyChangedEvent& PropertyChangedEvent)
 {
-	FProperty* PropertyThatChanged = PropertyChangedEvent.Property;
-	
-	if (PropertyThatChanged && PropertyThatChanged->GetFName() == FName(TEXT("MaterialFunction")))
+	if (PropertyChangedEvent.GetPropertyName() == GET_MEMBER_NAME_STRING_CHECKED(ThisClass, MaterialFunction))
 	{
 		// Set the new material function
 		SetMaterialFunctionEx(SavedMaterialFunction, MaterialFunction);
@@ -17293,8 +17266,7 @@ void UMaterialExpressionFunctionInput::PreEditChange(FProperty* PropertyAboutToC
 
 void UMaterialExpressionFunctionInput::PostEditChangeProperty(FPropertyChangedEvent& PropertyChangedEvent)
 {
-	FProperty* PropertyThatChanged = PropertyChangedEvent.Property;
-	if (PropertyThatChanged && PropertyThatChanged->GetFName() == GET_MEMBER_NAME_CHECKED(UMaterialExpressionFunctionInput, InputName))
+	if (PropertyChangedEvent.GetPropertyName() == GET_MEMBER_NAME_CHECKED(ThisClass, InputName))
 	{
 		if (Material)
 		{
@@ -17696,8 +17668,7 @@ void UMaterialExpressionFunctionOutput::PreEditChange(FProperty* PropertyAboutTo
 
 void UMaterialExpressionFunctionOutput::PostEditChangeProperty(FPropertyChangedEvent& PropertyChangedEvent)
 {
-	FProperty* PropertyThatChanged = PropertyChangedEvent.Property;
-	if (PropertyThatChanged && PropertyThatChanged->GetFName() == GET_MEMBER_NAME_CHECKED(UMaterialExpressionFunctionOutput, OutputName))
+	if (PropertyChangedEvent.GetPropertyName() == GET_MEMBER_NAME_CHECKED(UMaterialExpressionFunctionOutput, OutputName))
 	{
 		if (Material)
 		{
@@ -19007,8 +18978,7 @@ void UMaterialExpressionNamedRerouteDeclaration::PostEditChangeProperty(FPropert
 {
 	Super::PostEditChangeProperty(PropertyChangedEvent);
 
-	if (PropertyChangedEvent.MemberProperty &&
-		PropertyChangedEvent.MemberProperty->GetFName() == GET_MEMBER_NAME_CHECKED(UMaterialExpressionNamedRerouteDeclaration, Name))
+	if (PropertyChangedEvent.GetMemberPropertyName() == GET_MEMBER_NAME_CHECKED(ThisClass, Name))
 	{
 		MakeNameUnique();
 	}
@@ -21609,14 +21579,14 @@ bool UMaterialExpressionSpeedTree::CanEditChange(const FProperty* InProperty) co
 
 	if (GeometryType == STG_Billboard)
 	{
-		if (InProperty->GetFName() == TEXT("LODType"))
+		if (InProperty->GetFName() == GET_MEMBER_NAME_STRING_CHECKED(ThisClass, LODType))
 		{
 			bIsEditable = false;
 		}
 	}
 	else
 	{
-		if (InProperty->GetFName() == TEXT("BillboardThreshold"))
+		if (InProperty->GetFName() == GET_MEMBER_NAME_STRING_CHECKED(ThisClass, BillboardThreshold))
 		{
 			bIsEditable = false;
 		}
@@ -22621,9 +22591,9 @@ bool UMaterialExpressionCurveAtlasRowParameter::SetParameterValue(const FName& N
 		DefaultValue = Meta.Value.AsScalar();
 		if (EnumHasAnyFlags(Flags, EMaterialExpressionSetParameterValueFlags::SendPostEditChangeProperty))
 		{
-			SendPostEditChangeProperty(this, TEXT("Curve"));
-			SendPostEditChangeProperty(this, TEXT("Atlas"));
-			SendPostEditChangeProperty(this, TEXT("DefaultValue"));
+			SendPostEditChangeProperty(this, GET_MEMBER_NAME_STRING_CHECKED(ThisClass, Curve));
+			SendPostEditChangeProperty(this, GET_MEMBER_NAME_STRING_CHECKED(ThisClass, Atlas));
+			SendPostEditChangeProperty(this, GET_MEMBER_NAME_STRING_CHECKED(ThisClass, DefaultValue));
 		}
 		if (EnumHasAnyFlags(Flags, EMaterialExpressionSetParameterValueFlags::AssignGroupAndSortPriority))
 		{
@@ -22637,8 +22607,8 @@ bool UMaterialExpressionCurveAtlasRowParameter::SetParameterValue(const FName& N
 
 void UMaterialExpressionCurveAtlasRowParameter::PostEditChangeProperty(FPropertyChangedEvent& PropertyChangedEvent)
 {
-	if (PropertyChangedEvent.Property 
-		&& (PropertyChangedEvent.Property->GetName() == TEXT("Atlas") || PropertyChangedEvent.Property->GetName() == TEXT("Curve")))
+	const FName PropertyName = PropertyChangedEvent.GetPropertyName();
+	if (PropertyName == GET_MEMBER_NAME_STRING_CHECKED(ThisClass, Atlas) || PropertyName == GET_MEMBER_NAME_STRING_CHECKED(ThisClass, Curve))
 	{
 		int32 SlotIndex = INDEX_NONE;
 		if (Atlas && Curve)
@@ -23644,7 +23614,7 @@ FName UMaterialExpressionSceneDepthWithoutWater::GetInputName(int32 InputIndex)
 	if (InputIndex == 0)
 	{
 		// Display the current InputMode enum's display name.
-		FByteProperty* InputModeProperty = FindFProperty<FByteProperty>(UMaterialExpressionSceneDepthWithoutWater::StaticClass(), "InputMode");
+		FByteProperty* InputModeProperty = FindFProperty<FByteProperty>(UMaterialExpressionSceneDepthWithoutWater::StaticClass(), GET_MEMBER_NAME_STRING_CHECKED(ThisClass, InputMode));
 		// Can't use GetNameByValue as GetNameStringByValue does name mangling that GetNameByValue does not
 		return *InputModeProperty->Enum->GetNameStringByValue((int64)InputMode.GetValue());
 	}
@@ -27814,7 +27784,7 @@ UMaterialExpressionSparseVolumeTextureObject::UMaterialExpressionSparseVolumeTex
 void UMaterialExpressionSparseVolumeTextureObject::PostEditChangeProperty(FPropertyChangedEvent& PropertyChangedEvent)
 {
 	// Update what needs to be when the texture is changed
-	if (PropertyChangedEvent.Property && PropertyChangedEvent.Property->GetName() == TEXT("SparseVirtualTexture"))
+	if (PropertyChangedEvent.GetPropertyName() == GET_MEMBER_NAME_STRING_CHECKED(ThisClass, SparseVolumeTexture))
 	{
 		if (SparseVolumeTexture != nullptr)
 		{
@@ -27951,7 +27921,7 @@ void UMaterialExpressionSparseVolumeTextureSample::PostLoad()
 void UMaterialExpressionSparseVolumeTextureSample::PostEditChangeProperty(FPropertyChangedEvent& PropertyChangedEvent)
 {
 	// Update what needs to be when the texture is changed
-	if (PropertyChangedEvent.Property && PropertyChangedEvent.Property->GetName() == TEXT("SparseVirtualTexture"))
+	if (PropertyChangedEvent.GetPropertyName() == GET_MEMBER_NAME_STRING_CHECKED(ThisClass, SparseVolumeTexture))
 	{
 		if (SparseVolumeTexture != nullptr)
 		{
@@ -28103,7 +28073,7 @@ bool UMaterialExpressionSparseVolumeTextureSampleParameter::SetParameterValue(FN
 		SparseVolumeTexture = InValue;
 		if (EnumHasAnyFlags(Flags, EMaterialExpressionSetParameterValueFlags::SendPostEditChangeProperty))
 		{
-			SendPostEditChangeProperty(this, TEXT("SparseVolumeTexture"));
+			SendPostEditChangeProperty(this, GET_MEMBER_NAME_STRING_CHECKED(ThisClass, SparseVolumeTexture));
 		}
 		return true;
 	}
