World Partition Runtime HashSet: fix adding duplicate packages to cook as it's already done in UWorldPartitionRuntimeHash::PopulateRuntimeCell.

diff --git a/Engine/Source/Runtime/Engine/Private/WorldPartition/RuntimeHashSet/WorldPartitionRuntimeHashSetStreamingGeneration.cpp b/Engine/Source/Runtime/Engine/Private/WorldPartition/RuntimeHashSet/WorldPartitionRuntimeHashSetStreamingGeneration.cpp
--- a/Engine/Source/Runtime/Engine/Private/WorldPartition/RuntimeHashSet/WorldPartitionRuntimeHashSetStreamingGeneration.cpp
+++ b/Engine/Source/Runtime/Engine/Private/WorldPartition/RuntimeHashSet/WorldPartitionRuntimeHashSetStreamingGeneration.cpp
@@ -244,27 +244,6 @@ bool UWorldPartitionRuntimeHashSet::GenerateStreaming(UWorldPartitionStreamingPo
 		RuntimeStreamingData.Emplace(MoveTemp(StreamingData));
 	}
 
-	//
-	// Output generated packages
-	//
-	if (OutPackagesToGenerate)
-	{
-		for (UWorldPartitionRuntimeCell* RuntimeCell : RuntimeCells)
-		{
-			// Always loaded cell actors are transfered to World's Persistent Level (see UWorldPartitionRuntimeSpatialHash::PopulateGeneratorPackageForCook)
-			if (RuntimeCell->GetActorCount() && !RuntimeCell->IsAlwaysLoaded())
-			{
-				const FString PackageRelativePath = RuntimeCell->GetPackageNameToCreate();
-				check(!PackageRelativePath.IsEmpty());
-
-				OutPackagesToGenerate->Add(PackageRelativePath);
-
-				// Map relative package to StreamingCell for PopulateGeneratedPackageForCook/PopulateGeneratorPackageForCook/GetCellForPackage
-				PackagesToGenerateForCook.Add(PackageRelativePath, RuntimeCell);
-			}
-		}
-	}
-
 	return true;
 }
 
