ClothSim speed up copying position and normal data into upload buffer using more cache friendly approach.

Provided by jaroslaw.bronisz.

diff --git a/Engine/Source/Runtime/Engine/Private/GPUSkinCache.cpp b/Engine/Source/Runtime/Engine/Private/GPUSkinCache.cpp
--- a/Engine/Source/Runtime/Engine/Private/GPUSkinCache.cpp
+++ b/Engine/Source/Runtime/Engine/Private/GPUSkinCache.cpp
@@ -1780,24 +1780,23 @@ bool FGPUSkinCache::ProcessEntry(
 		InOutEntry->InputWeightStreamSRV = WeightBuffer->GetDataVertexBuffer()->GetSRV();
 	}
 
-    FVertexBufferAndSRV ClothPositionAndNormalsBuffer;
-    TSkeletalMeshVertexData<FClothSimulEntry> VertexAndNormalData(true);
     if (ClothVertexBuffer)
     {
+		FVertexBufferAndSRV ClothPositionAndNormalsBuffer;
+		TSkeletalMeshVertexData<FVector3f> VertexAndNormalData(true);
         InOutEntry->ClothBuffer = ClothVertexBuffer->GetSRV();
         check(InOutEntry->ClothBuffer);
 
 		if (SimData->Positions.Num() > 0)
 		{
 	        check(SimData->Positions.Num() == SimData->Normals.Num());
-	        VertexAndNormalData.ResizeBuffer( SimData->Positions.Num() );
+	        VertexAndNormalData.ResizeBuffer( 2 * SimData->Positions.Num() );
 
-	        uint8* Data = VertexAndNormalData.GetDataPointer();
+			FVector3f* Data = (FVector3f*)VertexAndNormalData.GetDataPointer();
 	        uint32 Stride = VertexAndNormalData.GetStride();
 
 	        // Copy the vertices into the buffer.
-	        checkSlow(Stride*VertexAndNormalData.GetNumVertices() == sizeof(FClothSimulEntry) * SimData->Positions.Num());
-	        check(sizeof(FClothSimulEntry) == 6 * sizeof(float));
+	        checkSlow(Stride*VertexAndNormalData.GetNumVertices() == sizeof(FVector3f) * 2 * SimData->Positions.Num());
 
 			if (ClothVertexBuffer && ClothVertexBuffer->GetClothIndexMapping().Num() > Section)
 			{
@@ -1815,13 +1814,24 @@ bool FGPUSkinCache::ProcessEntry(
 					ClothBufferIndexMapping.MappingOffset;  // Otherwise fallback to a 0 ClothLODBias to prevent from reading pass the buffer (but still raytrace broken shadows/reflections/etc.)
 			}
 
-	        for (int32 Index = 0;Index < SimData->Positions.Num();Index++)
-	        {
-	            FClothSimulEntry NewEntry;
-	            NewEntry.Position = SimData->Positions[Index];
-	            NewEntry.Normal = SimData->Normals[Index];
-	            *((FClothSimulEntry*)(Data + Index * Stride)) = NewEntry;
-	        }
+#if PLATFORM_ENABLE_VECTORINTRINSICS
+			const int32 LastIndex = SimData->Positions.Num() - 1;
+			// Utilize SIMD to populate positions and normals data
+			for (int32 Index = 0; Index < LastIndex; Index++)
+			{
+				VectorStoreAligned(VectorLoadAligned((const float*)(SimData->Positions.GetData() + Index)), (float*)(Data + Index * 2));
+				VectorStoreAligned(VectorLoadAligned((const float*)(SimData->Normals.GetData() + Index)), (float*)(Data + Index * 2 + 1));
+			}
+			// Manually fetch last element 
+			*((Data + LastIndex * 2)) = SimData->Positions[LastIndex];
+			*((Data + LastIndex * 2 + 1)) = SimData->Normals[LastIndex];
+#else
+			for (int32 Index = 0; Index < SimData->Positions.Num(); Index++)
+			{
+				*(Data + Index * 2) = SimData->Positions[Index];
+				*(Data + Index * 2 + 1) = SimData->Normals[Index];
+			}
+#endif
 
 	        FResourceArrayInterface* ResourceArray = VertexAndNormalData.GetResourceArray();
 	        check(ResourceArray->GetResourceDataSize() > 0);
