WorldPartitionActorDescArchive: Fix check which was using wrong sizeof and could fail on last field

#rnx

diff --git a/Engine/Source/Runtime/Engine/Public/WorldPartition/WorldPartitionActorDescArchive.h b/Engine/Source/Runtime/Engine/Public/WorldPartition/WorldPartitionActorDescArchive.h
--- a/Engine/Source/Runtime/Engine/Public/WorldPartition/WorldPartitionActorDescArchive.h
+++ b/Engine/Source/Runtime/Engine/Public/WorldPartition/WorldPartitionActorDescArchive.h
@@ -52,7 +52,7 @@ public:
 
 				if (PropertyOffset < ActorDescAr.ClassDescSizeof)
 				{
-					check((PropertyOffset + sizeof(V)) <= ActorDescAr.ClassDescSizeof);
+					check((PropertyOffset + sizeof(V.Value)) <= ActorDescAr.ClassDescSizeof);
 					const DestPropertyType* RefValue = (const DestPropertyType*)(*(UPTRINT*)&ActorDescAr.ClassDesc + PropertyOffset);
 					return RefValue;
 				}
