Copy bone data - vectorization
ISPC implementation seems to be slower than stanard one on current gen consoles.
Implementation using SIMD instruction explicitely seems to be faster on current gen consoles.

Provided by jaroslaw.bronisz

diff --git a/Engine/Source/Runtime/Engine/Private/GPUSkinVertexFactory.cpp b/Engine/Source/Runtime/Engine/Private/GPUSkinVertexFactory.cpp
--- a/Engine/Source/Runtime/Engine/Private/GPUSkinVertexFactory.cpp
+++ b/Engine/Source/Runtime/Engine/Private/GPUSkinVertexFactory.cpp
@@ -98,8 +98,9 @@ static TAutoConsoleVariable<int32> CVarVelocityTest(
 	ECVF_Cheat | ECVF_RenderThreadSafe);
 #endif // if !(UE_BUILD_SHIPPING || UE_BUILD_TEST)
 
+// Disable it by default as it seems to be up to 20% slower on current gen platforms
 #if !defined(GPU_SKIN_COPY_BONES_ISPC_ENABLED_DEFAULT)
-#define GPU_SKIN_COPY_BONES_ISPC_ENABLED_DEFAULT 1
+#define GPU_SKIN_COPY_BONES_ISPC_ENABLED_DEFAULT 0
 #endif
 
 // Support run-time toggling on supported platforms in non-shipping configurations
@@ -276,7 +277,7 @@ void FGPUBaseSkinVertexFactory::FShaderDataType::UpdateBoneData(FRHICommandList&
 		}
 		else
 		{
-			const int32 PreFetchStride = 2; // FPlatformMisc::Prefetch stride
+			constexpr int32 PreFetchStride = 2; // FPlatformMisc::Prefetch stride
 			for (uint32 BoneIdx = 0; BoneIdx < NumBones; BoneIdx++)
 			{
 				const FBoneIndexType RefToLocalIdx = BoneMap[BoneIdx];
@@ -285,7 +286,29 @@ void FGPUBaseSkinVertexFactory::FShaderDataType::UpdateBoneData(FRHICommandList&
 
 				FMatrix3x4& BoneMat = ChunkMatrices[BoneIdx];
 				const FMatrix44f& RefToLocal = ReferenceToLocalMatrices[RefToLocalIdx];
+				// Explicit SIMD implementation seems to be faster than standard implementation
+#if PLATFORM_ENABLE_VECTORINTRINSICS
+				VectorRegister4Float InRow0 = VectorLoadAligned(&(RefToLocal.M[0][0]));
+				VectorRegister4Float InRow1 = VectorLoadAligned(&(RefToLocal.M[1][0]));
+				VectorRegister4Float InRow2 = VectorLoadAligned(&(RefToLocal.M[2][0]));
+				VectorRegister4Float InRow3 = VectorLoadAligned(&(RefToLocal.M[3][0]));
+
+				VectorRegister4Float Temp0 = VectorShuffle(InRow0, InRow1, 0, 1, 0, 1);
+				VectorRegister4Float Temp1 = VectorShuffle(InRow2, InRow3, 0, 1, 0, 1);
+				VectorRegister4Float Temp2 = VectorShuffle(InRow0, InRow1, 2, 3, 2, 3);
+				VectorRegister4Float Temp3 = VectorShuffle(InRow2, InRow3, 2, 3, 2, 3);
+
+				Temp0 = VectorSwizzle(Temp0, 0, 2, 1, 3);
+				Temp1 = VectorSwizzle(Temp1, 0, 2, 1, 3);
+				Temp2 = VectorSwizzle(Temp2, 0, 2, 1, 3);
+				Temp3 = VectorSwizzle(Temp3, 0, 2, 1, 3);
+
+				VectorStoreAligned(VectorShuffle(Temp0, Temp1, 0, 1, 0, 1), &(BoneMat.M[0][0]));
+				VectorStoreAligned(VectorShuffle(Temp0, Temp1, 2, 3, 2, 3), &(BoneMat.M[1][0]));
+				VectorStoreAligned(VectorShuffle(Temp2, Temp3, 0, 1, 0, 1), &(BoneMat.M[2][0]));
+#else
 				RefToLocal.To3x4MatrixTranspose((float*)BoneMat.M);
+#endif
 			}
 		}
 	}
