[StaticMeshComponent] Fixed another codepath for navigation update when associated mesh gets compiled.

diff --git a/Engine/Source/Runtime/Engine/Private/Components/StaticMeshComponent.cpp b/Engine/Source/Runtime/Engine/Private/Components/StaticMeshComponent.cpp
--- a/Engine/Source/Runtime/Engine/Private/Components/StaticMeshComponent.cpp
+++ b/Engine/Source/Runtime/Engine/Private/Components/StaticMeshComponent.cpp
@@ -694,7 +694,11 @@ void UStaticMeshComponent::NotifyIfStaticMeshChanged()
 				{
 					// Our associated mesh is compiling so we are no longer navigation relevant.
 					// This will invalidate any pending add to the octree and dirty tiles until compilation completes.
-					bNavigationRelevant = IsNavigationRelevant();
+					// Note that we force the value instead of calling IsNavigationRelevant since the IsCompiling
+					// condition will not return the proper value in all code paths where OnPreMeshBuild delegate is called.
+					// (e.g. async task created right after calling OnPreMeshBuild)
+					// Member 'bNavigationRelevant' will be updated by PostStaticMeshCompilation.
+					bNavigationRelevant = false;
 					FNavigationSystem::UpdateComponentData(*this);
 				});
 		}
