Fixed forward declaration struct/class discrepancy for FActorComponentInstanceData

#rnx

diff --git a/Engine/Source/Runtime/Engine/Public/EngineFwd.h b/Engine/Source/Runtime/Engine/Public/EngineFwd.h
--- a/Engine/Source/Runtime/Engine/Public/EngineFwd.h
+++ b/Engine/Source/Runtime/Engine/Public/EngineFwd.h
@@ -113,7 +113,6 @@ class UUserDefinedEnum;
 class UUserDefinedStruct;
 class UWorld;
 
-class FActorComponentInstanceData;
 class FAudioDevice;
 class FAudioDeviceManager;
 class FCanvas;
@@ -162,6 +161,7 @@ enum class ECacheApplyPhase;
 enum class ETeleportType : uint8;
 
 struct FActiveSound;
+struct FActorComponentInstanceData;
 struct FActorComponentTickFunction;
 struct FActorTickFunction;
 struct FAnimAssetTickContext;
