Add an explicit casting path for StaticBool nodes in the material editor. Ensures that the bool is re-cast as 1.0f or 0.0f (mirroring HLSL casting specification) when a float destination type requires it, and allows more flexible use cases for the StaticBool node.

diff --git a/Engine/Source/Runtime/Engine/Private/Materials/HLSLMaterialDerivativeAutogen.cpp b/Engine/Source/Runtime/Engine/Private/Materials/HLSLMaterialDerivativeAutogen.cpp
--- a/Engine/Source/Runtime/Engine/Private/Materials/HLSLMaterialDerivativeAutogen.cpp
+++ b/Engine/Source/Runtime/Engine/Private/Materials/HLSLMaterialDerivativeAutogen.cpp
@@ -119,6 +119,7 @@ EDerivativeType GetDerivType(EMaterialValueType ValueType, bool bAllowNonFloat)
 {
 	switch (ValueType)
 	{
+	case MCT_StaticBool:
 	case MCT_Float:
 	case MCT_Float1:
 		return EDerivativeType::Float1;
diff --git a/Engine/Source/Runtime/Engine/Private/Materials/HLSLMaterialTranslator.cpp b/Engine/Source/Runtime/Engine/Private/Materials/HLSLMaterialTranslator.cpp
--- a/Engine/Source/Runtime/Engine/Private/Materials/HLSLMaterialTranslator.cpp
+++ b/Engine/Source/Runtime/Engine/Private/Materials/HLSLMaterialTranslator.cpp
@@ -4588,6 +4588,19 @@ int32 FHLSLMaterialTranslator::ValidCast(int32 Code, EMaterialValueType DestType
 	return CompiledResult;
 }
 
+EMaterialCastFlags GetForceCastFlags(uint32 ForceCastFlags)
+{
+	EMaterialCastFlags CastFlags = EMaterialCastFlags::AllowTruncate | EMaterialCastFlags::AllowAppendZeroes;
+	if ((ForceCastFlags & MFCF_ReplicateValue) || !(ForceCastFlags & MFCF_ExactMatch))
+	{
+		// Replicate scalar if requested, or if we don't require an exact match (this can happen when force-casting to/from LWC
+		// The only way we *don't* replicate scalar is requesting an exact match without the ReplicateValue flag
+		// TODO - My guess is that case probably isn't relevant, and we should just always replicate scalar on cast, but trying to preserve behavior for now
+		CastFlags |= EMaterialCastFlags::ReplicateScalar;
+	}
+	return CastFlags;
+}
+
 int32 FHLSLMaterialTranslator::ForceCast(int32 Code, EMaterialValueType DestType, uint32 ForceCastFlags)
 {
 	if(Code == INDEX_NONE)
@@ -4601,9 +4614,8 @@ int32 FHLSLMaterialTranslator::ForceCast(int32 Code, EMaterialValueType DestType
 	}
 
 	const EMaterialValueType SourceType = GetParameterType(Code);
-	const bool bExactMatch = (ForceCastFlags & MFCF_ExactMatch) ? true : false;
 
-	if (bExactMatch ? (SourceType == DestType) : (SourceType & DestType))
+	if ((ForceCastFlags & MFCF_ExactMatch) ? (SourceType == DestType) : (SourceType & DestType))
 	{
 		return Code;
 	}
@@ -4611,14 +4623,7 @@ int32 FHLSLMaterialTranslator::ForceCast(int32 Code, EMaterialValueType DestType
 	{
 		const FDerivInfo CodeDerivInfo = GetDerivInfo(Code);
 
-		EMaterialCastFlags CastFlags = EMaterialCastFlags::AllowTruncate | EMaterialCastFlags::AllowAppendZeroes;
-		if ((ForceCastFlags & MFCF_ReplicateValue) || !bExactMatch)
-		{
-			// Replicate scalar if requested, or if we don't require an exact match (this can happen when force-casting to/from LWC
-			// The only way we *don't* replicate scalar is requesting an exact match without the ReplicateValue flag
-			// TODO - My guess is that case probably isn't relevant, and we should just always replicate scalar on cast, but trying to preserve behavior for now
-			CastFlags |= EMaterialCastFlags::ReplicateScalar;
-		}
+		EMaterialCastFlags CastFlags = GetForceCastFlags(ForceCastFlags);
 
 		FString FiniteCode = CastValue(GetParameterCode(Code), SourceType, DestType, CastFlags);
 		if (IsAnalyticDerivEnabled() && IsDerivativeValid(CodeDerivInfo.DerivativeStatus))
@@ -4641,6 +4646,12 @@ int32 FHLSLMaterialTranslator::ForceCast(int32 Code, EMaterialValueType DestType
 			return AddInlinedCodeChunk(DestType, *FiniteCode);
 		}
 	}
+	else if ((SourceType & MCT_StaticBool) && IsFloatNumericType(DestType))
+	{
+		FString StaticBoolToFloat = GetParameterCode(Code).Equals("true") ? "1.0f" : "0.0f";
+		StaticBoolToFloat = CastValue(StaticBoolToFloat, MCT_Float, DestType, GetForceCastFlags(ForceCastFlags));
+		return AddInlinedCodeChunk(DestType, *StaticBoolToFloat);
+	}
 	else if ((SourceType & MCT_TextureVirtual) && (DestType & MCT_Texture2D))
 	{
 		return Code;
