WorldPartitionRuntimeHashSet: fix generating always loaded cells for data layers and content bundles.

diff --git a/Engine/Source/Runtime/Engine/Private/WorldPartition/RuntimeHashSet/WorldPartitionRuntimeHashSetStreamingGeneration.cpp b/Engine/Source/Runtime/Engine/Private/WorldPartition/RuntimeHashSet/WorldPartitionRuntimeHashSetStreamingGeneration.cpp
--- a/Engine/Source/Runtime/Engine/Private/WorldPartition/RuntimeHashSet/WorldPartitionRuntimeHashSetStreamingGeneration.cpp
+++ b/Engine/Source/Runtime/Engine/Private/WorldPartition/RuntimeHashSet/WorldPartitionRuntimeHashSetStreamingGeneration.cpp
@@ -181,7 +181,6 @@ bool UWorldPartitionRuntimeHashSet::GenerateStreaming(UWorldPartitionStreamingPo
 
 		UWorldPartitionRuntimeCell* RuntimeCell = Super::CreateRuntimeCell(CellClass, CellDataClass, CellUniqueId.Name, TEXT(""));
 
-		RuntimeCell->SetIsAlwaysLoaded(!CellDescInstance.bIsSpatiallyLoaded);
 		RuntimeCell->SetDataLayers(CellDescInstance.DataLayerInstances);
 		RuntimeCell->SetContentBundleUID(CellDescInstance.ContentBundleID);
 		RuntimeCell->SetClientOnlyVisible(CellDescInstance.bClientOnlyVisible);
@@ -209,6 +208,7 @@ bool UWorldPartitionRuntimeHashSet::GenerateStreaming(UWorldPartitionStreamingPo
 			if (PopulateCellActorInstances(CellDescInstance.ActorSetInstances, bIsMainWorldPartition, bIsCellAlwaysLoaded, CellActorInstances))
 			{
 				UWorldPartitionRuntimeCell* RuntimeCell = RuntimeCells.Emplace_GetRef(CreateRuntimeCellFromCellDesc(CellDescInstance, StreamingPolicy->GetRuntimeCellClass(), UWorldPartitionRuntimeCellData::StaticClass()));
+				RuntimeCell->SetIsAlwaysLoaded(bIsCellAlwaysLoaded);
 				PopulateRuntimeCell(RuntimeCell, CellActorInstances, OutPackagesToGenerate);
 
 				// Override the cell bounds if the runtime partition provided one
