World Partition Runtime HashSet: fix ForEachStreamingCellsQuery not doing spatial queries.

diff --git a/Engine/Source/Runtime/Engine/Private/WorldPartition/RuntimeHashSet/WorldPartitionRuntimeHashSet.cpp b/Engine/Source/Runtime/Engine/Private/WorldPartition/RuntimeHashSet/WorldPartitionRuntimeHashSet.cpp
--- a/Engine/Source/Runtime/Engine/Private/WorldPartition/RuntimeHashSet/WorldPartitionRuntimeHashSet.cpp
+++ b/Engine/Source/Runtime/Engine/Private/WorldPartition/RuntimeHashSet/WorldPartitionRuntimeHashSet.cpp
@@ -333,16 +333,21 @@ void UWorldPartitionRuntimeHashSet::ForEachStreamingCellsQuery(const FWorldParti
 		return false;
 	};
 
-	auto ForEachStreamingCells = [&ShouldAddCell, &QuerySource, &Func](FStaticSpatialIndexType* InSpatialIndex)
+	auto ForEachStreamingCells = [&ShouldAddCell, &QuerySource, &Func](FStaticSpatialIndexType* InSpatialIndex, int32 InLoadingRange, FName InGridName)
 	{
 		if (InSpatialIndex)
 		{
-			InSpatialIndex->ForEachElement([&ShouldAddCell, &QuerySource, &Func](UWorldPartitionRuntimeCell* RuntimeCell)
+			QuerySource.ForEachShape(InLoadingRange, InGridName, false, [InSpatialIndex, &ShouldAddCell, &QuerySource, &Func](const FSphericalSector& Shape)
 			{
-				if (ShouldAddCell(RuntimeCell, QuerySource))
+				const FSphere ShapeSphere(Shape.GetCenter(), Shape.GetRadius());
+
+				InSpatialIndex->ForEachIntersectingElement(ShapeSphere, [&ShouldAddCell, &QuerySource, &Func](UWorldPartitionRuntimeCell* RuntimeCell)
 				{
-					Func(RuntimeCell);
-				}
+					if (ShouldAddCell(RuntimeCell, QuerySource))
+					{
+						Func(RuntimeCell);
+					}
+				});
 			});
 		}
 	};
@@ -358,9 +363,9 @@ void UWorldPartitionRuntimeHashSet::ForEachStreamingCellsQuery(const FWorldParti
 		}
 	};
 
-	ForEachStreamingData([&ForEachStreamingCells, &ForEachNonStreamingCells](const FRuntimePartitionStreamingData& StreamingData)
+	ForEachStreamingData([&QuerySource, &ForEachStreamingCells, &ForEachNonStreamingCells](const FRuntimePartitionStreamingData& StreamingData)
 	{
-		ForEachStreamingCells(StreamingData.SpatialIndex.Get());
+		ForEachStreamingCells(StreamingData.SpatialIndex.Get(), StreamingData.LoadingRange, StreamingData.Name);
 		ForEachNonStreamingCells(StreamingData.NonStreamingCells);
 	});
 }
