Optimization: Only check in-disk assets in non-editor contexts in GetPrimaryAssetData and GetPrimaryAssetDataList.
Saves multiple seconds of time when transitioning a large number of game feature plugins through unloading.

diff --git a/Engine/Source/Runtime/Engine/Private/AssetManager.cpp b/Engine/Source/Runtime/Engine/Private/AssetManager.cpp
--- a/Engine/Source/Runtime/Engine/Private/AssetManager.cpp
+++ b/Engine/Source/Runtime/Engine/Private/AssetManager.cpp
@@ -1448,8 +1448,7 @@ bool UAssetManager::GetPrimaryAssetData(const FPrimaryAssetId& PrimaryAssetId, F
 
 	if (NameData)
 	{
-		FAssetData CachedAssetData = GetAssetRegistry().GetAssetByObjectPath(NameData->GetARLookupPath(),
-			false /* bIncludeOnlyOnDiskAssets */);
+		FAssetData CachedAssetData = GetAssetRegistry().GetAssetByObjectPath(NameData->GetARLookupPath(), bIncludeOnlyOnDiskAssets);
 
 		if (CachedAssetData.IsValid())
 		{
@@ -1472,8 +1471,7 @@ bool UAssetManager::GetPrimaryAssetDataList(FPrimaryAssetType PrimaryAssetType,
 
 		for (const TPair<FName, FPrimaryAssetData>& Pair : TypeData.GetAssets())
 		{
-			FAssetData CachedAssetData = Registry.GetAssetByObjectPath(Pair.Value.GetARLookupPath(),
-				false /* bIncludeOnlyOnDiskAssets */);
+			FAssetData CachedAssetData = Registry.GetAssetByObjectPath(Pair.Value.GetARLookupPath(), bIncludeOnlyOnDiskAssets);
 
 			if (CachedAssetData.IsValid())
 			{
