Fix MetaSounds not playing within sound cue quality nodes

#localization none
#tests https://horde.devtools.epicgames.com/job/654ab9320f6a6579c21c38eb

diff --git a/Engine/Source/Runtime/Engine/Private/SoundNodeQualityLevel.cpp b/Engine/Source/Runtime/Engine/Private/SoundNodeQualityLevel.cpp
--- a/Engine/Source/Runtime/Engine/Private/SoundNodeQualityLevel.cpp
+++ b/Engine/Source/Runtime/Engine/Private/SoundNodeQualityLevel.cpp
@@ -6,6 +6,7 @@
 #include "EdGraph/EdGraph.h"
 #include "Sound/AudioSettings.h"
 #include "Sound/SoundCue.h"
+#include "Sound/SoundNodeWavePlayer.h"
 #include "AudioCompressionSettingsUtils.h"
 
 #if WITH_EDITORONLY_DATA
@@ -78,7 +79,22 @@ void USoundNodeQualityLevel::ReleaseRetainerOnChildWavePlayers(bool bRecurse)
 
 void USoundNodeQualityLevel::LoadChildWavePlayers(bool bAddToRoot, bool bRecurse)
 {
-	ForCurrentQualityLevel([bAddToRoot, bRecurse](USoundNode* Node) { Node->LoadChildWavePlayerAssets(bAddToRoot, bRecurse); });
+	ForCurrentQualityLevel([bAddToRoot, bRecurse](USoundNode* Node) 
+	{ 
+		if (Node)
+		{
+			// If this node is a wave player node, load it, otherwise load children 
+			// (a wave player node cannot have children)
+			if (USoundNodeWavePlayer* WavePlayerNode = Cast<USoundNodeWavePlayer>(Node))
+			{
+				WavePlayerNode->LoadAsset(bAddToRoot);
+			}
+			else
+			{
+				Node->LoadChildWavePlayerAssets(bAddToRoot, bRecurse);
+			}
+		}
+	});
 }
 
 int32 USoundNodeQualityLevel::GetMaxChildNodes() const
