-Changed static_mesh_component to create it's own renderproxy instead of using a proxy actor

#okforversepublic

diff --git a/Engine/Source/Runtime/Engine/Public/PrimitiveSceneProxy.h b/Engine/Source/Runtime/Engine/Public/PrimitiveSceneProxy.h
--- a/Engine/Source/Runtime/Engine/Public/PrimitiveSceneProxy.h
+++ b/Engine/Source/Runtime/Engine/Public/PrimitiveSceneProxy.h
@@ -207,7 +207,7 @@ public:
 	 * @param bInParentSelected - true if the parent actor is selected in the editor
 	 * @param bInIndividuallySelected - true if the component is selected in the editor directly
 	 */
-	void SetSelection_GameThread(const bool bInParentSelected, const bool bInIndividuallySelected=false);
+	ENGINE_API void SetSelection_GameThread(const bool bInParentSelected, const bool bInIndividuallySelected=false);
 
 	/**
 	 * Updates the LevelInstance editing state for the primitive proxy. This simply sends a message to the rendering thread to call SetLevelInstanceEditingState_RenderThread.
