Refactored CompileDebugViewModeShaders() to include material shader error checking.

diff --git a/Engine/Source/Runtime/Engine/Private/DebugViewModeHelpers.cpp b/Engine/Source/Runtime/Engine/Private/DebugViewModeHelpers.cpp
--- a/Engine/Source/Runtime/Engine/Private/DebugViewModeHelpers.cpp
+++ b/Engine/Source/Runtime/Engine/Private/DebugViewModeHelpers.cpp
@@ -304,58 +304,69 @@ bool CompileDebugViewModeShaders(EDebugViewShaderMode ShaderMode, EMaterialQuali
 	{
 		return false;
 	}
+	
+	EShaderPlatform Platform = GetFeatureLevelShaderPlatform(FeatureLevel);
+
+	FMaterialShaderTypes ShaderTypes;
+	DebugViewModeInterface->AddShaderTypes(FeatureLevel, LocalVertexFactory, ShaderTypes);
+
+	TArray<FMaterial*> PendingMaterials;
+	PendingMaterials.Reserve(Materials.Num());
 
-	// remove materials incompatible with debug view modes (e.g. landscape materials can only be compiled with the landscape VF)
 	for (TSet<UMaterialInterface*>::TIterator It(Materials); It; ++It)
 	{
 		UMaterialInterface* MaterialInterface = *It;
 		check(MaterialInterface); // checked for null in GetTextureStreamingBuildMaterials
-		const FMaterial* Material = MaterialInterface->GetMaterialResource(FeatureLevel, QualityLevel);
+		
+		FMaterial* Material = MaterialInterface->GetMaterialResource(FeatureLevel, QualityLevel);
+		if (!Material)
+		{
+			continue;
+		}
 
-		if (Material && (Material->GetMaterialDomain() != MD_Surface || Material->IsUsedWithLandscape()))
+		// Remove materials incompatible with debug view modes (e.g. landscape materials can only be compiled with the landscape VF)
+		if (Material->GetMaterialDomain() != MD_Surface || Material->IsUsedWithLandscape())
 		{
 			It.RemoveCurrent();
+			continue;
+		}
+		
+		// If material needs the shaders for this platform cached, begin the operation.
+		if (Material->GetGameThreadShaderMap()
+			&& Material->ShouldCacheShaders(Platform, ShaderTypes, LocalVertexFactory)
+			&& !Material->HasShaders(ShaderTypes, LocalVertexFactory))
+		{
+			Material->CacheShaders(Platform, EMaterialShaderPrecompileMode::Default);
+			PendingMaterials.Push(Material);
 		}
 	}
 
-	TSet<UMaterialInterface*> PendingMaterials = Materials;
+	bool bAllMaterialsCompiledSuccessfully = true;
 	while (PendingMaterials.Num() > 0)
 	{
-		for(TSet<UMaterialInterface*>::TIterator It(PendingMaterials); It; ++It )
-		{
-			UMaterialInterface* MaterialInterface = *It;
-			check(MaterialInterface); // checked for null in GetTextureStreamingBuildMaterials
-
-			const FMaterial* Material = MaterialInterface->GetMaterialResource(FeatureLevel, QualityLevel);
-			bool bMaterialFinished = true;
-			if (Material && Material->GetGameThreadShaderMap())
-			{
-				FMaterialShaderTypes ShaderTypes;
-				DebugViewModeInterface->AddShaderTypes(FeatureLevel, LocalVertexFactory, ShaderTypes);
-				if (Material->ShouldCacheShaders(GetFeatureLevelShaderPlatform(FeatureLevel), ShaderTypes, LocalVertexFactory) && !Material->HasShaders(ShaderTypes, LocalVertexFactory))
-				{
-					bMaterialFinished = false;
-				}
-			}
+		FMaterial* Material = PendingMaterials.Last();
 
-			if (bMaterialFinished)
-			{
-				It.RemoveCurrent();
-			}
+		// Check if material has completed compiling the shaders.
+		if (Material->IsCompilationFinished())
+		{
+			bAllMaterialsCompiledSuccessfully &= Material->HasShaders(ShaderTypes, LocalVertexFactory);
+			PendingMaterials.Pop();
+			continue;
 		}
 
-		if (PendingMaterials.Num() > 0)
+		// Are we asked to cancel the operation?
+		if (GWarn->ReceivedUserCancel())
 		{
-			FPlatformProcess::Sleep(0.1f);
-			GShaderCompilingManager->ProcessAsyncResults(false, false);
-			if (GWarn->ReceivedUserCancel())
-			{
-				break;
-			}
+			bAllMaterialsCompiledSuccessfully = false;
+			break;
 		}
+
+		// Wait a little then try again.
+		FPlatformProcess::Sleep(0.1f);
+		GShaderCompilingManager->ProcessAsyncResults(false, false);
 	}
 
-	return PendingMaterials.Num() == 0;
+	return bAllMaterialsCompiledSuccessfully;
 
 #else
 	return false;
