Fix for failing engine tests
#tests EngineTest

diff --git a/Engine/Source/Runtime/Engine/Private/UnrealClient.cpp b/Engine/Source/Runtime/Engine/Private/UnrealClient.cpp
--- a/Engine/Source/Runtime/Engine/Private/UnrealClient.cpp
+++ b/Engine/Source/Runtime/Engine/Private/UnrealClient.cpp
@@ -177,7 +177,7 @@ FUnorderedAccessViewRHIRef FRenderTarget::GetRenderTargetUAV() const
 
 const FTextureRHIRef& FRenderTarget::GetShaderResourceTexture() const
 {
-	return RenderTargetTextureRHI;
+	return GetRenderTargetTexture();
 }
 
 void FScreenshotRequest::RequestScreenshot(bool bInShowUI)
