LocalFogVolume Actor: Fixed Empty Ticks issue.
- The actor was set to have ticks enabled, but no tick functionality was actually being used.
- The fix involves in setting PrimaryActorTick.bCanEverTick = false, for ALocalFogVolume

#tests I was able to repro this in the editor. I tested enabling and disabling ticks and comparing Insights captures before and after. As well as visually comparing the results and making sure local fog is still visible after the fix.

#localization none

diff --git a/Engine/Source/Runtime/Engine/Private/Components/LocalFogVolumeComponent.cpp b/Engine/Source/Runtime/Engine/Private/Components/LocalFogVolumeComponent.cpp
--- a/Engine/Source/Runtime/Engine/Private/Components/LocalFogVolumeComponent.cpp
+++ b/Engine/Source/Runtime/Engine/Private/Components/LocalFogVolumeComponent.cpp
@@ -43,8 +43,8 @@ void ULocalFogVolumeComponent::SendRenderTransformCommand()
 
 void ULocalFogVolumeComponent::CreateRenderState_Concurrent(FRegisterComponentContext* Context)
 {
-	Super::CreateRenderState_Concurrent(Context);	
-	
+	Super::CreateRenderState_Concurrent(Context);
+
 	bool bHidden = false;
 #if WITH_EDITORONLY_DATA
 	bHidden = GetOwner() ? GetOwner()->bHiddenEdLevel : false;
@@ -209,6 +209,7 @@ ALocalFogVolume::ALocalFogVolume(const FObjectInitializer& ObjectInitializer)
 	}
 #endif // WITH_EDITORONLY_DATA
 
-	PrimaryActorTick.bCanEverTick = true;
+	PrimaryActorTick.bCanEverTick = false;
+
 	SetHidden(false);
 }
