change PDI thick line rendering to use 4 triangles instead of 8 per line, and fix the appearance when clipping through the line

diff --git a/Engine/Source/Runtime/Engine/Private/BatchedElements.cpp b/Engine/Source/Runtime/Engine/Private/BatchedElements.cpp
--- a/Engine/Source/Runtime/Engine/Private/BatchedElements.cpp
+++ b/Engine/Source/Runtime/Engine/Private/BatchedElements.cpp
@@ -986,8 +986,19 @@ bool FBatchedElements::Draw(FRHICommandList& RHICmdList, const FMeshPassProcesso
 					OrthoZoomFactor = 1.0f / View.ViewMatrices.GetProjectionMatrix().M[0][0];
 				}
 
+				const double CameraScreenXScale = FVector(WorldToClip.TransformVector(CameraX)).Length();
+				const double CameraScreenYScale = FVector(WorldToClip.TransformVector(CameraY)).Length();
+
 				int32 LineIndex = 0;
-				const int32 MaxLinesPerBatch = 2048;
+				int32 MaxVerticesAllowed = ((GDrawUPVertexCheckCount / sizeof(FSimpleElementVertex)) / 3) * 3;
+				/*
+				hack to avoid a crash when trying to render large numbers of line segments.
+				*/
+				MaxVerticesAllowed = FMath::Min(MaxVerticesAllowed, 64 * 1024);
+
+				constexpr int32 TrisPerLine = 4;
+				constexpr int32 VertPerLine = 3 * TrisPerLine;
+				const int32 MaxLinesPerBatch = MaxVerticesAllowed / VertPerLine;
 				while (LineIndex < ThickLines.Num())
 				{
 					int32 FirstLineThisBatch = LineIndex;
@@ -1010,82 +1021,222 @@ bool FBatchedElements::Draw(FRHICommandList& RHICmdList, const FMeshPassProcesso
 					PrepareShaders(RHICmdList, GraphicsPSOInit, StencilRef, FeatureLevel, SE_BLEND_AlphaBlend, RelativeMatrices, BatchedElementParameters, GWhiteTexture, bHitTesting, Gamma, NULL, &View);
 
 					FRHIResourceCreateInfo CreateInfo(TEXT("ThickLines"));
-					FBufferRHIRef VertexBufferRHI = RHICmdList.CreateBuffer(sizeof(FSimpleElementVertex) * 8 * 3 * NumLinesThisBatch, BUF_VertexBuffer | BUF_Volatile, 0, ERHIAccess::VertexOrIndexBuffer, CreateInfo);
-					void* ThickVertexData = RHICmdList.LockBuffer(VertexBufferRHI, 0, sizeof(FSimpleElementVertex) * 8 * 3 * NumLinesThisBatch, RLM_WriteOnly);
+					FBufferRHIRef VertexBufferRHI = RHICmdList.CreateBuffer(sizeof(FSimpleElementVertex) * VertPerLine * NumLinesThisBatch, BUF_VertexBuffer | BUF_Volatile, 0, ERHIAccess::VertexOrIndexBuffer, CreateInfo);
+					void* ThickVertexData = RHICmdList.LockBuffer(VertexBufferRHI, 0, sizeof(FSimpleElementVertex) * VertPerLine * NumLinesThisBatch, RLM_WriteOnly);
 					FSimpleElementVertex* ThickVertices = (FSimpleElementVertex*)ThickVertexData;
 					check(ThickVertices);
 
+					int32 AddedTris = 0;
 					for (int i = 0; i < NumLinesThisBatch; ++i)
 					{
 						const FBatchedThickLines& Line = ThickLines[FirstLineThisBatch + i];
-						const float Thickness = FMath::Abs( Line.Thickness );
-
-						const float StartW			= WorldToClip.TransformFVector4(Line.Start).W;
-						const float EndW			= WorldToClip.TransformFVector4(Line.End).W;
-
-						// Negative thickness means that thickness is calculated in screen space, positive thickness should be used for world space thickness.
-						const float ScalingStart	= Line.bScreenSpace ? StartW / ViewportSizeX : 1.0f;
-						const float ScalingEnd		= Line.bScreenSpace ? EndW   / ViewportSizeX : 1.0f;
-
-						const float CurrentOrthoZoomFactor = Line.bScreenSpace ? OrthoZoomFactor : 1.0f;
-
-						const float ScreenSpaceScaling = Line.bScreenSpace ? 2.0f : 1.0f;
-
-						const float StartThickness	= Thickness * ScreenSpaceScaling * CurrentOrthoZoomFactor * ScalingStart;
-						const float EndThickness	= Thickness * ScreenSpaceScaling * CurrentOrthoZoomFactor * ScalingEnd;
-
-						const FVector WorldPointXS	= CameraX * StartThickness * 0.5f;
-						const FVector WorldPointYS	= CameraY * StartThickness * 0.5f;
+						const double Thickness = FMath::Abs( Line.Thickness );
 
-						const FVector WorldPointXE	= CameraX * EndThickness * 0.5f;
-						const FVector WorldPointYE	= CameraY * EndThickness * 0.5f;
+						FVector4 StartClip	= WorldToClip.TransformFVector4(Line.Start);
+						FVector4 EndClip	= WorldToClip.TransformFVector4(Line.End);
 
-						// Generate vertices for the point such that the post-transform point size is constant.
-						const FVector WorldPointX = CameraX * Thickness * StartW / ViewportSizeX;
-						const FVector WorldPointY = CameraY * Thickness * StartW / ViewportSizeX;
+						// Manually clip thick lines start/end if they would go behind the near plane
+						FVector LineStart = Line.Start;
+						FVector LineEnd = Line.End;
+						const double ClipAt = View.NearClippingDistance + UE_DOUBLE_KINDA_SMALL_NUMBER;
+						if (StartClip.W < ClipAt && EndClip.W < ClipAt)
+						{
+							continue;
+						}
+						else if (StartClip.W < ClipAt)
+						{
+							double Along = (ClipAt - StartClip.W) / (EndClip.W - StartClip.W);
+							LineStart = FMath::Lerp(LineStart, LineEnd, Along);
+							StartClip = FMath::Lerp(StartClip, EndClip, Along);
+						}
+						else if (EndClip.W < ClipAt)
+						{
+							double Along = (ClipAt - EndClip.W) / (StartClip.W - EndClip.W);
+							LineEnd = FMath::Lerp(LineEnd, LineStart, Along);
+							EndClip = FMath::Lerp(EndClip, StartClip, Along);
+						}
 
-						// Begin point
-						ThickVertices[0] = FSimpleElementVertex(Line.Start + WorldPointXS - WorldPointYS,FVector2D(1,0),Line.Color,Line.HitProxyColor); // 0S
-						ThickVertices[1] = FSimpleElementVertex(Line.Start + WorldPointXS + WorldPointYS,FVector2D(1,1),Line.Color,Line.HitProxyColor); // 1S
-						ThickVertices[2] = FSimpleElementVertex(Line.Start - WorldPointXS - WorldPointYS,FVector2D(0,0),Line.Color,Line.HitProxyColor); // 2S
-					
-						ThickVertices[3] = FSimpleElementVertex(Line.Start + WorldPointXS + WorldPointYS,FVector2D(1,1),Line.Color,Line.HitProxyColor); // 1S
-						ThickVertices[4] = FSimpleElementVertex(Line.Start - WorldPointXS - WorldPointYS,FVector2D(0,0),Line.Color,Line.HitProxyColor); // 2S
-						ThickVertices[5] = FSimpleElementVertex(Line.Start - WorldPointXS + WorldPointYS,FVector2D(0,1),Line.Color,Line.HitProxyColor); // 3S
-
-						// Ending point
-						ThickVertices[0+ 6] = FSimpleElementVertex(Line.End + WorldPointXE - WorldPointYE,FVector2D(1,0),Line.Color,Line.HitProxyColor); // 0E
-						ThickVertices[1+ 6] = FSimpleElementVertex(Line.End + WorldPointXE + WorldPointYE,FVector2D(1,1),Line.Color,Line.HitProxyColor); // 1E
-						ThickVertices[2+ 6] = FSimpleElementVertex(Line.End - WorldPointXE - WorldPointYE,FVector2D(0,0),Line.Color,Line.HitProxyColor); // 2E
-																																							  
-						ThickVertices[3+ 6] = FSimpleElementVertex(Line.End + WorldPointXE + WorldPointYE,FVector2D(1,1),Line.Color,Line.HitProxyColor); // 1E
-						ThickVertices[4+ 6] = FSimpleElementVertex(Line.End - WorldPointXE - WorldPointYE,FVector2D(0,0),Line.Color,Line.HitProxyColor); // 2E
-						ThickVertices[5+ 6] = FSimpleElementVertex(Line.End - WorldPointXE + WorldPointYE,FVector2D(0,1),Line.Color,Line.HitProxyColor); // 3E
-
-						// First part of line
-						ThickVertices[0+12] = FSimpleElementVertex(Line.Start - WorldPointXS - WorldPointYS,FVector2D(0,0),Line.Color,Line.HitProxyColor); // 2S
-						ThickVertices[1+12] = FSimpleElementVertex(Line.Start + WorldPointXS + WorldPointYS,FVector2D(1,1),Line.Color,Line.HitProxyColor); // 1S
-						ThickVertices[2+12] = FSimpleElementVertex(Line.End   - WorldPointXE - WorldPointYE,FVector2D(0,0),Line.Color,Line.HitProxyColor); // 2E
-
-						ThickVertices[3+12] = FSimpleElementVertex(Line.Start + WorldPointXS + WorldPointYS,FVector2D(1,1),Line.Color,Line.HitProxyColor); // 1S
-						ThickVertices[4+12] = FSimpleElementVertex(Line.End   + WorldPointXE + WorldPointYE,FVector2D(1,1),Line.Color,Line.HitProxyColor); // 1E
-						ThickVertices[5+12] = FSimpleElementVertex(Line.End   - WorldPointXE - WorldPointYE,FVector2D(0,0),Line.Color,Line.HitProxyColor); // 2E
-
-						// Second part of line
-						ThickVertices[0+18] = FSimpleElementVertex(Line.Start - WorldPointXS + WorldPointYS,FVector2D(0,1),Line.Color,Line.HitProxyColor); // 3S
-						ThickVertices[1+18] = FSimpleElementVertex(Line.Start + WorldPointXS - WorldPointYS,FVector2D(1,0),Line.Color,Line.HitProxyColor); // 0S
-						ThickVertices[2+18] = FSimpleElementVertex(Line.End   - WorldPointXE + WorldPointYE,FVector2D(0,1),Line.Color,Line.HitProxyColor); // 3E
-
-						ThickVertices[3+18] = FSimpleElementVertex(Line.Start + WorldPointXS - WorldPointYS,FVector2D(1,0),Line.Color,Line.HitProxyColor); // 0S
-						ThickVertices[4+18] = FSimpleElementVertex(Line.End   + WorldPointXE - WorldPointYE,FVector2D(1,0),Line.Color,Line.HitProxyColor); // 0E
-						ThickVertices[5+18] = FSimpleElementVertex(Line.End   - WorldPointXE + WorldPointYE,FVector2D(0,1),Line.Color,Line.HitProxyColor); // 3E
-
-						ThickVertices += 24;
+						const double StartW = StartClip.W;
+						const double EndW = EndClip.W;
+
+						// bScreenSpace controls if thickness is calculated in screen space or world space.
+						const double ScalingStart	= Line.bScreenSpace ? StartW / ViewportSizeX : 1.0;
+						const double ScalingEnd		= Line.bScreenSpace ? EndW   / ViewportSizeX : 1.0;
+
+						const double CurrentOrthoZoomFactor = Line.bScreenSpace ? OrthoZoomFactor : 1.0;
+
+						const double ScreenSpaceScaling = Line.bScreenSpace ? 2.0 : 1.0;
+
+						const double StartThickness	= Thickness * ScreenSpaceScaling * CurrentOrthoZoomFactor * ScalingStart;
+						const double EndThickness	= Thickness * ScreenSpaceScaling * CurrentOrthoZoomFactor * ScalingEnd;
+
+						// Line start and end points S and E are expanded into quads based on the line thickness, and
+						// the line is filled in by trapezoids S2,S0,E0,E1 and S2,S3,E3,E1, w/ vertices as pictured:
+						//      E3---E1
+						//     / |...|
+						//    /  |...|
+						//   /  E2---E0
+						// S3---S1  /
+						//  |...|  /
+						//  |...| /
+						// S2---S0
+						// This labelling assume the direction from S->E on screen is up-right (pictured)
+						// We mirror point labels as needed to handle the down- or left- oriented cases
+
+						// Figure out whether we're in an up-right/down-left case or not
+						const double InvStartW = 1.0 / StartW, InvEndW = 1.0 / EndW;
+						const double ScreenShiftX = EndClip.X * InvEndW - StartClip.X * InvStartW;
+						const double ScreenShiftY = EndClip.Y * InvEndW - StartClip.Y * InvStartW;
+						const bool bLineRight = ScreenShiftX >= 0;
+						const bool bLineUp = ScreenShiftY >= 0;
+
+						// Mirror X offsets in the up-left or down-right case
+						const double XSign = bLineRight ? 1.0 : -1.0;
+						const double YSign = bLineUp ? 1.0 : -1.0;
+						// Mirror X for UVs as well
+						const double RightUVX = (double)bLineRight, LeftUVX = (double)!bLineRight;
+						const double UpUVY = (double)bLineUp, DownUVY = (double)!bLineUp;
+
+						// Create the X and Y world offsets for each point, mirrored as needed
+						const FVector WorldPointXS	= CameraX * (StartThickness * 0.5 * XSign);
+						const FVector WorldPointYS	= CameraY * (StartThickness * 0.5 * YSign);
+						const FVector WorldPointXE	= CameraX * (EndThickness * 0.5 * XSign);
+						const FVector WorldPointYE	= CameraY * (EndThickness * 0.5 * YSign);
+
+						// Vertex positions
+						FVector S0 = LineStart + WorldPointXS - WorldPointYS;
+						FVector S1 = LineStart + WorldPointXS + WorldPointYS;
+						FVector S2 = LineStart - WorldPointXS - WorldPointYS;
+						FVector S3 = LineStart - WorldPointXS + WorldPointYS;
+
+						FVector E0 = LineEnd + WorldPointXE - WorldPointYE;
+						FVector E1 = LineEnd + WorldPointXE + WorldPointYE;
+						FVector E2 = LineEnd - WorldPointXE - WorldPointYE;
+						FVector E3 = LineEnd - WorldPointXE + WorldPointYE;
+
+						// UVs per vertex
+						FVector2D S0UV = FVector2D(RightUVX, DownUVY);
+						FVector2D S1UV = FVector2D(RightUVX, UpUVY);
+						FVector2D S2UV = FVector2D(LeftUVX, DownUVY);
+						FVector2D S3UV = FVector2D(LeftUVX, UpUVY);
+
+						FVector2D E0UV = FVector2D(RightUVX, DownUVY);
+						FVector2D E1UV = FVector2D(RightUVX, UpUVY);
+						FVector2D E2UV = FVector2D(LeftUVX, DownUVY);
+						FVector2D E3UV = FVector2D(LeftUVX, UpUVY);
+
+						// Handle special cases due to one end of the line being larger on screen than the other -- these don't happen w/ screen space lines
+						// Note: If you skip handling these cases, the lines will still mostly look fine, just not perfect if you get close to them
+						if (!Line.bScreenSpace)
+						{
+							// Figure out if we're in a only-vertical (pictured) or only-horizontal case, like:
+							//      E3---E1
+							//     / |...| \
+							//    /  |...|  \
+							//   /  E2---E0  \
+							// S3-------------S1
+							//  |.............|
+							// ... where S contains E along one axis, or vice versa
+							// In these cases, replace one of the vertices with S1 or E2 to fix the missing corner
+							const double ScreenStartRadius = StartThickness * InvStartW * .5;
+							const double ScreenEndRadius = EndThickness * InvEndW * .5;
+							const double RadDiff = ScreenEndRadius - ScreenStartRadius;
+							const bool bStartIsBigger = RadDiff < 0;
+							const double AbsRadDiff = bStartIsBigger ? -RadDiff : RadDiff;
+
+							bool bOnlyVertical = ScreenShiftX * XSign < AbsRadDiff * CameraScreenXScale, bOnlyHorizontal = ScreenShiftY * YSign < AbsRadDiff * CameraScreenYScale;
+							// If the larger point contains the smaller on both axes, we can just draw the larger point alone
+							if (bOnlyVertical && bOnlyHorizontal)
+							{
+								if (bStartIsBigger)
+								{
+									// Tri: S0,S1,S2
+									ThickVertices[0] = FSimpleElementVertex(S0, S0UV, Line.Color, Line.HitProxyColor);
+									ThickVertices[1] = FSimpleElementVertex(S1, S1UV, Line.Color, Line.HitProxyColor);
+									ThickVertices[2] = FSimpleElementVertex(S2, S2UV, Line.Color, Line.HitProxyColor);
+									// Tri: S2,S1,S3
+									ThickVertices[3] = FSimpleElementVertex(S2, S2UV, Line.Color, Line.HitProxyColor);
+									ThickVertices[4] = FSimpleElementVertex(S1, S1UV, Line.Color, Line.HitProxyColor);
+									ThickVertices[5] = FSimpleElementVertex(S3, S3UV, Line.Color, Line.HitProxyColor);
+								}
+								else
+								{
+									// Tri: E0,E1,E2
+									ThickVertices[0] = FSimpleElementVertex(E0, E0UV, Line.Color, Line.HitProxyColor);
+									ThickVertices[1] = FSimpleElementVertex(E1, E1UV, Line.Color, Line.HitProxyColor);
+									ThickVertices[2] = FSimpleElementVertex(E2, E2UV, Line.Color, Line.HitProxyColor);
+									// Tri: E2,E1,E3
+									ThickVertices[3] = FSimpleElementVertex(E2, E2UV, Line.Color, Line.HitProxyColor);
+									ThickVertices[4] = FSimpleElementVertex(E1, E1UV, Line.Color, Line.HitProxyColor);
+									ThickVertices[5] = FSimpleElementVertex(E3, E3UV, Line.Color, Line.HitProxyColor);
+								}
+								ThickVertices += 6;
+								AddedTris += 2;
+								continue;
+							}
+							else if (bOnlyVertical) // only one direction is dominant; figure out which vertex to replace
+							{
+								if (bStartIsBigger)
+								{
+									// replace E0 with S1
+									E0 = S1;
+									E0UV = S1UV;
+								}
+								else
+								{
+									// replace S3 with E2
+									S3 = E2;
+									S3UV = E2UV;
+								}
+							}
+							else if (bOnlyHorizontal)
+							{
+								if (bStartIsBigger)
+								{
+									// replace E3 with S1
+									E3 = S1;
+									E3UV = S1UV;
+								}
+								else
+								{
+									// replace S0 with E2
+									S0 = E2;
+									S0UV = E2UV;
+								}
+							}
+						}
+						
+
+						// First trapezoid (S2,S0,E0,E1)
+						// Tri: S2,S0,E0
+						ThickVertices[0] = FSimpleElementVertex(S2, S2UV, Line.Color, Line.HitProxyColor);
+						ThickVertices[1] = FSimpleElementVertex(S0, S0UV, Line.Color, Line.HitProxyColor);
+						ThickVertices[2] = FSimpleElementVertex(E0, E0UV, Line.Color, Line.HitProxyColor);
+						// Tri: S2,E1,E0
+						ThickVertices[3] = FSimpleElementVertex(S2, S2UV, Line.Color, Line.HitProxyColor);
+						ThickVertices[4] = FSimpleElementVertex(E1, E1UV, Line.Color, Line.HitProxyColor);
+						ThickVertices[5] = FSimpleElementVertex(E0, E0UV, Line.Color, Line.HitProxyColor);
+
+						// Second trapezoid (S2,S3,E3,E1)
+						// Tri: S2,E3,E1
+						ThickVertices[6] = FSimpleElementVertex(S2, S2UV, Line.Color, Line.HitProxyColor);
+						ThickVertices[7] = FSimpleElementVertex(E3, E3UV, Line.Color, Line.HitProxyColor);
+						ThickVertices[8] = FSimpleElementVertex(E1, E1UV, Line.Color, Line.HitProxyColor);
+						// Tri: S2,S3,E3
+						ThickVertices[9] = FSimpleElementVertex(S2, S2UV, Line.Color, Line.HitProxyColor);
+						ThickVertices[10] = FSimpleElementVertex(S3, S3UV, Line.Color, Line.HitProxyColor);
+						ThickVertices[11] = FSimpleElementVertex(E3, E3UV, Line.Color, Line.HitProxyColor);
+
+						ThickVertices += VertPerLine;
+						AddedTris += 4;
 					}
 
-					RHICmdList.UnlockBuffer(VertexBufferRHI);
-					RHICmdList.SetStreamSource(0, VertexBufferRHI, 0);
-					RHICmdList.DrawPrimitive(0, 8 * NumLinesThisBatch, 1);
+					if (AddedTris > 0)
+					{
+						RHICmdList.UnlockBuffer(VertexBufferRHI);
+						RHICmdList.SetStreamSource(0, VertexBufferRHI, 0);
+						RHICmdList.DrawPrimitive(0, AddedTris, 1);
+					}
 				}
 
 				GraphicsPSOInit.RasterizerState = TStaticRasterizerState<FM_Solid, CM_None>::GetRHI();
