Avoid a memory allocation in UPrimitiveComponent::IsSimulatingPhysics

diff --git a/Engine/Source/Runtime/Engine/Private/PrimitiveComponentPhysics.cpp b/Engine/Source/Runtime/Engine/Private/PrimitiveComponentPhysics.cpp
--- a/Engine/Source/Runtime/Engine/Private/PrimitiveComponentPhysics.cpp
+++ b/Engine/Source/Runtime/Engine/Private/PrimitiveComponentPhysics.cpp
@@ -1014,11 +1014,14 @@ bool UPrimitiveComponent::IsSimulatingPhysics(FName BoneName) const
 		}
 
 		FLockedReadPhysicsObjectExternalInterface Interface = FPhysicsObjectExternalInterface::LockRead(PhysicsObjects);
-		PhysicsObjects = PhysicsObjects.FilterByPredicate(
-			[&Interface](Chaos::FPhysicsObjectHandle Object) {
-				return !Interface->AreAllDisabled({ &Object, 1 });
-			}
-		);
+
+		constexpr bool bAllowShrinking = false;
+		PhysicsObjects.RemoveAllSwap(
+			[&Interface](Chaos::FPhysicsObjectHandle Object)
+			{
+				return Interface->AreAllDisabled({ &Object, 1 });
+			},
+			bAllowShrinking);
 
 		return Interface->AreAllDynamicOrSleeping(PhysicsObjects);
 	}
