Modify Cvar r.Shaders.WarningsAsErrors to only apply to global shaders when value is 1 (and all shaders only when value is 2 or more)

This is still not enabled by default, but gives a convenient way to see how many global shaders produce warnings.

diff --git a/Engine/Source/Runtime/Engine/Private/ShaderCompiler/ShaderCompiler.cpp b/Engine/Source/Runtime/Engine/Private/ShaderCompiler/ShaderCompiler.cpp
--- a/Engine/Source/Runtime/Engine/Private/ShaderCompiler/ShaderCompiler.cpp
+++ b/Engine/Source/Runtime/Engine/Private/ShaderCompiler/ShaderCompiler.cpp
@@ -2200,7 +2200,7 @@ static TAutoConsoleVariable<int32> CVarShaderBoundsChecking(
 static TAutoConsoleVariable<int32> CVarShaderWarningsAsErrors(
 	TEXT("r.Shaders.WarningsAsErrors"),
 	0,
-	TEXT("Whether to treat warnings as errors when running the shader compiler. Defaults to 0 (disabled). Not all compilers support this mode."),
+	TEXT("Whether to treat warnings as errors when compiling shaders. (0: disabled (default), 1: global shaders only, 2: all shaders)). This setting may be ignored on older platforms."),
 	ECVF_ReadOnly);
 
 static TAutoConsoleVariable<int32> CVarShaderFlowControl(
@@ -7779,9 +7779,12 @@ void GlobalBeginCompileShader(
 		SET_SHADER_DEFINE(Input.Environment, DO_GUARD_SLOW, GSShaderCheckLevel > 1 ? 1 : 0);
 	}
 
-	if (CVarShaderWarningsAsErrors.GetValueOnAnyThread())
+	if (int WarnLevel = CVarShaderWarningsAsErrors.GetValueOnAnyThread(); WarnLevel != 0)
 	{
-		Input.Environment.CompilerFlags.Add(CFLAG_WarningsAsErrors);
+		if ((WarnLevel == 1 && ShaderType->GetTypeForDynamicCast() == FShaderType::EShaderTypeForDynamicCast::Global) || WarnLevel > 1)
+		{
+			Input.Environment.CompilerFlags.Add(CFLAG_WarningsAsErrors);
+		}
 	}
 
 	if (UseRemoveUnsedInterpolators((EShaderPlatform)Target.Platform) && !IsOpenGLPlatform((EShaderPlatform)Target.Platform))
