Fix for inertial/dead blend carrying over previous blend when returning to relevancy

diff --git a/Engine/Source/Runtime/Engine/Classes/Animation/AnimNode_DeadBlending.h b/Engine/Source/Runtime/Engine/Classes/Animation/AnimNode_DeadBlending.h
--- a/Engine/Source/Runtime/Engine/Classes/Animation/AnimNode_DeadBlending.h
+++ b/Engine/Source/Runtime/Engine/Classes/Animation/AnimNode_DeadBlending.h
@@ -139,6 +139,12 @@ private:
 	UPROPERTY(EditAnywhere, Category = Memory)
 	bool bPreallocateMemory = false;
 
+	/**
+	* Clear any active blends if we just became relevant, to avoid carrying over undesired blends.
+	*/	
+	UPROPERTY(EditAnywhere, Category = Blending)
+	bool bResetOnBecomingRelevant = true;
+
 #if WITH_EDITORONLY_DATA
 	
 	// This setting can be used to show what the extrapolation of the animation looks like.
@@ -213,6 +219,9 @@ private:
 	UPROPERTY(Transient)
 	TArray<FInertializationRequest> RequestQueue;
 
+	// Update Counter for detecting being relevant
+	FGraphTraversalCounter UpdateCounter;
+
 private:
 
 	// Recorded pose state at point of transition.
diff --git a/Engine/Source/Runtime/Engine/Classes/Animation/AnimNode_Inertialization.h b/Engine/Source/Runtime/Engine/Classes/Animation/AnimNode_Inertialization.h
--- a/Engine/Source/Runtime/Engine/Classes/Animation/AnimNode_Inertialization.h
+++ b/Engine/Source/Runtime/Engine/Classes/Animation/AnimNode_Inertialization.h
@@ -382,6 +382,12 @@ private:
 	UPROPERTY(EditAnywhere, Category = Memory)
 	bool bPreallocateMemory = false;
 
+	/**
+	* Clear any active blends if we just became relevant, to avoid carrying over undesired blends.
+	*/	
+	UPROPERTY(EditAnywhere, Category = Blending)
+	bool bResetOnBecomingRelevant = false;
+
 public: // FAnimNode_Inertialization
 
 	ENGINE_API FAnimNode_Inertialization();
@@ -442,6 +448,9 @@ private:
 	UPROPERTY(Transient)
 	TArray<FInertializationRequest> RequestQueue;
 
+	// Update Counter for detecting being relevant
+	FGraphTraversalCounter UpdateCounter;
+
 	// Teleport type
 	ETeleportType TeleportType;
 
diff --git a/Engine/Source/Runtime/Engine/Private/Animation/AnimNode_DeadBlending.cpp b/Engine/Source/Runtime/Engine/Private/Animation/AnimNode_DeadBlending.cpp
--- a/Engine/Source/Runtime/Engine/Private/Animation/AnimNode_DeadBlending.cpp
+++ b/Engine/Source/Runtime/Engine/Private/Animation/AnimNode_DeadBlending.cpp
@@ -652,6 +652,17 @@ void FAnimNode_DeadBlending::Update_AnyThread(const FAnimationUpdateContext& Con
 	LLM_SCOPE_BYNAME(TEXT("Animation/DeadBlending"));
 	DECLARE_SCOPE_HIERARCHICAL_COUNTER_ANIMNODE(Update_AnyThread);
 
+	const bool bNeedsReset =
+		bResetOnBecomingRelevant &&
+		UpdateCounter.HasEverBeenUpdated() &&
+		!UpdateCounter.WasSynchronizedCounter(Context.AnimInstanceProxy->GetUpdateCounter());
+
+	if (bNeedsReset)
+	{
+		Deactivate();
+	}
+	UpdateCounter.SynchronizeWith(Context.AnimInstanceProxy->GetUpdateCounter());
+
 	const int32 NodeId = Context.GetCurrentNodeId();
 	const FAnimInstanceProxy& Proxy = *Context.AnimInstanceProxy;
 
diff --git a/Engine/Source/Runtime/Engine/Private/Animation/AnimNode_Inertialization.cpp b/Engine/Source/Runtime/Engine/Private/Animation/AnimNode_Inertialization.cpp
--- a/Engine/Source/Runtime/Engine/Private/Animation/AnimNode_Inertialization.cpp
+++ b/Engine/Source/Runtime/Engine/Private/Animation/AnimNode_Inertialization.cpp
@@ -203,6 +203,17 @@ void FAnimNode_Inertialization::Update_AnyThread(const FAnimationUpdateContext&
 	LLM_SCOPE_BYNAME(TEXT("Animation/Inertialization"));
 	DECLARE_SCOPE_HIERARCHICAL_COUNTER_ANIMNODE(Update_AnyThread);
 
+	const bool bNeedsReset =
+		bResetOnBecomingRelevant &&
+		UpdateCounter.HasEverBeenUpdated() &&
+		!UpdateCounter.WasSynchronizedCounter(Context.AnimInstanceProxy->GetUpdateCounter());
+
+	if (bNeedsReset)
+	{
+		Deactivate();
+	}
+	UpdateCounter.SynchronizeWith(Context.AnimInstanceProxy->GetUpdateCounter());
+
 	const int32 NodeId = Context.GetCurrentNodeId();
 	const FAnimInstanceProxy& Proxy = *Context.AnimInstanceProxy;
 
