UEWaveComp: Fixed 8-channel audio bug where it would only UEWaveComp encode the first channel.

diff --git a/Engine/Source/Runtime/Engine/Classes/Sound/SoundWave.h b/Engine/Source/Runtime/Engine/Classes/Sound/SoundWave.h
--- a/Engine/Source/Runtime/Engine/Classes/Sound/SoundWave.h
+++ b/Engine/Source/Runtime/Engine/Classes/Sound/SoundWave.h
@@ -904,6 +904,18 @@ public:
 	struct FEditorAudioBulkData
 	{
 		UE::Serialization::FEditorBulkData RawData;
+		USoundWave* SoundWave;
+
+		FEditorAudioBulkData()
+			: RawData()
+			, SoundWave()
+		{
+		}
+			
+		FEditorAudioBulkData(USoundWave* Owner)
+		{
+			SoundWave = Owner;
+		}
 
 		ENGINE_API void CreateFromBulkData(FBulkData& InBulkData, const FGuid& InGuid, UObject* Owner);
 		ENGINE_API void Serialize(FArchive& Ar, UObject* Owner, bool bAllowRegister=true);
diff --git a/Engine/Source/Runtime/Engine/Private/SoundWave.cpp b/Engine/Source/Runtime/Engine/Private/SoundWave.cpp
--- a/Engine/Source/Runtime/Engine/Private/SoundWave.cpp
+++ b/Engine/Source/Runtime/Engine/Private/SoundWave.cpp
@@ -892,6 +892,9 @@ uint32 FStreamedAudioChunk::StoreInDerivedDataCache(const FString& InDerivedData
 
 USoundWave::USoundWave(const FObjectInitializer& ObjectInitializer)
 	: Super(ObjectInitializer)
+#if WITH_EDITORONLY_DATA
+	, RawData(this)
+#endif
 {
 	Volume = 1.0;
 	Pitch = 1.0;
@@ -4856,20 +4859,45 @@ TFuture<FSharedBuffer> USoundWave::FEditorAudioBulkData::GetPayload() const
 	Data = (const uint8*) Buffer.GetData();
 	DataSize = Buffer.GetSize();
 	FWaveModInfo WaveInfo;
-	if (!WaveInfo.ReadWaveInfo(Data, DataSize, 0)) 
+	if (SoundWave && SoundWave->ChannelOffsets.Num() > 0)
 	{
-		UE_LOG(LogAudio, Warning, TEXT("Failed to read wave data."));
+		TArray<int16> scratch_buffer;
+		for (int i = 0; i < SoundWave->ChannelOffsets.Num(); ++i)
+		{
+			if (!WaveInfo.ReadWaveInfo(Data + SoundWave->ChannelOffsets[i], DataSize))
+			{
+				UE_LOG(LogAudio, Warning, TEXT("Failed to read wave data out of '%s'."), *SoundWave->GetFullName());
+			}
+			if (*WaveInfo.pFormatTag == FWaveModInfo::WAVE_INFO_FORMAT_OODLE_WAVE)
+			{
+				// Convert UEWavComp data back to a WAV file
+				*WaveInfo.pFormatTag = FWaveModInfo::WAVE_INFO_FORMAT_PCM;
+				int16* samples = (int16*)WaveInfo.SampleDataStart;
+				int64 num_samples = WaveInfo.GetNumSamples();
+				int64 num_channels = *WaveInfo.pChannels;
+				scratch_buffer.Reset();
+				scratch_buffer.AddUninitialized(num_samples);
+				uewav_decode16(samples, scratch_buffer.GetData(), num_samples, num_channels);
+			}
+		}
 	}
-	if (*WaveInfo.pFormatTag == FWaveModInfo::WAVE_INFO_FORMAT_OODLE_WAVE)
+	else
 	{
-		// Convert UEWavComp data back to a WAV file
-		*WaveInfo.pFormatTag = FWaveModInfo::WAVE_INFO_FORMAT_PCM;
-		int16* samples = (int16*)WaveInfo.SampleDataStart;
-		int64 num_samples = WaveInfo.GetNumSamples();
-		int64 num_channels = *WaveInfo.pChannels;
-		TArray<int16> scratch_buffer;
-		scratch_buffer.AddUninitialized(num_samples);
-		uewav_decode16(samples, scratch_buffer.GetData(), num_samples, num_channels);
+		if (!WaveInfo.ReadWaveInfo(Data, DataSize)) 
+		{
+			UE_LOG(LogAudio, Warning, TEXT("Failed to read wave data."));
+		}
+		if (*WaveInfo.pFormatTag == FWaveModInfo::WAVE_INFO_FORMAT_OODLE_WAVE)
+		{
+			// Convert UEWavComp data back to a WAV file
+			*WaveInfo.pFormatTag = FWaveModInfo::WAVE_INFO_FORMAT_PCM;
+			int16* samples = (int16*)WaveInfo.SampleDataStart;
+			int64 num_samples = WaveInfo.GetNumSamples();
+			int64 num_channels = *WaveInfo.pChannels;
+			TArray<int16> scratch_buffer;
+			scratch_buffer.AddUninitialized(num_samples);
+			uewav_decode16(samples, scratch_buffer.GetData(), num_samples, num_channels);
+		}
 	}
 	TPromise<FSharedBuffer> promise;
 	promise.EmplaceValue(Buffer);
@@ -4883,26 +4911,55 @@ void USoundWave::FEditorAudioBulkData::UpdatePayload(FSharedBuffer InPayload, UO
 	FSharedBuffer Buffer = FSharedBuffer::Clone(Data, DataSize);
 	Data = (const uint8*) Buffer.GetData();
 	DataSize = Buffer.GetSize();
-	FWaveModInfo WaveInfo;
-	if (!WaveInfo.ReadWaveInfo(Data, DataSize, 0)) 
+	bool bEnableUEWavComp = false;
+	GConfig->GetBool(TEXT("AudioImporter"), TEXT("EnableUEWavComp"), bEnableUEWavComp, GEditorIni);
+	if (SoundWave && SoundWave->ChannelOffsets.Num() > 0) 
 	{
-		UE_LOG(LogAudio, Warning, TEXT("Failed to read wave data."));
-	}
-	if (*WaveInfo.pFormatTag == FWaveModInfo::WAVE_INFO_FORMAT_PCM)
+		TArray<int16> scratch_buffer;
+		for (int i = 0; i < SoundWave->ChannelOffsets.Num(); ++i)
+		{
+			FWaveModInfo WaveInfo;
+			if (!WaveInfo.ReadWaveInfo(Data + SoundWave->ChannelOffsets[i], DataSize))
+			{
+				UE_LOG(LogAudio, Warning, TEXT("Failed to read wave data out of '%s'."), *SoundWave->GetFullName());
+			}
+			if (*WaveInfo.pFormatTag == FWaveModInfo::WAVE_INFO_FORMAT_PCM)
+			{
+				if (bEnableUEWavComp)
+				{
+					int16* samples = (int16*)WaveInfo.SampleDataStart;
+					int64 num_samples = WaveInfo.GetNumSamples();
+					int64 num_channels = *WaveInfo.pChannels;
+					scratch_buffer.Reset();
+					scratch_buffer.AddUninitialized(num_samples);
+					uewav_encode16(samples, scratch_buffer.GetData(), num_samples, num_channels);
+					*WaveInfo.pFormatTag = FWaveModInfo::WAVE_INFO_FORMAT_OODLE_WAVE;
+				}
+			}
+		}
+	} 
+	else
 	{
-		bool bEnableUEWavComp = false;
-		GConfig->GetBool(TEXT("AudioImporter"), TEXT("EnableUEWavComp"), bEnableUEWavComp, GEditorIni);
-		if (bEnableUEWavComp)
+		FWaveModInfo WaveInfo;
+		if (!WaveInfo.ReadWaveInfo(Data, DataSize, 0)) 
 		{
-			int16* samples = (int16*)WaveInfo.SampleDataStart;
-			int64 num_samples = WaveInfo.GetNumSamples();
-			int64 num_channels = *WaveInfo.pChannels;
-			TArray<int16> scratch_buffer;
-			scratch_buffer.AddUninitialized(num_samples);
-			uewav_encode16(samples, scratch_buffer.GetData(), num_samples, num_channels);
-			*WaveInfo.pFormatTag = FWaveModInfo::WAVE_INFO_FORMAT_OODLE_WAVE;
+			UE_LOG(LogAudio, Warning, TEXT("Failed to read wave data."));
+		}
+		if (*WaveInfo.pFormatTag == FWaveModInfo::WAVE_INFO_FORMAT_PCM)
+		{
+			if (bEnableUEWavComp)
+			{
+				int16* samples = (int16*)WaveInfo.SampleDataStart;
+				int64 num_samples = WaveInfo.GetNumSamples();
+				int64 num_channels = *WaveInfo.pChannels;
+				TArray<int16> scratch_buffer;
+				scratch_buffer.AddUninitialized(num_samples);
+				uewav_encode16(samples, scratch_buffer.GetData(), num_samples, num_channels);
+				*WaveInfo.pFormatTag = FWaveModInfo::WAVE_INFO_FORMAT_OODLE_WAVE;
+			}
 		}
 	}
+
 	RawData.UpdatePayload(Buffer, Owner);
 }
 
