Made DeadBlending store pose and bone data sparsely according to which bones are present in the CompactPose.

Previously we allocated buffers as large as the full number of bones in the skeleton asset and kept a mask saying which bones were valid or invalid. In this version we allocate buffers just for the bones which are valid and store an array of indices indexing into these buffers for each bone in the skeleton asset. The actual logic of the node should be identical - this change just reduces the memory usage.

#localization none
#tests PIE

diff --git a/Engine/Source/Runtime/Engine/Classes/Animation/AnimNode_DeadBlending.h b/Engine/Source/Runtime/Engine/Classes/Animation/AnimNode_DeadBlending.h
--- a/Engine/Source/Runtime/Engine/Classes/Animation/AnimNode_DeadBlending.h
+++ b/Engine/Source/Runtime/Engine/Classes/Animation/AnimNode_DeadBlending.h
@@ -188,10 +188,10 @@ private:
 	 * @param SrcPoseCurr	The pose recorded as output of the inertializer on the previous frame.
 	 */
 	void InitFrom(
-		const FCompactPose& InPose,
-		const FBlendedCurve& InCurves,
-		const FInertializationPose& SrcPosePrev,
-		const FInertializationPose& SrcPoseCurr);
+		const FCompactPose& InPose, 
+		const FBlendedCurve& InCurves, 
+		const FInertializationSparsePose& SrcPosePrev,
+		const FInertializationSparsePose& SrcPoseCurr);
 
 	/**
 	 * Computes the extrapolated pose and blends it with the input pose.
@@ -213,7 +213,8 @@ private:
 	TArray<FCompactPoseBoneIndex, TInlineAllocator<8>> BoneFilter;
 
 	// Snapshots of the actor pose generated as output.
-	TArray<FInertializationPose, TInlineAllocator<2>> PoseSnapshots;
+	FInertializationSparsePose PrevPoseSnapshot;
+	FInertializationSparsePose CurrPoseSnapshot;
 
 	// Pending inertialization requests.
 	UPROPERTY(Transient)
@@ -226,7 +227,8 @@ private:
 
 	// Recorded pose state at point of transition.
 
-	TBitArray<TInlineAllocator<16>> BoneValid;
+	TArray<uint16> BoneIndices;
+	
 	TArray<FVector> BoneTranslations;
 	TArray<FQuat> BoneRotations;
 	TArray<FQuat4f> BoneRotationDirections;
diff --git a/Engine/Source/Runtime/Engine/Classes/Animation/AnimNode_Inertialization.h b/Engine/Source/Runtime/Engine/Classes/Animation/AnimNode_Inertialization.h
--- a/Engine/Source/Runtime/Engine/Classes/Animation/AnimNode_Inertialization.h
+++ b/Engine/Source/Runtime/Engine/Classes/Animation/AnimNode_Inertialization.h
@@ -234,6 +234,39 @@ struct TUseBitwiseSwap<FInertializationPose>
 	enum { Value = false };
 };
 
+// Structure for storing a pose snapshot sparsely (i.e. when we may not have the full set of transform for every bone)
+struct FInertializationSparsePose
+{
+	friend struct FAnimNode_Inertialization;
+	friend struct FAnimNode_DeadBlending;
+
+private:
+
+	FTransform ComponentTransform;
+	
+	// For each SkeletonPoseBoneIndex this array stores the index into the BoneTranslations, BoneRotations, and 
+	// BoneScales arrays which contains that bone's data. Or INDEX_NONE if this bone's data is not in the snapshot.
+	TArray<uint16> BoneIndices;
+	
+	// Bone translation Data
+	TArray<FVector> BoneTranslations;
+	
+	// Bone Rotation Data
+	TArray<FQuat> BoneRotations;
+	
+	// Bone Scale Data
+	TArray<FVector> BoneScales;
+
+    // Curve Data
+	FInertializationCurve Curves;
+
+	FName AttachParentName;
+	float DeltaTime;
+
+	void InitFrom(const FCompactPose& Pose, const FBlendedCurve& InCurves, const FTransform& InComponentTransform, const FName& InAttachParentName, float InDeltaTime);
+	bool IsEmpty() const;
+	void Empty();
+};
 
 USTRUCT()
 struct FInertializationBoneDiff
diff --git a/Engine/Source/Runtime/Engine/Private/Animation/AnimNode_DeadBlending.cpp b/Engine/Source/Runtime/Engine/Private/Animation/AnimNode_DeadBlending.cpp
--- a/Engine/Source/Runtime/Engine/Private/Animation/AnimNode_DeadBlending.cpp
+++ b/Engine/Source/Runtime/Engine/Private/Animation/AnimNode_DeadBlending.cpp
@@ -66,8 +66,6 @@ namespace UE::Anim {
 
 namespace UE::Anim::DeadBlending::Private
 {
-	static constexpr int32 MaxPoseSnapShotNum = 2;
-
 	static constexpr float Ln2 = 0.69314718056f;
 
 	static int32 GetNumSkeletonBones(const FBoneContainer& BoneContainer)
@@ -262,7 +260,8 @@ void FAnimNode_DeadBlending::Deactivate()
 
 	if (!bPreallocateMemory)
 	{
-		BoneValid.Empty();
+		BoneIndices.Empty();
+
 		BoneTranslations.Empty();
 		BoneRotations.Empty();
 		BoneRotationDirections.Empty();
@@ -280,11 +279,7 @@ void FAnimNode_DeadBlending::Deactivate()
 	}
 }
 
-void FAnimNode_DeadBlending::InitFrom(
-	const FCompactPose& InPose,
-	const FBlendedCurve& InCurves,
-	const FInertializationPose& SrcPosePrev,
-	const FInertializationPose& SrcPoseCurr)
+void FAnimNode_DeadBlending::InitFrom(const FCompactPose& InPose, const FBlendedCurve& InCurves, const FInertializationSparsePose& SrcPosePrev, const FInertializationSparsePose& SrcPoseCurr)
 {
 	TRACE_CPUPROFILER_EVENT_SCOPE(FAnimNode_DeadBlending::InitFrom);
 
@@ -292,96 +287,126 @@ void FAnimNode_DeadBlending::InitFrom(
 
 	const int32 NumSkeletonBones = UE::Anim::DeadBlending::Private::GetNumSkeletonBones(BoneContainer);
 
-	BoneValid.Init(false, NumSkeletonBones);
-	BoneTranslations.Init(FVector::ZeroVector, NumSkeletonBones);
-	BoneRotations.Init(FQuat::Identity, NumSkeletonBones);
-	BoneRotationDirections.Init(FQuat4f::Identity, NumSkeletonBones);
-	BoneScales.Init(FVector::OneVector, NumSkeletonBones);
+	// Compute the Inertialization Bone Indices which we will use to index into BoneTranslations, BoneRotations, etc
 
-	BoneTranslationVelocities.Init(FVector3f::ZeroVector, NumSkeletonBones);
-	BoneRotationVelocities.Init(FVector3f::ZeroVector, NumSkeletonBones);
-	BoneScaleVelocities.Init(FVector3f::ZeroVector, NumSkeletonBones);
+	BoneIndices.Init(INDEX_NONE, NumSkeletonBones);
 
-	BoneTranslationDecayHalfLives.Init(ExtrapolationHalfLifeMin * FVector3f::OneVector, NumSkeletonBones);
-	BoneRotationDecayHalfLives.Init(ExtrapolationHalfLifeMin * FVector3f::OneVector, NumSkeletonBones);
-	BoneScaleDecayHalfLives.Init(ExtrapolationHalfLifeMin * FVector3f::OneVector, NumSkeletonBones);
-
-	// Record bone state
+	int32 NumInertializationBones = 0;
 
 	for (FCompactPoseBoneIndex BoneIndex : InPose.ForEachBoneIndex())
 	{
 		const int32 SkeletonPoseBoneIndex = BoneContainer.GetSkeletonIndex(BoneIndex);
 
 		if (SkeletonPoseBoneIndex == INDEX_NONE ||
-			SrcPosePrev.BoneStates[SkeletonPoseBoneIndex] != EInertializationBoneState::Valid ||
-			SrcPoseCurr.BoneStates[SkeletonPoseBoneIndex] != EInertializationBoneState::Valid)
+			SrcPoseCurr.BoneIndices[SkeletonPoseBoneIndex] == INDEX_NONE ||
+			SrcPosePrev.BoneIndices[SkeletonPoseBoneIndex] == INDEX_NONE)
 		{
 			continue;
 		}
 
-		// Mark bone as valid
+		BoneIndices[SkeletonPoseBoneIndex] = NumInertializationBones;
+		NumInertializationBones++;
+	}
+
+	// Allocate Inertialization Bones
+
+	BoneTranslations.Init(FVector::ZeroVector, NumInertializationBones);
+	BoneRotations.Init(FQuat::Identity, NumInertializationBones);
+	BoneRotationDirections.Init(FQuat4f::Identity, NumInertializationBones);
+	BoneScales.Init(FVector::OneVector, NumInertializationBones);
+
+	BoneTranslationVelocities.Init(FVector3f::ZeroVector, NumInertializationBones);
+	BoneRotationVelocities.Init(FVector3f::ZeroVector, NumInertializationBones);
+	BoneScaleVelocities.Init(FVector3f::ZeroVector, NumInertializationBones);
 
-		BoneValid[SkeletonPoseBoneIndex] = true;
+	BoneTranslationDecayHalfLives.Init(ExtrapolationHalfLifeMin * FVector3f::OneVector, NumInertializationBones);
+	BoneRotationDecayHalfLives.Init(ExtrapolationHalfLifeMin * FVector3f::OneVector, NumInertializationBones);
+	BoneScaleDecayHalfLives.Init(ExtrapolationHalfLifeMin * FVector3f::OneVector, NumInertializationBones);
+
+	for (FCompactPoseBoneIndex BoneIndex : InPose.ForEachBoneIndex())
+	{
+		const int32 SkeletonPoseBoneIndex = BoneContainer.GetSkeletonIndex(BoneIndex);
+
+		if (SkeletonPoseBoneIndex == INDEX_NONE || 
+			SrcPoseCurr.BoneIndices[SkeletonPoseBoneIndex] == INDEX_NONE ||
+			SrcPosePrev.BoneIndices[SkeletonPoseBoneIndex] == INDEX_NONE)
+		{
+			continue;
+		}
+
+		// Get Bone Indices for Inertialization Bone, Prev and Curr Pose Bones
+
+		const int32 InertializationBoneIndex = BoneIndices[SkeletonPoseBoneIndex];
+		const int32 CurrPoseBoneIndex = SrcPoseCurr.BoneIndices[SkeletonPoseBoneIndex];
+		const int32 PrevPoseBoneIndex = SrcPosePrev.BoneIndices[SkeletonPoseBoneIndex];
+
+		check(InertializationBoneIndex != INDEX_NONE);
+		check(CurrPoseBoneIndex != INDEX_NONE);
+		check(PrevPoseBoneIndex != INDEX_NONE);
 
 		// Get Source Animation Transform
 
-		const FTransform SrcTransformCurr = SrcPoseCurr.BoneTransforms[SkeletonPoseBoneIndex];
+		const FVector SrcTranslationCurr = SrcPoseCurr.BoneTranslations[CurrPoseBoneIndex];
+		const FQuat SrcRotationCurr = SrcPoseCurr.BoneRotations[CurrPoseBoneIndex];
+		const FVector SrcScaleCurr = SrcPoseCurr.BoneScales[CurrPoseBoneIndex];
 
-		BoneTranslations[SkeletonPoseBoneIndex] = SrcTransformCurr.GetTranslation();
-		BoneRotations[SkeletonPoseBoneIndex] = SrcTransformCurr.GetRotation();
-		BoneScales[SkeletonPoseBoneIndex] = SrcTransformCurr.GetScale3D();
+		BoneTranslations[InertializationBoneIndex] = SrcTranslationCurr;
+		BoneRotations[InertializationBoneIndex] = SrcRotationCurr;
+		BoneScales[InertializationBoneIndex] = SrcScaleCurr;
 
 		if (SrcPoseCurr.DeltaTime > UE_SMALL_NUMBER)
 		{
 			// Get Source Animation Velocity
 
-			const FTransform SrcTransformPrev = SrcPosePrev.BoneTransforms[SkeletonPoseBoneIndex];
+			const FVector SrcTranslationPrev = SrcPosePrev.BoneTranslations[PrevPoseBoneIndex];
+			const FQuat SrcRotationPrev = SrcPosePrev.BoneRotations[PrevPoseBoneIndex];
+			const FVector SrcScalePrev = SrcPosePrev.BoneScales[PrevPoseBoneIndex];
 
-			const FVector TranslationDiff = SrcTransformCurr.GetTranslation() - SrcTransformPrev.GetTranslation();
+			const FVector TranslationDiff = SrcTranslationCurr - SrcTranslationPrev;
 
-			FQuat RotationDiff = SrcTransformCurr.GetRotation() * SrcTransformPrev.GetRotation().Inverse();
+			FQuat RotationDiff = SrcRotationCurr * SrcRotationPrev.Inverse();
 			RotationDiff.EnforceShortestArcWith(FQuat::Identity);
 
-			const FVector ScaleDiff = UE::Anim::DeadBlending::Private::VectorDivMax(SrcTransformCurr.GetScale3D(), SrcTransformPrev.GetScale3D());
+			const FVector ScaleDiff = UE::Anim::DeadBlending::Private::VectorDivMax(SrcScaleCurr, SrcScalePrev);
 
-			BoneTranslationVelocities[SkeletonPoseBoneIndex] = (FVector3f)(TranslationDiff / SrcPoseCurr.DeltaTime);
-			BoneRotationVelocities[SkeletonPoseBoneIndex] = (FVector3f)(RotationDiff.ToRotationVector() / SrcPoseCurr.DeltaTime);
-			BoneScaleVelocities[SkeletonPoseBoneIndex] = (FVector3f)(UE::Anim::DeadBlending::Private::VectorLogSafe(ScaleDiff) / SrcPoseCurr.DeltaTime);
+			BoneTranslationVelocities[InertializationBoneIndex] = (FVector3f)(TranslationDiff / SrcPoseCurr.DeltaTime);
+			BoneRotationVelocities[InertializationBoneIndex] = (FVector3f)(RotationDiff.ToRotationVector() / SrcPoseCurr.DeltaTime);
+			BoneScaleVelocities[InertializationBoneIndex] = (FVector3f)(UE::Anim::DeadBlending::Private::VectorLogSafe(ScaleDiff) / SrcPoseCurr.DeltaTime);
 
 			// Clamp Maximum Velocity
 
-			BoneTranslationVelocities[SkeletonPoseBoneIndex] = BoneTranslationVelocities[SkeletonPoseBoneIndex].GetClampedToMaxSize(MaximumTranslationVelocity);
-			BoneRotationVelocities[SkeletonPoseBoneIndex] = BoneRotationVelocities[SkeletonPoseBoneIndex].GetClampedToMaxSize(FMath::DegreesToRadians(MaximumRotationVelocity));
-			BoneScaleVelocities[SkeletonPoseBoneIndex] = BoneScaleVelocities[SkeletonPoseBoneIndex].GetClampedToMaxSize(MaximumScaleVelocity);
+			BoneTranslationVelocities[InertializationBoneIndex] = BoneTranslationVelocities[InertializationBoneIndex].GetClampedToMaxSize(MaximumTranslationVelocity);
+			BoneRotationVelocities[InertializationBoneIndex] = BoneRotationVelocities[InertializationBoneIndex].GetClampedToMaxSize(FMath::DegreesToRadians(MaximumRotationVelocity));
+			BoneScaleVelocities[InertializationBoneIndex] = BoneScaleVelocities[InertializationBoneIndex].GetClampedToMaxSize(MaximumScaleVelocity);
 
 			// Compute Decay HalfLives
 
 			const FTransform DstTransform = InPose[BoneIndex];
 
-			const FVector TranslationSrcDstDiff = DstTransform.GetTranslation() - SrcTransformCurr.GetTranslation();
+			const FVector TranslationSrcDstDiff = DstTransform.GetTranslation() - SrcTranslationCurr;
 
-			FQuat RotationSrcDstDiff = DstTransform.GetRotation() * SrcTransformCurr.GetRotation().Inverse();
+			FQuat RotationSrcDstDiff = DstTransform.GetRotation() * SrcRotationCurr.Inverse();
 			RotationSrcDstDiff.EnforceShortestArcWith(FQuat::Identity);
 
-			const FVector ScaleSrcDstDiff = UE::Anim::DeadBlending::Private::VectorDivMax(DstTransform.GetScale3D(), SrcTransformCurr.GetScale3D());
+			const FVector ScaleSrcDstDiff = UE::Anim::DeadBlending::Private::VectorDivMax(DstTransform.GetScale3D(), SrcScaleCurr);
 
-			BoneTranslationDecayHalfLives[SkeletonPoseBoneIndex] = UE::Anim::DeadBlending::Private::ComputeDecayHalfLifeFromDiffAndVelocity(
+			BoneTranslationDecayHalfLives[InertializationBoneIndex] = UE::Anim::DeadBlending::Private::ComputeDecayHalfLifeFromDiffAndVelocity(
 				TranslationSrcDstDiff,
-				BoneTranslationVelocities[SkeletonPoseBoneIndex],
+				BoneTranslationVelocities[InertializationBoneIndex],
 				ExtrapolationHalfLife,
 				ExtrapolationHalfLifeMin,
 				ExtrapolationHalfLifeMax);
 
-			BoneRotationDecayHalfLives[SkeletonPoseBoneIndex] = UE::Anim::DeadBlending::Private::ComputeDecayHalfLifeFromDiffAndVelocity(
+			BoneRotationDecayHalfLives[InertializationBoneIndex] = UE::Anim::DeadBlending::Private::ComputeDecayHalfLifeFromDiffAndVelocity(
 				RotationSrcDstDiff.ToRotationVector(),
-				BoneRotationVelocities[SkeletonPoseBoneIndex],
+				BoneRotationVelocities[InertializationBoneIndex],
 				ExtrapolationHalfLife,
 				ExtrapolationHalfLifeMin,
 				ExtrapolationHalfLifeMax);
 
-			BoneScaleDecayHalfLives[SkeletonPoseBoneIndex] = UE::Anim::DeadBlending::Private::ComputeDecayHalfLifeFromDiffAndVelocity(
+			BoneScaleDecayHalfLives[InertializationBoneIndex] = UE::Anim::DeadBlending::Private::ComputeDecayHalfLifeFromDiffAndVelocity(
 				ScaleSrcDstDiff,
-				BoneScaleVelocities[SkeletonPoseBoneIndex],
+				BoneScaleVelocities[InertializationBoneIndex],
 				ExtrapolationHalfLife,
 				ExtrapolationHalfLifeMin,
 				ExtrapolationHalfLifeMax);
@@ -449,30 +474,33 @@ void FAnimNode_DeadBlending::ApplyTo(FCompactPose& InOutPose, FBlendedCurve& InO
 	{
 		const int32 SkeletonPoseBoneIndex = BoneContainer.GetSkeletonIndex(BoneIndex);
 
-		if (SkeletonPoseBoneIndex == INDEX_NONE || !BoneValid[SkeletonPoseBoneIndex] || BoneFilter.Contains(BoneIndex))
+		if (SkeletonPoseBoneIndex == INDEX_NONE || BoneIndices[SkeletonPoseBoneIndex] == INDEX_NONE || BoneFilter.Contains(BoneIndex))
 		{
 			continue;
 		}
 
+		const int32 InertializationBoneIndex = BoneIndices[SkeletonPoseBoneIndex];
+		check(InertializationBoneIndex != INDEX_NONE);
+
 		// Compute Extrapolated Bone State
 
 		const FVector ExtrapolatedTranslation = UE::Anim::DeadBlending::Private::ExtrapolateTranslation(
-			BoneTranslations[SkeletonPoseBoneIndex],
-			BoneTranslationVelocities[SkeletonPoseBoneIndex],
+			BoneTranslations[InertializationBoneIndex],
+			BoneTranslationVelocities[InertializationBoneIndex],
 			InertializationTime,
-			BoneTranslationDecayHalfLives[SkeletonPoseBoneIndex]);
+			BoneTranslationDecayHalfLives[InertializationBoneIndex]);
 
 		const FQuat ExtrapolatedRotation = UE::Anim::DeadBlending::Private::ExtrapolateRotation(
-			BoneRotations[SkeletonPoseBoneIndex],
-			BoneRotationVelocities[SkeletonPoseBoneIndex],
+			BoneRotations[InertializationBoneIndex],
+			BoneRotationVelocities[InertializationBoneIndex],
 			InertializationTime,
-			BoneRotationDecayHalfLives[SkeletonPoseBoneIndex]);
+			BoneRotationDecayHalfLives[InertializationBoneIndex]);
 
 		const FVector ExtrapolatedScale = UE::Anim::DeadBlending::Private::ExtrapolateScale(
-			BoneScales[SkeletonPoseBoneIndex],
-			BoneScaleVelocities[SkeletonPoseBoneIndex],
+			BoneScales[InertializationBoneIndex],
+			BoneScaleVelocities[InertializationBoneIndex],
 			InertializationTime,
-			BoneScaleDecayHalfLives[SkeletonPoseBoneIndex]);
+			BoneScaleDecayHalfLives[InertializationBoneIndex]);
 
 #if WITH_EDITORONLY_DATA
 		if (bShowExtrapolations)
@@ -490,10 +518,10 @@ void FAnimNode_DeadBlending::ApplyTo(FCompactPose& InOutPose, FBlendedCurve& InO
 		// side of this rotation.
 
 		FQuat RotationDiff = ExtrapolatedRotation * InOutPose[BoneIndex].GetRotation().Inverse();
-		RotationDiff.EnforceShortestArcWith((FQuat)BoneRotationDirections[SkeletonPoseBoneIndex]);
+		RotationDiff.EnforceShortestArcWith((FQuat)BoneRotationDirections[InertializationBoneIndex]);
 
 		// Update BoneRotationDirections to match our current path
-		BoneRotationDirections[SkeletonPoseBoneIndex] = (FQuat4f)RotationDiff;
+		BoneRotationDirections[InertializationBoneIndex] = (FQuat4f)RotationDiff;
 
 		// Compute Blend Alpha
 
@@ -593,13 +621,15 @@ void FAnimNode_DeadBlending::Initialize_AnyThread(const FAnimationInitializeCont
 
 	BoneFilter.Init(FCompactPoseBoneIndex(INDEX_NONE), FilteredBones.Num());
 
-	PoseSnapshots.Empty(UE::Anim::DeadBlending::Private::MaxPoseSnapShotNum);
+	PrevPoseSnapshot.Empty();
+	CurrPoseSnapshot.Empty();
 
 	RequestQueue.Reserve(8);
 
 	const int32 NumSkeletonBones = bPreallocateMemory ? Context.AnimInstanceProxy->GetSkeleton()->GetReferenceSkeleton().GetNum() : 0;
 
-	BoneValid.Empty(NumSkeletonBones);
+	BoneIndices.Empty(NumSkeletonBones);
+	
 	BoneTranslations.Empty(NumSkeletonBones);
 	BoneRotations.Empty(NumSkeletonBones);
 	BoneRotationDirections.Empty(NumSkeletonBones);
@@ -722,7 +752,8 @@ void FAnimNode_DeadBlending::Evaluate_AnyThread(FPoseContext& Output)
 		Deactivate();
 
 		// Clear the pose history
-		PoseSnapshots.Reset();
+		PrevPoseSnapshot.Empty();
+		CurrPoseSnapshot.Empty();
 
 		// Reset the cached time accumulator
 		DeltaTime = 0.0f;
@@ -740,14 +771,20 @@ void FAnimNode_DeadBlending::Evaluate_AnyThread(FPoseContext& Output)
 
 	const float TeleportDistanceThreshold = Output.AnimInstanceProxy->GetSkelMeshComponent()->GetTeleportDistanceThreshold();
 
-	if (PoseSnapshots.Num() > 0 && TeleportDistanceThreshold > 0.0f)
+	if (!CurrPoseSnapshot.IsEmpty() && TeleportDistanceThreshold > 0.0f)
 	{
 		const FVector RootWorldSpaceLocation = ComponentTransform.TransformPosition(Output.Pose[FCompactPoseBoneIndex(0)].GetTranslation());
-		const FVector PrevRootWorldSpaceLocation = PoseSnapshots.Last().ComponentTransform.TransformPosition(PoseSnapshots.Last().BoneTransforms[0].GetTranslation());
+		
+		const uint16 RootBoneIndex = CurrPoseSnapshot.BoneIndices[0];
 
-		if (FVector::DistSquared(RootWorldSpaceLocation, PrevRootWorldSpaceLocation) > FMath::Square(TeleportDistanceThreshold))
+		if (RootBoneIndex != INDEX_NONE)
 		{
-			bTeleported = true;
+			const FVector PrevRootWorldSpaceLocation = CurrPoseSnapshot.ComponentTransform.TransformPosition(CurrPoseSnapshot.BoneTranslations[RootBoneIndex]);
+
+			if (FVector::DistSquared(RootWorldSpaceLocation, PrevRootWorldSpaceLocation) > FMath::Square(TeleportDistanceThreshold))
+			{
+				bTeleported = true;
+			}
 		}
 	}
 
@@ -762,7 +799,7 @@ void FAnimNode_DeadBlending::Evaluate_AnyThread(FPoseContext& Output)
 	// which case there shouldn't be any discontinuity to remove, so no inertialization needs to be done, and we can 
 	// discard any requests.
 
-	if (PoseSnapshots.IsEmpty())
+	if (CurrPoseSnapshot.IsEmpty() && PrevPoseSnapshot.IsEmpty())
 	{
 		RequestQueue.Reset();
 	}
@@ -837,25 +874,25 @@ void FAnimNode_DeadBlending::Evaluate_AnyThread(FPoseContext& Output)
 
 		// Initialize the recorded pose state at the point of transition
 
-		if (PoseSnapshots.Num() > 1)
+		if (!PrevPoseSnapshot.IsEmpty() && !CurrPoseSnapshot.IsEmpty())
 		{
 			// We have two previous poses and so can initialize as normal.
 
 			InitFrom(
 				Output.Pose,
 				Output.Curve,
-				PoseSnapshots[PoseSnapshots.Num() - 2],
-				PoseSnapshots[PoseSnapshots.Num() - 1]);
+				PrevPoseSnapshot,
+				CurrPoseSnapshot);
 		}
-		else if (PoseSnapshots.Num() > 0)
+		else if (!CurrPoseSnapshot.IsEmpty())
 		{
 			// We only have a single previous pose. Repeat this pose assuming zero velocity.
 
 			InitFrom(
 				Output.Pose,
 				Output.Curve,
-				PoseSnapshots.Last(),
-				PoseSnapshots.Last());
+				CurrPoseSnapshot,
+				CurrPoseSnapshot);
 		}
 		else
 		{
@@ -901,22 +938,18 @@ void FAnimNode_DeadBlending::Evaluate_AnyThread(FPoseContext& Output)
 
 	// Record Pose Snapshot
 
-	if (PoseSnapshots.Num() < UE::Anim::DeadBlending::Private::MaxPoseSnapShotNum)
+	if (CurrPoseSnapshot.IsEmpty())
 	{
-		// Add the pose to the end of the buffer
-		PoseSnapshots.AddDefaulted_GetRef().InitFrom(Output.Pose, Output.Curve, ComponentTransform, AttachParentName, DeltaTime);
+		// Initialize the current pose
+		CurrPoseSnapshot.InitFrom(Output.Pose, Output.Curve, ComponentTransform, AttachParentName, DeltaTime);
 	}
 	else
 	{
-		// Bubble the old poses forward in the buffer (using swaps to avoid allocations and copies)
-		for (int32 SnapshotIndex = 0; SnapshotIndex < UE::Anim::DeadBlending::Private::MaxPoseSnapShotNum - 1; ++SnapshotIndex)
-		{
-			Swap(PoseSnapshots[SnapshotIndex], PoseSnapshots[SnapshotIndex + 1]);
-		}
+		// Directly swap the memory of the current pose with the prev pose snapshot (to avoid allocations and copies)
+		Swap(PrevPoseSnapshot, CurrPoseSnapshot);
 
-		// Overwrite the (now irrelevant) pose in the last slot with the new post snapshot
-		// (thereby avoiding the reallocation costs we would have incurred had we simply added a new pose at the end)
-		PoseSnapshots.Last().InitFrom(Output.Pose, Output.Curve, ComponentTransform, AttachParentName, DeltaTime);
+		// Initialize the (now irrelevant) current pose
+		CurrPoseSnapshot.InitFrom(Output.Pose, Output.Curve, ComponentTransform, AttachParentName, DeltaTime);
 	}
 
 	// Reset Delta Time
diff --git a/Engine/Source/Runtime/Engine/Private/Animation/AnimNode_Inertialization.cpp b/Engine/Source/Runtime/Engine/Private/Animation/AnimNode_Inertialization.cpp
--- a/Engine/Source/Runtime/Engine/Private/Animation/AnimNode_Inertialization.cpp
+++ b/Engine/Source/Runtime/Engine/Private/Animation/AnimNode_Inertialization.cpp
@@ -621,6 +621,88 @@ void FInertializationPose::InitFrom(const FCompactPose& Pose, const FBlendedCurv
 	DeltaTime = InDeltaTime;
 }
 
+void FInertializationSparsePose::InitFrom(
+	const FCompactPose& Pose,
+	const FBlendedCurve& InCurves,
+	const FTransform& InComponentTransform,
+	const FName& InAttachParentName,
+	float InDeltaTime)
+{
+	const FBoneContainer& BoneContainer = Pose.GetBoneContainer();
+
+	const int32 NumSkeletonBones = UE::Anim::GetNumSkeletonBones(BoneContainer);
+
+	// Allocate Bone Index Array
+
+	BoneIndices.Init(INDEX_NONE, NumSkeletonBones);
+
+	int32 NumInertializationBones = 0;
+
+	for (const FCompactPoseBoneIndex BoneIndex : Pose.ForEachBoneIndex())
+	{
+		const int32 SkeletonPoseBoneIndex = BoneContainer.GetSkeletonIndex(BoneIndex);
+
+		if (SkeletonPoseBoneIndex == INDEX_NONE)
+		{
+			continue;
+		}
+
+		// For each valid bone in the Compact Pose we write into BoneIndices the InertializationBoneIndex -
+		// i.e. the index into BoneTranslations, BoneRotations, and BoneScales we are going to use to store
+		// the transform data
+
+		BoneIndices[SkeletonPoseBoneIndex] = NumInertializationBones;
+		NumInertializationBones++;
+	}
+
+	// Initialize the BoneTranslations, BoneRotations, and BoneScales arrays
+
+	BoneTranslations.Init(FVector::ZeroVector, NumInertializationBones);
+	BoneRotations.Init(FQuat::Identity, NumInertializationBones);
+	BoneScales.Init(FVector::OneVector, NumInertializationBones);
+
+	for (const FCompactPoseBoneIndex BoneIndex : Pose.ForEachBoneIndex())
+	{
+		const int32 SkeletonPoseBoneIndex = BoneContainer.GetSkeletonIndex(BoneIndex);
+
+		if (SkeletonPoseBoneIndex == INDEX_NONE)
+		{
+			continue;
+		}
+
+		// Get the InertializationBoneIndex and write the transform data
+
+		const uint16 InertializationBoneIndex = BoneIndices[SkeletonPoseBoneIndex];
+		check(InertializationBoneIndex != INDEX_NONE);
+
+		const FTransform BoneTransform = Pose[BoneIndex];
+		BoneTranslations[InertializationBoneIndex] = BoneTransform.GetTranslation();
+		BoneRotations[InertializationBoneIndex] = BoneTransform.GetRotation();
+		BoneScales[InertializationBoneIndex] = BoneTransform.GetScale3D();
+	}
+
+	// Init the rest of the snapshot data
+
+	Curves.InitFrom(InCurves);
+	ComponentTransform = InComponentTransform;
+	AttachParentName = InAttachParentName;
+	DeltaTime = InDeltaTime;
+}
+
+bool FInertializationSparsePose::IsEmpty() const
+{
+	return BoneIndices.IsEmpty();
+}
+
+void FInertializationSparsePose::Empty()
+{
+	BoneIndices.Empty();
+	BoneTranslations.Empty();
+	BoneRotations.Empty();
+	BoneScales.Empty();
+	Curves.BlendedCurve.Empty();
+}
+
 
 // Initialize the pose difference from the current pose and the two previous snapshots
 //
