Improving packing on SkeletalMeshComponent which should get a 16 byte gain. This will cause SkeletalMeshComponentBudgeted to drop to the 4080 bin for all platforms other than Switch, getting 288 bytes back for each.

diff --git a/Engine/Plugins/Runtime/AnimationBudgetAllocator/Source/AnimationBudgetAllocator/Private/SkeletalMeshComponentBudgeted.cpp b/Engine/Plugins/Runtime/AnimationBudgetAllocator/Source/AnimationBudgetAllocator/Private/SkeletalMeshComponentBudgeted.cpp
--- a/Engine/Plugins/Runtime/AnimationBudgetAllocator/Source/AnimationBudgetAllocator/Private/SkeletalMeshComponentBudgeted.cpp
+++ b/Engine/Plugins/Runtime/AnimationBudgetAllocator/Source/AnimationBudgetAllocator/Private/SkeletalMeshComponentBudgeted.cpp
@@ -12,6 +12,10 @@
 
 CSV_DECLARE_CATEGORY_EXTERN(AnimationBudget);
 
+#if UE_BUILD_SHIPPING
+static_assert(sizeof(USkeletalMeshComponentBudgeted) <= 4080, "We expect SkeletalMeshComponentBudgeted to be 4080 or less, otherwise we waste memory on 288bytes of padding on MallocBinned2. Please improve padding on this class hierarchy to pay for your added variable.");
+#endif
+
 FOnCalculateSignificance USkeletalMeshComponentBudgeted::OnCalculateSignificanceDelegate;
 
 USkeletalMeshComponentBudgeted::USkeletalMeshComponentBudgeted(const FObjectInitializer& ObjectInitializer)
diff --git a/Engine/Source/Runtime/Engine/Classes/Components/SkeletalMeshComponent.h b/Engine/Source/Runtime/Engine/Classes/Components/SkeletalMeshComponent.h
--- a/Engine/Source/Runtime/Engine/Classes/Components/SkeletalMeshComponent.h
+++ b/Engine/Source/Runtime/Engine/Classes/Components/SkeletalMeshComponent.h
@@ -809,6 +809,10 @@ private:
 
 	uint8 bPostEvaluatingAnimation:1;
 
+public:
+	/** Physics-engine representation of aggregate which contains a physics asset instance with more than numbers of bodies. */
+	FPhysicsAggregateHandle Aggregate;
+
 public:
 
 	/** Cache AnimCurveUidVersion from Skeleton and this will be used to identify if it needs to be updated */
@@ -1381,9 +1385,6 @@ public:
 	/** Array of FConstraintInstance structs, storing per-instance state about each constraint. */
 	TArray<struct FConstraintInstance*> Constraints;
 
-	/** Physics-engine representation of aggregate which contains a physics asset instance with more than numbers of bodies. */
-	FPhysicsAggregateHandle Aggregate;
-
 	FSkeletalMeshComponentClothTickFunction ClothTickFunction;
 
 	/**
@@ -1443,13 +1444,6 @@ private:
 	ENGINE_API void ComputeTeleportRotationThresholdInRadians();
 	ENGINE_API void ComputeTeleportDistanceThresholdInRadians();
 
-	// Can't rely on time value, because those may be affected by dilation and whether or not
-	// the game is paused.
-	// Also can't just rely on a flag as other components (like CharacterMovementComponent) may tick
-	// the pose and we can't guarantee tick order.
-	UPROPERTY(Transient)
-	uint32 LastPoseTickFrame;
-
 public:
 
 	/** Checked whether we have already ticked the pose this frame */
@@ -2650,6 +2644,15 @@ private:
 	 */
 	ENGINE_API bool AreRequiredCurvesUpToDate() const;
 
+
+private:
+	// Can't rely on time value, because those may be affected by dilation and whether or not
+	// the game is paused.
+	// Also can't just rely on a flag as other components (like CharacterMovementComponent) may tick
+	// the pose and we can't guarantee tick order.
+	UPROPERTY(Transient)
+	uint32 LastPoseTickFrame; 
+
 public:
 	ENGINE_API void ConditionallyDispatchQueuedAnimEvents();
 
