Prevent creation of RHI classes on the server in material params and misc places.

- Also added FApp::CanEverRenderOrProduceRenderData() for more clarity around the conditions when a certain path can be entered.

diff --git a/Engine/Source/Runtime/Core/Public/Misc/App.h b/Engine/Source/Runtime/Core/Public/Misc/App.h
--- a/Engine/Source/Runtime/Core/Public/Misc/App.h
+++ b/Engine/Source/Runtime/Core/Public/Misc/App.h
@@ -404,6 +404,20 @@ public:
 #endif // UE_SERVER
 	}
 
+	/**
+	 * Checks whether this application can render anything or produce a derived data needed for rednering.
+	 * Certain application types never render, but produce DDC used during rendering and as such need to step into some rendering paths.
+	 *
+	 * A meaningful distinction from FApp::CanEverRender() is that commandlets like cooker will have FApp::CanEverRender() == false, but FApp::CanEverRenderOrProduceRenderData() == true.
+	 * As such, this function can be used to guard paths that e.g. load assets' render data.
+	 *
+	 * @return true if the application can render, false otherwise.
+	 */
+	INLINE_CANEVERRENDER static bool CanEverRenderOrProduceRenderData()
+	{
+		return !FPlatformProperties::RequiresCookedData() || FApp::CanEverRender();
+	}
+
 	/**
 	 * Checks whether this application can render audio.
 	 * Certain application types produce sound, while for others this can be controlled via the -nosound cmdline.
diff --git a/Engine/Source/Runtime/Engine/Private/Materials/ParameterCollection.cpp b/Engine/Source/Runtime/Engine/Private/Materials/ParameterCollection.cpp
--- a/Engine/Source/Runtime/Engine/Private/Materials/ParameterCollection.cpp
+++ b/Engine/Source/Runtime/Engine/Private/Materials/ParameterCollection.cpp
@@ -32,7 +32,7 @@ void UMaterialParameterCollection::PostInitProperties()
 {
 	Super::PostInitProperties();
 
-	if (!HasAnyFlags(RF_ClassDefaultObject))
+	if (LIKELY(!HasAnyFlags(RF_ClassDefaultObject) && FApp::CanEverRender()))
 	{
 		DefaultResource = new FMaterialParameterCollectionInstanceResource();
 	}
@@ -535,6 +535,11 @@ const FCollectionVectorParameter* UMaterialParameterCollection::GetVectorParamet
 
 void UMaterialParameterCollection::CreateBufferStruct()
 {	
+	if (UNLIKELY(!FApp::CanEverRenderOrProduceRenderData()))
+	{
+		return;
+	}
+
 	TArray<FShaderParametersMetadata::FMember> Members;
 	uint32 NextMemberOffset = 0;
 
@@ -591,6 +596,11 @@ void UMaterialParameterCollection::GetDefaultParameterData(TArray<FVector4f>& Pa
 
 void UMaterialParameterCollection::UpdateDefaultResource(bool bRecreateUniformBuffer)
 {
+	if (UNLIKELY(!FApp::CanEverRender()))
+	{
+		return;
+	}
+
 	// Propagate the new values to the rendering thread
 	TArray<FVector4f> ParameterData;
 	GetDefaultParameterData(ParameterData);
@@ -819,6 +829,11 @@ void UMaterialParameterCollectionInstance::FinishDestroy()
 
 void FMaterialParameterCollectionInstanceResource::GameThread_UpdateContents(const FGuid& InGuid, const TArray<FVector4f>& Data, const FName& InOwnerName, bool bRecreateUniformBuffer)
 {
+	if (UNLIKELY(!FApp::CanEverRender()))
+	{
+		return;
+	}
+
 	FMaterialParameterCollectionInstanceResource* Resource = this;
 	ENQUEUE_RENDER_COMMAND(UpdateCollectionCommand)(
 		[InGuid, Data, InOwnerName, Resource, bRecreateUniformBuffer](FRHICommandListImmediate& RHICmdList)
diff --git a/Engine/Source/Runtime/Engine/Private/RendererSettings.cpp b/Engine/Source/Runtime/Engine/Private/RendererSettings.cpp
--- a/Engine/Source/Runtime/Engine/Private/RendererSettings.cpp
+++ b/Engine/Source/Runtime/Engine/Private/RendererSettings.cpp
@@ -551,6 +551,11 @@ void URendererSettings::UpdateWorkingColorSpaceAndChromaticities()
 {
 	using namespace UE::Color;
 
+	if (!FApp::CanEverRenderOrProduceRenderData())
+	{
+		return;
+	}
+
 	switch (WorkingColorSpaceChoice)
 	{
 	case EWorkingColorSpace::sRGB:
diff --git a/Engine/Source/Runtime/Engine/Private/SceneManagement.cpp b/Engine/Source/Runtime/Engine/Private/SceneManagement.cpp
--- a/Engine/Source/Runtime/Engine/Private/SceneManagement.cpp
+++ b/Engine/Source/Runtime/Engine/Private/SceneManagement.cpp
@@ -1263,7 +1263,7 @@ FSharedSamplerState* Clamp_WorldGroupSettings = NULL;
 
 void InitializeSharedSamplerStates()
 {
-	if (!Wrap_WorldGroupSettings)
+	if (!Wrap_WorldGroupSettings && FApp::CanEverRender())
 	{
 		Wrap_WorldGroupSettings = new FSharedSamplerState(true);
 		Clamp_WorldGroupSettings = new FSharedSamplerState(false);
