Fix for crash in AnimNode_DeadBlend
#rnx

diff --git a/Engine/Source/Runtime/Engine/Private/Animation/AnimNode_DeadBlending.cpp b/Engine/Source/Runtime/Engine/Private/Animation/AnimNode_DeadBlending.cpp
--- a/Engine/Source/Runtime/Engine/Private/Animation/AnimNode_DeadBlending.cpp
+++ b/Engine/Source/Runtime/Engine/Private/Animation/AnimNode_DeadBlending.cpp
@@ -328,8 +328,8 @@ void FAnimNode_DeadBlending::InitFrom(const FCompactPose& InPose, const FBlended
 		const int32 SkeletonPoseBoneIndex = BoneContainer.GetSkeletonIndex(BoneIndex);
 
 		if (SkeletonPoseBoneIndex == INDEX_NONE || 
-			SrcPoseCurr.BoneIndices[SkeletonPoseBoneIndex] == INDEX_NONE ||
-			SrcPosePrev.BoneIndices[SkeletonPoseBoneIndex] == INDEX_NONE)
+			SrcPoseCurr.BoneIndices[SkeletonPoseBoneIndex] == static_cast<uint16>(INDEX_NONE) ||
+			SrcPosePrev.BoneIndices[SkeletonPoseBoneIndex] == static_cast<uint16>(INDEX_NONE))
 		{
 			continue;
 		}
@@ -474,7 +474,7 @@ void FAnimNode_DeadBlending::ApplyTo(FCompactPose& InOutPose, FBlendedCurve& InO
 	{
 		const int32 SkeletonPoseBoneIndex = BoneContainer.GetSkeletonIndex(BoneIndex);
 
-		if (SkeletonPoseBoneIndex == INDEX_NONE || BoneIndices[SkeletonPoseBoneIndex] == INDEX_NONE || BoneFilter.Contains(BoneIndex))
+		if (SkeletonPoseBoneIndex == INDEX_NONE || BoneIndices[SkeletonPoseBoneIndex] == static_cast<uint16>(INDEX_NONE) || BoneFilter.Contains(BoneIndex))
 		{
 			continue;
 		}
@@ -777,7 +777,7 @@ void FAnimNode_DeadBlending::Evaluate_AnyThread(FPoseContext& Output)
 		
 		const uint16 RootBoneIndex = CurrPoseSnapshot.BoneIndices[0];
 
-		if (RootBoneIndex != INDEX_NONE)
+		if (RootBoneIndex != static_cast<uint16>(INDEX_NONE))
 		{
 			const FVector PrevRootWorldSpaceLocation = CurrPoseSnapshot.ComponentTransform.TransformPosition(CurrPoseSnapshot.BoneTranslations[RootBoneIndex]);
 
