Fix non-unity CIS (ParameterCollection.cpp).

diff --git a/Engine/Source/Runtime/Engine/Private/Materials/ParameterCollection.cpp b/Engine/Source/Runtime/Engine/Private/Materials/ParameterCollection.cpp
--- a/Engine/Source/Runtime/Engine/Private/Materials/ParameterCollection.cpp
+++ b/Engine/Source/Runtime/Engine/Private/Materials/ParameterCollection.cpp
@@ -10,6 +10,7 @@
 #include "Materials/Material.h"
 #include "Materials/MaterialParameterCollectionInstance.h"
 #include "Materials/MaterialParameterCollection.h"
+#include "Misc/App.h"
 
 int32 GDeferUpdateRenderStates = 1;
 FAutoConsoleVariableRef CVarDeferUpdateRenderStates(
