cluster unions need to remove from the ActorToComponents map even if the actor is no longer valid because it's pending kill

diff --git a/Engine/Source/Runtime/Engine/Private/PhysicsEngine/ClusterUnionComponent.cpp b/Engine/Source/Runtime/Engine/Private/PhysicsEngine/ClusterUnionComponent.cpp
--- a/Engine/Source/Runtime/Engine/Private/PhysicsEngine/ClusterUnionComponent.cpp
+++ b/Engine/Source/Runtime/Engine/Private/PhysicsEngine/ClusterUnionComponent.cpp
@@ -1199,20 +1199,23 @@ void UClusterUnionComponent::HandleRemovedClusteredComponent(TObjectKey<UPrimiti
 		}
 	}
 
-	if (AActor* Owner = ComponentData.Owner.Get())
+	// Need to get even if pending kill just in case we received a removal because the actor got destroyed!
+	TObjectKey<AActor> OwnerKey { ComponentData.Owner.Get(true) };
+	if (FClusteredActorData* ActorData = ActorToComponents.Find(OwnerKey))
 	{
-		if (FClusteredActorData* ActorData = ActorToComponents.Find(Owner))
-		{
-			ActorData->Components.Remove(RemovedComponent);
+		ActorData->Components.Remove(RemovedComponent);
 
-			if (ActorData->Components.IsEmpty())
+		if (ActorData->Components.IsEmpty())
+		{
+			if (AActor* Owner = ComponentData.Owner.Get())
 			{
 				if (IsAuthority())
 				{
 					Owner->SetReplicatingMovement(ActorData->bWasReplicatingMovement);
 				}
-				ActorToComponents.Remove(Owner);
 			}
+
+			ActorToComponents.Remove(OwnerKey);
 		}
 	}
 
