ActorDescContainer: Fix Duplicate Actor Guid not being reported (also add comment in WorldPartition::RegisterActorDescContainer to explain the difference with the container check)

#rnx

diff --git a/Engine/Source/Runtime/Engine/Private/WorldPartition/ActorDescContainer.cpp b/Engine/Source/Runtime/Engine/Private/WorldPartition/ActorDescContainer.cpp
--- a/Engine/Source/Runtime/Engine/Private/WorldPartition/ActorDescContainer.cpp
+++ b/Engine/Source/Runtime/Engine/Private/WorldPartition/ActorDescContainer.cpp
@@ -126,15 +126,16 @@ void UActorDescContainer::Initialize(const FInitializeParams& InitParams)
 					ValidActorDescs.Remove(ActorDesc->GetGuid());
 				}
 			}
-			else if (FWorldPartitionActorDesc* ExistingDescGuid = FActorDescList::GetActorDesc(ActorDesc->GetGuid()))
+			else if (TUniquePtr<FWorldPartitionActorDesc>* ExistingActorDescPtr = ValidActorDescs.Find(ActorDesc->GetGuid()))
 			{
-				check(ExistingDescGuid->GetGuid() == ActorDesc->GetGuid());
+				const FWorldPartitionActorDesc* ExistingActorDesc = ExistingActorDescPtr->Get();
+				check(ExistingActorDesc->GetGuid() == ActorDesc->GetGuid());
 				UE_LOG(LogWorldPartition, Warning, TEXT("Duplicate actor descriptor guid `%s`: Actor: '%s' from package '%s' -> Existing actor '%s' from package '%s'"), 
 					*ActorDesc->GetGuid().ToString(), 
 					*ActorDesc->GetActorName().ToString(), 
 					*ActorDesc->GetActorPackage().ToString(),
-					*ExistingDescGuid->GetActorName().ToString(),
-					*ExistingDescGuid->GetActorPackage().ToString());
+					*ExistingActorDesc->GetActorName().ToString(),
+					*ExistingActorDesc->GetActorPackage().ToString());
 				InvalidActors.Emplace(Asset);
 			}
 			else
diff --git a/Engine/Source/Runtime/Engine/Private/WorldPartition/WorldPartition.cpp b/Engine/Source/Runtime/Engine/Private/WorldPartition/WorldPartition.cpp
--- a/Engine/Source/Runtime/Engine/Private/WorldPartition/WorldPartition.cpp
+++ b/Engine/Source/Runtime/Engine/Private/WorldPartition/WorldPartition.cpp
@@ -2076,7 +2076,9 @@ UActorDescContainer* UWorldPartition::RegisterActorDescContainer(const FContaine
 				WorldDataLayerActorsDesc = FoundWorldDataLayerActorsDesc;
 			}
 
-			// Filter actors with duplicated GUID in WorldPartition
+			// Filter actors with duplicated GUID in WorldPartition (across containers):
+			// difference with the duplicate check in UActorDescContainer is that WorldPartition is a collection of containers so same Guid could exist across those containers
+			// which wouldn't be validated by the container itself.
 			if (GetActorDesc(ActorDesc->GetGuid()))
 			{
 				UE_LOG(LogWorldPartition, Warning, TEXT("Found existing actor descriptor guid `%s`: Actor: '%s' from package '%s'"),
