World Partition HashSet: fix missing spatial index when streaming sub partitioned worlds.

diff --git a/Engine/Source/Runtime/Engine/Private/WorldPartition/RuntimeHashSet/WorldPartitionRuntimeHashSetStreamingGeneration.cpp b/Engine/Source/Runtime/Engine/Private/WorldPartition/RuntimeHashSet/WorldPartitionRuntimeHashSetStreamingGeneration.cpp
--- a/Engine/Source/Runtime/Engine/Private/WorldPartition/RuntimeHashSet/WorldPartitionRuntimeHashSetStreamingGeneration.cpp
+++ b/Engine/Source/Runtime/Engine/Private/WorldPartition/RuntimeHashSet/WorldPartitionRuntimeHashSetStreamingGeneration.cpp
@@ -241,6 +241,7 @@ bool UWorldPartitionRuntimeHashSet::GenerateStreaming(UWorldPartitionStreamingPo
 	check(RuntimeStreamingData.IsEmpty());
 	for (auto& [Partition, StreamingData] : RuntimePartitionsStreamingData)
 	{
+		StreamingData.CreatePartitionsSpatialIndex();
 		RuntimeStreamingData.Emplace(MoveTemp(StreamingData));
 	}
 
