Remove some transient client data from USceneComponent on the server.

- Also rearranged the data on the client to bring some 1MB improvements on one of the platforms.

diff --git a/Engine/Source/Runtime/Engine/Classes/Components/SceneComponent.h b/Engine/Source/Runtime/Engine/Classes/Components/SceneComponent.h
--- a/Engine/Source/Runtime/Engine/Classes/Components/SceneComponent.h
+++ b/Engine/Source/Runtime/Engine/Classes/Components/SceneComponent.h
@@ -112,17 +112,21 @@ private:
 	UPROPERTY(ReplicatedUsing = OnRep_AttachSocketName)
 	FName AttachSocketName;
 
+#if !UE_SERVER
+	/** Transient data used by clients when the attach parent changes during replication. */
 	FName NetOldAttachSocketName;
+#endif
 
 	/** List of child SceneComponents that are attached to us. */
 	UPROPERTY(ReplicatedUsing = OnRep_AttachChildren, Transient)
 	TArray<TObjectPtr<USceneComponent>> AttachChildren;
 
+#if !UE_SERVER
 	/** Set of attached SceneComponents that were attached by the client so we can fix up AttachChildren when it is replicated to us. */
-	UPROPERTY(Transient)
 	TArray<TObjectPtr<USceneComponent>> ClientAttachedChildren;
 
 	USceneComponent* NetOldAttachParent;
+#endif
 
 public:
 	/** Current bounds of the component */
@@ -933,9 +937,11 @@ public:
 	ENGINE_API virtual void PostNetReceive() override;
 	ENGINE_API virtual void PostRepNotifies() override;
 	ENGINE_API virtual void Serialize(FArchive& Ar) override;
-#if WITH_EDITORONLY_DATA
+#if (WITH_EDITORONLY_DATA || !UE_SERVER)
 	static ENGINE_API void AddReferencedObjects(UObject* InThis, FReferenceCollector& Collector);
+#endif
 
+#if WITH_EDITORONLY_DATA
 	ENGINE_API virtual void PostLoad() override;
 #endif
 
diff --git a/Engine/Source/Runtime/Engine/Private/Components/SceneComponent.cpp b/Engine/Source/Runtime/Engine/Private/Components/SceneComponent.cpp
--- a/Engine/Source/Runtime/Engine/Private/Components/SceneComponent.cpp
+++ b/Engine/Source/Runtime/Engine/Private/Components/SceneComponent.cpp
@@ -72,10 +72,11 @@ USceneComponent::USceneComponent(const FObjectInitializer& ObjectInitializer /*=
 	SetShouldBeAttached(AttachParent != nullptr);
 }
 
-#if WITH_EDITORONLY_DATA
+#if (WITH_EDITORONLY_DATA || !UE_SERVER)
 void USceneComponent::AddReferencedObjects(UObject* InThis, FReferenceCollector& Collector)
 {
 	USceneComponent* This = CastChecked<USceneComponent>(InThis);
+#if WITH_EDITORONLY_DATA
 	Collector.AddReferencedObject(This->SpriteComponent);
 
 	if (GComponentsWithLegacyLightmaps.GetAnnotationMap().Num() > 0)
@@ -87,10 +88,20 @@ void USceneComponent::AddReferencedObjects(UObject* InThis, FReferenceCollector&
 			LegacyMeshData.Data[EntryIndex].Value->AddReferencedObjects(Collector);
 		}
 	}
+#endif // WITH_EDITORONLY_DATA
+
+#if !UE_SERVER
+	for (TObjectPtr<USceneComponent> AttachedChild : This->ClientAttachedChildren)
+	{
+		Collector.AddReferencedObject(AttachedChild);
+	}
+#endif // !UE_SERVER
 
 	Super::AddReferencedObjects(InThis, Collector);
 }
+#endif // (WITH_EDITORONLY_DATA || !UE_SERVER)
 
+#if WITH_EDITORONLY_DATA
 void USceneComponent::PostLoad()
 {
 	Super::PostLoad();
@@ -2227,7 +2238,10 @@ bool USceneComponent::AttachToComponent(USceneComponent* Parent, const FAttachme
 
 		if (Parent->IsNetSimulating() && !IsNetSimulating())
 		{
+			checkf(UE_SERVER == 0, TEXT("Usage of ClientAttachedChildren detected on a server"));
+#if !UE_SERVER
 			Parent->ClientAttachedChildren.AddUnique(this);
+#endif
 		}
 
 		// Now apply attachment rules
@@ -2393,7 +2407,9 @@ void USceneComponent::DetachFromComponent(const FDetachmentTransformRules& Detac
 		PrimaryComponentTick.RemovePrerequisite(GetAttachParent(), GetAttachParent()->PrimaryComponentTick); // no longer required to tick after the attachment
 
 		GetAttachParent()->AttachChildren.Remove(this);
+#if !UE_SERVER
 		GetAttachParent()->ClientAttachedChildren.Remove(this);
+#endif
 		GetAttachParent()->OnChildDetached(this);
 		GetAttachParent()->ModifiedAttachChildren();
 
@@ -3444,6 +3460,7 @@ void USceneComponent::OnRep_AttachChildren()
 		}
 	}
 
+#if !UE_SERVER
 	if (ClientAttachedChildren.Num())
 	{
 		for (USceneComponent* AttachChild : AttachChildren)
@@ -3465,6 +3482,9 @@ void USceneComponent::OnRep_AttachChildren()
 			}
 		}
 	}
+#else
+	checkNoEntry();
+#endif
 
 	// It's possible AttachChildren are spawned before the AttachParent. This results in the AttachParent never being set.
 	for (USceneComponent* ChildComponent : AttachChildren)
@@ -3493,8 +3513,12 @@ void USceneComponent::PreNetReceive()
 
 	bNetUpdateTransform = false;
 	bNetUpdateAttachment = false;
+#if !UE_SERVER
 	NetOldAttachSocketName = GetAttachSocketName();
 	NetOldAttachParent = GetAttachParent();
+#else
+	checkNoEntry();
+#endif
 }
 
 void USceneComponent::PostNetReceive()
@@ -3515,6 +3539,7 @@ void USceneComponent::PostNetReceive()
 	
 void USceneComponent::PostRepNotifies()
 {
+#if !UE_SERVER
 	if (bNetUpdateAttachment)
 	{
 		Exchange(NetOldAttachParent, AttachParent);
@@ -3564,6 +3589,9 @@ void USceneComponent::PostRepNotifies()
 		UpdateComponentToWorld();
 		bNetUpdateTransform = false;
 	}
+#else
+	checkNoEntry();
+#endif
 }
 
 void USceneComponent::Serialize(FArchive& Ar)
