Fixed missing debug lines when parallel GDME is enabled.

diff --git a/Engine/Source/Runtime/Engine/Private/BatchedElements.cpp b/Engine/Source/Runtime/Engine/Private/BatchedElements.cpp
--- a/Engine/Source/Runtime/Engine/Private/BatchedElements.cpp
+++ b/Engine/Source/Runtime/Engine/Private/BatchedElements.cpp
@@ -1465,6 +1465,7 @@ bool FBatchedElements::Draw(FRHICommandList& RHICmdList, const FMeshPassProcesso
 
 void FBatchedElements::Clear()
 {
+	LineVertices.Empty();
 	Points.Empty();
 	WireTris.Empty();
 	WireTriVerts.Empty();
@@ -1476,6 +1477,7 @@ void FBatchedElements::Clear()
 
 void FBatchedElements::AddAllocationInfo(FAllocationInfo& AllocationInfo) const
 {
+	AllocationInfo.NumLineVertices += LineVertices.Num();
 	AllocationInfo.NumPoints += Points.Num();
 	AllocationInfo.NumWireTris += WireTris.Num();
 	AllocationInfo.NumWireTriVerts += WireTriVerts.Num();
@@ -1487,6 +1489,7 @@ void FBatchedElements::AddAllocationInfo(FAllocationInfo& AllocationInfo) const
 
 void FBatchedElements::Reserve(const FAllocationInfo& AllocationInfo)
 {
+	LineVertices.Reserve(AllocationInfo.NumLineVertices);
 	Points.Reserve(AllocationInfo.NumPoints);
 	WireTris.Reserve(AllocationInfo.NumWireTris);
 	WireTriVerts.Reserve(AllocationInfo.NumWireTriVerts);
@@ -1498,6 +1501,7 @@ void FBatchedElements::Reserve(const FAllocationInfo& AllocationInfo)
 
 void FBatchedElements::Append(FBatchedElements& Other)
 {
+	LineVertices.Append(Other.LineVertices);
 	Points.Append(Other.Points);
 	WireTris.Append(Other.WireTris);
 	WireTriVerts.Append(Other.WireTriVerts);
diff --git a/Engine/Source/Runtime/Engine/Public/BatchedElements.h b/Engine/Source/Runtime/Engine/Public/BatchedElements.h
--- a/Engine/Source/Runtime/Engine/Public/BatchedElements.h
+++ b/Engine/Source/Runtime/Engine/Public/BatchedElements.h
@@ -246,6 +246,7 @@ public:
 		FAllocationInfo() = default;
 
 	private:
+		int32 NumLineVertices = 0;
 		int32 NumPoints = 0;
 		int32 NumWireTris = 0;
 		int32 NumWireTriVerts = 0;
