Align cloth buffers to support SIMD loads.

diff --git a/Engine/Source/Runtime/ClothingSystemRuntimeInterface/Public/ClothingSystemRuntimeTypes.h b/Engine/Source/Runtime/ClothingSystemRuntimeInterface/Public/ClothingSystemRuntimeTypes.h
--- a/Engine/Source/Runtime/ClothingSystemRuntimeInterface/Public/ClothingSystemRuntimeTypes.h
+++ b/Engine/Source/Runtime/ClothingSystemRuntimeInterface/Public/ClothingSystemRuntimeTypes.h
@@ -19,11 +19,11 @@ struct FClothSimulData
 		LODIndex = INDEX_NONE;
 	}
 
-	// Positions of the simulation mesh particles
-	TArray<FVector3f> Positions;
+	// Positions of the simulation mesh particles (aligned for SIMD loads)
+	TArray<FVector3f, TAlignedHeapAllocator<16>> Positions;
 
-	// Normals at the simulation mesh particles
-	TArray<FVector3f> Normals;
+	// Normals at the simulation mesh particles (aligned for SIMD loads)
+	TArray<FVector3f, TAlignedHeapAllocator<16>> Normals;
 
 	// Transform applied per position/normal element when loaded
 	FTransform Transform;
diff --git a/Engine/Source/Runtime/Engine/Private/GPUSkinVertexFactory.cpp b/Engine/Source/Runtime/Engine/Private/GPUSkinVertexFactory.cpp
--- a/Engine/Source/Runtime/Engine/Private/GPUSkinVertexFactory.cpp
+++ b/Engine/Source/Runtime/Engine/Private/GPUSkinVertexFactory.cpp
@@ -992,8 +992,8 @@ IMPLEMENT_TYPE_LAYOUT(TGPUSkinAPEXClothVertexFactoryShaderParameters);
 	TGPUSkinAPEXClothVertexFactory::ClothShaderType
 -----------------------------------------------------------------------------*/
 
-void FGPUBaseSkinAPEXClothVertexFactory::ClothShaderType::UpdateClothSimulData(FRHICommandList& RHICmdList, const TArray<FVector3f>& InSimulPositions,
-	const TArray<FVector3f>& InSimulNormals, uint32 RevisionNumber, ERHIFeatureLevel::Type FeatureLevel, bool bForceUpdateImmediately, const FName& AssetPathName)
+void FGPUBaseSkinAPEXClothVertexFactory::ClothShaderType::UpdateClothSimulData(FRHICommandList& RHICmdList, TConstArrayView<FVector3f> InSimulPositions,
+	TConstArrayView<FVector3f> InSimulNormals, uint32 RevisionNumber, ERHIFeatureLevel::Type FeatureLevel, bool bForceUpdateImmediately, const FName& AssetPathName)
 {
 	QUICK_SCOPE_CYCLE_COUNTER(STAT_FGPUBaseSkinAPEXClothVertexFactory_UpdateClothSimulData);
 
diff --git a/Engine/Source/Runtime/Engine/Public/GPUSkinVertexFactory.h b/Engine/Source/Runtime/Engine/Public/GPUSkinVertexFactory.h
--- a/Engine/Source/Runtime/Engine/Public/GPUSkinVertexFactory.h
+++ b/Engine/Source/Runtime/Engine/Public/GPUSkinVertexFactory.h
@@ -492,7 +492,7 @@ public:
 			Reset();
 		}
 
-		void UpdateClothSimulData(FRHICommandList& RHICmdList, const TArray<FVector3f>& InSimulPositions, const TArray<FVector3f>& InSimulNormals, uint32 RevisionNumber, 
+		void UpdateClothSimulData(FRHICommandList& RHICmdList, TConstArrayView<FVector3f> InSimulPositions, TConstArrayView<FVector3f> InSimulNormals, uint32 RevisionNumber, 
 									ERHIFeatureLevel::Type FeatureLevel, bool bForceUpdateImmediately, const FName& AssetPathName);
 
 		void ReleaseClothSimulData()
