Clear uniform buffer async task reference after waiting on the render thread. This avoids the atomic load with each call to InvalidateUniformExpressionCache.

diff --git a/Engine/Source/Runtime/Engine/Private/Materials/MaterialRenderProxy.cpp b/Engine/Source/Runtime/Engine/Private/Materials/MaterialRenderProxy.cpp
--- a/Engine/Source/Runtime/Engine/Private/Materials/MaterialRenderProxy.cpp
+++ b/Engine/Source/Runtime/Engine/Private/Materials/MaterialRenderProxy.cpp
@@ -52,6 +52,7 @@ public:
 		if (--ReferenceCount == 0)
 		{
 			Wait();
+			Task = {};
 		}
 	}
 
