Added StripAnimDataOnDedicatedServer property to AnimSequence to have the option to override the project settings value per animation.

diff --git a/Engine/Source/Runtime/Engine/Classes/Animation/AnimSequence.h b/Engine/Source/Runtime/Engine/Classes/Animation/AnimSequence.h
--- a/Engine/Source/Runtime/Engine/Classes/Animation/AnimSequence.h
+++ b/Engine/Source/Runtime/Engine/Classes/Animation/AnimSequence.h
@@ -210,6 +210,18 @@ struct FRequestAnimCompressionParams
 	const ITargetPlatform* TargetPlatform;
 };
 
+/** Enum used to decide whether we should strip animation data on dedicated server */
+UENUM()
+enum class EStripAnimDataOnDedicatedServerSettings : uint8
+{
+	/** Strip track data on dedicated server if 'Strip Animation Data on Dedicated Server' option in Project Settings is true and EnableRootMotion is false */
+	UseProjectSetting,
+	/** Strip track data on dedicated server regardless of the value of 'Strip Animation Data on Dedicated Server' option in Project Settings as long as EnableRootMotion is false  */
+	StripAnimDataOnDedicatedServer,
+	/** Do not strip track data on dedicated server regardless of the value of 'Strip Animation Data on Dedicated Server' option in Project Settings  */
+	DoNotStripAnimDataOnDedicatedServer
+};
+
 UCLASS(config=Engine, hidecategories=(UObject, Length), BlueprintType, MinimalAPI)
 class UAnimSequence : public UAnimSequenceBase
 {
@@ -373,7 +385,9 @@ public:
 	int32 MarkerDataUpdateCounter;
 #endif // WITH_EDITORONLY_DATA
 
-
+	/** Enum used to decide whether we should strip animation data on dedicated server */
+	UPROPERTY(EditAnywhere, Category = Compression)
+	EStripAnimDataOnDedicatedServerSettings StripAnimDataOnDedicatedServer = EStripAnimDataOnDedicatedServerSettings::UseProjectSetting;
 
 public:
 	//~ Begin UObject Interface
diff --git a/Engine/Source/Runtime/Engine/Private/Animation/AnimSequence.cpp b/Engine/Source/Runtime/Engine/Private/Animation/AnimSequence.cpp
--- a/Engine/Source/Runtime/Engine/Private/Animation/AnimSequence.cpp
+++ b/Engine/Source/Runtime/Engine/Private/Animation/AnimSequence.cpp
@@ -681,7 +681,11 @@ void UAnimSequence::Serialize(FArchive& Ar)
 		const bool bIsTransacting = Ar.IsTransacting();
 		const bool bIsCookingWithoutAVData = bIsCooking && !Ar.CookingTarget()->AllowAudioVisualData();
 		const bool bIsCountingMemory = Ar.IsCountingMemory();
-		const bool bCookingTargetNeedsCompressedData = bIsCooking && (!UAnimationSettings::Get()->bStripAnimationDataOnDedicatedServer || !bIsCookingWithoutAVData || bEnableRootMotion);
+
+		const bool bStripAnimDataOnDedicatedServer = StripAnimDataOnDedicatedServer == EStripAnimDataOnDedicatedServerSettings::UseProjectSetting ? UAnimationSettings::Get()->bStripAnimationDataOnDedicatedServer :
+			StripAnimDataOnDedicatedServer == EStripAnimDataOnDedicatedServerSettings::StripAnimDataOnDedicatedServer;
+
+		const bool bCookingTargetNeedsCompressedData = bIsCooking && (!bStripAnimDataOnDedicatedServer || !bIsCookingWithoutAVData || bEnableRootMotion);
 
 		bool bSerializeCompressedData = bCookingTargetNeedsCompressedData || bIsDuplicating || bIsTransacting || bIsCountingMemory;
 		Ar << bSerializeCompressedData;
