Fixed issue with DeadBlending node creating NaNs when encountering negative scales and also made it so that the bLinearlyInterpolateScales flag causes scalar velocities to be calculated linearly. This allows for full support of negative scales when this option is enabled.
#localization none
#tests PIE

diff --git a/Engine/Source/Runtime/Engine/Classes/Animation/AnimNode_DeadBlending.h b/Engine/Source/Runtime/Engine/Classes/Animation/AnimNode_DeadBlending.h
--- a/Engine/Source/Runtime/Engine/Classes/Animation/AnimNode_DeadBlending.h
+++ b/Engine/Source/Runtime/Engine/Classes/Animation/AnimNode_DeadBlending.h
@@ -64,9 +64,9 @@ private:
 	float BlendTimeMultiplier = 1.0f;
 
 	/**
-	 * When enabled, bone scales will be linearly interpolated. This is slightly more performant and consistent with the
-	 * rest of Unreal but visually gives the appearance of the rate of change of scale being affected by the overall
-	 * size of the bone.
+	 * When enabled, bone scales will be linearly interpolated and extrapolated. This is slightly more performant and
+	 * consistent with the rest of Unreal but visually gives the appearance of the rate of change of scale being affected
+	 * by the overall size of the bone. Note: this option must be enabled if you want this node to support negative scales.
 	 */
 	UPROPERTY(EditAnywhere, Category = Blending)
 	bool bLinearlyInterpolateScales = false;
diff --git a/Engine/Source/Runtime/Engine/Private/Animation/AnimNode_DeadBlending.cpp b/Engine/Source/Runtime/Engine/Private/Animation/AnimNode_DeadBlending.cpp
--- a/Engine/Source/Runtime/Engine/Private/Animation/AnimNode_DeadBlending.cpp
+++ b/Engine/Source/Runtime/Engine/Private/Animation/AnimNode_DeadBlending.cpp
@@ -113,14 +113,17 @@ namespace UE::Anim::DeadBlending::Private
 	{
 		if (FVector::DistSquared(V, W) < Epsilon)
 		{
-			return FMath::Lerp(V, W, Alpha);
+			return FVector(
+				FMath::Lerp(FMath::Max(V.X, Epsilon), FMath::Max(W.X, Epsilon), Alpha),
+				FMath::Lerp(FMath::Max(V.Y, Epsilon), FMath::Max(W.Y, Epsilon), Alpha),
+				FMath::Lerp(FMath::Max(V.Z, Epsilon), FMath::Max(W.Z, Epsilon), Alpha));
 		}
 		else
 		{
 			return FVector(
-				FMath::Pow(V.X, (1.0f - Alpha)) * FMath::Pow(W.X, Alpha),
-				FMath::Pow(V.Y, (1.0f - Alpha)) * FMath::Pow(W.Y, Alpha),
-				FMath::Pow(V.Z, (1.0f - Alpha)) * FMath::Pow(W.Z, Alpha));
+				FMath::Pow(FMath::Max(V.X, Epsilon), (1.0f - Alpha)) * FMath::Pow(FMath::Max(W.X, Epsilon), Alpha),
+				FMath::Pow(FMath::Max(V.Y, Epsilon), (1.0f - Alpha)) * FMath::Pow(FMath::Max(W.Y, Epsilon), Alpha),
+				FMath::Pow(FMath::Max(V.Z, Epsilon), (1.0f - Alpha)) * FMath::Pow(FMath::Max(W.Z, Epsilon), Alpha));
 		}
 	}
 
@@ -367,11 +370,14 @@ void FAnimNode_DeadBlending::InitFrom(const FCompactPose& InPose, const FBlended
 			FQuat RotationDiff = SrcRotationCurr * SrcRotationPrev.Inverse();
 			RotationDiff.EnforceShortestArcWith(FQuat::Identity);
 
-			const FVector ScaleDiff = UE::Anim::DeadBlending::Private::VectorDivMax(SrcScaleCurr, SrcScalePrev);
+			const FVector ScaleDiffLinear = SrcScaleCurr - SrcScalePrev;
+			const FVector ScaleDiffExponential = UE::Anim::DeadBlending::Private::VectorDivMax(SrcScaleCurr, SrcScalePrev);
 
 			BoneTranslationVelocities[InertializationBoneIndex] = (FVector3f)(TranslationDiff / SrcPoseCurr.DeltaTime);
 			BoneRotationVelocities[InertializationBoneIndex] = (FVector3f)(RotationDiff.ToRotationVector() / SrcPoseCurr.DeltaTime);
-			BoneScaleVelocities[InertializationBoneIndex] = (FVector3f)(UE::Anim::DeadBlending::Private::VectorLogSafe(ScaleDiff) / SrcPoseCurr.DeltaTime);
+			BoneScaleVelocities[InertializationBoneIndex] = bLinearlyInterpolateScales ?
+				(FVector3f)(ScaleDiffLinear / SrcPoseCurr.DeltaTime) :
+				(FVector3f)(UE::Anim::DeadBlending::Private::VectorLogSafe(ScaleDiffExponential) / SrcPoseCurr.DeltaTime);
 
 			// Clamp Maximum Velocity
 
@@ -496,11 +502,26 @@ void FAnimNode_DeadBlending::ApplyTo(FCompactPose& InOutPose, FBlendedCurve& InO
 			InertializationTime,
 			BoneRotationDecayHalfLives[InertializationBoneIndex]);
 
-		const FVector ExtrapolatedScale = UE::Anim::DeadBlending::Private::ExtrapolateScale(
-			BoneScales[InertializationBoneIndex],
-			BoneScaleVelocities[InertializationBoneIndex],
-			InertializationTime,
-			BoneScaleDecayHalfLives[InertializationBoneIndex]);
+		FVector ExtrapolatedScale = FVector::OneVector;
+
+		if (bLinearlyInterpolateScales)
+		{
+			// If we are handling scales linearly then treat them like a normal vector such as a translation.
+			ExtrapolatedScale = UE::Anim::DeadBlending::Private::ExtrapolateTranslation(
+				BoneScales[InertializationBoneIndex],
+				BoneScaleVelocities[InertializationBoneIndex],
+				InertializationTime,
+				BoneScaleDecayHalfLives[InertializationBoneIndex]);
+		}
+		else
+		{
+			// Otherwise extrapolate using the exponential version of scalar velocities.
+			ExtrapolatedScale = UE::Anim::DeadBlending::Private::ExtrapolateScale(
+				BoneScales[InertializationBoneIndex],
+				BoneScaleVelocities[InertializationBoneIndex],
+				InertializationTime,
+				BoneScaleDecayHalfLives[InertializationBoneIndex]);
+		}
 
 #if WITH_EDITORONLY_DATA
 		if (bShowExtrapolations)
@@ -538,9 +559,9 @@ void FAnimNode_DeadBlending::ApplyTo(FCompactPose& InOutPose, FBlendedCurve& InO
 
 			// Here we use `Eerp` rather than `Lerp` to interpolate scales by default (see: https://theorangeduck.com/page/scalar-velocity).
 			// This default is inconsistent with the rest of Unreal which (mostly) uses `Lerp` on scales. The decision 
-			// to use `Eerp` by default here is partially due to the fact we are also dealing properly with scalar 
-			// velocities in this node, and partially to try and not to lock this node into having the same less 
-			// accurate behavior by default. Users still have the option to interpolate scales with `Lerp` if they want.
+			// to use `Eerp` by default here is partially due to the fact we are also providing the option of dealing properly 
+			// with scalar velocities in this node, and partially to try and not to lock this node into having the same less 
+			// accurate behavior by default. Users still have the option to interpolate scales with `Lerp` if they want using bLinearlyInterpolateScales.
 			if (bLinearlyInterpolateScales)
 			{
 				InOutPose[BoneIndex].SetScale3D(FMath::Lerp(InOutPose[BoneIndex].GetScale3D(), ExtrapolatedScale, Alpha));
