Fix/workaround for simulation rewind when using a fixed timestep and small frame times. Adds an option to allow velocities to be rewound along with positions.

#localization none
#tests have run locally in PIE plus using cooked builds

diff --git a/Engine/Plugins/Experimental/PhysicsControl/Source/PhysicsControl/Private/AnimNode_RigidBodyWithControl.cpp b/Engine/Plugins/Experimental/PhysicsControl/Source/PhysicsControl/Private/AnimNode_RigidBodyWithControl.cpp
--- a/Engine/Plugins/Experimental/PhysicsControl/Source/PhysicsControl/Private/AnimNode_RigidBodyWithControl.cpp
+++ b/Engine/Plugins/Experimental/PhysicsControl/Source/PhysicsControl/Private/AnimNode_RigidBodyWithControl.cpp
@@ -1009,6 +1009,8 @@ void FAnimNode_RigidBodyWithControl::InitPhysics(const UAnimInstance* InAnimInst
 	if(bEnabled)
 	{
 		PhysicsSimulation = new ImmediatePhysics::FSimulation();
+		PhysicsSimulation->SetRewindVelocities(true);
+
 		const int32 NumBodies = PhysicsAssetToUse->SkeletalBodySetups.Num();
 		BodyNameToIndexMap.Reset();
 		Bodies.Empty(NumBodies);
diff --git a/Engine/Source/Runtime/Engine/Private/PhysicsEngine/ImmediatePhysics/ImmediatePhysicsChaos/ImmediatePhysicsSimulation_Chaos.cpp b/Engine/Source/Runtime/Engine/Private/PhysicsEngine/ImmediatePhysics/ImmediatePhysicsChaos/ImmediatePhysicsSimulation_Chaos.cpp
--- a/Engine/Source/Runtime/Engine/Private/PhysicsEngine/ImmediatePhysics/ImmediatePhysicsChaos/ImmediatePhysicsSimulation_Chaos.cpp
+++ b/Engine/Source/Runtime/Engine/Private/PhysicsEngine/ImmediatePhysics/ImmediatePhysicsChaos/ImmediatePhysicsSimulation_Chaos.cpp
@@ -804,6 +804,11 @@ namespace ImmediatePhysics_Chaos
 		return Implementation->RollingAverageStepTime;
 	}
 
+	void FSimulation::SetRewindVelocities(bool bRewindVelocities)
+	{
+		Implementation->Evolution.SetRewindVelocities(bRewindVelocities);
+	}
+
 	void FSimulation::Simulate(FReal InDeltaTime, FReal MaxStepTime, int32 MaxSubSteps, const FVector& InGravity)
 	{
 		SCOPE_CYCLE_COUNTER(STAT_ImmediateSimulate_Chaos);
diff --git a/Engine/Source/Runtime/Engine/Public/Physics/ImmediatePhysics/ImmediatePhysicsChaos/ImmediatePhysicsSimulation_Chaos.h b/Engine/Source/Runtime/Engine/Public/Physics/ImmediatePhysics/ImmediatePhysicsChaos/ImmediatePhysicsSimulation_Chaos.h
--- a/Engine/Source/Runtime/Engine/Public/Physics/ImmediatePhysics/ImmediatePhysicsChaos/ImmediatePhysicsSimulation_Chaos.h
+++ b/Engine/Source/Runtime/Engine/Public/Physics/ImmediatePhysics/ImmediatePhysicsChaos/ImmediatePhysicsSimulation_Chaos.h
@@ -68,6 +68,12 @@ namespace ImmediatePhysics_Chaos
 		/** Set up potential collisions between the actor and all other dynamic actors */
 		ENGINE_API void AddToCollidingPairs(FActorHandle* ActorHandle);
 
+		/** 
+		 * Sets whether velocities should be rewound when simulating - this may happen when the requested 
+		 * step size is smaller than the fixed simulation step.
+		 */
+		ENGINE_API void SetRewindVelocities(bool bRewindVelocities);
+
 		/** Advance the simulation by DeltaTime */
 		ENGINE_API void Simulate(FReal DeltaTime, FReal MaxStepTime, int32 MaxSubSteps, const FVector& InGravity);
 		void Simulate_AssumesLocked(FReal DeltaTime, FReal MaxStepTime, int32 MaxSubSteps, const FVector& InGravity) { Simulate(DeltaTime, MaxStepTime, MaxSubSteps, InGravity); }
diff --git a/Engine/Source/Runtime/Experimental/Chaos/Private/Chaos/Evolution/PBDMinEvolution.cpp b/Engine/Source/Runtime/Experimental/Chaos/Private/Chaos/Evolution/PBDMinEvolution.cpp
--- a/Engine/Source/Runtime/Experimental/Chaos/Private/Chaos/Evolution/PBDMinEvolution.cpp
+++ b/Engine/Source/Runtime/Experimental/Chaos/Private/Chaos/Evolution/PBDMinEvolution.cpp
@@ -17,7 +17,7 @@
 #include "ChaosStats.h"
 #include "ChaosVisualDebugger/ChaosVisualDebuggerTrace.h"
 
-//PRAGMA_DISABLE_OPTIMIZATION
+//UE_DISABLE_OPTIMIZATION
 
 namespace Chaos
 {
@@ -62,6 +62,7 @@ namespace Chaos
 		, ConstraintSolver(Private::FIterationSettings(0,0,0))
 		, Gravity(FVec3(0))
 		, SimulationSpaceSettings()
+		, bRewindVelocities(false)
 	{
 	}
 
@@ -181,6 +182,12 @@ namespace Chaos
 			{
 				Particle.X() = FVec3::Lerp(Particle.Handle()->AuxilaryValue(ParticlePrevXs), Particle.X(), T);
 				Particle.R() = FRotation3::Slerp(Particle.Handle()->AuxilaryValue(ParticlePrevRs), Particle.R(), (decltype(FQuat::X))T);	// LWC_TODO: Remove decltype cast once FQuat supports variants
+
+				if (bRewindVelocities)
+				{
+					Particle.V() = FVec3::Lerp(Particle.PreV(), Particle.V(), T);
+					Particle.W() = FVec3::Lerp(Particle.PreW(), Particle.W(), T);
+				}
 			}
 		}
 
diff --git a/Engine/Source/Runtime/Experimental/Chaos/Public/Chaos/Evolution/PBDMinEvolution.h b/Engine/Source/Runtime/Experimental/Chaos/Public/Chaos/Evolution/PBDMinEvolution.h
--- a/Engine/Source/Runtime/Experimental/Chaos/Public/Chaos/Evolution/PBDMinEvolution.h
+++ b/Engine/Source/Runtime/Experimental/Chaos/Public/Chaos/Evolution/PBDMinEvolution.h
@@ -69,6 +69,10 @@ namespace Chaos
 			Gravity = G;
 		}
 
+		void SetRewindVelocities(bool InRewindVelocities)
+		{
+			bRewindVelocities = InRewindVelocities;
+		}
 
 		void SetSimulationSpace(const FSimulationSpace& InSimulationSpace)
 		{
@@ -127,5 +131,6 @@ namespace Chaos
 		FVec3 Gravity;
 		FSimulationSpaceSettings SimulationSpaceSettings;
 		FSimulationSpace SimulationSpace;
+		bool bRewindVelocities;
 	};
 }
