expose particle bone ids and parent cluster union component from the replicated proxy component as exported functions

diff --git a/Engine/Source/Runtime/Engine/Classes/PhysicsEngine/ClusterUnionReplicatedProxyComponent.h b/Engine/Source/Runtime/Engine/Classes/PhysicsEngine/ClusterUnionReplicatedProxyComponent.h
--- a/Engine/Source/Runtime/Engine/Classes/PhysicsEngine/ClusterUnionReplicatedProxyComponent.h
+++ b/Engine/Source/Runtime/Engine/Classes/PhysicsEngine/ClusterUnionReplicatedProxyComponent.h
@@ -29,7 +29,7 @@ public:
 	ENGINE_API UClusterUnionReplicatedProxyComponent(const FObjectInitializer& ObjectInitializer);
 
 	UFUNCTION()
-	UClusterUnionComponent* GetParentClusterUnionComponent() const { return ParentClusterUnion.Get(); }
+	ENGINE_API UClusterUnionComponent* GetParentClusterUnionComponent() const { return ParentClusterUnion.Get(); }
 
 	UFUNCTION()
 	ENGINE_API void SetParentClusterUnion(UClusterUnionComponent* InComponent);
@@ -43,6 +43,9 @@ public:
 	UFUNCTION()
 	ENGINE_API void SetParticleBoneIds(const TArray<int32>& InIds);
 
+	UFUNCTION()
+	ENGINE_API const TArray<int32>& GetParticleBoneIds() const { return ParticleBoneIds; }
+
 	UFUNCTION()
 	ENGINE_API void SetParticleChildToParent(int32 BoneId, const FTransform& ChildToParent);
 
