Fix for end of lifetime tracing on linked anim graphs

diff --git a/Engine/Source/Runtime/Engine/Private/Animation/AnimNode_LinkedAnimGraph.cpp b/Engine/Source/Runtime/Engine/Private/Animation/AnimNode_LinkedAnimGraph.cpp
--- a/Engine/Source/Runtime/Engine/Private/Animation/AnimNode_LinkedAnimGraph.cpp
+++ b/Engine/Source/Runtime/Engine/Private/Animation/AnimNode_LinkedAnimGraph.cpp
@@ -8,6 +8,7 @@
 #include "Animation/BlendProfile.h"
 #include "Components/SkeletalMeshComponent.h"
 #include "Animation/ExposedValueHandler.h"
+#include "ObjectTrace.h"
 
 #include UE_INLINE_GENERATED_CPP_BY_NAME(AnimNode_LinkedAnimGraph)
 
@@ -252,6 +253,8 @@ void FAnimNode_LinkedAnimGraph::TeardownInstance(const UAnimInstance* InOwningAn
 	UAnimInstance* InstanceToRun = GetTargetInstance<UAnimInstance>();
 	if (InstanceToRun)
 	{
+		// trace lifetime end early, because by the time we get UninitializeAnimation below, the Owner has changed, and so the ObjectId has changed.
+		TRACE_OBJECT_LIFETIME_END(InstanceToRun);
 		DynamicUnlink(const_cast<UAnimInstance*>(InOwningAnimInstance));
 		// Never delete the owning animation instance
 		if (InstanceToRun != InOwningAnimInstance)
