Make sure FSubsystemCollectionBase::ActivateExternalSubsystem causes SubsystemArrayMap to be updated so that early calls to GetSubsystem can't populate the map with an empty array that is not replaced when the subsystem is initialized.

diff --git a/Engine/Source/Runtime/Engine/Private/Subsystems/SubsystemCollection.cpp b/Engine/Source/Runtime/Engine/Private/Subsystems/SubsystemCollection.cpp
--- a/Engine/Source/Runtime/Engine/Private/Subsystems/SubsystemCollection.cpp
+++ b/Engine/Source/Runtime/Engine/Private/Subsystems/SubsystemCollection.cpp
@@ -80,7 +80,7 @@ const TArray<USubsystem*>& FSubsystemCollectionBase::GetSubsystemArrayInternal(U
 	{
 		TArray<USubsystem*>& NewList = SubsystemArrayMap.Add(SubsystemClass);
 
-		UpdateSubsystemArrayInternal(SubsystemClass, NewList);
+		PopulateSubsystemArrayInternal(SubsystemClass, NewList);
 
 		return NewList;
 	}
@@ -89,8 +89,9 @@ const TArray<USubsystem*>& FSubsystemCollectionBase::GetSubsystemArrayInternal(U
 	return List;
 }
 
-void FSubsystemCollectionBase::UpdateSubsystemArrayInternal(UClass* SubsystemClass, TArray<USubsystem*>& SubsystemArray) const
+void FSubsystemCollectionBase::PopulateSubsystemArrayInternal(UClass* SubsystemClass, TArray<USubsystem*>& SubsystemArray) const
 {
+	check(SubsystemArray.Num() == 0);
 	for (auto Iter = SubsystemMap.CreateConstIterator(); Iter; ++Iter)
 	{
 		UClass* KeyClass = Iter.Key();
@@ -149,13 +150,6 @@ void FSubsystemCollectionBase::Initialize(UObject* NewOuter)
 			}
 		}
 
-		// Update Internal Arrays without emptying it so that existing refs remain valid
-		for (auto& Pair : SubsystemArrayMap)
-		{
-			Pair.Value.Empty();
-			UpdateSubsystemArrayInternal(Pair.Key, Pair.Value);
-		}
-
 		// Statically track collections
 		GlobalSubsystemCollections.Add(this);
 	}
@@ -257,6 +251,16 @@ USubsystem* FSubsystemCollectionBase::AddAndInitializeSubsystem(UClass* Subsyste
 				SubsystemMap.Add(SubsystemClass,Subsystem);
 				Subsystem->InternalOwningSubsystem = this;
 				Subsystem->Initialize(*this);
+				
+				// Add this new subsystem to any existing maps of base classes to lists of subsystems
+				for (TPair<UClass*, TArray<USubsystem*>>& Pair : SubsystemArrayMap)
+				{
+					if (SubsystemClass->IsChildOf(Pair.Key))
+					{
+						Pair.Value.Add(Subsystem);
+					}
+				}
+
 				return Subsystem;
 			}
 
