Missing header from previous submission

diff --git a/Engine/Source/Runtime/Engine/Public/Subsystems/SubsystemCollection.h b/Engine/Source/Runtime/Engine/Public/Subsystems/SubsystemCollection.h
--- a/Engine/Source/Runtime/Engine/Public/Subsystems/SubsystemCollection.h
+++ b/Engine/Source/Runtime/Engine/Public/Subsystems/SubsystemCollection.h
@@ -75,7 +75,7 @@ private:
 
 	ENGINE_API void RemoveAndDeinitializeSubsystem(USubsystem* Subsystem);
 
-	ENGINE_API void UpdateSubsystemArrayInternal(UClass* SubsystemClass, TArray<USubsystem*>& SubsystemArray) const;
+	ENGINE_API void PopulateSubsystemArrayInternal(UClass* SubsystemClass, TArray<USubsystem*>& SubsystemArray) const;
 
 	TMap<TObjectPtr<UClass>, TObjectPtr<USubsystem>> SubsystemMap;
 
