Iris
- Fixed crash due to getting piggybacked export bunches on iris channel

diff --git a/Engine/Source/Runtime/Engine/Private/Net/Experimental/Iris/DataStreamChannel.cpp b/Engine/Source/Runtime/Engine/Private/Net/Experimental/Iris/DataStreamChannel.cpp
--- a/Engine/Source/Runtime/Engine/Private/Net/Experimental/Iris/DataStreamChannel.cpp
+++ b/Engine/Source/Runtime/Engine/Private/Net/Experimental/Iris/DataStreamChannel.cpp
@@ -390,11 +390,11 @@ void UDataStreamChannel::ReceivedAck(int32 PacketId)
 	}
 
 	const FDataStreamChannelRecord& ChannelRecord = WriteRecords.Peek();
-	ensureMsgf((uint32)PacketId == ChannelRecord.PacketId, TEXT("PacketId %d != ChannelRecord.PacketId %d, WriteRecords.Num %d"), PacketId, ChannelRecord.PacketId, (int32)WriteRecords.Count());
-
-	DataStreamManager->ProcessPacketDeliveryStatus(UE::Net::EPacketDeliveryStatus::Delivered, static_cast<const FDataStreamRecord*>(ChannelRecord.Record));
-
-	WriteRecords.Pop();
+	if (ensureMsgf((uint32)PacketId == ChannelRecord.PacketId, TEXT("PacketId %d != ChannelRecord.PacketId %d, WriteRecords.Num %d"), PacketId, ChannelRecord.PacketId, (int32)WriteRecords.Count()))
+	{
+		DataStreamManager->ProcessPacketDeliveryStatus(UE::Net::EPacketDeliveryStatus::Delivered, static_cast<const FDataStreamRecord*>(ChannelRecord.Record));
+		WriteRecords.Pop();
+	}
 
 #endif // UE_WITH_IRIS
 }
@@ -412,10 +412,11 @@ void UDataStreamChannel::ReceivedNak(int32 PacketId)
 
 	const FDataStreamChannelRecord& ChannelRecord = WriteRecords.Peek();
 
-	check((uint32)PacketId == ChannelRecord.PacketId);
-
-	DataStreamManager->ProcessPacketDeliveryStatus(UE::Net::EPacketDeliveryStatus::Lost, static_cast<const FDataStreamRecord*>(ChannelRecord.Record));
-	WriteRecords.Pop();
+	if (ensureMsgf((uint32)PacketId == ChannelRecord.PacketId, TEXT("PacketId %d != ChannelRecord.PacketId %d, WriteRecords.Num %d"), PacketId, ChannelRecord.PacketId, (int32)WriteRecords.Count()))
+	{
+		DataStreamManager->ProcessPacketDeliveryStatus(UE::Net::EPacketDeliveryStatus::Lost, static_cast<const FDataStreamRecord*>(ChannelRecord.Record));
+		WriteRecords.Pop();
+	}
 
 #endif // UE_WITH_IRIS
 }
@@ -440,3 +441,11 @@ void UDataStreamChannel::AddReferencedObjects(UObject* Object, FReferenceCollect
 
 	Super::AddReferencedObjects(Channel, Collector);
 }
+
+void UDataStreamChannel::AppendExportBunches(TArray<FOutBunch *>& OutExportBunches)
+{
+}
+
+void UDataStreamChannel::AppendMustBeMappedGuids(FOutBunch* Bunch)
+{
+}
diff --git a/Engine/Source/Runtime/Engine/Private/Net/Experimental/Iris/DataStreamChannel.h b/Engine/Source/Runtime/Engine/Private/Net/Experimental/Iris/DataStreamChannel.h
--- a/Engine/Source/Runtime/Engine/Private/Net/Experimental/Iris/DataStreamChannel.h
+++ b/Engine/Source/Runtime/Engine/Private/Net/Experimental/Iris/DataStreamChannel.h
@@ -49,6 +49,11 @@ private:
 	/** Human readable information about the channel */
 	ENGINE_API virtual FString Describe() override;
 
+	/** We do not want to append orphaned exportbunches from other channels */
+	ENGINE_API virtual void AppendExportBunches(TArray<FOutBunch *>& OutExportBunches) override;
+	ENGINE_API virtual void AppendMustBeMappedGuids(FOutBunch* Bunch) override;
+
+
 	/** Packet delivery status handling */
 	ENGINE_API virtual void ReceivedAck(int32 PacketId) override;
 	ENGINE_API virtual void ReceivedNak(int32 PacketId) override;
