[shaders]
- don't bother submitting shader jobs which failed to preprocess; just complete straight away to log the error
- allow jobs with zero preprocess time in the case the preprocessing did not succeed; this can result due to shader pipelines earlying out if preprocessing one stage has an error (subsequent stages are skipped)

diff --git a/Engine/Source/Runtime/Engine/Private/ShaderCompiler/ShaderCompiler.cpp b/Engine/Source/Runtime/Engine/Private/ShaderCompiler/ShaderCompiler.cpp
--- a/Engine/Source/Runtime/Engine/Private/ShaderCompiler/ShaderCompiler.cpp
+++ b/Engine/Source/Runtime/Engine/Private/ShaderCompiler/ShaderCompiler.cpp
@@ -712,7 +712,7 @@ public:
 	void SubmitJobs(const TArray<FShaderCommonCompileJobPtr>& InJobs);
 
 	/** This is an entry point for all jobs that have finished the compilation (whether real or cached). Can be called from multiple threads. Returns mutex stall time. */
-	double ProcessFinishedJob(FShaderCommonCompileJob* FinishedJob, bool bWasCached);
+	double ProcessFinishedJob(FShaderCommonCompileJob* FinishedJob, bool bCompilationSkipped);
 
 	/** Adds the job to cache. */
 	void AddToCacheAndProcessPending(FShaderCommonCompileJob* FinishedJob);
@@ -1293,9 +1293,9 @@ void FShaderCompileJobCollection::SubmitJobs(const TArray<FShaderCommonCompileJo
 {
 	JobsCache->SubmitJobs(InJobs);
 }
-void FShaderCompileJobCollection::ProcessFinishedJob(FShaderCommonCompileJob* FinishedJob, bool bWasCached)
+void FShaderCompileJobCollection::ProcessFinishedJob(FShaderCommonCompileJob* FinishedJob, bool bCompilationSkipped)
 {
-	JobsCache->ProcessFinishedJob(FinishedJob, bWasCached);
+	JobsCache->ProcessFinishedJob(FinishedJob, bCompilationSkipped);
 }
 void FShaderCompileJobCollection::AddToCacheAndProcessPending(FShaderCommonCompileJob* FinishedJob)
 {
@@ -1742,13 +1742,22 @@ void FShaderJobCache::SubmitJobs(const TArray<FShaderCommonCompileJobPtr>& InJob
 						FPlatformProcess::Sleep(GShaderCompilerDebugStallSubmitJob * 0.001f);
 					}
 
+					bool bSubmitJob = true;
 					if (ShaderCompiler::IsJobCacheEnabled())
 					{
-						ConditionalPreprocessShader(Job);
+						bSubmitJob = ConditionalPreprocessShader(Job);
 						Job->GetInputHash();
 					}
-					SubmitJob(Job);
-
+					
+					if (bSubmitJob)
+					{
+						SubmitJob(Job);
+					}
+					else // if preprocessing ran and failed, finish the job immediately
+					{
+						ProcessFinishedJob(Job, /* bCompilationSkipped = */true);
+					}
+					
 					Job->TimeTaskSubmitJobs = FPlatformTime::Seconds() - TimeStart;
 				}, Prio);
 			}
@@ -1779,17 +1788,13 @@ void FShaderCompileJobCollection::HandlePrintStats()
 	GShaderCompilingManager->PrintStats();
 }
 
-double FShaderJobCache::ProcessFinishedJob(FShaderCommonCompileJob* FinishedJob, bool bWasCached)
+double FShaderJobCache::ProcessFinishedJob(FShaderCommonCompileJob* FinishedJob, bool bCompilationSkipped)
 {
 	double StallTime;
 
 	FinishedJob->OnComplete();
 
-	GShaderCompilerStats->RegisterFinishedJob(*FinishedJob, bWasCached);
-
-	if (!bWasCached)
-	{
-	}
+	GShaderCompilerStats->RegisterFinishedJob(*FinishedJob, bCompilationSkipped);
 
 	{
 		// Need to protect writes to FShaderMapCompileResults
@@ -1806,7 +1811,7 @@ double FShaderJobCache::ProcessFinishedJob(FShaderCommonCompileJob* FinishedJob,
 	}
 
 	InternalSubtractNumOutstandingJobs(1);
-	if (!bWasCached && ShaderCompiler::IsJobCacheEnabled())
+	if (!bCompilationSkipped && ShaderCompiler::IsJobCacheEnabled())
 	{
 		AddToCacheAndProcessPending(FinishedJob);
 	}
@@ -5036,11 +5041,11 @@ void FShaderCompilerStats::RegisterAssignedJob(FShaderCommonCompileJob& Job)
 }
 
 TRACE_DECLARE_INT_COUNTER(Shaders_Compiled, TEXT("Shaders/Compiled"));
-void FShaderCompilerStats::RegisterFinishedJob(FShaderCommonCompileJob& Job, bool bWasCached)
+void FShaderCompilerStats::RegisterFinishedJob(FShaderCommonCompileJob& Job, bool bCompilationSkipped)
 {
 	FScopeLock Lock(&CompileStatsLock);
 	
-	if (!bWasCached)
+	if (!bCompilationSkipped)
 	{
 		ensure(Job.TimeAssignedToExecution != 0.0);
 		Job.TimeExecutionCompleted = FPlatformTime::Seconds();
@@ -5059,18 +5064,18 @@ void FShaderCompilerStats::RegisterFinishedJob(FShaderCommonCompileJob& Job, boo
 		ensure(Job.TimeAddedToPendingQueue != 0.0 && Job.TimeAddedToPendingQueue <= Job.TimeExecutionCompleted);
 		AddToInterval(JobLifeTimeIntervals, TInterval<double>(Job.TimeAddedToPendingQueue, Job.TimeExecutionCompleted));
 	}
-
+	
 	if (Job.TimeTaskSubmitJobs)
 	{
 		Counters.AccumulatedTaskSubmitJobs += Job.TimeTaskSubmitJobs;
 		Counters.AccumulatedTaskSubmitJobsStall += Job.TimeTaskSubmitJobsStall;
 	}
 
-	auto RegisterStatsFromSingleJob = [this, bWasCached](const FShaderCompileJob& SingleJob)
+	auto RegisterStatsFromSingleJob = [this, bCompilationSkipped](const FShaderCompileJob& SingleJob)
 	{
 		// Register min/max/average shader code sizes for single job output
 		const int32 ShaderCodeSize = SingleJob.Output.ShaderCode.GetShaderCodeSize();
-		if (!bWasCached && ShaderCodeSize > 0)
+		if (!bCompilationSkipped && ShaderCodeSize > 0)
 		{
 			Counters.MinShaderCodeSize = (Counters.MinShaderCodeSize > 0 ? FMath::Min(Counters.MinShaderCodeSize, ShaderCodeSize) : ShaderCodeSize);
 			Counters.MaxShaderCodeSize = (Counters.MaxShaderCodeSize > 0 ? FMath::Max(Counters.MaxShaderCodeSize, ShaderCodeSize) : ShaderCodeSize);
@@ -5079,9 +5084,10 @@ void FShaderCompilerStats::RegisterFinishedJob(FShaderCommonCompileJob& Job, boo
 		}
 
 		// Sanity check; compile time should be 0 for cache hits
-		check(!bWasCached || SingleJob.Output.CompileTime == 0.0f);
-		// Preprocess time should always be non-zero if preprocessed job cache is enabled
-		check(!SingleJob.Input.bCachePreprocessed || SingleJob.Output.PreprocessTime > 0.0f);
+		check(!bCompilationSkipped || SingleJob.Output.CompileTime == 0.0f);
+		// Preprocess time should always be non-zero if preprocessed job cache is enabled and preprocessing succeeded;
+		// preprocessing for pipeline stage jobs may be skipped in the case preprocessing a preceding stage of the pipeline failed
+		check(!SingleJob.Input.bCachePreprocessed || !SingleJob.PreprocessOutput.GetSucceeded() || SingleJob.Output.PreprocessTime > 0.0f);
 
 		const FString ShaderName(SingleJob.Key.ShaderType->GetName());
 		if (FShaderTimings* Existing = ShaderTimings.Find(ShaderName))
@@ -5089,7 +5095,7 @@ void FShaderCompilerStats::RegisterFinishedJob(FShaderCommonCompileJob& Job, boo
 			// Always want to log preprocess time, in case preprocessed cache is enabled and preprocessing ran in the cooker prior to compilation
 			// (PreprocessTime will be 0 if preprocessed cache is disabled)
 			Existing->TotalPreprocessTime += SingleJob.Output.PreprocessTime;
-			if (!bWasCached)
+			if (!bCompilationSkipped)
 			{
 				// If no actual compiles have been logged yet, min compile time is just the compile time of this job (first to actually run)
 				Existing->MinCompileTime = Existing->NumCompiled ? FMath::Min(Existing->MinCompileTime, static_cast<float>(SingleJob.Output.CompileTime)) : SingleJob.Output.CompileTime;
@@ -5109,7 +5115,7 @@ void FShaderCompilerStats::RegisterFinishedJob(FShaderCommonCompileJob& Job, boo
 			New.AverageCompileTime = New.MinCompileTime;
 			// It's possible the first entry for a given shader didn't actually compile (i.e. hit in DDC)
 			// so we need to account for that in the stats
-			New.NumCompiled = bWasCached ? 0 : 1;
+			New.NumCompiled = bCompilationSkipped ? 0 : 1;
 			New.TotalPreprocessTime += SingleJob.Output.PreprocessTime;
 
 			ShaderTimings.Add(ShaderName, New);
diff --git a/Engine/Source/Runtime/Engine/Public/ShaderCompiler.h b/Engine/Source/Runtime/Engine/Public/ShaderCompiler.h
--- a/Engine/Source/Runtime/Engine/Public/ShaderCompiler.h
+++ b/Engine/Source/Runtime/Engine/Public/ShaderCompiler.h
@@ -66,8 +66,10 @@ public:
 
 	void SubmitJobs(const TArray<FShaderCommonCompileJobPtr>& InJobs);
 	
-	/** This is an entry point for all jobs that have finished the compilation (whether real or cached). Can be called from multiple threads.*/
-	void ProcessFinishedJob(FShaderCommonCompileJob* FinishedJob, bool bWasCached = false);
+	/** Called for all completed jobs, including those that were cache hits, duplicates of other in flight jobs, or skipped due to failed preprocessing.
+	 * Can be called from multiple threads.
+	 */
+	void ProcessFinishedJob(FShaderCommonCompileJob* FinishedJob, bool bCompilationSkipped = false);
 
 	/** Adds the job to cache. */
 	void AddToCacheAndProcessPending(FShaderCommonCompileJob* FinishedJob);
@@ -402,7 +404,7 @@ public:
 	void RegisterAssignedJob(FShaderCommonCompileJob& InOutJob);
 
 	/** Marks the job as finished for the stats purpose. Job will be modified to include the current timestamp. */
-	void RegisterFinishedJob(FShaderCommonCompileJob& InOutJob, bool bWasCached);
+	void RegisterFinishedJob(FShaderCommonCompileJob& InOutJob, bool bCompilationSkipped);
 
 	/** Informs statistics about a new job batch, so we can tally up batches. */
 	void RegisterJobBatch(int32 NumJobs, EExecutionType ExecType);
diff --git a/Engine/Source/Runtime/RenderCore/Private/ShaderCore.cpp b/Engine/Source/Runtime/RenderCore/Private/ShaderCore.cpp
--- a/Engine/Source/Runtime/RenderCore/Private/ShaderCore.cpp
+++ b/Engine/Source/Runtime/RenderCore/Private/ShaderCore.cpp
@@ -1273,6 +1273,7 @@ public:
 		Job.PreprocessOutput.ElapsedTime = FPlatformTime::Seconds() - StartPreprocessTime;
 		Job.Output.PreprocessTime = Job.PreprocessOutput.ElapsedTime;
 		Job.Output.ShaderDiagnosticDatas = Job.PreprocessOutput.GetDiagnosticDatas();
+		Job.Output.Errors.Append(Job.PreprocessOutput.Errors);
 		return Job.PreprocessOutput.bSucceeded;
 	}
 
@@ -1333,8 +1334,6 @@ public:
 			PreprocessShaderInternal(Compiler, Job);
 		}
 
-		Job.Output.Errors.Append(Job.PreprocessOutput.Errors);
-
 		if (Job.PreprocessOutput.bSucceeded)
 		{
 			if (Job.SecondaryPreprocessOutput.IsValid())
@@ -1375,7 +1374,7 @@ public:
 	}
 };
 
-void ConditionalPreprocessShader(FShaderCommonCompileJob* Job)
+bool ConditionalPreprocessShader(FShaderCommonCompileJob* Job)
 {
 	static ITargetPlatformManagerModule& TargetPlatformManager = GetTargetPlatformManagerRef();
 	if (FShaderCompileJob* SingleJob = Job->GetSingleShaderJob())
@@ -1383,28 +1382,39 @@ void ConditionalPreprocessShader(FShaderCommonCompileJob* Job)
 		if (SingleJob->Input.bCachePreprocessed)
 		{
 			const IShaderFormat* ShaderFormat = TargetPlatformManager.FindShaderFormat(SingleJob->Input.ShaderFormat);
-			FInternalShaderCompilerFunctions::PreprocessShaderInternal(ShaderFormat, *SingleJob);
+			return FInternalShaderCompilerFunctions::PreprocessShaderInternal(ShaderFormat, *SingleJob);
 		}
+		return true;
+
 	}
 	else if (FShaderPipelineCompileJob* PipelineJob = Job->GetShaderPipelineJob())
 	{
+		bool bAnyFailed = false;
 		for (FShaderCompileJob* StageJob : PipelineJob->StageJobs)
 		{
 			if (StageJob->Input.bCachePreprocessed)
 			{
 				const IShaderFormat* ShaderFormat = TargetPlatformManager.FindShaderFormat(StageJob->Input.ShaderFormat);
-				if (!FInternalShaderCompilerFunctions::PreprocessShaderInternal(ShaderFormat, *StageJob))
+
+				if (!bAnyFailed)
+				{
+					bAnyFailed |= !FInternalShaderCompilerFunctions::PreprocessShaderInternal(ShaderFormat, *StageJob);
+				}
+				else
 				{
-					// early out if preprocessing failed on one stage, no point in continuing
-					return;
+					// skip subsequent stage preprocessing if a prior stage failed to avoid unnecessary work, but log an error to indicate this
+					FString Error = FString::Printf(
+						TEXT("Preprocessing %s stage skipped due to earlier stage preprocessing failure."),
+						GetShaderFrequencyString(StageJob->Input.Target.GetFrequency()));
+					StageJob->Output.Errors.Add(FShaderCompilerError(*Error));
 				}
 			}
 		}
+		return !bAnyFailed;
 	}
-	else
-	{
-		checkf(0, TEXT("Unknown shader compile job type or bad job pointer"));
-	}
+
+	checkf(0, TEXT("Unknown shader compile job type or bad job pointer"));
+	return false;
 }
 
 void CompileShader(const TArray<const IShaderFormat*>& ShaderFormats, FShaderCompilerInput& Input, FShaderCompilerOutput& Output, const FString& WorkingDirectory, int32* CompileCount)
diff --git a/Engine/Source/Runtime/RenderCore/Public/ShaderCompilerCore.h b/Engine/Source/Runtime/RenderCore/Public/ShaderCompilerCore.h
--- a/Engine/Source/Runtime/RenderCore/Public/ShaderCompilerCore.h
+++ b/Engine/Source/Runtime/RenderCore/Public/ShaderCompilerCore.h
@@ -736,7 +736,7 @@ extern RENDERCORE_API int HandleShaderCompileException(Windows::LPEXCEPTION_POIN
 extern RENDERCORE_API const IShaderFormat* FindShaderFormat(FName Format, const TArray<const IShaderFormat*>& ShaderFormats);
 
 // Executes preprocessing for the given job, if the job is marked to be preprocessed independently prior to compilation.
-extern RENDERCORE_API void ConditionalPreprocessShader(FShaderCommonCompileJob* Job);
+extern RENDERCORE_API bool ConditionalPreprocessShader(FShaderCommonCompileJob* Job);
 UE_DEPRECATED(5.3, "Use CompileShader overload which takes an FShaderCompileJob& rather than passing input/output directly.")
 extern RENDERCORE_API void CompileShader(const TArray<const IShaderFormat*>& ShaderFormats, FShaderCompilerInput& Input, FShaderCompilerOutput& Output, const FString& WorkingDirectory, int32* CompileCount = nullptr);
 extern RENDERCORE_API void CompileShader(const TArray<const IShaderFormat*>& ShaderFormats, FShaderCompileJob& Job, const FString& WorkingDirectory, int32* CompileCount = nullptr);
diff --git a/Engine/Source/Runtime/RenderCore/Public/ShaderPreprocessTypes.h b/Engine/Source/Runtime/RenderCore/Public/ShaderPreprocessTypes.h
--- a/Engine/Source/Runtime/RenderCore/Public/ShaderPreprocessTypes.h
+++ b/Engine/Source/Runtime/RenderCore/Public/ShaderPreprocessTypes.h
@@ -122,6 +122,11 @@ public:
 		return bIsSecondary;
 	}
 
+	inline bool GetSucceeded() const
+	{
+		return bSucceeded;
+	}
+
 	inline void LogError(FString&& Message)
 	{
 		FShaderCompilerError& CompilerError = Errors.AddDefaulted_GetRef();;
