Fixed thick lines being clipped with ortho views because near clip distance is set to 1 then. This fixes the water landscape brush which relies on canvas draws for updating curves render target (!)

#tests editor

diff --git a/Engine/Source/Runtime/Engine/Private/BatchedElements.cpp b/Engine/Source/Runtime/Engine/Private/BatchedElements.cpp
--- a/Engine/Source/Runtime/Engine/Private/BatchedElements.cpp
+++ b/Engine/Source/Runtime/Engine/Private/BatchedElements.cpp
@@ -1031,7 +1031,7 @@ bool FBatchedElements::Draw(FRHICommandList& RHICmdList, const FMeshPassProcesso
 					{
 						const FBatchedThickLines& Line = ThickLines[FirstLineThisBatch + i];
 						const double Thickness = FMath::Abs( Line.Thickness );
-
+						
 						FVector4 StartClip	= WorldToClip.TransformFVector4(Line.Start);
 						FVector4 EndClip	= WorldToClip.TransformFVector4(Line.End);
 
@@ -1039,21 +1039,24 @@ bool FBatchedElements::Draw(FRHICommandList& RHICmdList, const FMeshPassProcesso
 						FVector LineStart = Line.Start;
 						FVector LineEnd = Line.End;
 						const double ClipAt = View.NearClippingDistance + UE_DOUBLE_KINDA_SMALL_NUMBER;
-						if (StartClip.W < ClipAt && EndClip.W < ClipAt)
-						{
-							continue;
-						}
-						else if (StartClip.W < ClipAt)
+						if (bIsPerspective)
 						{
-							double Along = (ClipAt - StartClip.W) / (EndClip.W - StartClip.W);
-							LineStart = FMath::Lerp(LineStart, LineEnd, Along);
-							StartClip = FMath::Lerp(StartClip, EndClip, Along);
-						}
-						else if (EndClip.W < ClipAt)
-						{
-							double Along = (ClipAt - EndClip.W) / (StartClip.W - EndClip.W);
-							LineEnd = FMath::Lerp(LineEnd, LineStart, Along);
-							EndClip = FMath::Lerp(EndClip, StartClip, Along);
+							if (StartClip.W < ClipAt && EndClip.W < ClipAt)
+							{
+								continue;
+							}
+							else if (StartClip.W < ClipAt)
+							{
+								double Along = (ClipAt - StartClip.W) / (EndClip.W - StartClip.W);
+								LineStart = FMath::Lerp(LineStart, LineEnd, Along);
+								StartClip = FMath::Lerp(StartClip, EndClip, Along);
+							}
+							else if (EndClip.W < ClipAt)
+							{
+								double Along = (ClipAt - EndClip.W) / (StartClip.W - EndClip.W);
+								LineEnd = FMath::Lerp(LineEnd, LineStart, Along);
+								EndClip = FMath::Lerp(EndClip, StartClip, Along);
+							}
 						}
 
 						const double StartW = StartClip.W;
@@ -1510,4 +1513,4 @@ void FBatchedElements::Append(FBatchedElements& Other)
 	MeshElements.Append(Other.MeshElements);
 	MeshVertices.Append(Other.MeshVertices);
 	Other.Clear();
-}
\ No newline at end of file
+}
