Fix crash initializing rigid body nodes for skeletal mesh components with no mesh assigned

diff --git a/Engine/Source/Runtime/Engine/Private/Components/SkeletalMeshComponent.cpp b/Engine/Source/Runtime/Engine/Private/Components/SkeletalMeshComponent.cpp
--- a/Engine/Source/Runtime/Engine/Private/Components/SkeletalMeshComponent.cpp
+++ b/Engine/Source/Runtime/Engine/Private/Components/SkeletalMeshComponent.cpp
@@ -4149,7 +4149,7 @@ void USkeletalMeshComponent::SetAllowRigidBodyAnimNode(bool bInAllow, bool bRein
 	{
 		bDisableRigidBodyAnimNode = !bInAllow;
 
-		if(bReinitAnim && bRegistered)
+		if(bReinitAnim && bRegistered && GetSkeletalMeshAsset() != nullptr)
 		{
 			// need to reinitialize rigid body nodes for new setting to take effect
 			if (AnimScriptInstance)
