Fixed case where static mesh will not be rendered if component overridden MinLOD is less than StaticMesh MinLOD

#localization none
#tests played a game on device

diff --git a/Engine/Source/Runtime/Engine/Private/StaticMeshRender.cpp b/Engine/Source/Runtime/Engine/Private/StaticMeshRender.cpp
--- a/Engine/Source/Runtime/Engine/Private/StaticMeshRender.cpp
+++ b/Engine/Source/Runtime/Engine/Private/StaticMeshRender.cpp
@@ -2753,7 +2753,7 @@ FPrimitiveSceneProxy* UStaticMeshComponent::CreateSceneProxy()
 	// Validate the LOD resources here
 	const FStaticMeshLODResourcesArray& LODResources = GetStaticMesh()->GetRenderData()->LODResources;
 	const int32 SMCurrentMinLOD = GetStaticMesh()->GetMinLODIdx();
-	const int32 EffectiveMinLOD = bOverrideMinLOD ? MinLOD : SMCurrentMinLOD;
+	const int32 EffectiveMinLOD = bOverrideMinLOD ? FMath::Max(MinLOD, SMCurrentMinLOD) : SMCurrentMinLOD;
 	if (LODResources.Num() == 0	|| LODResources[FMath::Clamp<int32>(EffectiveMinLOD, 0, LODResources.Num()-1)].VertexBuffers.StaticMeshVertexBuffer.GetNumVertices() == 0)
 	{
 		UE_LOG(LogStaticMesh, Verbose, TEXT("Skipping CreateSceneProxy for StaticMeshComponent %s (LOD problems)"), *GetFullName());
