Add some more FBox BP callable util functions.

diff --git a/Engine/Source/Runtime/Engine/Classes/Kismet/KismetMathLibrary.h b/Engine/Source/Runtime/Engine/Classes/Kismet/KismetMathLibrary.h
--- a/Engine/Source/Runtime/Engine/Classes/Kismet/KismetMathLibrary.h
+++ b/Engine/Source/Runtime/Engine/Classes/Kismet/KismetMathLibrary.h
@@ -3706,7 +3706,96 @@ class UKismetMathLibrary : public UBlueprintFunctionLibrary
 	UFUNCTION(BlueprintPure, Category="Math|Box", meta=(Keywords="construct build", NativeMakeFunc))
 	static ENGINE_API FBox MakeBox(FVector Min, FVector Max);
 
+	/** 
+	 * Utility function to build an box from an Origin and Extent 
+	 *
+	 * @param Origin The location of the bounding box.
+	 * @param Extent Half size of the bounding box.
+	 * @return A new axis-aligned bounding box.
+	 */
+	UFUNCTION(BlueprintPure, Category="Math|Box", meta=(Keywords="construct build"))
+	static ENGINE_API FBox MakeBoxWithOrigin(const FVector& Origin, const FVector& Extent);
+
+	/**
+	 * Returns true if the InnerTest Box is is completely inside of the OuterTest Box
+	 * 
+	 * @param InnerTest		The box to check if it is on the inside
+	 * @param OuterTest		The box to check if InnerTest is within.
+	 * 
+	 * @return True if InnerTest Box is is completely inside of OuterTest Box
+	 */
+	UFUNCTION(BlueprintPure, Category="Math|Box", meta=(DisplayName="Is Inside (Box)"))
+	static ENGINE_API bool Box_IsInside(const FBox& InnerTest, const FBox& OuterTest);
 
+	/**
+	 * Returns true if the InnerTest Box is is completely inside or on OuterTest Box
+	 * 
+	 * @param InnerTest		The box to check if it is on the inside
+	 * @param OuterTest		The box to check if InnerTest is within or on.
+	 * 
+	 * @return True if InnerTest Box is is completely inside of or on OuterTest Box
+	 */
+	UFUNCTION(BlueprintPure, Category="Math|Box", meta=(DisplayName="Is Inside Or On (Box)"))
+	static ENGINE_API bool Box_IsInsideOrOn(const FBox& InnerTest, const FBox& OuterTest);
+
+	/** 
+	 * Checks whether the given location is inside this box.
+	 *
+	 * @param Box	The box to test
+	 * @param Point The location to test for inside the bounding volume.
+	 * @return true if location is inside this volume.
+	 *
+	 * @note  This function assumes boxes have open bounds, i.e. points lying on the border of the box are not inside.
+	 *        Use IsPointInBox_Box to include borders in the test.
+	 */
+	UFUNCTION(BlueprintPure, Category="Math|Box", meta=(DisplayName="Is Vector Inside (Box)"))
+	static ENGINE_API bool Box_IsPointInside(const FBox& Box, const FVector& Point);
+	
+	/**
+	 * Checks whether the given bounding box A intersects this bounding box B.
+	 *
+	 * @param A The bounding box to check intersection against
+	 * @param B The bounding box to intersect with.
+	 * 
+	 * @return true if the boxes intersect, false otherwise.
+	 *
+	 * @note  This function assumes boxes have closed bounds, i.e. boxes with
+	 *        coincident borders on any edge will overlap.
+	 */
+	UFUNCTION(BlueprintPure, Category="Math|Box", meta=(DisplayName="Intersects (Box)"))
+	static ENGINE_API bool Box_Intersects(const FBox& A, const FBox& B);
+
+	/**
+	* Returns a box of increased size.
+	*
+	* @param Negative The size to increase the volume by in the negative direction (positive values move the bounds outwards)
+	* @param Positive The size to increase the volume by in the positive direction (positive values move the bounds outwards)
+	* @return A new bounding box.
+	*/
+	UFUNCTION(BlueprintPure, Category="Math|Box", meta=(DisplayName="Expand By (Box)"))
+	static ENGINE_API FBox Box_ExpandBy(const FBox& Box, const FVector& Negative, const FVector& Positive);
+
+	/**
+	 * Returns the overlap TBox<T> of two boxes
+	 *
+	 * @param A		The bounding box to test
+	 * @param B		The bounding box to test overlap against
+	 * 
+	 * @return the overlap box. It can be 0 if they don't overlap
+	 */
+	UFUNCTION(BlueprintPure, Category="Math|Box", meta=(DisplayName="Overlap (Box)"))
+	static ENGINE_API FBox Box_Overlap(const FBox& A, const FBox& B);
+
+	/**
+	 * Calculates the closest point on or inside the box to a given point in space.
+	 *
+	 * @param Box	The box to check if the point is inside of
+	 * @param Point	The point in space
+	 * @return The closest point on or inside the box.
+	 */
+	UFUNCTION(BlueprintPure, Category="Math|Box", meta=(DisplayName="GetClosestPointTo (Box)"))
+	static ENGINE_API FVector Box_GetClosestPointTo(const FBox& Box, const FVector& Point);
+	
 	//
 	// Box2D functions
 	//
diff --git a/Engine/Source/Runtime/Engine/Classes/Kismet/KismetStringLibrary.h b/Engine/Source/Runtime/Engine/Classes/Kismet/KismetStringLibrary.h
--- a/Engine/Source/Runtime/Engine/Classes/Kismet/KismetStringLibrary.h
+++ b/Engine/Source/Runtime/Engine/Classes/Kismet/KismetStringLibrary.h
@@ -68,6 +68,14 @@ class UKismetStringLibrary : public UBlueprintFunctionLibrary
 	UFUNCTION(BlueprintPure, meta=(DisplayName = "To String (Object)", CompactNodeTitle = "->", BlueprintAutocast), Category="Utilities|String")
 	static ENGINE_API FString Conv_ObjectToString(class UObject* InObj);
 
+	/** Converts a FBox value to a string */
+	UFUNCTION(BlueprintPure, meta=(DisplayName = "To String (Box)", CompactNodeTitle = "->", BlueprintAutocast), Category="Utilities|String")
+	static ENGINE_API FString Conv_BoxToString(const FBox& Box);
+
+	/** Converts a FBox value to a string of its Center and Extents values. */
+	UFUNCTION(BlueprintPure, meta=(DisplayName = "To String Center and Extents (Box)", CompactNodeTitle = "->", BlueprintAutocast), Category="Utilities|String")
+	static ENGINE_API FString Conv_BoxCenterAndExtentsToString(const FBox& Box);
+
 	/** Converts a InputDeviceId value to a string */
 	UFUNCTION(BlueprintPure, meta=(DisplayName = "To String (InputDeviceId)", CompactNodeTitle = "->", BlueprintAutocast), Category="Utilities|String")
 	static ENGINE_API FString Conv_InputDeviceIdToString(FInputDeviceId InDeviceId);
diff --git a/Engine/Source/Runtime/Engine/Private/KismetMathLibrary.cpp b/Engine/Source/Runtime/Engine/Private/KismetMathLibrary.cpp
--- a/Engine/Source/Runtime/Engine/Private/KismetMathLibrary.cpp
+++ b/Engine/Source/Runtime/Engine/Private/KismetMathLibrary.cpp
@@ -1372,6 +1372,46 @@ bool UKismetMathLibrary::IsPointInBox_Box(FVector Point, FBox Box)
 	return Box.IsInsideOrOn(Point);
 }
 
+FBox UKismetMathLibrary::MakeBoxWithOrigin(const FVector& Origin, const FVector& Extent)
+{
+	return FBox::BuildAABB(Origin, Extent);
+}
+
+bool UKismetMathLibrary::Box_IsInside(const FBox& InnerTest, const FBox& OuterTest)
+{
+	return InnerTest.IsInside(OuterTest);
+}
+
+bool UKismetMathLibrary::Box_IsInsideOrOn(const FBox& InnerTest, const FBox& OuterTest)
+{
+	return InnerTest.IsInsideOrOn(OuterTest);
+}
+
+bool UKismetMathLibrary::Box_IsPointInside(const FBox& Box, const FVector& Point)
+{
+	return Box.IsInside(Point);
+}
+
+bool UKismetMathLibrary::Box_Intersects(const FBox& A, const FBox& B)
+{
+	return A.Intersect(B);
+}
+
+FBox UKismetMathLibrary::Box_ExpandBy(const FBox& Box, const FVector& Negative, const FVector& Positive)
+{
+	return Box.ExpandBy(Negative, Positive);
+}
+
+FBox UKismetMathLibrary::Box_Overlap(const FBox& A, const FBox& B)
+{
+	return A.Overlap(B);
+}
+
+FVector UKismetMathLibrary::Box_GetClosestPointTo(const FBox& Box, const FVector& Point)
+{
+	return Box.GetClosestPointTo(Point);
+}
+
 double UKismetMathLibrary::GetBoxVolume(const FBox& InBox)
 {
 	return InBox.GetVolume();
diff --git a/Engine/Source/Runtime/Engine/Private/KismetStringLibrary.cpp b/Engine/Source/Runtime/Engine/Private/KismetStringLibrary.cpp
--- a/Engine/Source/Runtime/Engine/Private/KismetStringLibrary.cpp
+++ b/Engine/Source/Runtime/Engine/Private/KismetStringLibrary.cpp
@@ -1,6 +1,7 @@
 // Copyright Epic Games, Inc. All Rights Reserved.
 
 #include "Kismet/KismetStringLibrary.h"
+#include "Math/Box.h"
 
 #include UE_INLINE_GENERATED_CPP_BY_NAME(KismetStringLibrary)
 
@@ -123,6 +124,20 @@ FString UKismetStringLibrary::Conv_ObjectToString(class UObject* InObj)
 	return (InObj != nullptr) ? InObj->GetName() : FString(TEXT("None"));
 }
 
+FString UKismetStringLibrary::Conv_BoxToString(const FBox& Box)
+{
+	return Box.ToString();
+}
+
+FString UKismetStringLibrary::Conv_BoxCenterAndExtentsToString(const FBox& Box)
+{
+	FVector Center = FVector::ZeroVector;
+	FVector Extents = FVector::ZeroVector;
+	Box.GetCenterAndExtents(OUT Center, OUT Extents);
+
+	return FString::Printf(TEXT("Center: %s Extents: %s "), *Center.ToString(), *Extents.ToString());
+}
+
 FString UKismetStringLibrary::Conv_InputDeviceIdToString(FInputDeviceId InDeviceId)
 {
 	return FString::Printf(TEXT("%d"), InDeviceId.GetId());
