ISM Component Descriptor: Fixed incomplete hash computation for FISMComponentDescriptor and FSoftISMComponentDescriptor

#localization none
#tests Rebuilt HLODs

diff --git a/Engine/Source/Runtime/Engine/Private/ISMPartition/ISMComponentDescriptor.cpp b/Engine/Source/Runtime/Engine/Private/ISMPartition/ISMComponentDescriptor.cpp
--- a/Engine/Source/Runtime/Engine/Private/ISMPartition/ISMComponentDescriptor.cpp
+++ b/Engine/Source/Runtime/Engine/Private/ISMPartition/ISMComponentDescriptor.cpp
@@ -269,6 +269,36 @@ uint32 FISMComponentDescriptorBase::ComputeHash() const
 	return Hash;
 }
 
+uint32 FISMComponentDescriptor::ComputeHash() const
+{
+	Super::ComputeHash();
+
+	FISMComponentDescriptor& MutableSelf = *const_cast<FISMComponentDescriptor*>(this);
+	FArchiveCrc32 CrcArchive(Hash);
+	CrcArchive << MutableSelf.StaticMesh;
+	CrcArchive << MutableSelf.OverrideMaterials;
+	CrcArchive << MutableSelf.OverlayMaterial;
+	CrcArchive << MutableSelf.RuntimeVirtualTextures;
+	Hash = CrcArchive.GetCrc();
+
+	return Hash;
+}
+
+uint32 FSoftISMComponentDescriptor::ComputeHash() const
+{
+	Super::ComputeHash();
+
+	FSoftISMComponentDescriptor& MutableSelf = *const_cast<FSoftISMComponentDescriptor*>(this);
+	FArchiveCrc32 CrcArchive(Hash);
+	CrcArchive << MutableSelf.StaticMesh;
+	CrcArchive << MutableSelf.OverrideMaterials;
+	CrcArchive << MutableSelf.OverlayMaterial;
+	CrcArchive << MutableSelf.RuntimeVirtualTextures;
+	Hash = CrcArchive.GetCrc();
+
+	return Hash;
+}
+
 UInstancedStaticMeshComponent* FISMComponentDescriptorBase::CreateComponent(UObject* Outer, FName Name, EObjectFlags ObjectFlags) const
 {
 	UInstancedStaticMeshComponent* ISMComponent = NewObject<UInstancedStaticMeshComponent>(Outer, ComponentClass, Name, ObjectFlags);
diff --git a/Engine/Source/Runtime/Engine/Public/ISMPartition/ISMComponentDescriptor.h b/Engine/Source/Runtime/Engine/Public/ISMPartition/ISMComponentDescriptor.h
--- a/Engine/Source/Runtime/Engine/Public/ISMPartition/ISMComponentDescriptor.h
+++ b/Engine/Source/Runtime/Engine/Public/ISMPartition/ISMComponentDescriptor.h
@@ -218,6 +218,7 @@ struct FISMComponentDescriptor : public FISMComponentDescriptorBase
 	static ENGINE_API FISMComponentDescriptor CreateFrom(const TSubclassOf<UStaticMeshComponent>& ComponentClass);
 
 	ENGINE_API virtual void InitFrom(const UStaticMeshComponent* Component, bool bInitBodyInstance = true) override;
+	ENGINE_API virtual uint32 ComputeHash() const;
 	ENGINE_API virtual void InitComponent(UInstancedStaticMeshComponent* ISMComponent) const override;
 		
 	ENGINE_API bool operator!=(const FISMComponentDescriptor& Other) const;
@@ -252,6 +253,7 @@ struct FSoftISMComponentDescriptor : public FISMComponentDescriptorBase
 	static ENGINE_API FSoftISMComponentDescriptor CreateFrom(const TSubclassOf<UStaticMeshComponent>& ComponentClass);
 
 	ENGINE_API virtual void InitFrom(const UStaticMeshComponent* Component, bool bInitBodyInstance = true) override;
+	ENGINE_API virtual uint32 ComputeHash() const;
 	ENGINE_API virtual void InitComponent(UInstancedStaticMeshComponent* ISMComponent) const override;
 
 	ENGINE_API bool operator!=(const FSoftISMComponentDescriptor& Other) const;
