Strengthened tiling texture build path against ddc failures.

diff --git a/Engine/Source/Runtime/Engine/Private/TextureDerivedDataTask.cpp b/Engine/Source/Runtime/Engine/Private/TextureDerivedDataTask.cpp
--- a/Engine/Source/Runtime/Engine/Private/TextureDerivedDataTask.cpp
+++ b/Engine/Source/Runtime/Engine/Private/TextureDerivedDataTask.cpp
@@ -823,7 +823,9 @@ static void DDC1_BuildTexture(
 				bSucceeded = DerivedData->TryInlineMipData(InBuildSettingsPerLayer[0].LODBiasWithCinematicMips, TexturePathName);
 				if (bSucceeded == false)
 				{
-					UE_LOG(LogTexture, Display, TEXT("Failed to put and then read back mipmap data from DDC for %s"), *TexturePathName);
+					// This should only ever happen with DDC issues - it can technically be a transient issue if you lose connection
+					// in the middle of a build, but with a stable connection it's probably a ddc bug.
+					UE_LOG(LogTexture, Warning, TEXT("Failed to put and then read back mipmap data from DDC for %s"), *TexturePathName);
 				}
 			}
 		}
@@ -1783,6 +1785,21 @@ bool DDC1_BuildTiledClassicTexture(
 	BytesCached = LinearBytesCached;
 	bool bHasLinearDerivedData = bLinearSucceeded;
 
+	void* LinearMipData[MAX_TEXTURE_MIP_COUNT] = {};
+	int64 LinearMipSizes[MAX_TEXTURE_MIP_COUNT];
+	if (bHasLinearDerivedData)
+	{
+		// The linear bits are built - need to fetch
+		if (LinearDerivedData.TryLoadMipsWithSizes(0, LinearMipData, LinearMipSizes, TexturePathName) == false)
+		{
+			// This can technically happen with a DDC failure and there is an expectation that we can recover and regenerate in such situations.
+			// However, it should be very rare and most likely indicated a backend bug, so we still warn.
+			UE_LOG(LogTexture, Warning, TEXT("Tiling texture build was unable to load the linear texture mips after fetching, will try to build: %s"), *TexturePathName);
+			bHasLinearDerivedData = false;
+		}
+
+	}
+
 	if (bHasLinearDerivedData == false)
 	{
 		// Linear data didn't exist, need to build it.
@@ -1810,6 +1827,14 @@ bool DDC1_BuildTiledClassicTexture(
 				&LinearDerivedData,
 				LinearBytesCached,
 				bHasLinearDerivedData);
+
+			// This should succeed because we asked for inline mips if the build succeeded
+			if (bHasLinearDerivedData && 
+				LinearDerivedData.TryLoadMipsWithSizes(0, LinearMipData, LinearMipSizes, TexturePathName) == false)
+			{
+				UE_LOG(LogTexture, Error, TEXT("Tiling texture build was unable to load the linear texture mips after a successful build, bad bug!: %s"), *TexturePathName);
+				return false;
+			}
 		}
 	}
 
@@ -1826,15 +1851,6 @@ bool DDC1_BuildTiledClassicTexture(
 	UE::TextureBuildUtilities::FTextureBuildMetadata BuildMetadata;
 	UnpackTextureBuildMetadataFromPlatformData(&BuildMetadata, &LinearDerivedData);
 
-	// The linear bits are built - need to fetch
-	void* LinearMipData[MAX_TEXTURE_MIP_COUNT] = {};
-	int64 LinearMipSizes[MAX_TEXTURE_MIP_COUNT];
-	if (LinearDerivedData.TryLoadMipsWithSizes(0, LinearMipData, LinearMipSizes, TexturePathName) == false)
-	{
-		UE_LOG(LogTexture, Error, TEXT("Tiling texture build was unable to reload the linear texture mips after building/fetching: %s"), *TexturePathName);
-		return false;
-	}
-
 	// Have all the data - do some sanity checks as we convert to the metadata format the tiler expects.
 	TArray<FMemoryView, TInlineAllocator<MAX_TEXTURE_MIP_COUNT>> InputTextureMipViews;
 	FEncodedTextureDescription TextureDescription;
