PR #11066: Export CheckForErrors override from ANote

diff --git a/Engine/Source/Runtime/Engine/Classes/Engine/Note.h b/Engine/Source/Runtime/Engine/Classes/Engine/Note.h
--- a/Engine/Source/Runtime/Engine/Classes/Engine/Note.h
+++ b/Engine/Source/Runtime/Engine/Classes/Engine/Note.h
@@ -35,7 +35,7 @@ public:
 
 #if WITH_EDITOR
 	//~ Begin AActor Interface
-	virtual void CheckForErrors() override;
+	ENGINE_API virtual void CheckForErrors() override;
 	//~ End AActor Interface
 #endif
 
