Fix three separate bugs with cooked builds ODSC.

1. Shader compilation made the assumption that we should only update materials with complete shader maps.  In the editor w/ ODSC this is never true so the logic to update materials only happened if we had a complete shader map _or_ we were in the Editor.  If we were in the cooker we would only process completed shader maps.  Since COTF (the cooker) will process in-complete shader maps with ODSC we were never updating materials with new shader maps.  This was causing the client to not receive the updated shader map with newly compiled shaders.  The fix here is to have the COTF server explicitly enable a new setting on the Shader Manager to update materials with incomplete maps.
2. FMaterialShaderMap::Register doesn't override the GIdToMaterialShaderMap with a new map.  In the case of ODSC the same shader map with the same ID would be sent to the client, but it would have more shaders in it than before.  When re-initialized the material the _old_ shader map would be applied, and not the new one.  This has been fixed by adding a specific ODSC version of this function which can override the existing map in the global variable.
3. FMaterial::CacheGivenTypes didn't take into account shader permutations for mesh material shaders.

#rnx

diff --git a/Engine/Source/Editor/UnrealEd/Private/CookOnTheFlyServer.cpp b/Engine/Source/Editor/UnrealEd/Private/CookOnTheFlyServer.cpp
--- a/Engine/Source/Editor/UnrealEd/Private/CookOnTheFlyServer.cpp
+++ b/Engine/Source/Editor/UnrealEd/Private/CookOnTheFlyServer.cpp
@@ -571,6 +571,7 @@ bool UCookOnTheFlyServer::StartCookOnTheFly(FCookOnTheFlyStartupOptions InCookOn
 	if (!IsCookingInEditor())
 	{
 		GShaderCompilingManager->SkipShaderCompilation(true);
+		GShaderCompilingManager->SetAllowForIncompleteShaderMaps(true);
 	}
 
 	LLM_SCOPE_BYTAG(Cooker);
@@ -5706,6 +5707,7 @@ void UCookOnTheFlyServer::ShutdownCookOnTheFly()
 		if (!IsCookingInEditor())
 		{
 			GShaderCompilingManager->SkipShaderCompilation(false);
+			GShaderCompilingManager->SetAllowForIncompleteShaderMaps(false);
 		}
 	}
 
@@ -11251,6 +11253,7 @@ void UCookOnTheFlyServer::StartCookAsCookWorker()
 	if (IsDirectorCookOnTheFly())
 	{
 		GShaderCompilingManager->SkipShaderCompilation(true);
+		GShaderCompilingManager->SetAllowForIncompleteShaderMaps(true);
 	}
 	CookWorkerClient->DoneWithInitialSettings();
 
@@ -11316,6 +11319,7 @@ void UCookOnTheFlyServer::CookAsCookWorkerFinished()
 	if (IsDirectorCookOnTheFly())
 	{
 		GShaderCompilingManager->SkipShaderCompilation(false);
+		GShaderCompilingManager->SetAllowForIncompleteShaderMaps(false);
 	}
 	LogCookWorkerStats();
 	if (IsDirectorCookByTheBook())
diff --git a/Engine/Source/Runtime/Engine/Private/Materials/MaterialShader.cpp b/Engine/Source/Runtime/Engine/Private/Materials/MaterialShader.cpp
--- a/Engine/Source/Runtime/Engine/Private/Materials/MaterialShader.cpp
+++ b/Engine/Source/Runtime/Engine/Private/Materials/MaterialShader.cpp
@@ -1719,7 +1719,7 @@ void FMaterialShaderMap::LoadForRemoteRecompile(FArchive& Ar, EShaderPlatform Sh
 				ShaderMap->Serialize(Ar, false);
 
 				// Register in the global map
-				ShaderMap->Register(ShaderPlatform);
+				ShaderMap->RegisterForODSC(ShaderPlatform);
 
 				LoadedShaderMaps.Add(ShaderMap);
 			}
@@ -2962,6 +2962,36 @@ void FMaterialShaderMap::Register(EShaderPlatform InShaderPlatform)
 	}
 }
 
+void FMaterialShaderMap::RegisterForODSC(EShaderPlatform InShaderPlatform)
+{
+	Register(InShaderPlatform);
+
+	{
+		FScopeLock ScopeLock(&GIdToMaterialShaderMapCS);
+
+		FMaterialShaderMap* CachedMap = GIdToMaterialShaderMap[GetShaderPlatform()].FindRef(ShaderMapId);
+		if (CachedMap == nullptr)
+		{
+			GIdToMaterialShaderMap[GetShaderPlatform()].Add(ShaderMapId, this);
+			bRegistered = true;
+		}
+		else
+		{	
+			if (CachedMap != this)
+			{
+				// deregister the existing map.
+				int RemoveIndex = GIdToMaterialShaderMap[GetShaderPlatform()].Remove(ShaderMapId);
+				check(RemoveIndex != INDEX_NONE);
+				CachedMap->bRegistered = false;
+
+				// register ourselves.
+				GIdToMaterialShaderMap[GetShaderPlatform()].Add(ShaderMapId, this);
+				bRegistered = true;
+			}
+		}
+	}
+}
+
 void FMaterialShaderMap::AddRef()
 {
 	//#todo-mw: re-enable to try to find potential corruption of the global shader map ID array
diff --git a/Engine/Source/Runtime/Engine/Private/Materials/MaterialShared.cpp b/Engine/Source/Runtime/Engine/Private/Materials/MaterialShared.cpp
--- a/Engine/Source/Runtime/Engine/Private/Materials/MaterialShared.cpp
+++ b/Engine/Source/Runtime/Engine/Private/Materials/MaterialShared.cpp
@@ -3269,20 +3269,33 @@ void FMaterial::CacheGivenTypes(EShaderPlatform Platform, const TArray<const FVe
 			}
 			else if (ShaderType->GetTypeForDynamicCast() == FShaderType::EShaderTypeForDynamicCast::MeshMaterial)
 			{
-				ShaderType->AsMeshMaterialShaderType()->BeginCompileShader(
-					EShaderCompileJobPriority::ForceLocal,
-					GetGameThreadCompilingShaderMapId(),
-					0,
-					Platform,
-					GameThreadShaderMap->GetPermutationFlags(),
-					this,
-					GameThreadShaderMap->GetShaderMapId(),
-					GameThreadPendingCompilerEnvironment,
-					VFType,
-					CompileJobs,
-					DebugGroupName,
-					nullptr,
-					nullptr);
+				const EShaderPermutationFlags ShaderPermutation = GameThreadShaderMap->GetPermutationFlags();
+
+				const FMeshMaterialShaderType* MeshMaterialShaderType = ShaderType->GetMeshMaterialShaderType();
+
+				for (int32 PermutationId = 0; PermutationId < ShaderType->GetPermutationCount(); ++PermutationId)
+				{
+					const bool bShaderShouldCompile = MeshMaterialShaderType->ShouldCompilePermutation(Platform, this, VFType, PermutationId, ShaderPermutation);
+					if (!bShaderShouldCompile)
+					{
+						continue;
+					}
+
+					ShaderType->AsMeshMaterialShaderType()->BeginCompileShader(
+						EShaderCompileJobPriority::ForceLocal,
+						GetGameThreadCompilingShaderMapId(),
+						PermutationId,
+						Platform,
+						ShaderPermutation,
+						this,
+						GameThreadShaderMap->GetShaderMapId(),
+						GameThreadPendingCompilerEnvironment,
+						VFType,
+						CompileJobs,
+						DebugGroupName,
+						nullptr,
+						nullptr);
+				}
 			}
 		}
 
diff --git a/Engine/Source/Runtime/Engine/Private/ShaderCompiler/ShaderCompiler.cpp b/Engine/Source/Runtime/Engine/Private/ShaderCompiler/ShaderCompiler.cpp
--- a/Engine/Source/Runtime/Engine/Private/ShaderCompiler/ShaderCompiler.cpp
+++ b/Engine/Source/Runtime/Engine/Private/ShaderCompiler/ShaderCompiler.cpp
@@ -5372,6 +5372,7 @@ FShaderCompilingManager::FShaderCompilingManager() :
 	SuppressedShaderPlatforms(0),
 	BuildDistributionController(nullptr),
 	bNoShaderCompilation(false),
+	bAllowForIncompleteShaderMaps(false),
 	Notification(GetAssetNameFormat())
 {
 	// don't perform any initialization if compiling is not allowed
@@ -6463,7 +6464,7 @@ void FShaderCompilingManager::ProcessCompiledShaderMaps(
 					bool bIsComplete = CompilingShaderMap->IsComplete(Material, true);
 
 					// If running a cook, only process complete shader maps, as there's no rendering of partially complete shader maps to worry about.
-					if (bIsComplete || IsRunningCookCommandlet() == false)
+					if (bIsComplete || IsRunningCookCommandlet() == false || bAllowForIncompleteShaderMaps)
 					{
 						if (ShaderMapToUseForRendering == nullptr)
 						{
diff --git a/Engine/Source/Runtime/Engine/Public/MaterialShared.h b/Engine/Source/Runtime/Engine/Public/MaterialShared.h
--- a/Engine/Source/Runtime/Engine/Public/MaterialShared.h
+++ b/Engine/Source/Runtime/Engine/Public/MaterialShared.h
@@ -1333,6 +1333,11 @@ public:
 	/** Registers a material shader map in the global map so it can be used by materials. */
 	void Register(EShaderPlatform InShaderPlatform);
 
+	/** Registers a material shader map in the global map so it can be used by materials.
+		Since the shader maps content can change this will overwrite entries in the global map.
+	*/
+	void RegisterForODSC(EShaderPlatform InShaderPlatform);
+
 	// Reference counting.
 	ENGINE_API void AddRef();
 	ENGINE_API void Release();
diff --git a/Engine/Source/Runtime/Engine/Public/ShaderCompiler.h b/Engine/Source/Runtime/Engine/Public/ShaderCompiler.h
--- a/Engine/Source/Runtime/Engine/Public/ShaderCompiler.h
+++ b/Engine/Source/Runtime/Engine/Public/ShaderCompiler.h
@@ -709,6 +709,9 @@ private:
 	/** Opt out of material shader compilation and instead place an empty shader map. */
 	bool bNoShaderCompilation;
 
+	/** If we are using ODSC (On Demand Shader Compilation) we should allow for incomplete maps to still be processed. */
+	bool bAllowForIncompleteShaderMaps;
+
 	/** Used to show a notification accompanying progress. */
 	FAsyncCompilationNotification Notification;
 
@@ -885,6 +888,11 @@ public:
 		}
 	}
 
+	void SetAllowForIncompleteShaderMaps(bool toggle)
+	{
+		bAllowForIncompleteShaderMaps = toggle;
+	}
+
 	ENGINE_API bool IsCompilingShaderMap(uint32 Id);
 
 	/** Prepares a job of the given type for compilation.  If a job with the given Id/Key already exists, it will attempt to adjust to the higher priority if possible, and nullptr will be returned.
