Do not set bAnyMaterialHasPixelAnimation if TSR is not supported on the shader platform. Setting the flag forces velocity outputs but only TSR uses the information currently so don't set the flag if TSR is not supported

diff --git a/Engine/Source/Runtime/Engine/Private/PrimitiveSceneProxy.cpp b/Engine/Source/Runtime/Engine/Private/PrimitiveSceneProxy.cpp
--- a/Engine/Source/Runtime/Engine/Private/PrimitiveSceneProxy.cpp
+++ b/Engine/Source/Runtime/Engine/Private/PrimitiveSceneProxy.cpp
@@ -656,7 +656,9 @@ FPrimitiveSceneProxy::FPrimitiveSceneProxy(const FPrimitiveSceneProxyDesc& InPro
 			}
 		}
 
-		if (VelocityEncodeHasPixelAnimation(GetScene().GetShaderPlatform()))
+		if (VelocityEncodeHasPixelAnimation(GetScene().GetShaderPlatform())
+			// Currently, only TSR checks the HasPixelAnimation flag but setting it will force velocity writes even if TSR is not used
+			&& SupportsTSR(GetScene().GetShaderPlatform()))
 		{
 			for (const UMaterialInterface* MaterialInterface : UsedMaterials)
 			{
