Fixed: missing Texture2DArray information in a 'listtextures' dump

diff --git a/Engine/Source/Runtime/Engine/Classes/Engine/Texture2DArray.cpp b/Engine/Source/Runtime/Engine/Classes/Engine/Texture2DArray.cpp
--- a/Engine/Source/Runtime/Engine/Classes/Engine/Texture2DArray.cpp
+++ b/Engine/Source/Runtime/Engine/Classes/Engine/Texture2DArray.cpp
@@ -688,4 +688,21 @@ bool UTexture2DArray::StreamIn(int32 NewMipCount, bool bHighPrio)
 	return false;
 }
 
+
+int32 UTexture2DArray::GetNumResidentMips() const
+{
+	if (GetResource())
+	{
+		if (CachedSRRState.IsValid())
+		{
+			return CachedSRRState.NumResidentLODs;
+		}
+		else
+		{
+			return GetResource()->GetCurrentMipCount();
+		}
+	}
+	return 0;
+}
+
 #undef LOCTEXT_NAMESPACE
diff --git a/Engine/Source/Runtime/Engine/Classes/Engine/Texture2DArray.h b/Engine/Source/Runtime/Engine/Classes/Engine/Texture2DArray.h
--- a/Engine/Source/Runtime/Engine/Classes/Engine/Texture2DArray.h
+++ b/Engine/Source/Runtime/Engine/Classes/Engine/Texture2DArray.h
@@ -133,5 +133,6 @@ public:
 	virtual bool StreamOut(int32 NewMipCount) final override;
 	virtual bool StreamIn(int32 NewMipCount, bool bHighPrio) final override;
 	//~ End UStreamableRenderAsset Interface
+	int32 GetNumResidentMips() const;
 
 };
\ No newline at end of file
diff --git a/Engine/Source/Runtime/Engine/Private/UnrealEngine.cpp b/Engine/Source/Runtime/Engine/Private/UnrealEngine.cpp
--- a/Engine/Source/Runtime/Engine/Private/UnrealEngine.cpp
+++ b/Engine/Source/Runtime/Engine/Private/UnrealEngine.cpp
@@ -6177,6 +6177,7 @@ bool UEngine::HandleListTexturesCommand( const TCHAR* Cmd, FOutputDevice& Ar )
 	{
 		UTexture*			Texture				= *It;
 		UTexture2D*			Texture2D			= Cast<UTexture2D>(Texture);
+		UTexture2DArray*	Texture2DArray		= Cast<UTexture2DArray>(Texture);
 		UTextureCube*		TextureCube			= Cast<UTextureCube>(Texture);
 		UVolumeTexture*		Texture3D			= Cast<UVolumeTexture>(Texture);
 		UTextureRenderTarget2D* RenderTexture	= Cast<UTextureRenderTarget2D>(Texture);
@@ -6229,6 +6230,15 @@ bool UEngine::HandleListTexturesCommand( const TCHAR* Cmd, FOutputDevice& Ar )
 				NumApplicableToMinSize++;
 			}
 		}
+		else if (Texture2DArray != nullptr)
+		{
+			NumMips				= Texture2DArray->GetNumMips();
+			Format				= Texture2DArray->GetPixelFormat();
+
+			DroppedMips			= Texture2DArray->GetNumMips() - Texture2DArray->GetNumResidentMips();
+			CurSizeX			= FMath::Max<int32>(Texture2DArray->GetSizeX() >> DroppedMips, 1);
+			CurSizeY			= FMath::Max<int32>(Texture2DArray->GetSizeY() >> DroppedMips, 1);
+		}
 		else if (TextureCube != nullptr)
 		{
 			NumMips				= TextureCube->GetNumMips();
