[Custom Controls] [Reticle] Reticle stuck in the corner if target is outside camera view
#rnx
#tests [How To Test]
* On your island make sure you have the following devices:
   * Gameplay Control: Third Person
   * Item Spawner with any ranged weapon
   * Sentry
* Play the game
* Get near the sentry and start rotating the camera to see the reticle aligns to the bottom when the target is outside the view
[Video after fix]
https://drive.google.com/file/d/1nw4ClXOeH8UmzCM3e-yo6EweWy2Gs9sP/view?usp=drive_link

diff --git a/Engine/Source/Runtime/Engine/Private/SceneView.cpp b/Engine/Source/Runtime/Engine/Private/SceneView.cpp
--- a/Engine/Source/Runtime/Engine/Private/SceneView.cpp
+++ b/Engine/Source/Runtime/Engine/Private/SceneView.cpp
@@ -1412,30 +1412,40 @@ void FSceneView::DeprojectScreenToWorld(const FVector2D& ScreenPos, const FIntRe
 	out_WorldDirection = RayDirWorldSpace;
 }
 
-bool FSceneView::ProjectWorldToScreen(const FVector& WorldPosition, const FIntRect& ViewRect, const FMatrix& ViewProjectionMatrix, FVector2D& out_ScreenPos)
+bool FSceneView::ProjectWorldToScreen(const FVector& WorldPosition, const FIntRect& ViewRect, const FMatrix& ViewProjectionMatrix, FVector2D& out_ScreenPos, bool bShouldCalcOutsideViewPosition /*= false*/)
 {
 	FPlane Result = ViewProjectionMatrix.TransformFVector4(FVector4(WorldPosition, 1.f));
-	if ( Result.W > 0.0f )
+	bool bIsInsideView = Result.W > 0.0f;
+	double W = Result.W;
+
+	// If WorldPosition is outside the ViewProjectionMatrix and we don't force to calc the outside view position, stop the calcs.
+	if (!bIsInsideView && !bShouldCalcOutsideViewPosition)
 	{
-		// the result of this will be x and y coords in -1..1 projection space
-		const float RHW = 1.0f / Result.W;
-		FPlane PosInScreenSpace = FPlane(Result.X * RHW, Result.Y * RHW, Result.Z * RHW, Result.W);
+		return false;
+	}
 
-		// Move from projection space to normalized 0..1 UI space
-		const float NormalizedX = ( PosInScreenSpace.X / 2.f ) + 0.5f;
-		const float NormalizedY = 1.f - ( PosInScreenSpace.Y / 2.f ) - 0.5f;
+	// Tweak our W value to allow the outside view position calcs if the variable is enabled.
+	if (bShouldCalcOutsideViewPosition)
+	{
+		W = FMath::Abs(Result.W);
+	}
+	
+	// the result of this will be x and y coords in -1..1 projection space
+	const float RHW = 1.0f / W;
+	FPlane PosInScreenSpace = FPlane(Result.X * RHW, Result.Y * RHW, Result.Z * RHW, W);
 
-		FVector2D RayStartViewRectSpace(
-			( NormalizedX * (float)ViewRect.Width() ),
-			( NormalizedY * (float)ViewRect.Height() )
-			);
+	// Move from projection space to normalized 0..1 UI space
+	const float NormalizedX = ( PosInScreenSpace.X / 2.f ) + 0.5f;
+	const float NormalizedY = 1.f - ( PosInScreenSpace.Y / 2.f ) - 0.5f;
 
-		out_ScreenPos = RayStartViewRectSpace + FVector2D(static_cast<float>(ViewRect.Min.X), static_cast<float>(ViewRect.Min.Y));
+	FVector2D RayStartViewRectSpace(
+		( NormalizedX * (float)ViewRect.Width() ),
+		( NormalizedY * (float)ViewRect.Height() )
+		);
 
-		return true;
-	}
+	out_ScreenPos = RayStartViewRectSpace + FVector2D(static_cast<float>(ViewRect.Min.X), static_cast<float>(ViewRect.Min.Y));
 
-	return false;
+	return bIsInsideView;
 }
 
 
diff --git a/Engine/Source/Runtime/Engine/Public/SceneView.h b/Engine/Source/Runtime/Engine/Public/SceneView.h
--- a/Engine/Source/Runtime/Engine/Public/SceneView.h
+++ b/Engine/Source/Runtime/Engine/Public/SceneView.h
@@ -1645,8 +1645,9 @@ public:
 	 * @param ViewRect - view rectangle
 	 * @param ViewProjectionMatrix - combined view projection matrix
 	 * @param out_ScreenPos (out) - screen coordinates in pixels
+	 * @param bShouldCalcOutsideViewPosition - if enabled, calculates the out_ScreenPos if the WorldPosition is outside of ViewProjectionMatrix
 	 */
-	static ENGINE_API bool ProjectWorldToScreen(const FVector& WorldPosition, const FIntRect& ViewRect, const FMatrix& ViewProjectionMatrix, FVector2D& out_ScreenPos);
+	static ENGINE_API bool ProjectWorldToScreen(const FVector& WorldPosition, const FIntRect& ViewRect, const FMatrix& ViewProjectionMatrix, FVector2D& out_ScreenPos, bool bShouldCalcOutsideViewPosition = false);
 
 	inline FVector GetViewRight() const { return ViewMatrices.GetViewMatrix().GetColumn(0); }
 	inline FVector GetViewUp() const { return ViewMatrices.GetViewMatrix().GetColumn(1); }
