Fix non-unity compilation error.

diff --git a/Engine/Source/Runtime/Engine/Private/Components/TextRenderComponent.cpp b/Engine/Source/Runtime/Engine/Private/Components/TextRenderComponent.cpp
--- a/Engine/Source/Runtime/Engine/Private/Components/TextRenderComponent.cpp
+++ b/Engine/Source/Runtime/Engine/Private/Components/TextRenderComponent.cpp
@@ -19,6 +19,7 @@
 #include "Engine/TextRenderActor.h"
 #include "Materials/MaterialInstanceDynamic.h"
 #include "StaticMeshResources.h"
+#include "RayTracingInstance.h"
 #include "Containers/Ticker.h"
 #include "PSOPrecache.h"
 #include "DataDrivenShaderPlatformInfo.h"
