Fix player base not clearing properly when falling off base at high altitude.
Also fix client authority for lift-off not working with bStayBasedInAir

diff --git a/Engine/Source/Runtime/Engine/Private/Components/CharacterMovementComponent.cpp b/Engine/Source/Runtime/Engine/Private/Components/CharacterMovementComponent.cpp
--- a/Engine/Source/Runtime/Engine/Private/Components/CharacterMovementComponent.cpp
+++ b/Engine/Source/Runtime/Engine/Private/Components/CharacterMovementComponent.cpp
@@ -2487,7 +2487,7 @@ void UCharacterMovementComponent::UpdateBasedMovement(float DeltaSeconds)
 		ComputeFloorDist(PawnLocation, StayBasedInAirHeight, StayBasedInAirHeight, OutFloorResult, CharacterOwner->GetCapsuleComponent()->GetScaledCapsuleRadius(), NULL);
 
 		UPrimitiveComponent* HitComponent = OutFloorResult.HitResult.Component.Get();
-		if (HitComponent && HitComponent->GetAttachmentRoot() != MovementBase->GetAttachmentRoot())
+		if (!HitComponent || HitComponent->GetAttachmentRoot() != MovementBase->GetAttachmentRoot())
 		{
 			// New or no base under the character
 			ApplyImpartedMovementBaseVelocity();
@@ -9846,7 +9846,13 @@ void UCharacterMovementComponent::ServerMoveHandleClientError(float ClientTimeSt
 			bCanTrustClientOnLanding = false;
 		}
 
-		if (bServerIsFalling && bLastServerIsWalking && !bTeleportedSinceLastUpdate)
+		// Check for lift-off, going from walking to falling.
+		// Note that if bStayBasedInAir is enabled we can't rely on the walking to falling transition, instead run logic on the first tick after clearing the MovementBase
+		const bool bCanLiftOffFromBase = bStayBasedInAir
+			? !MovementBase && LastServerMovementBase.Get() // If we keep the base while in air, consider lift-off if base gets set to null and we had a base last tick
+			: bLastServerIsWalking; // If walking last tick, we were can consider lift-off logic
+
+		if (bServerIsFalling && bCanLiftOffFromBase && !bTeleportedSinceLastUpdate)
 		{
 			float ClientForwardFactor = 1.f;
 			UPrimitiveComponent* LastServerMovementBasePtr = LastServerMovementBase.Get();
