- Converted a Warning to Log in DP code, it was not necessary to be a warning

diff --git a/Engine/Source/Runtime/Engine/Private/DeviceProfiles/DeviceProfileManager.cpp b/Engine/Source/Runtime/Engine/Private/DeviceProfiles/DeviceProfileManager.cpp
--- a/Engine/Source/Runtime/Engine/Private/DeviceProfiles/DeviceProfileManager.cpp
+++ b/Engine/Source/Runtime/Engine/Private/DeviceProfiles/DeviceProfileManager.cpp
@@ -437,7 +437,7 @@ void UDeviceProfileManager::SetDeviceProfileCVars(const FString& DeviceProfileNa
 		}
 		else
 		{
-			UE_LOG(LogDeviceProfileManager, Warning, TEXT("Creating unregistered Device Profile CVar: [[%s:%s]]"), *CVarKey, *Pair.Value);
+			UE_LOG(LogDeviceProfileManager, Log, TEXT("Creating unregistered Device Profile CVar: [[%s:%s]]"), *CVarKey, *Pair.Value);
 		}
 
 		// Cache any scalability related cvars so we can conveniently reapply them later as a way to reset the device defaults
