Fix crash in FAudioChunkCache when CacheSizeKB=0

diff --git a/Engine/Source/Runtime/Engine/Private/AudioStreamingCache.cpp b/Engine/Source/Runtime/Engine/Private/AudioStreamingCache.cpp
--- a/Engine/Source/Runtime/Engine/Private/AudioStreamingCache.cpp
+++ b/Engine/Source/Runtime/Engine/Private/AudioStreamingCache.cpp
@@ -748,6 +748,11 @@ FAudioChunkCache::FAudioChunkCache(uint32 InMaxChunkSize, uint32 NumChunks, uint
 		CachePool.Emplace(Index);
 	}
 	CacheOverflowCount.Set(0);
+
+	if (MemoryLimitBytes == 0)
+	{
+		UE_LOG(LogAudioStreamCaching, Warning, TEXT("Audio stream cache size is 0 bytes. Audio will not play. To resolve this issue, set \'CacheSizeKB\' to a non-zero value."));
+	}
 }
 
 FAudioChunkCache::~FAudioChunkCache()
@@ -1554,6 +1559,12 @@ FAudioChunkCache::FCacheElement* FAudioChunkCache::EvictLeastRecentChunk(bool bB
 {
 	FCacheElement* CacheElement = LeastRecentElement;
 
+	if (!CacheElement)
+	{
+		// This can happen if the MemoryLimitBytes is 0, prevting LeastRecentElement from being set to a valid element.
+		return nullptr;
+	}
+
 	// If the least recent chunk is evictable, evict it.
 	bool bIsChunkEvictable = CacheElement->CanEvictChunk();
 	bool bIsChunkLoadingButUnreferenced = (CacheElement->IsLoadInProgress() && !CacheElement->IsInUse());
