Changes to allow make cluster union kinematic based on a distance range from any local players.

This is needed to ensure the simulation in the client do not differs to much from the server when elements are added/removed from the Cluster Union due to network relevancy changes, causing visuals de-syncs.

When in kinematic mode, the Cluster Unions will update their positions based on replication.

Also:

- Fixed an issue where Kinematic Particles positions were not being updated during replication.
- Fixed a case where we can end up with Infinite InvM and InvI if we update the state of a particle with 0 mass to be dynamic. Cluster Union particles can have 0 mass when they are first created and empty.

#rnx

diff --git a/Engine/Source/Runtime/Engine/Classes/PhysicsEngine/ClusterUnionComponent.h b/Engine/Source/Runtime/Engine/Classes/PhysicsEngine/ClusterUnionComponent.h
--- a/Engine/Source/Runtime/Engine/Classes/PhysicsEngine/ClusterUnionComponent.h
+++ b/Engine/Source/Runtime/Engine/Classes/PhysicsEngine/ClusterUnionComponent.h
@@ -228,7 +228,7 @@ public:
 	ENGINE_API TArray<AActor*> GetActors();
 
 	UFUNCTION(BlueprintCallable, Category = "Cluster Union")
-	ENGINE_API void SetIsAnchored(bool bIsAnchored);
+	ENGINE_API virtual void SetIsAnchored(bool bIsAnchored);
 
 	ENGINE_API bool IsAnchored() const;
 
diff --git a/Engine/Source/Runtime/Engine/Private/PhysicsEngine/PhysicsReplication.cpp b/Engine/Source/Runtime/Engine/Private/PhysicsEngine/PhysicsReplication.cpp
--- a/Engine/Source/Runtime/Engine/Private/PhysicsEngine/PhysicsReplication.cpp
+++ b/Engine/Source/Runtime/Engine/Private/PhysicsEngine/PhysicsReplication.cpp
@@ -1411,14 +1411,24 @@ bool FPhysicsReplicationAsync::PredictiveInterpolation(Chaos::FPBDRigidParticleH
 
 	if (bHardSnap)
 	{
-		// Too much error so just snap state here and be done with it
 		Target.AccumulatedErrorSeconds = 0.0f;
-		Handle->SetX(Target.PrevPosTarget);
-		Handle->SetP(Target.PrevPosTarget);
-		Handle->SetR(Target.PrevRotTarget);
-		Handle->SetQ(Target.PrevRotTarget);
-		Handle->SetV(Target.TargetState.LinVel);
-		Handle->SetW(FMath::DegreesToRadians(Target.TargetState.AngVel));
+
+		if (Handle->IsKinematic())
+		{
+			// Set a FKinematicTarget to hard snap kinematic object
+			const Chaos::FKinematicTarget KinTarget = Chaos::FKinematicTarget::MakePositionTarget(Target.PrevPosTarget, Target.PrevRotTarget); // Uses EKinematicTargetMode::Position
+			RigidsSolver->GetEvolution()->SetParticleKinematicTarget(Handle, KinTarget);
+		}
+		else 
+		{
+			// Set XPRQVW to hard snap dynamic object
+			Handle->SetX(Target.PrevPosTarget);
+			Handle->SetP(Target.PrevPosTarget);
+			Handle->SetR(Target.PrevRotTarget);
+			Handle->SetQ(Target.PrevRotTarget);
+			Handle->SetV(Target.TargetState.LinVel);
+			Handle->SetW(FMath::DegreesToRadians(Target.TargetState.AngVel));
+		}
 
 		// Cache data for next replication
 		Target.PrevLinVel = FVector(Target.TargetState.LinVel);
diff --git a/Engine/Source/Runtime/Experimental/Chaos/Public/Chaos/PBDRigidParticles.h b/Engine/Source/Runtime/Experimental/Chaos/Public/Chaos/PBDRigidParticles.h
--- a/Engine/Source/Runtime/Experimental/Chaos/Public/Chaos/PBDRigidParticles.h
+++ b/Engine/Source/Runtime/Experimental/Chaos/Public/Chaos/PBDRigidParticles.h
@@ -156,13 +156,10 @@ class TPBDRigidParticles : public TRigidParticles<T, d>
 		else if ((CurrentState == EObjectStateType::Kinematic || CurrentState == EObjectStateType::Static || CurrentState == EObjectStateType::Uninitialized) && (InObjectState == EObjectStateType::Dynamic || InObjectState == EObjectStateType::Sleeping))
 		{
 			// Transitioning from kinematic or static to dynamic, compute the inverses.
-			checkSlow(this->M(Index) != 0.0);
-			checkSlow(this->I(Index)[0] != 0.0);
-			checkSlow(this->I(Index)[1] != 0.0);
-			checkSlow(this->I(Index)[2] != 0.0);
-			this->InvM(Index) = 1.f / this->M(Index);
-			this->InvI(Index) = TVec3<FRealSingle>(
-				1.f / this->I(Index)[0], 
+			this->InvM(Index) = FMath::IsNearlyZero(this->M(Index)) ? 0.0f : 1.f / this->M(Index);
+			this->InvI(Index) = this->I(Index).IsNearlyZero() ? TVec3<FRealSingle>::ZeroVector :
+				TVec3<FRealSingle>(
+				1.f / this->I(Index)[0],
 				1.f / this->I(Index)[1],
 				1.f / this->I(Index)[2]);
 
