AggregateGeom (more accessors);

Adding ways to get a shape by the name and get the shape index by name.

-- Code --
AggregateGeom:
GetElementByName() - Goes through shape arrays and returns the first found element that has the passed in name.
GetElementIndexByName() - Goes through the shape arrays and returns the first found index for the element that has the passed in name.
GetElementByName() - Templated version to reduce shared code.
GetElementIndexByName() - Templated version to reduce shared code.

review [REDACTED_HANDLE], [REDACTED_HANDLE], [REDACTED_HANDLE]

diff --git a/Engine/Source/Runtime/Engine/Classes/PhysicsEngine/AggregateGeom.h b/Engine/Source/Runtime/Engine/Classes/PhysicsEngine/AggregateGeom.h
--- a/Engine/Source/Runtime/Engine/Classes/PhysicsEngine/AggregateGeom.h
+++ b/Engine/Source/Runtime/Engine/Classes/PhysicsEngine/AggregateGeom.h
@@ -135,6 +135,87 @@ struct FKAggregateGeom
 		return nullptr;
 	}
 
+	const FKShapeElem* GetElementByName(const FName InName) const
+	{
+		if (const FKShapeElem* FoundSphereElem = GetElementByName<FKSphereElem>(MakeArrayView(SphereElems), InName))
+		{
+			return FoundSphereElem;
+		}
+		else if (const FKShapeElem* FoundBoxElem = GetElementByName<FKBoxElem>(MakeArrayView(BoxElems), InName))
+		{
+			return FoundBoxElem;
+		}
+		else if (const FKShapeElem* FoundSphylElem = GetElementByName<FKSphylElem>(MakeArrayView(SphylElems), InName))
+		{
+			return FoundSphylElem;
+		}
+		else if (const FKShapeElem* FoundConvexElem = GetElementByName<FKConvexElem>(MakeArrayView(ConvexElems), InName))
+		{
+			return FoundConvexElem;
+		}
+		else if (const FKShapeElem* FoundTaperedCapsuleElem = GetElementByName<FKTaperedCapsuleElem>(MakeArrayView(TaperedCapsuleElems), InName))
+		{
+			return FoundTaperedCapsuleElem;
+		}
+		else if (const FKShapeElem* FoundLevelSetElem = GetElementByName<FKLevelSetElem>(MakeArrayView(LevelSetElems), InName))
+		{
+			return FoundLevelSetElem;
+		}
+		else if (const FKShapeElem* FoundSkinnedLevelSetElem = GetElementByName<FKSkinnedLevelSetElem>(MakeArrayView(SkinnedLevelSetElems), InName))
+		{
+			return FoundSkinnedLevelSetElem;
+		}
+
+		return nullptr;
+	}
+
+	int32 GetElementIndexByName(const FName InName) const
+	{
+		int32 FoundIndex = GetElementIndexByName<FKSphereElem>(MakeArrayView(SphereElems), InName);
+		if (FoundIndex != INDEX_NONE)
+		{
+			return FoundIndex;
+		}
+
+		FoundIndex = GetElementIndexByName<FKBoxElem>(MakeArrayView(BoxElems), InName);
+		if (FoundIndex != INDEX_NONE)
+		{
+			return FoundIndex;
+		}
+
+		FoundIndex = GetElementIndexByName<FKSphylElem>(MakeArrayView(SphylElems), InName);
+		if (FoundIndex != INDEX_NONE)
+		{
+			return FoundIndex;
+		}
+
+		FoundIndex = GetElementIndexByName<FKConvexElem>(MakeArrayView(ConvexElems), InName);
+		if (FoundIndex != INDEX_NONE)
+		{
+			return FoundIndex;
+		}
+
+		FoundIndex = GetElementIndexByName<FKTaperedCapsuleElem>(MakeArrayView(TaperedCapsuleElems), InName);
+		if (FoundIndex != INDEX_NONE)
+		{
+			return FoundIndex;
+		}
+
+		FoundIndex = GetElementIndexByName<FKLevelSetElem>(MakeArrayView(LevelSetElems), InName);
+		if (FoundIndex != INDEX_NONE)
+		{
+			return FoundIndex;
+		}
+
+		FoundIndex = GetElementIndexByName<FKSkinnedLevelSetElem>(MakeArrayView(SkinnedLevelSetElems), InName);
+		if (FoundIndex != INDEX_NONE)
+		{
+			return FoundIndex;
+		}
+
+		return INDEX_NONE;
+	}
+
 	void EmptyElements()
 	{
 		BoxElems.Empty();
@@ -190,4 +271,26 @@ private:
 		LevelSetElems = Other.LevelSetElems;
 		SkinnedLevelSetElems = Other.SkinnedLevelSetElems;
 	}
+
+	template <class T>
+	const FKShapeElem* GetElementByName(TArrayView<const T> Elements, const FName InName) const
+	{
+		const FKShapeElem* FoundElem = Elements.FindByPredicate(
+			[InName](const T& Elem)
+			{
+				return InName == Elem.GetName();
+			});
+		return FoundElem;
+	}
+
+	template <class T>
+	int32 GetElementIndexByName(TArrayView<const T> Elements, const FName InName) const
+	{
+		int32 FoundIndex = Elements.IndexOfByPredicate(
+			[InName](const T& Elem)
+			{
+				return InName == Elem.GetName();
+			});
+		return FoundIndex;
+	}
 };
