Fix simulated platforms not simulating when spawning via persistence.
1. Since gameplay code (ie PhysicsObjectComponent) may have set the object to net dormant by the time that the server has set a CU to be unanchored, OnRep_RigidState callback may never occur on client. Add FlushNetDormancy to SyncClusterUnionFromProxy on server. Controlled with CVar bFlushNetDormancyOnSyncProxy
2. Since sometimes OnRep_RigidState might occur before Proxy exists on client, set rigid state properties in OnCreatePhysicsState on clients. Controlled with CVar bApplyReplicatedRigidStateOnCreatePhysicsState

#tests repro'd in cooked build, tested fix by replacing server bin with rebuild bin, unable to repro with new code.

#localization none

diff --git a/Engine/Source/Runtime/Engine/Classes/PhysicsEngine/ClusterUnionComponent.h b/Engine/Source/Runtime/Engine/Classes/PhysicsEngine/ClusterUnionComponent.h
--- a/Engine/Source/Runtime/Engine/Classes/PhysicsEngine/ClusterUnionComponent.h
+++ b/Engine/Source/Runtime/Engine/Classes/PhysicsEngine/ClusterUnionComponent.h
@@ -128,6 +128,10 @@ struct FClusterUnionReplicatedData
 
 	UPROPERTY()
 	bool bIsAnchored = false;
+
+	FClusterUnionReplicatedData& operator=(const FClusterUnionReplicatedData&) = default;
+	bool operator==(const FClusterUnionReplicatedData&) const = default;
+	bool operator!=(const FClusterUnionReplicatedData&) const = default;
 };
 
 USTRUCT()
diff --git a/Engine/Source/Runtime/Engine/Private/PhysicsEngine/ClusterUnionComponent.cpp b/Engine/Source/Runtime/Engine/Private/PhysicsEngine/ClusterUnionComponent.cpp
--- a/Engine/Source/Runtime/Engine/Private/PhysicsEngine/ClusterUnionComponent.cpp
+++ b/Engine/Source/Runtime/Engine/Private/PhysicsEngine/ClusterUnionComponent.cpp
@@ -34,6 +34,16 @@ namespace
 	float LocalBoneDataMapGrowFactor = 1.2f;
 	FAutoConsoleVariableRef CVarLocalBoneDataMapGrowFactor(TEXT("ClusterUnion.LocalBoneDataMapGrowFactor"), LocalBoneDataMapGrowFactor, TEXT("Grow factor to apply to the size of bone data array of pre-existing component when preallocating the local bones data map"));
 
+	bool bApplyReplicatedRigidStateOnCreatePhysicsState = true;
+	FAutoConsoleVariableRef CVarApplyReplicatedRigidStateOnCreatePhysicsState(TEXT("ClusterUnion.ApplyReplicatedRigidStateOnCreatePhysicsState"), bApplyReplicatedRigidStateOnCreatePhysicsState, TEXT("When physics state is created, apply replicated rigid state. Useful because sometimes the initial OnRep will have been called before a proxy exists, so initial properties will be unset"));
+
+	bool bDirtyRigidStateOnlyIfChanged = false;
+	FAutoConsoleVariableRef CVarDirtyRigidStateOnlyIfChanged(TEXT("ClusterUnion.DirtyRigidStateOnlyIfChanged"), bDirtyRigidStateOnlyIfChanged, TEXT("Add a check for changed rigid state before marking it dirty and updating the replicated data. No need to flush an update if there was no change."));
+
+	bool bFlushNetDormancyOnSyncProxy = true;
+	FAutoConsoleVariableRef CVarFlushNetDormancyOnSyncProxy(TEXT("ClusterUnion.FlushNetDormancyOnSyncProxy"), bFlushNetDormancyOnSyncProxy, TEXT("When there is a new rigid state on the authority, flush net dormancy so that even if this object is net dorman the rigid state will come through to the client."));
+
+
 	template<typename PayloadType>
 	struct TClusterUnionAABBTreeStorageTraits
 	{
@@ -725,6 +735,13 @@ void UClusterUnionComponent::OnCreatePhysicsState()
 		// if it is not set and the callback does not get called then the cluster union component will be moved to the origin of the world
 		const FTransform Transform = GetComponentTransform();
 		PhysicsProxy->SetXR_External(Transform.GetLocation(), Transform.GetRotation());
+
+		// Since the initial OnRep_RigidState might've occurred before we actually had a physics state,
+		// do OnRep actions again now to make sure we start off with the right data.
+		if (bApplyReplicatedRigidStateOnCreatePhysicsState)
+		{
+			OnRep_RigidState();
+		}
 	}
 }
 
@@ -886,9 +903,29 @@ void UClusterUnionComponent::SyncClusterUnionFromProxy()
 
 	if (IsAuthority())
 	{
-		ReplicatedRigidState.bIsAnchored = PhysicsProxy->IsAnchored_External();
-		ReplicatedRigidState.ObjectState = static_cast<uint8>(PhysicsProxy->GetObjectState_External());
-		MARK_PROPERTY_DIRTY_FROM_NAME(UClusterUnionComponent, ReplicatedRigidState, this);
+		// Create the new replicated state
+		const FClusterUnionReplicatedData NewReplicatedRigidState =
+		{
+			static_cast<uint8>(PhysicsProxy->GetObjectState_External()),
+			PhysicsProxy->IsAnchored_External()
+		};
+
+		// Only dirty the state if it has changed
+		if (ReplicatedRigidState != NewReplicatedRigidState || bDirtyRigidStateOnlyIfChanged == false)
+		{
+			// Make sure that the new dirty data gets flushed through to clients even if the actor
+			// has been made net dormant.
+			if (bFlushNetDormancyOnSyncProxy)
+			{
+				if (AActor* Owner = GetOwner())
+				{
+					Owner->FlushNetDormancy();
+				}
+			}
+
+			ReplicatedRigidState = NewReplicatedRigidState;
+			MARK_PROPERTY_DIRTY_FROM_NAME(UClusterUnionComponent, ReplicatedRigidState, this);
+		}
 	}
 	
 	const Chaos::FClusterUnionSyncedData& FullData = PhysicsProxy->GetSyncedData_External();
