UEJpegComp: Fixed access out of bounds

diff --git a/Engine/Source/Runtime/Engine/Private/SoundWave.cpp b/Engine/Source/Runtime/Engine/Private/SoundWave.cpp
--- a/Engine/Source/Runtime/Engine/Private/SoundWave.cpp
+++ b/Engine/Source/Runtime/Engine/Private/SoundWave.cpp
@@ -4864,7 +4864,7 @@ TFuture<FSharedBuffer> USoundWave::FEditorAudioBulkData::GetPayload() const
 		TArray<int16> scratch_buffer;
 		for (int i = 0; i < SoundWave->ChannelOffsets.Num(); ++i)
 		{
-			if (!WaveInfo.ReadWaveInfo(Data + SoundWave->ChannelOffsets[i], DataSize))
+			if (!WaveInfo.ReadWaveInfo(Data + SoundWave->ChannelOffsets[i], SoundWave->ChannelSizes[i]))
 			{
 				UE_LOG(LogAudio, Warning, TEXT("Failed to read wave data out of '%s'."), *SoundWave->GetFullName());
 			}
@@ -4919,7 +4919,7 @@ void USoundWave::FEditorAudioBulkData::UpdatePayload(FSharedBuffer InPayload, UO
 		for (int i = 0; i < SoundWave->ChannelOffsets.Num(); ++i)
 		{
 			FWaveModInfo WaveInfo;
-			if (!WaveInfo.ReadWaveInfo(Data + SoundWave->ChannelOffsets[i], DataSize))
+			if (!WaveInfo.ReadWaveInfo(Data + SoundWave->ChannelOffsets[i], SoundWave->ChannelSizes[i]))
 			{
 				UE_LOG(LogAudio, Warning, TEXT("Failed to read wave data out of '%s'."), *SoundWave->GetFullName());
 			}
