Disable PVS warning V758
Silence V570 false positives for bit field assignments
Silence various other PVS warnings
#rnx

diff --git a/Engine/Plugins/Experimental/Mutable/Source/MutableTools/Private/MuT/CodeGenerator.cpp b/Engine/Plugins/Experimental/Mutable/Source/MutableTools/Private/MuT/CodeGenerator.cpp
--- a/Engine/Plugins/Experimental/Mutable/Source/MutableTools/Private/MuT/CodeGenerator.cpp
+++ b/Engine/Plugins/Experimental/Mutable/Source/MutableTools/Private/MuT/CodeGenerator.cpp
@@ -1369,8 +1369,8 @@ namespace mu
 						lastSurfOp = op;
 					}
 
-					else if (ImageLayoutStrategy == CompilerOptions::TextureLayoutStrategy::Pack)
-					{
+					else if (ImageLayoutStrategy == CompilerOptions::TextureLayoutStrategy::Pack) //-V547
+					{						
 						if (LayoutIndex >= meshResults.GeneratedLayouts.Num() ||
 							LayoutIndex >= meshResults.layoutOps.Num())
 						{
diff --git a/Engine/Plugins/FX/Niagara/Source/Niagara/Public/NiagaraMeshRendererProperties.h b/Engine/Plugins/FX/Niagara/Source/Niagara/Public/NiagaraMeshRendererProperties.h
--- a/Engine/Plugins/FX/Niagara/Source/Niagara/Public/NiagaraMeshRendererProperties.h
+++ b/Engine/Plugins/FX/Niagara/Source/Niagara/Public/NiagaraMeshRendererProperties.h
@@ -365,7 +365,7 @@ public:
 	The component controls if shadows are enabled, this flag allows you to disable the renderer casting shadows.
 	*/
 	UPROPERTY(EditAnywhere, Category = "Rendering")
-	uint8 bCastShadows : 1 = 1;
+	uint8 bCastShadows : 1 = 1; //-V570
 
 	/** The materials to be used instead of the StaticMesh's materials. Note that each material must have the Niagara Mesh Particles flag checked. If the ParticleMesh
 	requires more materials than exist in this array or any entry in this array is set to None, we will use the ParticleMesh's existing Material instead.*/
diff --git a/Engine/Plugins/FX/Niagara/Source/Niagara/Public/NiagaraRibbonRendererProperties.h b/Engine/Plugins/FX/Niagara/Source/Niagara/Public/NiagaraRibbonRendererProperties.h
--- a/Engine/Plugins/FX/Niagara/Source/Niagara/Public/NiagaraRibbonRendererProperties.h
+++ b/Engine/Plugins/FX/Niagara/Source/Niagara/Public/NiagaraRibbonRendererProperties.h
@@ -336,7 +336,7 @@ public:
 	The component controls if shadows are enabled, this flag allows you to disable the renderer casting shadows.
 	*/
 	UPROPERTY(EditAnywhere, Category = "Rendering")
-	uint8 bCastShadows : 1 = 1;
+	uint8 bCastShadows : 1 = 1; //-V570
 
 	/** Tessellation factor to apply to the width of the ribbon.
 	* Ranges from 1 to 16. Greater values increase amount of tessellation.
diff --git a/Engine/Plugins/FX/Niagara/Source/Niagara/Public/NiagaraSpriteRendererProperties.h b/Engine/Plugins/FX/Niagara/Source/Niagara/Public/NiagaraSpriteRendererProperties.h
--- a/Engine/Plugins/FX/Niagara/Source/Niagara/Public/NiagaraSpriteRendererProperties.h
+++ b/Engine/Plugins/FX/Niagara/Source/Niagara/Public/NiagaraSpriteRendererProperties.h
@@ -219,7 +219,7 @@ public:
 	The component controls if shadows are enabled, this flag allows you to disable the renderer casting shadows.
 	*/
 	UPROPERTY(EditAnywhere, Category = "Rendering")
-	uint8 bCastShadows : 1 = 1;
+	uint8 bCastShadows : 1 = 1; //-V570
 
 	/** Sort precision to use when sorting is active. */
 	UPROPERTY(EditAnywhere, Category = "Sorting", meta = (EditCondition = "SortMode != ENiagaraSortMode::None", EditConditionHides))
diff --git a/Engine/Plugins/MovieScene/MovieRenderPipeline/Source/MovieRenderPipelineCore/Public/Graph/Nodes/MovieGraphImagePassBaseNode.h b/Engine/Plugins/MovieScene/MovieRenderPipeline/Source/MovieRenderPipelineCore/Public/Graph/Nodes/MovieGraphImagePassBaseNode.h
--- a/Engine/Plugins/MovieScene/MovieRenderPipeline/Source/MovieRenderPipelineCore/Public/Graph/Nodes/MovieGraphImagePassBaseNode.h
+++ b/Engine/Plugins/MovieScene/MovieRenderPipeline/Source/MovieRenderPipelineCore/Public/Graph/Nodes/MovieGraphImagePassBaseNode.h
@@ -56,7 +56,7 @@ protected:
 	// Note: Since *individual* show flags are overridden instead of the entire ShowFlags property, manually set to
 	// overridden so the traversal picks the changes up (otherwise they will be ignored).
 	UPROPERTY(EditAnywhere, BlueprintReadWrite, Category = Overrides, meta = (InlineEditConditionToggle))
-	uint8 bOverride_ShowFlags : 1 = 1;
+	uint8 bOverride_ShowFlags : 1 = 1; //-V570
 
 	/** The show flags that should be active during a render for this node. */
 	UPROPERTY(EditAnywhere, BlueprintReadWrite, Category = "Show Flags")
diff --git a/Engine/Plugins/MovieScene/MovieRenderPipeline/Source/MovieRenderPipelineCore/Public/Graph/Nodes/MovieGraphModifierNode.h b/Engine/Plugins/MovieScene/MovieRenderPipeline/Source/MovieRenderPipelineCore/Public/Graph/Nodes/MovieGraphModifierNode.h
--- a/Engine/Plugins/MovieScene/MovieRenderPipeline/Source/MovieRenderPipelineCore/Public/Graph/Nodes/MovieGraphModifierNode.h
+++ b/Engine/Plugins/MovieScene/MovieRenderPipeline/Source/MovieRenderPipelineCore/Public/Graph/Nodes/MovieGraphModifierNode.h
@@ -94,10 +94,10 @@ public:
 
 private:
 	UPROPERTY()
-	uint8 bOverride_Collections : 1 = 1;
+	uint8 bOverride_Collections : 1 = 1; //-V570
 
 	UPROPERTY()
-	uint8 bOverride_ModifiersContainer : 1 = 1;
+	uint8 bOverride_ModifiersContainer : 1 = 1; //-V570
 	
 	/** The names of collections being modified. */
 	UPROPERTY()
diff --git a/Engine/Plugins/Runtime/GeometryProcessing/Source/DynamicMesh/Private/ShapeApproximation/ShapeDetection3.cpp b/Engine/Plugins/Runtime/GeometryProcessing/Source/DynamicMesh/Private/ShapeApproximation/ShapeDetection3.cpp
--- a/Engine/Plugins/Runtime/GeometryProcessing/Source/DynamicMesh/Private/ShapeApproximation/ShapeDetection3.cpp
+++ b/Engine/Plugins/Runtime/GeometryProcessing/Source/DynamicMesh/Private/ShapeApproximation/ShapeDetection3.cpp
@@ -151,7 +151,7 @@ bool UE::Geometry::IsBoxMesh(const FDynamicMesh3& Mesh, FOrientedBox3d& BoxOut,
 			if (Normals[Idx].Dot(Normal) > ParallelDotTolerance)
 			{
 				double Distance = Pos.Dot(Normal);
-				if (FMath::IsNearlyEqual(Distance, Distances[Idx], PlaneDistanceTolerance))
+				if (FMath::IsNearlyEqual(Distance, Distances[Idx], PlaneDistanceTolerance)) //-V614
 				{
 					return Idx;
 				}
diff --git a/Engine/Source/Developer/ShaderCompilerCommon/Private/ShaderCompilerCommon.cpp b/Engine/Source/Developer/ShaderCompilerCommon/Private/ShaderCompilerCommon.cpp
--- a/Engine/Source/Developer/ShaderCompilerCommon/Private/ShaderCompilerCommon.cpp
+++ b/Engine/Source/Developer/ShaderCompilerCommon/Private/ShaderCompilerCommon.cpp
@@ -1446,7 +1446,7 @@ static bool SearchText(const TCHAR*& InOut, FAssertParsingStack& OutAssertStack)
 		// Check if the current assert is valid, i.e., containt a TEXT() argument
 		const TCHAR* EndPtr = nullptr;
 		const FString Tmp = ParseText(AssertPtr, EndPtr);
-		const bool bIsValid = FCString::Strstr(&Tmp[0], TextIdentifier) != nullptr;
+		const bool bIsValid = FCString::Strstr(&Tmp[0], TextIdentifier) != nullptr; //-V547
 		if (bIsValid)
 		{
 			OutAssertStack.BeginPtr = AssertPtr;
diff --git a/Engine/Source/Developer/ShaderPreprocessor/Private/stb_preprocess/preprocessor.c b/Engine/Source/Developer/ShaderPreprocessor/Private/stb_preprocess/preprocessor.c
--- a/Engine/Source/Developer/ShaderPreprocessor/Private/stb_preprocess/preprocessor.c
+++ b/Engine/Source/Developer/ShaderPreprocessor/Private/stb_preprocess/preprocessor.c
@@ -45,6 +45,9 @@
 // not obvious how to get it to sort itself out, so just disabling the warning for expediency
 #pragma warning(disable:6011) 
 
+// Silence completely incorrect misparsing of the arrsetlen macros
+//-V::521
+
 //////////////////////////////////////////////////////////////////////////////
 //
 //  Limits to prevent infinite recursion
diff --git a/Engine/Source/Runtime/AudioCaptureImplementations/Windows/AudioCaputureWasapi/Private/WasapiAudioFormat.cpp b/Engine/Source/Runtime/AudioCaptureImplementations/Windows/AudioCaputureWasapi/Private/WasapiAudioFormat.cpp
--- a/Engine/Source/Runtime/AudioCaptureImplementations/Windows/AudioCaputureWasapi/Private/WasapiAudioFormat.cpp
+++ b/Engine/Source/Runtime/AudioCaptureImplementations/Windows/AudioCaputureWasapi/Private/WasapiAudioFormat.cpp
@@ -25,7 +25,7 @@ namespace Audio
 	{
 		if ((InFormat->wFormatTag == WAVE_FORMAT_EXTENSIBLE) && (InFormat->cbSize >= WaveFormatExtensibleCbsize))
 		{
-			FMemory::Memcpy(&WaveFormat, InFormat, sizeof(WAVEFORMATEXTENSIBLE));
+			FMemory::Memcpy(&WaveFormat, InFormat, sizeof(WAVEFORMATEXTENSIBLE)); //-V512
 		}
 		else
 		{
diff --git a/Engine/Source/Runtime/Core/Public/Windows/WindowsPlatformCodeAnalysis.h b/Engine/Source/Runtime/Core/Public/Windows/WindowsPlatformCodeAnalysis.h
--- a/Engine/Source/Runtime/Core/Public/Windows/WindowsPlatformCodeAnalysis.h
+++ b/Engine/Source/Runtime/Core/Public/Windows/WindowsPlatformCodeAnalysis.h
@@ -103,7 +103,7 @@
 	#pragma warning(disable : 6240) // warning C6240 : (<expression> && <non-zero constant>) always evaluates to the result of <expression>. Did you intend to use the bitwise-and operator?
 
 	//PVS-Studio settings:
-	//-V::505,542,581,601,623,668,677,690,704,719,720,730,735,751,1002,1008,1062
+	//-V::505,542,581,601,623,668,677,690,704,719,720,730,735,751,758,1002,1008,1062
 	//-V:TRYCOMPRESSION:519,547
 	//-V:check(:501,547,560,605
 	//-V:checkf(:510
@@ -229,6 +229,8 @@
 	//-V:Linker:678
 	//-V:self:678
 	//-V:AccumulateParentID:678
+	//-V:BasicLog:510
+	//-V:BasicFatalLog:510
 
 	// The following classes retain a reference to data supplied in the constructor by the derived class which can not yet be initialized.
 	//-V:FMemoryWriter(:1050
diff --git a/Engine/Source/Runtime/CoreUObject/Private/UObject/GarbageCollection.cpp b/Engine/Source/Runtime/CoreUObject/Private/UObject/GarbageCollection.cpp
--- a/Engine/Source/Runtime/CoreUObject/Private/UObject/GarbageCollection.cpp
+++ b/Engine/Source/Runtime/CoreUObject/Private/UObject/GarbageCollection.cpp
@@ -3782,7 +3782,7 @@ class FRealtimeGC : public FGarbageCollectionTracer
 
 #if !UE_BUILD_SHIPPING
 		TDebugReachabilityProcessor<Options> DebugProcessor;
-		if (DebugProcessor.IsForceEnabled() |
+		if (DebugProcessor.IsForceEnabled() | //-V792
 			DebugProcessor.TracksHistory() | 
 			DebugProcessor.TracksGarbage() & Stats.bFoundGarbageRef)
 		{
diff --git a/Engine/Source/Runtime/Engine/Private/Materials/MaterialHLSLGenerator.cpp b/Engine/Source/Runtime/Engine/Private/Materials/MaterialHLSLGenerator.cpp
--- a/Engine/Source/Runtime/Engine/Private/Materials/MaterialHLSLGenerator.cpp
+++ b/Engine/Source/Runtime/Engine/Private/Materials/MaterialHLSLGenerator.cpp
@@ -531,11 +531,11 @@ const UE::HLSLTree::FExpression* FMaterialHLSLGenerator::AcquireExpression(UE::H
 				const FConnectedInput& ConnectedInput = FunctionCall->ConnectedInputs[FunctionInputIndex];
 				if (ConnectedInput.Input)
 				{
-					Hasher.AppendData(&ConnectedInput.Input, sizeof(ConnectedInput.Input));
+					Hasher.AppendData(&ConnectedInput.Input, sizeof(ConnectedInput.Input)); //-V568
 				}
 				else
 				{
-					Hasher.AppendData(&ConnectedInput.Expression, sizeof(ConnectedInput.Expression));
+					Hasher.AppendData(&ConnectedInput.Expression, sizeof(ConnectedInput.Expression)); //-V568
 				}
 			}
 		}
@@ -777,11 +777,11 @@ const UE::HLSLTree::FExpression* FMaterialHLSLGenerator::GenerateFunctionCall(UE
 			const FConnectedInput& ConnectedInput = ConnectedInputs[InputIndex];
 			if (ConnectedInput.Input)
 			{
-				Hasher.Update(&ConnectedInput.Input, sizeof(ConnectedInput.Input));
+				Hasher.Update(&ConnectedInput.Input, sizeof(ConnectedInput.Input)); //-V568
 			}
 			else
 			{
-				Hasher.Update(&ConnectedInput.Expression, sizeof(ConnectedInput.Expression));
+				Hasher.Update(&ConnectedInput.Expression, sizeof(ConnectedInput.Expression)); //-V568
 			}
 		}
 
diff --git a/Engine/Source/Runtime/MovieScene/Public/EntitySystem/MovieSceneEntityGroupingSystem.h b/Engine/Source/Runtime/MovieScene/Public/EntitySystem/MovieSceneEntityGroupingSystem.h
--- a/Engine/Source/Runtime/MovieScene/Public/EntitySystem/MovieSceneEntityGroupingSystem.h
+++ b/Engine/Source/Runtime/MovieScene/Public/EntitySystem/MovieSceneEntityGroupingSystem.h
@@ -190,7 +190,7 @@ struct TEntityGroupingHandlerImpl<GroupingPolicy, TIntegerSequence<int, Componen
 			VisitUnlinkedEntities(Item, EntityIDs, GroupIDs, Builder);
 		}
 
-		PostAllocationCallbackImpl(&Policy, Item);
+		PostAllocationCallbackImpl(&Policy, Item); //-V510
 	}
 	
 	void VisitLinkedEntities(FEntityAllocationIteratorItem Item, FReadEntityIDs EntityIDs, TWrite<FEntityGroupID> GroupIDs, FEntityGroupBuilder* Builder)
diff --git a/Engine/Source/Runtime/RenderCore/Private/RenderingThread.cpp b/Engine/Source/Runtime/RenderCore/Private/RenderingThread.cpp
--- a/Engine/Source/Runtime/RenderCore/Private/RenderingThread.cpp
+++ b/Engine/Source/Runtime/RenderCore/Private/RenderingThread.cpp
@@ -1270,7 +1270,7 @@ void FRenderCommandFence::BeginFence(bool bSyncToRHIAndGPU)
 
 		CompletionTask = MoveTemp(CompletionTaskEvent);
 	}
-}
+} //-V773
 
 bool FRenderCommandFence::IsFenceComplete() const
 {
diff --git a/Engine/Source/Runtime/Renderer/Private/InstanceCulling/InstanceCullingContext.cpp b/Engine/Source/Runtime/Renderer/Private/InstanceCulling/InstanceCullingContext.cpp
--- a/Engine/Source/Runtime/Renderer/Private/InstanceCulling/InstanceCullingContext.cpp
+++ b/Engine/Source/Runtime/Renderer/Private/InstanceCulling/InstanceCullingContext.cpp
@@ -832,9 +832,8 @@ void FInstanceCullingContext::BuildRenderingCommandsInternal(
 
 	PassParametersTmp.ViewIds = GraphBuilder.CreateSRV(ViewIdsBuffer);
 	PassParametersTmp.NumCullingViews = 0;
-	if (bCullInstances || bAllowWPODisable)
+	if (bCullInstances || bAllowWPODisable) //-V595
 	{
-		checkSlow(InstanceCullingManager);
 #if DO_CHECK
 		for (int32 ViewId : ViewIds)
 		{
diff --git a/Engine/Source/Runtime/Renderer/Private/PostProcess/TemporalSuperResolution.cpp b/Engine/Source/Runtime/Renderer/Private/PostProcess/TemporalSuperResolution.cpp
--- a/Engine/Source/Runtime/Renderer/Private/PostProcess/TemporalSuperResolution.cpp
+++ b/Engine/Source/Runtime/Renderer/Private/PostProcess/TemporalSuperResolution.cpp
@@ -2514,7 +2514,7 @@ FDefaultTemporalUpscaler::FOutputs AddTemporalSuperResolutionPasses(
 		TShaderMapRef<FTSRResolveHistoryCS> ComputeShader(View.ShaderMap, PermutationVector);
 		FComputeShaderUtils::AddPass(
 			GraphBuilder,
-			RDG_EVENT_NAME("TSR ResolveHistory(#%d WaveSize=%d%s%s) %dx%d",
+			RDG_EVENT_NAME("TSR ResolveHistory(#%d WaveSize=%d%s%s) %dx%d", //-V510
 				PermutationVector.ToDimensionValueId(),
 				PermutationVector.Get<FTSRResolveHistoryCS::FNyquistDim>(),
 				PermutationVector.Get<FTSRShader::F16BitVALUDim>() ? TEXT(" 16bit") : TEXT(""),
