[NetPhysics] Physics replication to make sure the object goes to sleep and is asleep before stopping the replication so that something else doesn't instantly put the object else an attached simulated object can keep the object awake and make it desync from the server.

diff --git a/Engine/Source/Runtime/Engine/Private/PhysicsEngine/PhysicsReplication.cpp b/Engine/Source/Runtime/Engine/Private/PhysicsEngine/PhysicsReplication.cpp
--- a/Engine/Source/Runtime/Engine/Private/PhysicsEngine/PhysicsReplication.cpp
+++ b/Engine/Source/Runtime/Engine/Private/PhysicsEngine/PhysicsReplication.cpp
@@ -1307,10 +1307,11 @@ bool FPhysicsReplicationAsync::PredictiveInterpolation(Chaos::FPBDRigidParticleH
 		return false;
 	}
 
-	const bool bCanSimulate = Handle->IsDynamic() || Handle->IsSleeping();
+	const bool bIsSleeping = Handle->IsSleeping();
+	const bool bCanSimulate = Handle->IsDynamic() || bIsSleeping;
 
 	// Helper for sleep and target clearing at replication end
-	auto EndReplicationHelper = [RigidsSolver, Handle, bCanSimulate](FReplicatedPhysicsTargetAsync& Target, bool bOkToClear) -> bool
+	auto EndReplicationHelper = [RigidsSolver, Handle, bCanSimulate, bIsSleeping](FReplicatedPhysicsTargetAsync& Target, bool bOkToClear) -> bool
 	{
 		const bool bShouldSleep = (Target.TargetState.Flags & ERigidBodyFlags::Sleeping) != 0;
 		const bool bReplicatingPhysics = (Target.TargetState.Flags & ERigidBodyFlags::RepPhysics) != 0;
@@ -1320,7 +1321,12 @@ bool FPhysicsReplicationAsync::PredictiveInterpolation(Chaos::FPBDRigidParticleH
 			RigidsSolver->GetEvolution()->SetParticleObjectState(Handle, Chaos::EObjectStateType::Sleeping);
 		}
 
-		const bool bClearTarget = (!bCanSimulate || (bOkToClear && bShouldSleep) || (bOkToClear && !bReplicatingPhysics)) && !PhysicsReplicationCVars::PredictiveInterpolationCVars::bDontClearTarget;
+		const bool bClearTarget =
+			(!bCanSimulate 
+			|| (bOkToClear && bIsSleeping && bShouldSleep) // Don't clear the target due to sleeping until the object both should sleep and is sleeping
+			|| (bOkToClear && !bReplicatingPhysics))
+			&& !PhysicsReplicationCVars::PredictiveInterpolationCVars::bDontClearTarget;
+	
 		return bClearTarget;
 	};
 
@@ -1347,7 +1353,7 @@ bool FPhysicsReplicationAsync::PredictiveInterpolation(Chaos::FPBDRigidParticleH
 	}
 	
 	// Wake up if sleeping
-	if (Handle->IsSleeping())
+	if (bIsSleeping)
 	{
 		RigidsSolver->GetEvolution()->SetParticleObjectState(Handle, Chaos::EObjectStateType::Dynamic);
 	}
