Fixed HISM instance data proxy build to respect the number of instances in the proxy when they are not matching the source instance count OR the reorder table. This can happen when a component is being serialized (the serialized cluster tree has a different number of instances than the reorder table).

#br benjamin.rouveyrol

diff --git a/Engine/Source/Runtime/Engine/Private/InstancedStaticMesh/ISMInstanceDataSceneProxy.cpp b/Engine/Source/Runtime/Engine/Private/InstancedStaticMesh/ISMInstanceDataSceneProxy.cpp
--- a/Engine/Source/Runtime/Engine/Private/InstancedStaticMesh/ISMInstanceDataSceneProxy.cpp
+++ b/Engine/Source/Runtime/Engine/Private/InstancedStaticMesh/ISMInstanceDataSceneProxy.cpp
@@ -384,7 +384,7 @@ void FISMCInstanceDataSceneProxyLegacyReordered::Update(FISMInstanceUpdateChange
 			int32 InstanceIndex = InstanceIdIndexMap.IdToIndex(FPrimitiveInstanceId{It.GetIndex()});
 			if (IndexRemapOld.RemapIndex(InstanceIndex))
 			{
-				LOG_INST_DATA(TEXT("Update/HideInstance, ID: %d, IDX: %d"), BitIt.GetIndex(), InstanceIndex);
+				LOG_INST_DATA(TEXT("Update/HideInstance, ID: %d, IDX: %d"), It.GetIndex(), InstanceIndex);
 				ProxyData.VisibleInstances[InstanceIndex] = false;
 			}
 		}
@@ -432,11 +432,11 @@ void FISMCInstanceDataSceneProxyLegacyReordered::Build(FISMInstanceUpdateChangeS
 	{
 		ProxyData.VisibleInstances.Reset();
 		ProxyData.VisibleInstances.SetNum(ChangeSet.PostUpdateNumInstances, false);
-		for (int32 Index : LegacyInstanceReorderTable)
+		for (int32 InstanceIndex : LegacyInstanceReorderTable)
 		{
-			if (Index != INDEX_NONE)
+			if (IndexRemap.RemapIndex(InstanceIndex))
 			{
-				ProxyData.VisibleInstances[Index] = true;
+				ProxyData.VisibleInstances[InstanceIndex] = true;
 			}
 		}
 		ProxyData.Flags.bHasPerInstanceVisible = true;
@@ -731,7 +731,7 @@ void FISMCInstanceDataSceneProxyNoGPUScene::Update(FISMInstanceUpdateChangeSet&&
 			int32 InstanceIndex = InstanceIdIndexMap.IdToIndex(FPrimitiveInstanceId{It.GetIndex()});
 			if (IndexRemapOld.RemapIndex(InstanceIndex))
 			{
-				LOG_INST_DATA(TEXT("Update/HideInstance, ID: %d, IDX: %d"), BitIt.GetIndex(), InstanceIndex);
+				LOG_INST_DATA(TEXT("Update/HideInstance, ID: %d, IDX: %d"), It.GetIndex(), InstanceIndex);
 				LegacyInstanceBuffer->InstanceData->NullifyInstance(InstanceIndex);
 			}
 		}
