When a node conversion occurs in the Editor, ensure that Pins are connectable before copying or moving the connections from the old node to the new node.

diff --git a/Engine/Source/Runtime/Engine/Private/EdGraph/EdGraphPin.cpp b/Engine/Source/Runtime/Engine/Private/EdGraph/EdGraphPin.cpp
--- a/Engine/Source/Runtime/Engine/Private/EdGraph/EdGraphPin.cpp
+++ b/Engine/Source/Runtime/Engine/Private/EdGraph/EdGraphPin.cpp
@@ -624,6 +624,14 @@ void TransferPersistentDataFromOldPin(UEdGraphPin& DestPin, T& SourcePin, const
 	{
 		check(OtherPin);
 
+		if (const UEdGraphSchema* Schema = DestPin.GetSchema())
+		{
+			if(Schema->CanCreateConnection(&DestPin, OtherPin).Response == CONNECT_RESPONSE_DISALLOW)
+			{
+				continue;
+			}
+		}
+
 		OtherPin->Modify();
 
 		DestPin.LinkedTo.Add(OtherPin);
