World Partition Runtime Hash Conversion: log source and destination runtime hash class when converting and don't log validation errors.

diff --git a/Engine/Source/Runtime/Engine/Private/WorldPartition/RuntimeHashSet/WorldPartitionRuntimeHashSetConversions.cpp b/Engine/Source/Runtime/Engine/Private/WorldPartition/RuntimeHashSet/WorldPartitionRuntimeHashSetConversions.cpp
--- a/Engine/Source/Runtime/Engine/Private/WorldPartition/RuntimeHashSet/WorldPartitionRuntimeHashSetConversions.cpp
+++ b/Engine/Source/Runtime/Engine/Private/WorldPartition/RuntimeHashSet/WorldPartitionRuntimeHashSetConversions.cpp
@@ -3,6 +3,7 @@
 #include "WorldPartition/RuntimeHashSet/WorldPartitionRuntimeHashSet.h"
 #include "WorldPartition/RuntimeHashSet/RuntimePartitionLHGrid.h"
 #include "WorldPartition/RuntimeHashSet/RuntimePartitionPersistent.h"
+#include "WorldPartition/ErrorHandling/WorldPartitionStreamingGenerationNullErrorHandler.h"
 #include "WorldPartition/WorldPartitionRuntimeSpatialHash.h"
 #include "WorldPartition/HLOD/HLODLayer.h"
 
@@ -15,8 +16,10 @@ UWorldPartitionRuntimeHashSet* UWorldPartitionRuntimeHashSet::CreateFrom(const U
 	UWorldPartition* WorldPartition = SrcHash->GetTypedOuter<UWorldPartition>();
 	check(WorldPartition);
 
+	FStreamingGenerationNullErrorHandler NullErrorHandler;
 	UWorldPartition::FGenerateStreamingParams Params = UWorldPartition::FGenerateStreamingParams()
-		.SetActorDescContainer(WorldPartition->GetActorDescContainer());
+		.SetActorDescContainer(WorldPartition->GetActorDescContainer())
+		.SetErrorHandler(&NullErrorHandler);
 
 	UWorldPartition::FGenerateStreamingContext Context;
 	TUniquePtr<IStreamingGenerationContext> StreamingGenerationContext = WorldPartition->GenerateStreamingGenerationContext(Params, Context);
diff --git a/Engine/Source/Runtime/Engine/Private/WorldPartition/WorldPartition.cpp b/Engine/Source/Runtime/Engine/Private/WorldPartition/WorldPartition.cpp
--- a/Engine/Source/Runtime/Engine/Private/WorldPartition/WorldPartition.cpp
+++ b/Engine/Source/Runtime/Engine/Private/WorldPartition/WorldPartition.cpp
@@ -47,6 +47,7 @@
 #include "WorldPartition/Cook/WorldPartitionCookPackageContextInterface.h"
 #include "WorldPartition/ContentBundle/ContentBundleEditorSubsystemInterface.h"
 #include "WorldPartition/ContentBundle/ContentBundleEditor.h"
+#include "WorldPartition/ErrorHandling/WorldPartitionStreamingGenerationMapCheckErrorHandler.h"
 #include "Modules/ModuleManager.h"
 #include "GameDelegates.h"
 #else
@@ -420,7 +421,10 @@ void UWorldPartition::OnBeginPlay()
 	check(!bIsPIE);
 	bIsPIE = true;
 
-	FGenerateStreamingParams Params;
+	// In PIE, we always want to populate the map check dialog
+	FStreamingGenerationMapCheckErrorHandler MapCheckErrorHandler;
+	FGenerateStreamingParams Params = FGenerateStreamingParams()
+		.SetErrorHandler(&MapCheckErrorHandler);
 
 	TArray<FString> OutGeneratedStreamingPackageNames;
 	FGenerateStreamingContext Context = FGenerateStreamingContext()
diff --git a/Engine/Source/Runtime/Engine/Private/WorldPartition/WorldPartitionRuntimeHash.cpp b/Engine/Source/Runtime/Engine/Private/WorldPartition/WorldPartitionRuntimeHash.cpp
--- a/Engine/Source/Runtime/Engine/Private/WorldPartition/WorldPartitionRuntimeHash.cpp
+++ b/Engine/Source/Runtime/Engine/Private/WorldPartition/WorldPartitionRuntimeHash.cpp
@@ -513,19 +513,28 @@ UWorldPartitionRuntimeHash* UWorldPartitionRuntimeHash::ConvertWorldPartitionHas
 	check(InDstClass->IsChildOf<UWorldPartitionRuntimeHash>());
 	check(!InDstClass->HasAnyClassFlags(CLASS_Abstract));
 
+	UWorldPartitionRuntimeHash* NewHash = nullptr;
+
 	const UClass* CurrentSrcClass = InSrcHash->GetClass();
 	while(CurrentSrcClass != UWorldPartitionRuntimeHash::StaticClass())
 	{
 		if (FRuntimeHashConvertFunc* Converter = WorldPartitionRuntimeHashConverters.Find({ CurrentSrcClass, InDstClass }))
 		{
-			return (*Converter)(InSrcHash);
+			NewHash = (*Converter)(InSrcHash);
+			break;
 		}
 		CurrentSrcClass = CurrentSrcClass->GetSuperClass();
 	}
 
-	UWorldPartitionRuntimeHash* RuntimeHash = NewObject<UWorldPartitionRuntimeHash>(InSrcHash->GetOuter(), InDstClass, NAME_None, RF_Transactional);
-	RuntimeHash->SetDefaultValues();
-	return RuntimeHash;
+	if (!NewHash)
+	{
+		NewHash = NewObject<UWorldPartitionRuntimeHash>(InSrcHash->GetOuter(), InDstClass, NAME_None, RF_Transactional);
+		NewHash->SetDefaultValues();
+	}
+
+	UE_LOG(LogWorldPartition, Log, TEXT("Converted '%s' runtime hash class from '%s' to '%s'"), *InSrcHash->GetPackage()->GetName(), *InSrcHash->GetClass()->GetName(), *InDstClass->GetName());
+
+	return NewHash;
 }
 #endif
 
diff --git a/Engine/Source/Runtime/Engine/Private/WorldPartition/WorldPartitionStreamingGeneration.cpp b/Engine/Source/Runtime/Engine/Private/WorldPartition/WorldPartitionStreamingGeneration.cpp
--- a/Engine/Source/Runtime/Engine/Private/WorldPartition/WorldPartitionStreamingGeneration.cpp
+++ b/Engine/Source/Runtime/Engine/Private/WorldPartition/WorldPartitionStreamingGeneration.cpp
@@ -26,7 +26,6 @@
 #include "WorldPartition/DataLayer/DataLayerUtils.h"
 #include "WorldPartition/ErrorHandling/WorldPartitionStreamingGenerationNullErrorHandler.h"
 #include "WorldPartition/ErrorHandling/WorldPartitionStreamingGenerationLogErrorHandler.h"
-#include "WorldPartition/ErrorHandling/WorldPartitionStreamingGenerationMapCheckErrorHandler.h"
 #include "WorldPartition/HLOD/HLODActor.h"
 #include "WorldPartition/HLOD/HLODLayer.h"
 #include "WorldPartition/WorldPartitionSubsystem.h"
@@ -54,6 +53,29 @@ static FAutoConsoleCommand DumpStreamingGenerationLog(
 	})
 );
 
+template <class T>
+class TErrorHandlerSelector
+{
+public:
+	TErrorHandlerSelector(IStreamingGenerationErrorHandler* InErrorHandler = nullptr)
+	{
+		ErrorHandler = InErrorHandler ? InErrorHandler : &BaseErrorHandler;
+	}
+
+	IStreamingGenerationErrorHandler* Get()
+	{
+		if (UWorldPartition::StreamingGenerationErrorHandlerOverride)
+		{
+			return (*UWorldPartition::StreamingGenerationErrorHandlerOverride)(ErrorHandler);
+		}
+		return ErrorHandler;
+	}
+
+private:
+	T BaseErrorHandler;
+	IStreamingGenerationErrorHandler* ErrorHandler;
+};
+
 FActorDescViewMap::FActorDescViewMap()
 {}
 
@@ -1361,16 +1383,9 @@ bool UWorldPartition::GenerateContainerStreaming(const FGenerateStreamingParams&
 
 	FActorDescList* ModifiedActorsDescList = nullptr;
 
-	FStreamingGenerationLogErrorHandler LogErrorHandler;
-	FStreamingGenerationMapCheckErrorHandler MapCheckErrorHandler;	
-	IStreamingGenerationErrorHandler* ErrorHandler = &LogErrorHandler;	
-
 	if (bIsPIE)
 	{
 		ModifiedActorsDescList = &RuntimeHash->ModifiedActorDescListForPIE;
-		
-		// In PIE, we always want to populate the map check dialog
-		ErrorHandler = &MapCheckErrorHandler;
 	}
 
 	const FString ContainerPackageName = InParams.ActorDescCollection.GetMainContainerPackageName().ToString();
@@ -1409,12 +1424,14 @@ bool UWorldPartition::GenerateContainerStreaming(const FGenerateStreamingParams&
 		}
 	}
 
+	TErrorHandlerSelector<FStreamingGenerationLogErrorHandler> ErrorHandlerSelector(InParams.ErrorHandler);
+
 	FWorldPartitionStreamingGenerator::FWorldPartitionStreamingGeneratorParams StreamingGeneratorParams = FWorldPartitionStreamingGenerator::FWorldPartitionStreamingGeneratorParams()
 		.SetWorldPartitionContext(this)
 		.SetModifiedActorsDescList( ModifiedActorsDescList)
 		.SetIsValidGrid([this](FName GridName) { return RuntimeHash->IsValidGrid(GridName); })
 		.SetIsValidHLODLayer([this](FName GridName, const FSoftObjectPath& HLODLayerPath) { return RuntimeHash->IsValidHLODLayer(GridName, HLODLayerPath); })
-		.SetErrorHandler(StreamingGenerationErrorHandlerOverride ? (*StreamingGenerationErrorHandlerOverride)(ErrorHandler) : ErrorHandler)
+		.SetErrorHandler(ErrorHandlerSelector.Get())
 		.SetEnableStreaming(IsStreamingEnabled())
 		.SetCreateContainerResolver(FEditorPathHelper::IsEnabled());
 
@@ -1471,12 +1488,13 @@ TUniquePtr<IStreamingGenerationContext> UWorldPartition::GenerateStreamingGenera
 {
 	TRACE_CPUPROFILER_EVENT_SCOPE(UWorldPartition::GenerateStreamingGenerationContext);
 
-	FStreamingGenerationLogErrorHandler LogErrorHandler;
+	TErrorHandlerSelector<FStreamingGenerationLogErrorHandler> ErrorHandlerSelector(InParams.ErrorHandler);
+
 	FWorldPartitionStreamingGenerator::FWorldPartitionStreamingGeneratorParams StreamingGeneratorParams = FWorldPartitionStreamingGenerator::FWorldPartitionStreamingGeneratorParams()
 		.SetWorldPartitionContext(this)
 		.SetIsValidGrid([this](FName GridName) { return RuntimeHash->IsValidGrid(GridName); })
 		.SetIsValidHLODLayer([this](FName GridName, const FSoftObjectPath& HLODLayerPath) { return RuntimeHash->IsValidHLODLayer(GridName, HLODLayerPath); })
-		.SetErrorHandler(&LogErrorHandler)
+		.SetErrorHandler(ErrorHandlerSelector.Get())
 		.SetEnableStreaming(IsStreamingEnabled())
 		.SetCreateContainerResolver(FEditorPathHelper::IsEnabled());
 
@@ -1514,10 +1532,11 @@ void UWorldPartition::SetupHLODActors(const FSetupHLODActorsParams& Params)
 {
 	ForEachActorDescContainer([this, &Params](UActorDescContainer* InActorDescContainer)
 	{
-		FStreamingGenerationLogErrorHandler LogErrorHandler;
+		TErrorHandlerSelector<FStreamingGenerationLogErrorHandler> ErrorHandlerSelector;
+
 		FWorldPartitionStreamingGenerator::FWorldPartitionStreamingGeneratorParams StreamingGeneratorParams = FWorldPartitionStreamingGenerator::FWorldPartitionStreamingGeneratorParams()
 			.SetWorldPartitionContext(this)
-			.SetErrorHandler(StreamingGenerationErrorHandlerOverride ? (*StreamingGenerationErrorHandlerOverride)(&LogErrorHandler) : &LogErrorHandler)
+			.SetErrorHandler(ErrorHandlerSelector.Get())
 			.SetEnableStreaming(IsStreamingEnabled())
 			.SetFilteredClasses({ AWorldPartitionHLOD::StaticClass() })
 			.SetIsValidGrid([this](FName GridName) { return RuntimeHash->IsValidGrid(GridName); })
@@ -1649,10 +1668,12 @@ void UWorldPartition::CheckForErrors(IStreamingGenerationErrorHandler* ErrorHand
 {
 	FActorDescList ModifiedActorDescList;
 
+	TErrorHandlerSelector<FStreamingGenerationLogErrorHandler> ErrorHandlerSelector(ErrorHandler);
+
 	FWorldPartitionStreamingGenerator::FWorldPartitionStreamingGeneratorParams StreamingGeneratorParams = FWorldPartitionStreamingGenerator::FWorldPartitionStreamingGeneratorParams()
 		.SetWorldPartitionContext(this)
 		.SetModifiedActorsDescList(&ModifiedActorDescList)
-		.SetErrorHandler(StreamingGenerationErrorHandlerOverride ? (*StreamingGenerationErrorHandlerOverride)(ErrorHandler) : ErrorHandler)
+		.SetErrorHandler(ErrorHandlerSelector.Get())
 		.SetIsValidGrid([this](FName GridName) { return RuntimeHash->IsValidGrid(GridName); })
 		.SetIsValidHLODLayer([this](FName GridName, const FSoftObjectPath& HLODLayerPath) { return RuntimeHash->IsValidHLODLayer(GridName, HLODLayerPath); })
 		.SetEnableStreaming(IsStreamingEnabled());
@@ -1703,10 +1724,12 @@ void UWorldPartition::CheckForErrors(const FCheckForErrorsParams& Params)
 	FStreamingGenerationActorDescCollection Collection { Containers };
 	const UActorDescContainer* MainActorDescContainer = Collection.GetMainActorDescContainer();
 
+	TErrorHandlerSelector<FStreamingGenerationLogErrorHandler> ErrorHandlerSelector(Params.ErrorHandler);
+
 	FWorldPartitionStreamingGenerator::FWorldPartitionStreamingGeneratorParams StreamingGeneratorParams = FWorldPartitionStreamingGenerator::FWorldPartitionStreamingGeneratorParams()
 		.SetWorldPartitionContext(MainActorDescContainer->GetWorldPartition())
 		.SetModifiedActorsDescList(!MainActorDescContainer->IsTemplateContainer() ? &ModifiedActorDescList : nullptr)
-		.SetErrorHandler(StreamingGenerationErrorHandlerOverride ? (*StreamingGenerationErrorHandlerOverride)(Params.ErrorHandler) : Params.ErrorHandler)
+		.SetErrorHandler(ErrorHandlerSelector.Get())
 		.SetIsValidGrid([](FName) { return true; })
 		.SetIsValidHLODLayer([](FName, const FSoftObjectPath&) { return true; })
 		.SetEnableStreaming(Params.bEnableStreaming)
diff --git a/Engine/Source/Runtime/Engine/Public/WorldPartition/WorldPartition.h b/Engine/Source/Runtime/Engine/Public/WorldPartition/WorldPartition.h
--- a/Engine/Source/Runtime/Engine/Public/WorldPartition/WorldPartition.h
+++ b/Engine/Source/Runtime/Engine/Public/WorldPartition/WorldPartition.h
@@ -238,17 +238,16 @@ public:
 	struct FGenerateStreamingParams
 	{
 		FGenerateStreamingParams()
+			: ErrorHandler(nullptr)
 		{}
 
 		FStreamingGenerationActorDescCollection ActorDescCollection;
 		TOptional<const FString> OutputLogPath;
+		IStreamingGenerationErrorHandler* ErrorHandler;
 
-		FGenerateStreamingParams& SetActorDescContainer(const UActorDescContainer* InActorDescContainer) 
-		{ 
-			ActorDescCollection.AddContainer(InActorDescContainer);
-			return *this;
-		}
+		FGenerateStreamingParams& SetActorDescContainer(const UActorDescContainer* InActorDescContainer) { ActorDescCollection.AddContainer(InActorDescContainer); return *this; }
 		FGenerateStreamingParams& SetOutputLogPath(const FString& InOutputLogPath) { OutputLogPath = InOutputLogPath; return *this; }
+		FGenerateStreamingParams& SetErrorHandler(IStreamingGenerationErrorHandler* InErrorHandler) { ErrorHandler = InErrorHandler; return *this; }
 	};
 
 	struct FGenerateStreamingContext
