World Partition: fix wrongly setting bIsPIE when running in editor game.
Also fix using the check map error handler when running editor game.

diff --git a/Engine/Source/Runtime/Engine/Private/WorldPartition/WorldPartition.cpp b/Engine/Source/Runtime/Engine/Private/WorldPartition/WorldPartition.cpp
--- a/Engine/Source/Runtime/Engine/Private/WorldPartition/WorldPartition.cpp
+++ b/Engine/Source/Runtime/Engine/Private/WorldPartition/WorldPartition.cpp
@@ -47,6 +47,7 @@
 #include "WorldPartition/Cook/WorldPartitionCookPackageContextInterface.h"
 #include "WorldPartition/ContentBundle/ContentBundleEditorSubsystemInterface.h"
 #include "WorldPartition/ContentBundle/ContentBundleEditor.h"
+#include "WorldPartition/ErrorHandling/WorldPartitionStreamingGenerationLogErrorHandler.h"
 #include "WorldPartition/ErrorHandling/WorldPartitionStreamingGenerationMapCheckErrorHandler.h"
 #include "Modules/ModuleManager.h"
 #include "GameDelegates.h"
@@ -419,12 +420,14 @@ void UWorldPartition::OnPrePIEEnded(bool bWasSimulatingInEditor)
 void UWorldPartition::OnBeginPlay()
 {
 	check(!bIsPIE);
-	bIsPIE = true;
+	bIsPIE = !IsRunningGame();
 
 	// In PIE, we always want to populate the map check dialog
 	FStreamingGenerationMapCheckErrorHandler MapCheckErrorHandler;
+	FStreamingGenerationLogErrorHandler LogErrorHandler;
+	
 	FGenerateStreamingParams Params = FGenerateStreamingParams()
-		.SetErrorHandler(&MapCheckErrorHandler);
+		.SetErrorHandler(bIsPIE ? (IStreamingGenerationErrorHandler*)&MapCheckErrorHandler : (IStreamingGenerationErrorHandler*)&LogErrorHandler);
 
 	TArray<FString> OutGeneratedStreamingPackageNames;
 	FGenerateStreamingContext Context = FGenerateStreamingContext()
