Server crash on executing async physics command.

Check that the cached lambda function is callable before calling it.
Null check logic in the source function enqueing the lambda.
Use locally cached async physics command struct instead of accessing array each time.

diff --git a/Engine/Source/Runtime/Engine/Private/PhysicsEngine/Experimental/PhysScene_Chaos.cpp b/Engine/Source/Runtime/Engine/Private/PhysicsEngine/Experimental/PhysScene_Chaos.cpp
--- a/Engine/Source/Runtime/Engine/Private/PhysicsEngine/Experimental/PhysScene_Chaos.cpp
+++ b/Engine/Source/Runtime/Engine/Private/PhysicsEngine/Experimental/PhysScene_Chaos.cpp
@@ -157,33 +157,33 @@ public:
 		{
 			const int32 CurrentFrame = static_cast<FPBDRigidsSolver*>(GetSolver())->GetCurrentFrame();
 
-			FPendingAsyncPhysicsCommand* PendingCommand = &PendingCommands[Idx];
-			bool bRemove = PendingCommand->OwningObject.IsStale();
+			const FPendingAsyncPhysicsCommand& PendingCommand = PendingCommands[Idx];
+			bool bRemove = PendingCommand.OwningObject.IsStale() || !PendingCommand.Command;
 
-			if (bAllowResim && PendingCommand->bEnableResim && PendingCommands[Idx].PhysicsStep > (CurrentFrame - NumFrames))
+			if (!bRemove && bAllowResim && PendingCommand.bEnableResim && PendingCommand.PhysicsStep > (CurrentFrame - NumFrames))
 			{
-				if (PendingCommands[Idx].PhysicsStep < CurrentFrame && PendingCommands[Idx].OwningObject.IsValid())
+				if (PendingCommand.PhysicsStep < CurrentFrame)
 				{
 					if (Chaos::FRewindData* RewindData = GetSolver()->GetRewindData())
 					{
 						int32 ResimFrame = RewindData->GetResimFrame();
-						ResimFrame = (ResimFrame == INDEX_NONE) ? PendingCommands[Idx].PhysicsStep :
-							FMath::Min(ResimFrame, PendingCommands[Idx].PhysicsStep);
+						ResimFrame = (ResimFrame == INDEX_NONE) ? PendingCommand.PhysicsStep :
+							FMath::Min(ResimFrame, PendingCommand.PhysicsStep);
 
 						RewindData->SetResimFrame(ResimFrame);
 					}
 				}
-				else if (!bRemove && PendingCommands[Idx].PhysicsStep == CurrentFrame)
+				else if (PendingCommand.PhysicsStep == CurrentFrame)
 				{
-					PendingCommands[Idx].Command();
+					PendingCommand.Command();
 					bRemove = true;
 				}
 			}
 			else
 			{
-				if (!bRemove && PendingCommands[Idx].PhysicsStep <= CurrentFrame)
+				if (!bRemove && PendingCommand.PhysicsStep <= CurrentFrame)
 				{
-					PendingCommands[Idx].Command();
+					PendingCommand.Command();
 					bRemove = true;
 				}
 			}
diff --git a/Engine/Source/Runtime/Experimental/GeometryCollectionEngine/Private/GeometryCollection/GeometryCollectionComponent.cpp b/Engine/Source/Runtime/Experimental/GeometryCollectionEngine/Private/GeometryCollection/GeometryCollectionComponent.cpp
--- a/Engine/Source/Runtime/Experimental/GeometryCollectionEngine/Private/GeometryCollection/GeometryCollectionComponent.cpp
+++ b/Engine/Source/Runtime/Experimental/GeometryCollectionEngine/Private/GeometryCollection/GeometryCollectionComponent.cpp
@@ -4363,17 +4363,19 @@ void UGeometryCollectionComponent::UpdateRemovalIfNeeded()
 
 void UGeometryCollectionComponent::RequestUpdateRepData()
 {
-	if (FPhysScene* PhysScene = GetInnerChaosScene())
+	if (FPhysScene* PhysScene = GetInnerChaosScene(); PhysScene && PhysicsProxy)
 	{
-		Chaos::FPBDRigidsSolver* Solver = PhysicsProxy->GetSolver<Chaos::FPBDRigidsSolver>();
-		const int32 CurrentFrame = Solver->GetCurrentFrame();
+		if (Chaos::FPBDRigidsSolver* Solver = PhysicsProxy->GetSolver<Chaos::FPBDRigidsSolver>())
+		{
+			const int32 CurrentFrame = Solver->GetCurrentFrame();
 
-		PhysScene->EnqueueAsyncPhysicsCommand(CurrentFrame, this,
-			[this]()
-			{
-				UpdateRepData();
-			}, false
-		);
+			PhysScene->EnqueueAsyncPhysicsCommand(CurrentFrame, this,
+				[this]()
+				{
+					UpdateRepData();
+				}, false
+			);
+		}
 	}
 }
 
