Change World Parition Resolve Helper to use FTopLevelAssetPath to avoid unnecessary string manipulation

diff --git a/Engine/Source/Runtime/Engine/Private/WorldPartition/WorldPartitionLevelHelper.cpp b/Engine/Source/Runtime/Engine/Private/WorldPartition/WorldPartitionLevelHelper.cpp
--- a/Engine/Source/Runtime/Engine/Private/WorldPartition/WorldPartitionLevelHelper.cpp
+++ b/Engine/Source/Runtime/Engine/Private/WorldPartition/WorldPartitionLevelHelper.cpp
@@ -33,7 +33,7 @@ bool FWorldPartitionResolveData::ResolveObject(UWorld* InWorld, const FSoftObjec
 	OutObject = nullptr;
 	if (InWorld)
 	{
-		if (IsValid() && SourceWorldAssetPath == InObjectPath.GetAssetPath().ToString())
+		if (IsValid() && SourceWorldAssetPath == InObjectPath.GetAssetPath())
 		{
 			const FString SubPathString = FWorldPartitionLevelHelper::AddActorContainerIDToSubPathString(ContainerID, InObjectPath.GetSubPathString());
 			// We don't read the return value as we always want to return true when using the resolve data.
@@ -575,7 +575,7 @@ bool FWorldPartitionLevelHelper::LoadActors(const FLoadActorsParams& InParams)
 
 					if (IWorldPartitionObjectResolver* ObjectResolver = Cast<IWorldPartitionObjectResolver>(Actor))
 					{
-						ObjectResolver->SetWorldPartitionResolveData(FWorldPartitionResolveData(PackageObjectMapping->ContainerID, SourceWorldPath));
+						ObjectResolver->SetWorldPartitionResolveData(FWorldPartitionResolveData(PackageObjectMapping->ContainerID, FTopLevelAssetPath(SourceWorldPath)));
 					}
 				}
 
diff --git a/Engine/Source/Runtime/Engine/Public/WorldPartition/IWorldPartitionObjectResolver.h b/Engine/Source/Runtime/Engine/Public/WorldPartition/IWorldPartitionObjectResolver.h
--- a/Engine/Source/Runtime/Engine/Public/WorldPartition/IWorldPartitionObjectResolver.h
+++ b/Engine/Source/Runtime/Engine/Public/WorldPartition/IWorldPartitionObjectResolver.h
@@ -14,7 +14,12 @@ struct FWorldPartitionResolveData
 	GENERATED_USTRUCT_BODY()
 
 	FWorldPartitionResolveData() {}
-	FWorldPartitionResolveData(const FActorContainerID& InContainerID, const FString& InSourceWorldAssetPath) 
+
+	UE_DEPRECATED(5.4, "Please use FTopLevelAssetPath for the asset path directly")
+	FWorldPartitionResolveData(const FActorContainerID& InContainerID, const FString& InSourceWorldAssetPath)
+		: ContainerID(InContainerID), SourceWorldAssetPath(InSourceWorldAssetPath) {}
+
+	FWorldPartitionResolveData(const FActorContainerID& InContainerID, const FTopLevelAssetPath& InSourceWorldAssetPath)
 		: ContainerID(InContainerID), SourceWorldAssetPath(InSourceWorldAssetPath) {}
 
 	bool IsValid() const { return !ContainerID.IsMainContainer(); }
@@ -25,7 +30,7 @@ struct FWorldPartitionResolveData
 	FActorContainerID ContainerID;
 
 	UPROPERTY()
-	FString SourceWorldAssetPath;
+	FTopLevelAssetPath SourceWorldAssetPath;
 };
 
 UINTERFACE(MinimalAPI)
diff --git a/Engine/Source/Runtime/Engine/Public/WorldPartition/WorldPartitionLevelHelper.h b/Engine/Source/Runtime/Engine/Public/WorldPartition/WorldPartitionLevelHelper.h
--- a/Engine/Source/Runtime/Engine/Public/WorldPartition/WorldPartitionLevelHelper.h
+++ b/Engine/Source/Runtime/Engine/Public/WorldPartition/WorldPartitionLevelHelper.h
@@ -20,8 +20,9 @@ struct FActorContainerID;
 class FWorldPartitionLevelHelper
 {
 public:
-	static FString AddActorContainerIDToSubPathString(const FActorContainerID& InContainerID, const FString& InSubPathString);
-	static FString AddActorContainerID(const FActorContainerID& InContainerID, const FString& InActorName);
+	ENGINE_API static FString AddActorContainerIDToSubPathString(const FActorContainerID& InContainerID, const FString& InSubPathString);
+	ENGINE_API static FString AddActorContainerID(const FActorContainerID& InContainerID, const FString& InActorName);
+
 #if WITH_EDITOR
 public:
 	static FWorldPartitionLevelHelper& Get();
