ChildActorComponent
* Fix simulated actors using the wrong template to spawn the actor if the CAC template instance is of the wrong type.

#rnx

diff --git a/Engine/Source/Runtime/Engine/Classes/Components/ChildActorComponent.h b/Engine/Source/Runtime/Engine/Classes/Components/ChildActorComponent.h
--- a/Engine/Source/Runtime/Engine/Classes/Components/ChildActorComponent.h
+++ b/Engine/Source/Runtime/Engine/Classes/Components/ChildActorComponent.h
@@ -197,6 +197,7 @@ public:
 
 	AActor* GetChildActor() const { return ChildActor; }
 	AActor* GetChildActorTemplate() const { return ChildActorTemplate; }
+	ENGINE_API AActor* GetSpawnableChildActorTemplate() const;
 
 	FName GetChildActorName() const { return ChildActorName; }
 	ENGINE_API void SetChildActorName(const FName InName);
diff --git a/Engine/Source/Runtime/Engine/Private/Components/ChildActorComponent.cpp b/Engine/Source/Runtime/Engine/Private/Components/ChildActorComponent.cpp
--- a/Engine/Source/Runtime/Engine/Private/Components/ChildActorComponent.cpp
+++ b/Engine/Source/Runtime/Engine/Private/Components/ChildActorComponent.cpp
@@ -292,6 +292,20 @@ void UChildActorComponent::PostEditUndo()
 }
 #endif
 
+AActor* UChildActorComponent::GetSpawnableChildActorTemplate() const
+{
+	// Only use the instance if it's the same type as the class it was supposedly built from
+	if (ChildActorTemplate && ChildActorTemplate->GetClass() == ChildActorClass)
+	{
+		return ChildActorTemplate;
+	}
+	// Use the CDO of the class as the template if the instance is wrong.
+	else
+	{
+		return ChildActorClass ? ChildActorClass->GetDefaultObject<AActor>() : nullptr;
+	}
+}
+
 void UChildActorComponent::GetLifetimeReplicatedProps(TArray<FLifetimeProperty>& OutLifetimeProps) const
 {
 	Super::GetLifetimeReplicatedProps(OutLifetimeProps);
diff --git a/Engine/Source/Runtime/Engine/Private/Net/Iris/ReplicationSystem/ActorReplicationBridge.cpp b/Engine/Source/Runtime/Engine/Private/Net/Iris/ReplicationSystem/ActorReplicationBridge.cpp
--- a/Engine/Source/Runtime/Engine/Private/Net/Iris/ReplicationSystem/ActorReplicationBridge.cpp
+++ b/Engine/Source/Runtime/Engine/Private/Net/Iris/ReplicationSystem/ActorReplicationBridge.cpp
@@ -655,7 +655,7 @@ FObjectReplicationBridgeInstantiateResult UActorReplicationBridge::BeginInstanti
 
 				if (!SubObj)
 				{
-					UE_LOG_ACTORREPLICATIONBRIDGE(Error, TEXT("BeginInstantiateFromRemote Failed to find subobjectReference for dynamic SubObject %s, Owner %s, RootObject %s"), *DescribeObjectReference(Header->ObjectReference, ResolveContext), *RootObjectOfSubObject.ToString(), *GetPathNameSafe(RootActor));
+					UE_LOG_ACTORREPLICATIONBRIDGE(Error, TEXT("BeginInstantiateFromRemote Failed to find stable name reference of dynamic SubObject %s, Owner %s, RootObject %s"), *DescribeObjectReference(Header->ObjectReference, ResolveContext), *RootObjectOfSubObject.ToString(), *GetPathNameSafe(RootActor));
 				}
 			}
 			else
@@ -852,7 +852,7 @@ void UActorReplicationBridge::GetInitialDependencies(FNetRefHandle Handle, FNetD
 			// customized properties will be incorrect on the Client.
 			if (UChildActorComponent* CAC = Actor->GetParentComponent())
 			{
-				Archetype = CAC->GetChildActorTemplate();
+				Archetype = CAC->GetSpawnableChildActorTemplate();
 			}
 			if (Archetype == nullptr)
 			{
@@ -951,7 +951,7 @@ void UActorReplicationBridge::GetActorCreationHeader(const AActor* Actor, UE::Ne
 		// customized properties will be incorrect on the Client.
 		if (UChildActorComponent* CAC = Actor->GetParentComponent())
 		{
-			Archetype = CAC->GetChildActorTemplate();
+			Archetype = CAC->GetSpawnableChildActorTemplate();
 		}
 		if (Archetype == nullptr)
 		{
diff --git a/Engine/Source/Runtime/Engine/Private/PackageMapClient.cpp b/Engine/Source/Runtime/Engine/Private/PackageMapClient.cpp
--- a/Engine/Source/Runtime/Engine/Private/PackageMapClient.cpp
+++ b/Engine/Source/Runtime/Engine/Private/PackageMapClient.cpp
@@ -548,7 +548,7 @@ bool UPackageMapClient::SerializeNewActor(FArchive& Ar, class UActorChannel *Cha
 			// customized properties will be incorrect on the Client.
 			if (UChildActorComponent* CAC = Actor->GetParentComponent())
 			{
-				Archetype = CAC->GetChildActorTemplate();
+				Archetype = CAC->GetSpawnableChildActorTemplate();
 			}
 			if (Archetype == nullptr)
 			{
@@ -941,8 +941,16 @@ void UPackageMapClient::InternalWriteObject(FArchive & Ar, FNetworkGUID NetGUID,
 			}
 		}
 
+#if WITH_EDITOR
+		FString TempObjectName = ObjectPathName;
+#endif
+
 		GEngine->NetworkRemapPath(Connection, ObjectPathName, false);
 
+#if WITH_EDITOR
+		ensureMsgf(!ObjectPathName.IsEmpty(), TEXT("NetworkRemapPath found PathName: %s to be an invalid name for %s. This object will not replicate!"), *TempObjectName, *GetPathNameSafe(Object));
+#endif
+
 		// Serialize Name of object
 		Ar << ObjectPathName;
 
@@ -1122,9 +1130,17 @@ FNetworkGUID UPackageMapClient::InternalLoadObject( FArchive & Ar, UObject *& Ob
 			return NetGUID;
 		}
 
+#if WITH_EDITOR
+		FString TempObjectName = ObjectName;
+#endif
+
 		// Remap name for PIE
 		GEngine->NetworkRemapPath( Connection, ObjectName, true );
 
+#if WITH_EDITOR
+		ensureMsgf(!ObjectName.IsEmpty(), TEXT("NetworkRemapPath found %s to be an invalid name. This object will not be binded and replicated!"), *TempObjectName);
+#endif
+
 		if (NetGUID.IsDefault())
 		{
 			// This should be from the client
