Optimize GetPerBoneInterpolationIndex for blend spaces by avoiding possible lock contention

diff --git a/Engine/Source/Runtime/Engine/Private/Animation/BlendSpace.cpp b/Engine/Source/Runtime/Engine/Private/Animation/BlendSpace.cpp
--- a/Engine/Source/Runtime/Engine/Private/Animation/BlendSpace.cpp
+++ b/Engine/Source/Runtime/Engine/Private/Animation/BlendSpace.cpp
@@ -758,21 +758,37 @@ void UBlendSpace::RuntimeValidateMarkerData()
 /** When per-bone blend data are required to be sorted, this stores the sorted copy and the index of the original */
 struct FSortedPerBoneInterpolation
 {
-	FSortedPerBoneInterpolation(const FPerBoneInterpolation& Original, int32 Index) :
-		PerBoneBlend(Original), OriginalIndex(Index) {}
+	FSortedPerBoneInterpolation(const FPerBoneInterpolation& Original, int32 Index)
+		: PerBoneBlend(Original)
+		, OriginalIndex(Index)
+	{
+	}
+
 	FPerBoneInterpolation PerBoneBlend;
+
 	/** Index into the original array */
 	int32 OriginalIndex;
 };
+
 struct FSortedPerBoneInterpolationData : public IInterpolationIndexProvider::FPerBoneInterpolationData
 {
+	explicit FSortedPerBoneInterpolationData(const FSkeletonRemapping& InSkeletonMapping)
+		: SkeletonMapping(InSkeletonMapping)
+	{
+	}
+
+	const FSkeletonRemapping& SkeletonMapping;
 	TArray<FSortedPerBoneInterpolation> Data;
 };
 
 TSharedPtr<IInterpolationIndexProvider::FPerBoneInterpolationData> UBlendSpace::GetPerBoneInterpolationData(const USkeleton* RuntimeSkeleton) const 
 {
-	FSortedPerBoneInterpolationData* Data = new FSortedPerBoneInterpolationData();
+	const USkeleton* SourceSkeleton = GetSkeleton();
+	const FSkeletonRemapping& SkeletonRemapping = UE::Anim::FSkeletonRemappingRegistry().Get().GetRemapping(SourceSkeleton, RuntimeSkeleton);
+
+	FSortedPerBoneInterpolationData* Data = new FSortedPerBoneInterpolationData(SkeletonRemapping);
 	Data->Data.SetNumUninitialized(PerBoneBlendValues.Num());
+
 	for (int32 Iter = 0 ; Iter != PerBoneBlendValues.Num() ; ++Iter)
 	{
 		Data->Data[Iter] = FSortedPerBoneInterpolation(PerBoneBlendValues[Iter], Iter);
@@ -793,16 +809,18 @@ TSharedPtr<IInterpolationIndexProvider::FPerBoneInterpolationData> UBlendSpace::
 int32 UBlendSpace::GetPerBoneInterpolationIndex(
 	const FCompactPoseBoneIndex&                                  InCompactPoseBoneIndex, 
 	const FBoneContainer&                                         RequiredBones, 
-	const IInterpolationIndexProvider::FPerBoneInterpolationData* Data) const
+	const IInterpolationIndexProvider::FPerBoneInterpolationData* InData) const
 {
-	if (!ensure(Data))
+	if (!ensure(InData))
 	{
 		return INDEX_NONE;
 	}
-	const TArray<FSortedPerBoneInterpolation>& SortedData = static_cast<const FSortedPerBoneInterpolationData*>(Data)->Data;
+
+	const FSortedPerBoneInterpolationData* Data = static_cast<const FSortedPerBoneInterpolationData*>(InData);
+	const TArray<FSortedPerBoneInterpolation>& SortedData = Data->Data;
 
 	const USkeleton* SourceSkeleton = GetSkeleton();
-	const FSkeletonRemapping& SkeletonRemapping = UE::Anim::FSkeletonRemappingRegistry().Get().GetRemapping(SourceSkeleton, RequiredBones.GetSkeletonAsset());
+	const FSkeletonRemapping& SkeletonRemapping = Data->SkeletonMapping;
 
 	for (int32 Iter = 0; Iter < SortedData.Num(); ++Iter)
 	{
@@ -840,18 +858,22 @@ namespace UE::Anim::Private
 	static FBoneContainer DummyContainer;
 }
 
-int32 UBlendSpace::GetPerBoneInterpolationIndex(const FSkeletonPoseBoneIndex InSkeletonBoneIndex, const USkeleton* TargetSkeleton, const IInterpolationIndexProvider::FPerBoneInterpolationData* Data) const
+int32 UBlendSpace::GetPerBoneInterpolationIndex(
+	const FSkeletonPoseBoneIndex InSkeletonBoneIndex,
+	const USkeleton* TargetSkeleton,
+	const IInterpolationIndexProvider::FPerBoneInterpolationData* InData) const
 {
-	if (!ensure(Data) || !InSkeletonBoneIndex.IsValid())
+	if (!ensure(InData) || !InSkeletonBoneIndex.IsValid())
 	{
 		return INDEX_NONE;
 	}
 
-	const TArray<FSortedPerBoneInterpolation>& SortedData = static_cast<const FSortedPerBoneInterpolationData*>(Data)->Data;
+	const FSortedPerBoneInterpolationData* Data = static_cast<const FSortedPerBoneInterpolationData*>(InData);
+	const TArray<FSortedPerBoneInterpolation>& SortedData = Data->Data;
 
 	const USkeleton* SourceSkeleton = GetSkeleton();
 	const FReferenceSkeleton& TargetReferenceSkeleton = TargetSkeleton->GetReferenceSkeleton();
-	const FSkeletonRemapping& SkeletonRemapping = UE::Anim::FSkeletonRemappingRegistry().Get().GetRemapping(SourceSkeleton, TargetSkeleton);
+	const FSkeletonRemapping& SkeletonRemapping = Data->SkeletonMapping;
 
 	for (const FSortedPerBoneInterpolation& SortedBoneData : SortedData)
 	{
diff --git a/Engine/Source/Runtime/Engine/Private/Animation/SkeletonRemappingRegistry.cpp b/Engine/Source/Runtime/Engine/Private/Animation/SkeletonRemappingRegistry.cpp
--- a/Engine/Source/Runtime/Engine/Private/Animation/SkeletonRemappingRegistry.cpp
+++ b/Engine/Source/Runtime/Engine/Private/Animation/SkeletonRemappingRegistry.cpp
@@ -66,30 +66,39 @@ void FSkeletonRemappingRegistry::Destroy()
 
 const FSkeletonRemapping& FSkeletonRemappingRegistry::GetRemapping(const USkeleton* InSourceSkeleton, const USkeleton* InTargetSkeleton)
 {
-	if(InSourceSkeleton == InTargetSkeleton)
+	if (InSourceSkeleton == InTargetSkeleton)
 	{
 		return DefaultMapping;
 	}
 	
 	const FWeakSkeletonPair Pair(InSourceSkeleton, InTargetSkeleton);
+	const uint32 PairHash = GetTypeHash(Pair);
 
 	{
 		FRWScopeLock ReadLock(MappingsLock, SLT_ReadOnly);
-		const TSharedPtr<FSkeletonRemapping>* ExistingMapping = Mappings.Find(Pair);
-		if(ExistingMapping && ExistingMapping->IsValid())
+
+		const TSharedPtr<FSkeletonRemapping>* ExistingMapping = Mappings.FindByHash(PairHash, Pair);
+		if (ExistingMapping && ExistingMapping->IsValid())
 		{
-			return *(*ExistingMapping).Get();
+			return *ExistingMapping->Get();
 		}
 	}
 
-	// No existing, so create a new mapping
+	// No valid mapping was found, so create a new one
 	TSharedPtr<FSkeletonRemapping> NewMapping = MakeShared<FSkeletonRemapping>(InSourceSkeleton, InTargetSkeleton);
 	
 	{
 		FRWScopeLock WriteLock(MappingsLock, SLT_Write);
 
+		// First check if another thread grabbed the write lock before us to do the same mapping
+		const TSharedPtr<FSkeletonRemapping>* ExistingMapping = Mappings.FindByHash(PairHash, Pair);
+		if (ExistingMapping && ExistingMapping->IsValid())
+		{
+			return *ExistingMapping->Get();
+		}
+
 		// Add to global mapping
-		Mappings.Add(Pair, NewMapping);
+		Mappings.AddByHash(PairHash, Pair, NewMapping);
 
 		// Add to per-skeleton mappings
 		PerSkeletonMappings.Add(InSourceSkeleton, NewMapping);
