World Partition Runtime HashSet: remove wp.Editor.WorldPartitionRuntimeHashSet.Enable and make the HashSet class visible by default.

diff --git a/Engine/Source/Runtime/Engine/Private/WorldPartition/RuntimeHashSet/WorldPartitionRuntimeHashSet.cpp b/Engine/Source/Runtime/Engine/Private/WorldPartition/RuntimeHashSet/WorldPartitionRuntimeHashSet.cpp
--- a/Engine/Source/Runtime/Engine/Private/WorldPartition/RuntimeHashSet/WorldPartitionRuntimeHashSet.cpp
+++ b/Engine/Source/Runtime/Engine/Private/WorldPartition/RuntimeHashSet/WorldPartitionRuntimeHashSet.cpp
@@ -8,15 +8,6 @@
 #include "WorldPartition/DataLayer/DataLayersID.h"
 #include "Misc/ArchiveMD5.h"
 
-FAutoConsoleCommand WorldPartitionRuntimeHashSetEnable(
-	TEXT("wp.Editor.WorldPartitionRuntimeHashSet.Enable"),
-	TEXT("Enable experimental runtime hash set class."),
-	FConsoleCommandWithArgsDelegate::CreateLambda([](const TArray<FString>& Args)
-	{
-		UWorldPartitionRuntimeHashSet::StaticClass()->ClassFlags &= ~CLASS_HideDropDown;
-	})
-);
-
 void FRuntimePartitionStreamingData::CreatePartitionsSpatialIndex() const
 {
 	if (!SpatialIndex)
diff --git a/Engine/Source/Runtime/Engine/Public/WorldPartition/RuntimeHashSet/WorldPartitionRuntimeHashSet.h b/Engine/Source/Runtime/Engine/Public/WorldPartition/RuntimeHashSet/WorldPartitionRuntimeHashSet.h
--- a/Engine/Source/Runtime/Engine/Public/WorldPartition/RuntimeHashSet/WorldPartitionRuntimeHashSet.h
+++ b/Engine/Source/Runtime/Engine/Public/WorldPartition/RuntimeHashSet/WorldPartitionRuntimeHashSet.h
@@ -115,7 +115,7 @@ public:
 	TArray<FRuntimePartitionStreamingData> RuntimeStreamingData;
 };
 
-UCLASS(HideDropdown, MinimalAPI)
+UCLASS(MinimalAPI)
 class UWorldPartitionRuntimeHashSet : public UWorldPartitionRuntimeHash
 {
 	GENERATED_UCLASS_BODY()
