Do not create render resources with null rhi for a number of classes.

diff --git a/Engine/Source/Runtime/Engine/Private/Materials/MaterialInstance.cpp b/Engine/Source/Runtime/Engine/Private/Materials/MaterialInstance.cpp
--- a/Engine/Source/Runtime/Engine/Private/Materials/MaterialInstance.cpp
+++ b/Engine/Source/Runtime/Engine/Private/Materials/MaterialInstance.cpp
@@ -711,7 +711,7 @@ void UMaterialInstance::PostInitProperties()
 	LLM_SCOPE(ELLMTag::MaterialInstance);
 	Super::PostInitProperties();
 
-	if(!HasAnyFlags(RF_ClassDefaultObject))
+	if(!HasAnyFlags(RF_ClassDefaultObject) && FApp::CanEverRenderOrProduceRenderData())
 	{
 		Resource = new FMaterialInstanceResource(this);
 		UsedByRT |= (uint32)EMaterialInstanceUsedByRTFlag::ResourceCreate;
@@ -723,7 +723,7 @@ void UMaterialInstance::PostInitProperties()
  */
 void GameThread_InitMIParameters(const UMaterialInstance& Instance)
 {
-	if (Instance.HasAnyFlags(RF_ClassDefaultObject))
+	if (Instance.HasAnyFlags(RF_ClassDefaultObject) || !FApp::CanEverRender())
 	{
 		return;
 	}
@@ -3333,8 +3333,11 @@ void UMaterialInstance::FinishDestroy()
 {
 	if(!HasAnyFlags(RF_ClassDefaultObject))
 	{
-		Resource->GameThread_Destroy();
-		Resource = nullptr;
+		if (Resource)
+		{
+			Resource->GameThread_Destroy();
+			Resource = nullptr;
+		}
 	}
 
 	for (FMaterialResource* CurrentResource : StaticPermutationMaterialResources)
diff --git a/Engine/Source/Runtime/Engine/Private/Materials/ParameterCollection.cpp b/Engine/Source/Runtime/Engine/Private/Materials/ParameterCollection.cpp
--- a/Engine/Source/Runtime/Engine/Private/Materials/ParameterCollection.cpp
+++ b/Engine/Source/Runtime/Engine/Private/Materials/ParameterCollection.cpp
@@ -628,7 +628,7 @@ void UMaterialParameterCollectionInstance::PostInitProperties()
 {
 	Super::PostInitProperties();
 
-	if (!HasAnyFlags(RF_ClassDefaultObject))
+	if (!HasAnyFlags(RF_ClassDefaultObject) && FApp::CanEverRender())
 	{
 		Resource = new FMaterialParameterCollectionInstanceResource();
 	}
diff --git a/Engine/Source/Runtime/Engine/Private/VT/RuntimeVirtualTexture.cpp b/Engine/Source/Runtime/Engine/Private/VT/RuntimeVirtualTexture.cpp
--- a/Engine/Source/Runtime/Engine/Private/VT/RuntimeVirtualTexture.cpp
+++ b/Engine/Source/Runtime/Engine/Private/VT/RuntimeVirtualTexture.cpp
@@ -275,14 +275,20 @@ URuntimeVirtualTexture::URuntimeVirtualTexture(const FObjectInitializer& ObjectI
 	: Super(ObjectInitializer)
 {
 	// Initialize the RHI resources with a null producer
-	Resource = new FRuntimeVirtualTextureRenderResource;
-	InitNullResource();
+	if (!HasAnyFlags(RF_ClassDefaultObject) && FApp::CanEverRender())
+	{
+		Resource = new FRuntimeVirtualTextureRenderResource;
+		InitNullResource();
+	}
 }
 
 URuntimeVirtualTexture::~URuntimeVirtualTexture()
 {
-	Resource->Release();
-	delete Resource;
+	if (Resource)
+	{
+		Resource->Release();
+		delete Resource;
+	}
 }
 
 int32 URuntimeVirtualTexture::GetMaxTileCountLog2(bool InAdaptive) 
@@ -508,12 +514,12 @@ bool URuntimeVirtualTexture::IsLayerYCoCg(int32 LayerIndex) const
 
 FVirtualTextureProducerHandle URuntimeVirtualTexture::GetProducerHandle() const
 {
-	return Resource->GetProducerHandle();
+	return Resource ? Resource->GetProducerHandle() : FVirtualTextureProducerHandle();
 }
 
 IAllocatedVirtualTexture* URuntimeVirtualTexture::GetAllocatedVirtualTexture() const
 {
-	return Resource->GetAllocatedVirtualTexture();
+	return Resource ? Resource->GetAllocatedVirtualTexture() : nullptr;
 }
 
 FVector4 URuntimeVirtualTexture::GetUniformParameter(int32 Index) const
diff --git a/Engine/Source/Runtime/Engine/Private/VectorField.cpp b/Engine/Source/Runtime/Engine/Private/VectorField.cpp
--- a/Engine/Source/Runtime/Engine/Private/VectorField.cpp
+++ b/Engine/Source/Runtime/Engine/Private/VectorField.cpp
@@ -1203,12 +1203,15 @@ UVectorFieldAnimated::UVectorFieldAnimated(const FObjectInitializer& ObjectIniti
 
 void UVectorFieldAnimated::InitInstance(FVectorFieldInstance* Instance, bool bPreviewInstance)
 {
-	FVectorFieldAnimatedResource* Resource = new FVectorFieldAnimatedResource(this);
-	if (!bPreviewInstance)
+	if (FApp::CanEverRender())
 	{
-		BeginInitResource(Resource);
+		FVectorFieldAnimatedResource* Resource = new FVectorFieldAnimatedResource(this);
+		if (!bPreviewInstance)
+		{
+			BeginInitResource(Resource);
+		}
+		Instance->Init(Resource, /*bInstanced=*/ true);
 	}
-	Instance->Init(Resource, /*bInstanced=*/ true);
 }
 
 static int32 ClampVolumeSize(int32 InVolumeSize)
