When flushing net dormancy, init object replicators with default state (instead of current state) if the client doesn't have the actor's level loaded. Ensures the correct replicated state is sent to the client.
Can be disabled with cvar net.FlushDormancyUseDefaultStateForUnloadedLevels.

diff --git a/Engine/Source/Runtime/Engine/Private/NetConnection.cpp b/Engine/Source/Runtime/Engine/Private/NetConnection.cpp
--- a/Engine/Source/Runtime/Engine/Private/NetConnection.cpp
+++ b/Engine/Source/Runtime/Engine/Private/NetConnection.cpp
@@ -213,6 +213,9 @@ namespace UE::Net::Connection::Private
 	int32 bEnableFlushDormantSubObjectsCheckConditions = true;
 	FAutoConsoleVariableRef CVarNetFlushDormantSubObjectsCheckConditions(TEXT("net.EnableFlushDormantSubObjectsCheckConditions"), bEnableFlushDormantSubObjectsCheckConditions, TEXT("If enabled, when net.EnableFlushDormantSubObjects is also true a dormancy flush will also check replicated subobject conditions"));
 
+	int32 bFlushDormancyUseDefaultStateForUnloadedLevels = true;
+	FAutoConsoleVariableRef CVarFlushDormancyUseDefaultStateForUnloadedLevels(TEXT("net.FlushDormancyUseDefaultStateForUnloadedLevels"), bFlushDormancyUseDefaultStateForUnloadedLevels, TEXT("If enabled, dormancy flushing will init replicators with default object state if the client doesn't have the actor's level loaded."));
+
 	// Tracking for dormancy-flushed subobjects for correct deletion, see UE-77163
 	void TrackFlushedSubObject(FDormantObjectMap& InOutFlushedGuids, UObject* FlushedObject, const TSharedPtr<FNetGUIDCache>& GuidCache)
 	{
@@ -5062,6 +5065,8 @@ void UNetConnection::ForcePropertyCompare( AActor* Actor )
 
 void UNetConnection::FlushDormancyForObject(AActor* DormantActor, UObject* ReplicatedObject)
 {
+	using namespace UE::Net::Connection::Private;
+
 	QUICK_SCOPE_CYCLE_COUNTER(STAT_NetConnection_FlushDormancyForObject)
 
 	bool bReuseReplicators = false;
@@ -5092,8 +5097,15 @@ void UNetConnection::FlushDormancyForObject(AActor* DormantActor, UObject* Repli
 		bool bOverwroteExistingReplicator = false;
 		const TSharedRef<FObjectReplicator>& ObjectReplicatorRef = DormantReplicatorSet.CreateAndStoreReplicator(DormantActor, ReplicatedObject, bOverwroteExistingReplicator);
 		
-		// Init using the objects current state
-		constexpr bool bUseDefaultState = false; 
+		bool bUseDefaultState = false;
+		
+		// Init using the object's current state if the client has this actor's level loaded. If the level
+		// is unloaded, we need to use the default state since the client has no current state.
+		if (bFlushDormancyUseDefaultStateForUnloadedLevels && !ClientHasInitializedLevelFor(DormantActor))
+		{
+			bUseDefaultState = true;
+		}
+
 		ObjectReplicatorRef->InitWithObject(ReplicatedObject, this, bUseDefaultState);
 
 #if UE_REPLICATED_OBJECT_REFCOUNTING
