Use a read only LockMips to prevent a source guid change, unintentionally invalidating the ddc key.

diff --git a/Engine/Source/Runtime/Engine/Private/PhysicsEngine/PhysicalMaterialMask.cpp b/Engine/Source/Runtime/Engine/Private/PhysicsEngine/PhysicalMaterialMask.cpp
--- a/Engine/Source/Runtime/Engine/Private/PhysicsEngine/PhysicalMaterialMask.cpp
+++ b/Engine/Source/Runtime/Engine/Private/PhysicsEngine/PhysicalMaterialMask.cpp
@@ -248,8 +248,8 @@ template<typename PixelDataType, int32 RIdx, int32 GIdx, int32 BIdx, int32 AIdx>
 {
 public:
 
-	MaskDataGenerator(int32 SizeX, int32 SizeY, uint8* SourceTextureData)
-		: SourceData(reinterpret_cast<PixelDataType*>(SourceTextureData))
+	MaskDataGenerator(int32 SizeX, int32 SizeY, const uint8* SourceTextureData)
+		: SourceData(reinterpret_cast<const PixelDataType*>(SourceTextureData))
 		, TextureWidth(SizeX)
 		, TextureHeight(SizeY)
 	{
@@ -261,7 +261,7 @@ public:
 
 		OutMaskData.Empty();
 
-		PixelDataType* PixelData = SourceData;
+		const PixelDataType* PixelData = SourceData;
 
 		for (int32 Y = 0; Y < TextureHeight; ++Y)
 		{
@@ -322,7 +322,7 @@ public:
 		return UPhysicalMaterialMask::INVALID_MASK_INDEX;
 	}
 
-	PixelDataType* SourceData;
+	const PixelDataType* SourceData;
 	int32 TextureWidth;
 	int32 TextureHeight;
 };
@@ -336,7 +336,7 @@ void UPhysicalMaterialMask::GenerateMaskData(TArray<uint32>& OutMaskData, int32&
 
 	if (MaskTexture)
 	{
-		uint8* TextureData = MaskTexture->Source.LockMip(0);
+		const uint8* TextureData = MaskTexture->Source.LockMipReadOnly(0);
 		if (TextureData)
 		{
 			const int32 TextureDataSize = MaskTexture->Source.CalcMipSize(0);
