World Partition: remove actors from dirty list when we undo the transcation that they were created in.

diff --git a/Engine/Source/Runtime/Engine/Private/WorldPartition/WorldPartition.cpp b/Engine/Source/Runtime/Engine/Private/WorldPartition/WorldPartition.cpp
--- a/Engine/Source/Runtime/Engine/Private/WorldPartition/WorldPartition.cpp
+++ b/Engine/Source/Runtime/Engine/Private/WorldPartition/WorldPartition.cpp
@@ -385,9 +385,17 @@ void UWorldPartition::OnPackageDirtyStateChanged(UPackage* Package)
 				DirtyActors.Add(Actor->GetActorGuid(), FDirtyActor(ActorHandle.ToReference(), Actor));
 			}
 		}
+		else if (DirtyActors.Contains(Actor->GetActorGuid()))
+		{
+			if (!Package->IsDirty())
+			{
+				// Remove newly created, unsaved actor that dissapeared after an undo
+				DirtyActors.Remove(Actor->GetActorGuid());
+			}
+		}
 		else
 		{
-			// This is handling a new actor (unsaved).
+			// Add newly created, unsaved actor after editor placement or a redo
 			DirtyActors.Add(Actor->GetActorGuid(), FDirtyActor(Actor));
 		}
 	}
