ISM-related memory savings on the server.

- Stop serializing old instance data for the server (it was dropped on load anyway)
- Hollow out Primitive Instance Data manager (reducing sizeof of each ISM by 512b)
- Do not maintain instance reorder tables that map instances to rendered instances

diff --git a/Engine/Source/Runtime/Engine/Classes/Components/HierarchicalInstancedStaticMeshComponent.h b/Engine/Source/Runtime/Engine/Classes/Components/HierarchicalInstancedStaticMeshComponent.h
--- a/Engine/Source/Runtime/Engine/Classes/Components/HierarchicalInstancedStaticMeshComponent.h
+++ b/Engine/Source/Runtime/Engine/Classes/Components/HierarchicalInstancedStaticMeshComponent.h
@@ -209,7 +209,14 @@ public:
 	{
 		TArray<FClusterNode> Nodes;
 		TArray<int32> SortedInstances;
+#if !USE_NULL_RHI
 		TArray<int32> InstanceReorderTable;
+		inline int32 GetNumBuiltInstances() const { return InstanceReorderTable.Num(); }
+#else
+		// some code still needs to know how many instances are there and it relied on the InstanceReorderTable size for that
+		int32 NumInstancesTotal = 0;
+		inline int32 GetNumBuiltInstances() const { return NumInstancesTotal; }
+#endif
 		int32 OutOcclusionLayerNum = 0;
 
 		bool PrintLevel(int32 NodeIndex, int32 Level, int32 CurrentLevel, int32 Parent);
diff --git a/Engine/Source/Runtime/Engine/Private/HierarchicalInstancedStaticMesh.cpp b/Engine/Source/Runtime/Engine/Private/HierarchicalInstancedStaticMesh.cpp
--- a/Engine/Source/Runtime/Engine/Private/HierarchicalInstancedStaticMesh.cpp
+++ b/Engine/Source/Runtime/Engine/Private/HierarchicalInstancedStaticMesh.cpp
@@ -288,6 +288,7 @@ void UHierarchicalInstancedStaticMeshComponent::FClusterBuilder::Split(int32 Sta
 
 void UHierarchicalInstancedStaticMeshComponent::FClusterBuilder::BuildInstanceBuffer()
 {
+#if !USE_NULL_RHI
 	// build new instance buffer
 	FRandomStream RandomStream = FRandomStream(InstancingRandomSeed);
 	BuiltInstanceData = MakeUnique<FStaticMeshInstanceData>(/*bInUseHalfFloat = */true);
@@ -319,6 +320,7 @@ void UHierarchicalInstancedStaticMeshComponent::FClusterBuilder::BuildInstanceBu
 			// correct light/shadow map bias will be setup on game thread side if needed
 		}
 	}
+#endif
 }
 
 void UHierarchicalInstancedStaticMeshComponent::FClusterBuilder::Init()
@@ -393,7 +395,11 @@ void UHierarchicalInstancedStaticMeshComponent::FClusterBuilder::BuildTree()
 		// Can happen if all instances are excluded due to scalability
 		// It doesn't only happen with a scalability factor of 0 - 
 		// even with a scalability factor of 0.99, if there's only one instance of this type you can end up with Num == 0 if you're unlucky
+#if !USE_NULL_RHI
 		Result->InstanceReorderTable.Init(INDEX_NONE, OriginalNum);
+#else
+		Result->NumInstancesTotal = OriginalNum;
+#endif
 		return;
 	}
 
@@ -626,12 +632,16 @@ void UHierarchicalInstancedStaticMeshComponent::FClusterBuilder::BuildTree()
 		}
 	}
 
+#if !USE_NULL_RHI
 	// Save inverse map
 	Result->InstanceReorderTable.Init(INDEX_NONE, OriginalNum);
 	for (int32 Index = 0; Index < Num; Index++)
 	{
 		Result->InstanceReorderTable[Result->SortedInstances[Index]] = Index;
 	}
+#else
+	Result->NumInstancesTotal = OriginalNum;
+#endif
 
 	// Output a general scale of 1 if we dont want the scaling range
 	if (!GenerateInstanceScalingRange)
@@ -2657,7 +2667,7 @@ void UHierarchicalInstancedStaticMeshComponent::ApplyBuildTreeAsync(ENamedThread
 	{
 		bConcurrentChanges = false;
 
-		UE_LOG(LogStaticMesh, Verbose, TEXT("Discarded foliage hierarchy of %d elements build due to concurrent removal (%.1fs)"), Builder->Result->InstanceReorderTable.Num(), (float)(FPlatformTime::Seconds() - StartTime));
+		UE_LOG(LogStaticMesh, Verbose, TEXT("Discarded foliage hierarchy of %d elements build due to concurrent removal (%.1fs)"), Builder->Result->GetNumBuiltInstances(), (float)(FPlatformTime::Seconds() - StartTime));
 
 		// There were changes while we were building, it's too slow to fix up the result now, so build async again.
 		BuildTreeAsync();
@@ -2688,14 +2698,18 @@ void UHierarchicalInstancedStaticMeshComponent::ApplyBuildTree(FClusterBuilder&
 {
 	bIsOutOfDate = false;
 
-	check(Builder.Result->InstanceReorderTable.Num() == PerInstanceSMData.Num());
+	check(Builder.Result->GetNumBuiltInstances() == PerInstanceSMData.Num());
 
-	NumBuiltInstances = Builder.Result->InstanceReorderTable.Num();
+	NumBuiltInstances = Builder.Result->GetNumBuiltInstances();
 	NumBuiltRenderInstances = Builder.Result->SortedInstances.Num();
 
 	ClusterTreePtr = MakeShareable(new TArray<FClusterNode>(MoveTemp(Builder.Result->Nodes)));
 
+#if !USE_NULL_RHI
 	InstanceReorderTable = MoveTemp(Builder.Result->InstanceReorderTable);
+#else
+	InstanceReorderTable.Reset();
+#endif
 	SortedInstances = MoveTemp(Builder.Result->SortedInstances);
 	CacheMeshExtendedBounds = GetStaticMesh()->GetBounds();
 	TUniquePtr<FStaticMeshInstanceData> BuiltInstanceData = MoveTemp(Builder.BuiltInstanceData);
diff --git a/Engine/Source/Runtime/Engine/Private/InstancedStaticMesh.cpp b/Engine/Source/Runtime/Engine/Private/InstancedStaticMesh.cpp
--- a/Engine/Source/Runtime/Engine/Private/InstancedStaticMesh.cpp
+++ b/Engine/Source/Runtime/Engine/Private/InstancedStaticMesh.cpp
@@ -2928,6 +2928,11 @@ struct FInstancedStaticMeshInstanceData_DEPRECATED
 static bool NeedRenderDataForTargetPlatform(const ITargetPlatform* TargetPlatform)
 {
 #if WITH_EDITOR
+	if (TargetPlatform->IsServerOnly())
+	{
+		return false;
+	}
+
 	const UDeviceProfile* DeviceProfile = UDeviceProfileManager::Get().FindProfile(TargetPlatform->IniPlatformName());
 	if (DeviceProfile)
 	{
@@ -3082,6 +3087,11 @@ void UInstancedStaticMeshComponent::Serialize(FArchive& Ar)
 		Ar << SelectedInstances;
 	}
 #endif
+
+	if (!FApp::CanEverRenderOrProduceRenderData() && InstanceReorderTable.GetAllocatedSize() > 0)
+	{
+		InstanceReorderTable.Empty();
+	}
 }
 
 void UInstancedStaticMeshComponent::PreAllocateInstancesMemory(int32 AddedInstanceCount)
diff --git a/Engine/Source/Runtime/Engine/Private/InstancedStaticMesh/ISMInstanceDataManager.cpp b/Engine/Source/Runtime/Engine/Private/InstancedStaticMesh/ISMInstanceDataManager.cpp
--- a/Engine/Source/Runtime/Engine/Private/InstancedStaticMesh/ISMInstanceDataManager.cpp
+++ b/Engine/Source/Runtime/Engine/Private/InstancedStaticMesh/ISMInstanceDataManager.cpp
@@ -36,7 +36,8 @@ static TAutoConsoleVariable<int32> CVarInstanceDataResetTrackingOnRegister(
 
 
 FPrimitiveInstanceDataManager::FPrimitiveInstanceDataManager(UPrimitiveComponent* InPrimitiveComponent) 
-	: PrimitiveComponent(InPrimitiveComponent) 
+#if !USE_NULL_RHI
+	: PrimitiveComponent(InPrimitiveComponent)
 {
 	// Don't do anything if this is not a "real" ISM being tracked (this logic shopuld move out).
 	if (PrimitiveComponent.IsValid() && PrimitiveComponent->HasAnyFlags(RF_ClassDefaultObject | RF_ArchetypeObject))
@@ -45,9 +46,14 @@ FPrimitiveInstanceDataManager::FPrimitiveInstanceDataManager(UPrimitiveComponent
 	}
 	LOG_INST_DATA(TEXT("FPrimitiveInstanceDataManager %s, TrackingState=%s"), *PrimitiveComponent->GetFullName(), TrackingState == ETrackingState::Disabled ? TEXT("Disabled") : TEXT("Initial"));
 }
+#else
+{
+}
+#endif
 
 void FPrimitiveInstanceDataManager::SetMode(EMode InMode)
 {
+#if !USE_NULL_RHI
 	if (InMode != Mode)
 	{
 		// This should never be called in mid-use
@@ -55,6 +61,7 @@ void FPrimitiveInstanceDataManager::SetMode(EMode InMode)
 		Invalidate(0);
 	}
 	Mode = InMode;
+#endif
 }
 
 void FPrimitiveInstanceDataManager::Add(int32 InInstanceAddAtIndex, bool bInsert)
@@ -64,6 +71,7 @@ void FPrimitiveInstanceDataManager::Add(int32 InInstanceAddAtIndex, bool bInsert
 		return;
 	}
 
+#if !USE_NULL_RHI
 	ValidateMapping();
 
 	// If the manager is marked for legacy-only mode, we should not see any tracking calls!
@@ -131,6 +139,7 @@ void FPrimitiveInstanceDataManager::Add(int32 InInstanceAddAtIndex, bool bInsert
 	LOG_INST_DATA(TEXT("Add(IDX: %d, bInsert: %d) -> Id: %d"), InInstanceAddAtIndex, bInsert, InstanceId.Id);
 
 	ValidateMapping();
+#endif
 }
 
 void FPrimitiveInstanceDataManager::RemoveAtSwap(int32 InstanceIndex)
@@ -140,6 +149,7 @@ void FPrimitiveInstanceDataManager::RemoveAtSwap(int32 InstanceIndex)
 		return;
 	}
 
+#if !USE_NULL_RHI
 	ValidateMapping();
 
 	check(Mode != EMode::ExternalLegacyData);
@@ -180,6 +190,7 @@ void FPrimitiveInstanceDataManager::RemoveAtSwap(int32 InstanceIndex)
 	}
 	ValidateMapping();
 	LOG_INST_DATA(TEXT("RemoveAtSwap(IDX: %d) -> Id: %d"), InstanceIndex, InstanceId.Id);
+#endif
 }
 	
 void FPrimitiveInstanceDataManager::RemoveAt(int32 InstanceIndex)
@@ -189,6 +200,7 @@ void FPrimitiveInstanceDataManager::RemoveAt(int32 InstanceIndex)
 		return;
 	}
 
+#if !USE_NULL_RHI
 	ValidateMapping();
 
 	check(Mode != EMode::ExternalLegacyData);
@@ -233,18 +245,23 @@ void FPrimitiveInstanceDataManager::RemoveAt(int32 InstanceIndex)
 	LOG_INST_DATA(TEXT("RemoveAt(IDX: %d) -> Id: %d"), InstanceIndex, InstanceId.Id);
 
 	ValidateMapping();
+#endif
 }
 
 void FPrimitiveInstanceDataManager::TransformChanged(int32 InstanceIndex)
 {
+#if !USE_NULL_RHI
 	LOG_INST_DATA(TEXT("TransformChanged(IDX: %d)"), InstanceIndex);
 	MarkChangeHelper<EChangeFlag::TransformChanged>( InstanceIndex);
+#endif
 }
 
 void FPrimitiveInstanceDataManager::TransformChanged(FPrimitiveInstanceId InstanceId)
 {
+#if !USE_NULL_RHI
 	LOG_INST_DATA(TEXT("TransformChanged(ID: %d)"), InstanceId.Id);
 	MarkChangeHelper<EChangeFlag::TransformChanged>(InstanceId);
+#endif
 }
 
 void FPrimitiveInstanceDataManager::TransformsChangedAll()
@@ -254,22 +271,28 @@ void FPrimitiveInstanceDataManager::TransformsChangedAll()
 		return;
 	}
 
+#if !USE_NULL_RHI
 	LOG_INST_DATA(TEXT("TransformsChangedAll(%s)"), TEXT(""));
 	bTransformChangedAllInstances = true;
 	MarkComponentRenderInstancesDirty();
+#endif
 }
 
 void FPrimitiveInstanceDataManager::CustomDataChanged(int32 InstanceIndex)
 {
+#if !USE_NULL_RHI
 	LOG_INST_DATA(TEXT("CustomDataChanged(IDX: %d)"), InstanceIndex);
 	MarkChangeHelper<EChangeFlag::CustomDataChanged>(InstanceIndex);
+#endif
 }
 
 void FPrimitiveInstanceDataManager::BakedLightingDataChanged(int32 InstanceIndex)
 {
+#if !USE_NULL_RHI
 	LOG_INST_DATA(TEXT("BakedLightingDataChanged(IDX: %d)"), InstanceIndex);
 	bBakedLightingDataChanged = true;
 	MarkComponentRenderInstancesDirty();
+#endif
 }
 
 void FPrimitiveInstanceDataManager::BakedLightingDataChangedAll()
@@ -279,9 +302,11 @@ void FPrimitiveInstanceDataManager::BakedLightingDataChangedAll()
 		return;
 	}
 
+#if !USE_NULL_RHI
 	LOG_INST_DATA(TEXT("BakedLightingDataChangedAll(%s)"), TEXT(""));
 	bBakedLightingDataChanged = true;
 	MarkComponentRenderInstancesDirty();
+#endif
 }
 
 void FPrimitiveInstanceDataManager::NumCustomDataChanged()
@@ -291,9 +316,11 @@ void FPrimitiveInstanceDataManager::NumCustomDataChanged()
 		return;
 	}
 
+#if !USE_NULL_RHI
 	LOG_INST_DATA(TEXT("NumCustomDataChanged(%s)"), TEXT(""));
 	bNumCustomDataChanged = true;
 	MarkComponentRenderInstancesDirty();
+#endif
 }
 
 #if WITH_EDITOR
@@ -319,18 +346,24 @@ void FPrimitiveInstanceDataManager::PrimitiveTransformChanged()
 		return;
 	}
 
+#if !USE_NULL_RHI
 	LOG_INST_DATA(TEXT("PrimitiveTransformChanged(%s)"), TEXT(""));
 	bPrimitiveTransformChanged = true;
 	MarkComponentRenderInstancesDirty();
+#endif
 }
 
 bool FPrimitiveInstanceDataManager::HasAnyInstanceChanges() const
 {
-	return InstanceUpdateTracker.HasAnyChanges() 
+#if !USE_NULL_RHI
+	return InstanceUpdateTracker.HasAnyChanges()
 #if WITH_EDITOR
 		|| bAnyEditorDataChanged 
 #endif
 		|| bNumCustomDataChanged || bBakedLightingDataChanged || bTransformChangedAllInstances;
+#else
+	return false;
+#endif // !USE_NULL_RHI
 }
 
 
@@ -353,6 +386,7 @@ void FPrimitiveInstanceDataManager::SerializeRenderData(FArchive& Ar, bool bCook
 	}
 }
 
+#if !USE_NULL_RHI
 /**
  * Describes what has changed, that can be derived from the primitive desc, or internal tracking state.
  */
@@ -485,9 +519,11 @@ struct FLegacyRebuildChangeSet
 #endif
 	int32 NumCustomDataFloats = 0; 
 };
+#endif // !USE_NULL_RHI
 
 void FPrimitiveInstanceDataManager::InitChangeSet(const FChangeDesc &ChangeDesc, const FInstanceUpdateComponentDesc &ComponentData, FISMInstanceUpdateChangeSet &ChangeSet)
 {
+#if !USE_NULL_RHI
 	// Collect the delta data to be able to update the index mapping.
 	ChangeSet.MaxInstanceId = GetMaxInstanceId();
 	ChangeSet.bIdentityIdMap = IsIdentity();
@@ -516,6 +552,7 @@ void FPrimitiveInstanceDataManager::InitChangeSet(const FChangeDesc &ChangeDesc,
 
 	// This is the odd one out
 	ChangeSet.SetInstanceLocalBounds(ComponentData.StaticMeshBounds);
+#endif // !USE_NULL_RHI
 }
 
 bool FPrimitiveInstanceDataManager::FlushChanges(FInstanceUpdateComponentDesc &&ComponentData, bool bNewPrimitiveProxy)
@@ -525,6 +562,7 @@ bool FPrimitiveInstanceDataManager::FlushChanges(FInstanceUpdateComponentDesc &&
 		return false;
 	}
 
+#if !USE_NULL_RHI
 	// Always clear the flag such that any subsequent change marks it as needing update again.
 	bComponentMarkedDirty = false;
 
@@ -771,7 +809,7 @@ bool FPrimitiveInstanceDataManager::FlushChanges(FInstanceUpdateComponentDesc &&
 
 	// After an update has been sent, we need to track all deltas.
 	TrackingState = ETrackingState::Tracked;
-
+#endif
 	return true;
 }
 
@@ -782,6 +820,7 @@ void FPrimitiveInstanceDataManager::PostLoad(int32 InNumInstances, TUniquePtr<FS
 		return;
 	}
 
+#if !USE_NULL_RHI
 	if (LegacyStaticMeshInstanceData.IsValid())
 	{
 		check(NumInstances == LegacyStaticMeshInstanceData->GetNumInstances());
@@ -808,6 +847,7 @@ void FPrimitiveInstanceDataManager::PostLoad(int32 InNumInstances, TUniquePtr<FS
 #endif
 	}
 	NumInstances = InNumInstances;
+#endif // !USE_NULL_RHI
 }
 
 void FPrimitiveInstanceDataManager::ClearIdTracking(int32 InNumInstances)
@@ -817,6 +857,7 @@ void FPrimitiveInstanceDataManager::ClearIdTracking(int32 InNumInstances)
 		return;
 	}
 
+#if !USE_NULL_RHI
 	// Reset the mapping to identity & clear allocator, this looses all association with ID:Index that existed before
 	IndexToIdMap.Empty();
 	IdToIndexMap.Empty();
@@ -826,6 +867,7 @@ void FPrimitiveInstanceDataManager::ClearIdTracking(int32 InNumInstances)
 
 	// Also clear the change tracking since it is not valid anymore
 	ClearChangeTracking();
+#endif
 }
 
 void FPrimitiveInstanceDataManager::ClearChangeTracking()
@@ -835,6 +877,7 @@ void FPrimitiveInstanceDataManager::ClearChangeTracking()
 		return;
 	}
 
+#if !USE_NULL_RHI
 	// When tracking data is cleared, we loose connection to previously tracked state until the next update is sent.
 	TrackingState = ETrackingState::Initial;
 
@@ -848,20 +891,30 @@ void FPrimitiveInstanceDataManager::ClearChangeTracking()
 	bAnyEditorDataChanged = false;
 #endif	
 	bPrimitiveTransformChanged = false;
+#endif
 }
 
 int32 FPrimitiveInstanceDataManager::GetMaxInstanceId() const
 {
+#if !USE_NULL_RHI
 	return HasIdentityMapping() ? NumInstances : ValidInstanceIdMask.Num();
+#else
+	return 0;
+#endif
 }
 
 int32 FPrimitiveInstanceDataManager::GetMaxInstanceIndex() const
 {
+#if !USE_NULL_RHI
 	return HasIdentityMapping() ? NumInstances : IndexToIdMap.Num();
+#else
+	return 0;
+#endif
 }
 
 void FPrimitiveInstanceDataManager::CreateExplicitIdentityMapping()
 {
+#if !USE_NULL_RHI
 	check(HasIdentityMapping());
 	IndexToIdMap.SetNumUninitialized(NumInstances);
 	IdToIndexMap.SetNumUninitialized(NumInstances);
@@ -873,11 +926,13 @@ void FPrimitiveInstanceDataManager::CreateExplicitIdentityMapping()
 	ValidInstanceIdMask.Reset();
 	ValidInstanceIdMask.SetNum(NumInstances, true);
 	IdSearchStartIndex = NumInstances;
+#endif
 }
 
 template<FPrimitiveInstanceDataManager::EChangeFlag Flag>
 void FPrimitiveInstanceDataManager::MarkChangeHelper(int32 InstanceIndex)
 {
+#if !USE_NULL_RHI
 	if (GetState() == ETrackingState::Disabled)
 	{
 		return;
@@ -892,11 +947,13 @@ void FPrimitiveInstanceDataManager::MarkChangeHelper(int32 InstanceIndex)
 	}
 	InstanceUpdateTracker.MarkIndex<Flag>(InstanceIndex, GetMaxInstanceIndex());
 	MarkComponentRenderInstancesDirty();
+#endif
 }
 
 template<FPrimitiveInstanceDataManager::EChangeFlag Flag>
 void FPrimitiveInstanceDataManager::MarkChangeHelper(FPrimitiveInstanceId InstanceId)
 {
+#if !USE_NULL_RHI
 	check(Mode != EMode::ExternalLegacyData);
 
 	if (GetState() != ETrackingState::Tracked)
@@ -905,10 +962,12 @@ void FPrimitiveInstanceDataManager::MarkChangeHelper(FPrimitiveInstanceId Instan
 		return;
 	}
 	MarkChangeHelper<Flag>(IdToIndex(InstanceId));
+#endif
 }
 
 void FPrimitiveInstanceDataManager::MarkComponentRenderInstancesDirty()
 {
+#if !USE_NULL_RHI
 	if (!bComponentMarkedDirty)
 	{
 		if (UPrimitiveComponent *PrimitiveComponentPtr = PrimitiveComponent.Get())
@@ -917,20 +976,30 @@ void FPrimitiveInstanceDataManager::MarkComponentRenderInstancesDirty()
 			bComponentMarkedDirty = true;
 		}
 	}
+#endif
 }
 
 bool FPrimitiveInstanceDataManager::HasIdentityMapping() const
 {
+#if !USE_NULL_RHI
 	return IndexToIdMap.IsEmpty();
+#else
+	return false;
+#endif
 }
 
 bool FPrimitiveInstanceDataManager::ShouldTrackIds() const
 {
+#if !USE_NULL_RHI
 	return Proxy != nullptr;
+#else
+	return false;
+#endif
 }
 
 void FPrimitiveInstanceDataManager::FreeInstanceId(FPrimitiveInstanceId InstanceId)
 {
+#if !USE_NULL_RHI
 	LOG_INST_DATA(TEXT("FreeInstanceId(Id: %d)"), InstanceId.Id);
 
 	if (!HasIdentityMapping())
@@ -942,10 +1011,12 @@ void FPrimitiveInstanceDataManager::FreeInstanceId(FPrimitiveInstanceId Instance
 	}
 
 	LOG_INST_DATA(TEXT("IdToIndexMap[%d] = %d"), InstanceId.Id, INDEX_NONE);
+#endif
 }
 
 TSharedPtr<FISMCInstanceDataSceneProxy, ESPMode::ThreadSafe> FPrimitiveInstanceDataManager::GetOrCreateProxy(FStaticShaderPlatform InShaderPlatform, ERHIFeatureLevel::Type InFeatureLevel)
 {
+#if !USE_NULL_RHI
 	LOG_INST_DATA(TEXT("GetOrCreateProxy"));
 	if (Proxy && !Proxy->CheckPlatformFeatureLevel(InShaderPlatform, InFeatureLevel))
 	{
@@ -972,10 +1043,15 @@ TSharedPtr<FISMCInstanceDataSceneProxy, ESPMode::ThreadSafe> FPrimitiveInstanceD
 	}
 
 	return Proxy;
+#else
+	checkNoEntry();
+	return nullptr;
+#endif
 }
 
 void FPrimitiveInstanceDataManager::Invalidate(int32 InNumInstances)
 {
+#if !USE_NULL_RHI
 	LOG_INST_DATA(TEXT("Invalidate"));
 	Proxy.Reset();
 	ClearIdTracking(InNumInstances);
@@ -985,9 +1061,10 @@ void FPrimitiveInstanceDataManager::Invalidate(int32 InNumInstances)
 		PrimitiveComponentPtr->MarkRenderStateDirty();
 		bComponentMarkedDirty = true;
 	}
+#endif 
 }
 
-#if DO_GUARD_SLOW
+#if DO_GUARD_SLOW && !USE_NULL_RHI
 void FPrimitiveInstanceDataManager::ValidateMapping() const
 {
 	check(HasIdentityMapping() || IndexToIdMap.Num() == NumInstances);
@@ -1017,6 +1094,7 @@ void FPrimitiveInstanceDataManager::ValidateMapping() const
 
 void FPrimitiveInstanceDataManager::MarkForRebuildFromLegacy(TUniquePtr<FStaticMeshInstanceData> &&InLegacyInstanceData, const TArray<int32> &InstanceReorderTable, const TArray<TRefCountPtr<HHitProxy>> &HitProxies)
 {
+#if !USE_NULL_RHI
 	if (GetState() == ETrackingState::Disabled)
 	{
 		return;
@@ -1052,16 +1130,22 @@ void FPrimitiveInstanceDataManager::MarkForRebuildFromLegacy(TUniquePtr<FStaticM
 		TrackingState = ETrackingState::Tracked;
 	}
 	MarkComponentRenderInstancesDirty();
+#endif
 }
 
 SIZE_T FPrimitiveInstanceDataManager::GetAllocatedSize() const
 {
+#if !USE_NULL_RHI
 	return ValidInstanceIdMask.GetAllocatedSize() +
 		InstanceUpdateTracker.GetAllocatedSize();
+#else
+	return 0;
+#endif
 }
 
 void FPrimitiveInstanceDataManager::OnRegister(int32 InNumInstances)
 {
+#if !USE_NULL_RHI
 	LOG_INST_DATA(TEXT("OnRegister(InNumInstances : %d) NumInstances: %d"), InNumInstances, NumInstances);
 
 	if (CVarInstanceDataResetTrackingOnRegister.GetValueOnGameThread())
@@ -1075,4 +1159,5 @@ void FPrimitiveInstanceDataManager::OnRegister(int32 InNumInstances)
 			ClearChangeTracking();
 		}
 	}
+#endif
 }
diff --git a/Engine/Source/Runtime/Engine/Public/InstancedStaticMesh/ISMInstanceDataManager.h b/Engine/Source/Runtime/Engine/Public/InstancedStaticMesh/ISMInstanceDataManager.h
--- a/Engine/Source/Runtime/Engine/Public/InstancedStaticMesh/ISMInstanceDataManager.h
+++ b/Engine/Source/Runtime/Engine/Public/InstancedStaticMesh/ISMInstanceDataManager.h
@@ -42,7 +42,10 @@ struct FInstanceUpdateComponentDesc
  * NOTE/TODO: This is tied to the ISM use-case, mostly because of legacy (HISM) interactions. Will be refactored and sub-classed or something.
  *            Also: Still somewhat tied to the UComponent, which also can be refactored a bit to make it more general.
  */
-class FPrimitiveInstanceDataManager : public FInstanceIdIndexMap
+class FPrimitiveInstanceDataManager 
+#if !USE_NULL_RHI	// unparenting reduces the sizeof of the manager to 1
+	: public FInstanceIdIndexMap
+#endif
 {
 public:
 	ENGINE_API FPrimitiveInstanceDataManager(UPrimitiveComponent* InPrimitiveComponent);
@@ -64,13 +67,34 @@ public:
 		ExternalLegacyData // In this mode, it is illegal to call the incremental state tracking methods & updates can only be sent to the proxy if there has been and extranal one queued.
 	};
 
+#if USE_NULL_RHI	// functions used from FInstanceIdIndexMap - copied here for USE_NULL_RHI case since we unparented this class from it.
+	inline bool IsValidId(FPrimitiveInstanceId InstanceId) const { return false; }
+
+	inline int32 IdToIndex(FPrimitiveInstanceId InstanceId) const
+	{
+		return InstanceId.Id;
+	}
+
+	inline FPrimitiveInstanceId IndexToId(int32 InstanceIndex) const
+	{
+		return FPrimitiveInstanceId{ InstanceIndex };
+	}
+#endif
+
 	/**
 	 */
 	ENGINE_API void SetMode(EMode InMode);
 
 	/**
 	 */
-	ENGINE_API EMode GetMode() const { return Mode; }
+	ENGINE_API EMode GetMode() const
+	{ 
+#if !USE_NULL_RHI
+		return Mode;
+#else
+		return EMode::Default;
+#endif
+	}
 
 	void Add(int32 InInstanceAddAtIndex, bool bInsert);
 
@@ -123,7 +147,14 @@ public:
 
 	void CreateExplicitIdentityMapping();
 
-	ETrackingState GetState() const { return TrackingState; }
+	ETrackingState GetState() const
+	{ 
+#if !USE_NULL_RHI
+		return TrackingState;
+#else
+		return ETrackingState::Disabled;
+#endif
+	}
 
 	ENGINE_API TSharedPtr<FISMCInstanceDataSceneProxy, ESPMode::ThreadSafe> GetOrCreateProxy(FStaticShaderPlatform InShaderPlatform, ERHIFeatureLevel::Type InFeatureLevel);
 
@@ -145,7 +176,12 @@ public:
 	ENGINE_API SIZE_T GetAllocatedSize() const;
 
 
-	inline void ResetComponentDirtyTracking() { bComponentMarkedDirty = false; }
+	inline void ResetComponentDirtyTracking()
+	{ 
+#if !USE_NULL_RHI
+		bComponentMarkedDirty = false; 
+#endif
+	}
 
 	/**
 	 * Called by the corresponding function in the owner UPrimitiveComponent
@@ -178,6 +214,7 @@ private:
 
 	void InitChangeSet(const union FChangeDesc &ChangeDesc, const FInstanceUpdateComponentDesc &ComponentData, FISMInstanceUpdateChangeSet &ChangeSet);
 
+#if !USE_NULL_RHI
 	EMode Mode = EMode::Default;
 	ETrackingState TrackingState = ETrackingState::Initial;
 
@@ -215,4 +252,5 @@ private:
 	float AbsMaxDisplacement = 0.0f;
 	FRenderBounds StaticMeshBounds;
 	bool bFirstFlush = true;
+#endif
 };
diff --git a/Engine/Source/Runtime/Foliage/Private/InstancedGrass.cpp b/Engine/Source/Runtime/Foliage/Private/InstancedGrass.cpp
--- a/Engine/Source/Runtime/Foliage/Private/InstancedGrass.cpp
+++ b/Engine/Source/Runtime/Foliage/Private/InstancedGrass.cpp
@@ -37,7 +37,11 @@ void UGrassInstancedStaticMeshComponent::BuildTreeAnyThread(
 	OutOcclusionLayerNum = Builder.Result->OutOcclusionLayerNum;
 
 	OutClusterTree = MoveTemp(Builder.Result->Nodes);
+#if !USE_NULL_RHI
 	OutInstanceReorderTable = MoveTemp(Builder.Result->InstanceReorderTable);
+#else
+	OutInstanceReorderTable.Empty();
+#endif
 	OutSortedInstances = MoveTemp(Builder.Result->SortedInstances);
 }
 
