NetDriver:
* Added tracking stats to NetConnection that track number of rpcs delayed by unmapped guids and how many frames they were delayed by.
* Added tracking stats to NetDriver that track number of reliable bunches.
* Added console cheats to force only unreliable RPCs to be dropped. This can be used to better simulate a scenario where specific unreliable RPCs get dropped and not resent but not the adjacent property data sampled the same frame. Useful to test gameplay code using unreliable RPCs.
   * NetEmulation.DropAnyUnreliable : drops a % of all unreliables sent
   * NetEmulation.DropUnreliableOfActorClass: drops a % of all unreliables sent by a specific class
   * NetEmulation.DropUnreliableOfSubObjectClass: drops a % of all unreliables sent by a specific class
   * NetEmulation.DropUnreliableRPC: drops a % of all unreliables of the passed function name.
   * NetEmulation.DropNothing: stops dropping unreliable rpcs

diff --git a/Engine/Plugins/Runtime/ReplicationGraph/Source/Private/ReplicationGraph.cpp b/Engine/Plugins/Runtime/ReplicationGraph/Source/Private/ReplicationGraph.cpp
--- a/Engine/Plugins/Runtime/ReplicationGraph/Source/Private/ReplicationGraph.cpp
+++ b/Engine/Plugins/Runtime/ReplicationGraph/Source/Private/ReplicationGraph.cpp
@@ -2383,7 +2383,7 @@ bool UReplicationGraph::ProcessRemoteFunction(class AActor* Actor, UFunction* Fu
 	{
 		const bool bIsReliable = EnumHasAnyFlags(Function->FunctionFlags, FUNC_NetReliable);
 
-		// If we're saturated and it's not a reliable multicast, drop it.
+		// If we're saturated and it's not a reliable, drop it.
 		if (!(bIsReliable || IsConnectionReady(Connection)))
 		{
 			return true;
diff --git a/Engine/Source/Runtime/Engine/Classes/Engine/NetConnection.h b/Engine/Source/Runtime/Engine/Classes/Engine/NetConnection.h
--- a/Engine/Source/Runtime/Engine/Classes/Engine/NetConnection.h
+++ b/Engine/Source/Runtime/Engine/Classes/Engine/NetConnection.h
@@ -535,6 +535,9 @@ public:
 	int32 InTotalPacketsLost, OutTotalPacketsLost;
 	/** total acks sent on this connection */
 	int32 OutTotalAcks;
+	/** Delayed RPCs and the total average frame delay */
+	int32 TotalDelayedRPCs = 0;
+	int32 TotalDelayedRPCsFrameCount = 0;
 
 private:
 	/** total packets received on this connection, including PacketHandler */
diff --git a/Engine/Source/Runtime/Engine/Classes/Engine/NetDriver.h b/Engine/Source/Runtime/Engine/Classes/Engine/NetDriver.h
--- a/Engine/Source/Runtime/Engine/Classes/Engine/NetDriver.h
+++ b/Engine/Source/Runtime/Engine/Classes/Engine/NetDriver.h
@@ -1104,6 +1104,10 @@ public:
 	uint32						InTotalBunches;
 	/** Total bunches sent since the net driver's creation  */
 	uint32						OutTotalBunches;
+	/** Total number of outgoing reliable bunches */
+	uint32						OutTotalReliableBunches;
+	/** Total number of incoming reliable bunches */
+	uint32						InTotalReliableBunches;
 	/** todo document */
 	uint32						InPacketsLost;
 	/** Total packets lost that have been sent by clients since the net driver's creation  */
@@ -1218,7 +1222,7 @@ public:
 	PRAGMA_ENABLE_DEPRECATION_WARNINGS
 
 	/** Used to invalidate properties marked "unchanged" in FRepChangedPropertyTracker's */
-	uint32																		ReplicationFrame;
+	uint32 ReplicationFrame;
 
 	/** Maps FRepLayout to the respective UClass */
 	TMap<TWeakObjectPtr<UObject>, TSharedPtr<FRepLayout>, FDefaultSetAllocator, TWeakObjectPtrMapKeyFuncs<TWeakObjectPtr<UObject>, TSharedPtr<FRepLayout> > >	RepLayoutMap;
diff --git a/Engine/Source/Runtime/Engine/Private/DataReplication.cpp b/Engine/Source/Runtime/Engine/Private/DataReplication.cpp
--- a/Engine/Source/Runtime/Engine/Private/DataReplication.cpp
+++ b/Engine/Source/Runtime/Engine/Private/DataReplication.cpp
@@ -967,6 +967,7 @@ bool FObjectReplicator::ReceivedBunch(FNetBitReader& Bunch, const FReplicationFl
 
 	const bool bIsServer = ConnectionNetDriver->IsServer();
 	const bool bCanDelayRPCs = (CVarDelayUnmappedRPCs.GetValueOnGameThread() > 0) && !bIsServer;
+    const uint32 DriverReplicationFrame = ConnectionNetDriver->ReplicationFrame;
 
 	const FClassNetCache* const ClassCache = ConnectionNetDriver->NetCache->GetClassNetCache(ObjectClass);
 
@@ -1163,7 +1164,7 @@ bool FObjectReplicator::ReceivedBunch(FNetBitReader& Bunch, const FReplicationFl
 			else if (bDelayFunction)
 			{
 				// This invalidates Reader's buffer
-				PendingLocalRPCs.Emplace(FieldCache, RepFlags, Reader, UnmappedGuids);
+				PendingLocalRPCs.Emplace(FieldCache, RepFlags, Reader, DriverReplicationFrame, UnmappedGuids);
 				bOutHasUnmapped = true;
 				bGuidsChanged = true;
 				bForceUpdateUnmapped = true;
@@ -2357,6 +2358,8 @@ void FObjectReplicator::UpdateUnmappedObjects(bool & bOutHasMoreUnmapped)
 
 	check(RepLayout);
 
+	const uint32 CurrentReplicationFrame = Connection->GetDriver()->ReplicationFrame;
+
 	const FRepLayout& LocalRepLayout = *RepLayout;
 
 	FNetSerializeCB NetSerializeCB(Connection->Driver);
@@ -2466,7 +2469,17 @@ void FObjectReplicator::UpdateUnmappedObjects(bool & bOutHasMoreUnmapped)
 			}
 			else
 			{
-				// We executed, remove this one and continue;
+				// Track RPCs delayed multiple frames
+				const uint32 DelayedFrames = (CurrentReplicationFrame >= Pending.FrameQueuedAt) ? (CurrentReplicationFrame - Pending.FrameQueuedAt) : 0u;
+				if (DelayedFrames > 0)
+				{
+					UE_LOG(LogNet, Verbose, TEXT("FObjectReplicator::UpdateUnmappedObjects: RPC %s on Object %s was finally executed after being delayed for %u frames (~%f ms)"),
+						*FunctionName, *Object->GetFullName(), DelayedFrames, 
+						DelayedFrames*(1000.f / ((GEngine->GetMaxTickRate(0.0f, true) > 0.0f) ? GEngine->GetMaxTickRate(0.0f, true) : 30.f)));
+					Connection->TotalDelayedRPCs++;
+					Connection->TotalDelayedRPCsFrameCount += DelayedFrames;
+				}
+				// We executed, remove this one and continue
 				PendingLocalRPCs.RemoveAt(RPCIndex);
 				RPCIndex--;
 			}
diff --git a/Engine/Source/Runtime/Engine/Private/Net/NetEmulationHelper.cpp b/Engine/Source/Runtime/Engine/Private/Net/NetEmulationHelper.cpp
new file mode 100644
--- /dev/null
+++ b/Engine/Source/Runtime/Engine/Private/Net/NetEmulationHelper.cpp
@@ -0,0 +1,365 @@
+// Copyright Epic Games, Inc. All Rights Reserved.
+
+#include "NetEmulationHelper.h"
+
+#include "Engine/Engine.h"
+#include "Engine/NetworkSettings.h"
+
+#if DO_ENABLE_NET_TEST
+
+namespace UE::Net::Private::NetEmulationHelper
+{
+/** Global that stores the network emulation values outside the NetDriver lifetime */
+static TOptional<FPacketSimulationSettings> PersistentPacketSimulationSettings;
+
+void CreatePersistentSimulationSettings()
+{
+	if (!PersistentPacketSimulationSettings.IsSet())
+	{
+		PersistentPacketSimulationSettings.Emplace(FPacketSimulationSettings());
+	}
+}
+
+void ApplySimulationSettingsOnNetDrivers(UWorld* World, const FPacketSimulationSettings& Settings)
+{
+	// Execute on all active NetDrivers
+	FWorldContext& Context = GEngine->GetWorldContextFromWorldChecked(World);
+	for (const FNamedNetDriver& ActiveNetDriver : Context.ActiveNetDrivers)
+	{
+		if (ActiveNetDriver.NetDriver)
+		{
+			ActiveNetDriver.NetDriver->SetPacketSimulationSettings(Settings);
+		}
+	}
+}
+
+bool HasPersistentPacketEmulationSettings()
+{
+	return PersistentPacketSimulationSettings.IsSet();
+}
+
+void ApplyPersistentPacketEmulationSettings(UNetDriver* NetDriver)
+{
+	NetDriver->SetPacketSimulationSettings(PersistentPacketSimulationSettings.GetValue());
+}
+
+FAutoConsoleCommandWithWorldArgsAndOutputDevice NetEmulationPktEmulationProfile(TEXT("NetEmulation.PktEmulationProfile"), TEXT("Apply a preconfigured emulation profile."),
+FConsoleCommandWithWorldArgsAndOutputDeviceDelegate::CreateStatic([](const TArray<FString>& Args, UWorld* World, FOutputDevice& Output)
+{
+	bool bProfileApplied(false);
+	if (Args.Num() > 0)
+	{
+		FString CmdParams = FString::Printf(TEXT("PktEmulationProfile=%s"), *(Args[0]));
+
+		CreatePersistentSimulationSettings();
+
+		bProfileApplied = PersistentPacketSimulationSettings.GetValue().ParseSettings(*CmdParams, nullptr);
+
+		if (bProfileApplied)
+		{
+			ApplySimulationSettingsOnNetDrivers(World, PersistentPacketSimulationSettings.GetValue());
+		}
+		else
+		{
+			Output.Log(FString::Printf(TEXT("EmulationProfile: %s was not found in Engine.ini"), *(Args[0])));
+		}
+	}
+	else
+	{
+		Output.Log(FString::Printf(TEXT("Missing emulation profile name")));
+	}
+
+	if (!bProfileApplied)
+	{
+		if (const UNetworkSettings* NetworkSettings = GetDefault<UNetworkSettings>())
+		{
+			Output.Log(TEXT("List of some supported emulation profiles:"));
+			for (const FNetworkEmulationProfileDescription& ProfileDesc : NetworkSettings->NetworkEmulationProfiles)
+			{
+				Output.Log(FString::Printf(TEXT("%s"), *ProfileDesc.ProfileName));
+			}
+		}
+	}
+}));
+
+FAutoConsoleCommandWithWorld NetEmulationOff(TEXT("NetEmulation.Off"), TEXT("Turn off network emulation"),
+FConsoleCommandWithWorldDelegate::CreateStatic([](UWorld* World)
+{
+	CreatePersistentSimulationSettings();
+	PersistentPacketSimulationSettings.GetValue().ResetSettings();
+	ApplySimulationSettingsOnNetDrivers(World, PersistentPacketSimulationSettings.GetValue());
+}));
+
+
+FAutoConsoleCommandWithWorldArgsAndOutputDevice NetEmulationDropNothingFunction(TEXT("NetEmulation.DropNothing"), TEXT("Disables any RPC drop settings previously set."),
+FConsoleCommandWithWorldArgsAndOutputDeviceDelegate::CreateStatic([](const TArray<FString>& Args, UWorld* World, FOutputDevice& Output)
+{
+	UNetDriver* NetDriver = World->NetDriver;
+	if (!NetDriver)
+	{
+		return;
+	}
+
+	NetDriver->SendRPCDel.Unbind();
+}));
+
+FAutoConsoleCommandWithWorldArgsAndOutputDevice NetEmulationDropAnyUnreliableFunction(TEXT("NetEmulation.DropAnyUnreliable"), TEXT("Drop any sent unreliable RPCs. (optional)<0-100> to set the drop percentage (default is 20)."),
+FConsoleCommandWithWorldArgsAndOutputDeviceDelegate::CreateStatic([](const TArray<FString>& Args, UWorld* World, FOutputDevice& Output)
+{
+	UNetDriver* NetDriver = World->NetDriver;
+	if (!NetDriver)
+	{
+		return;
+	}
+
+	float DropPercentage = 20.f;
+	if (Args.Num() > 0)
+	{
+		float DropParam = DropPercentage;
+		LexFromString(DropParam, *Args[0]);
+		if (DropParam > 0.0f)
+		{
+			DropPercentage = FMath::Min(DropParam, 100.f);
+		}
+	}
+
+	NetDriver->SendRPCDel.BindLambda([DropPercentage](AActor* Actor, UFunction* Function, void* Parameters, FOutParmRec* OutParms, FFrame* Stack, UObject* SubObject, bool& bOutBlockSendRPC)
+	{
+		if ((Function->FunctionFlags & FUNC_NetReliable) == 0)
+		{
+			const float RandValue = FMath::FRand();
+			bOutBlockSendRPC = RandValue > (DropPercentage * 0.01f);
+
+			if (SubObject)
+			{
+				UE_LOG(LogNetTraffic, Log, TEXT("      Dropped unreliable RPC %s::%s : %s"), *GetFullNameSafe(Actor), *GetNameSafe(SubObject), *GetNameSafe(Function));
+			}
+			else
+			{
+				UE_LOG(LogNetTraffic, Log, TEXT("      Dropped unreliable RPC %s : %s"), *GetFullNameSafe(Actor), *GetNameSafe(Function));
+			}
+		}
+	});
+
+	UE_LOG(LogNetTraffic, Warning, TEXT("Will start dropping %.2f%% of all unreliable RPCs"), DropPercentage);
+}));
+
+FAutoConsoleCommandWithWorldArgsAndOutputDevice NetEmulationDropUnreliableOfActorClassFunction(TEXT("NetEmulation.DropUnreliableOfActorClass"),
+	TEXT("Drop random unreliable RPCs sent on actors of the given class type. "
+		"<ActorClassName> Class name to match with (can be a substring). "
+		"(optional)<0-100> to set the drop percentage (default is 20)."),
+FConsoleCommandWithWorldArgsAndOutputDeviceDelegate::CreateStatic([](const TArray<FString>& Args, UWorld* World, FOutputDevice& Output)
+{
+	UNetDriver* NetDriver = World->NetDriver;
+	if (!NetDriver)
+	{
+		return;
+	}
+
+	if (Args.Num() < 1)
+	{
+		UE_LOG(LogNet, Warning, TEXT("No class name parameter passed to NetEmulation.DropUnreliableOfActorClass"));
+		return;
+	}
+
+	FString ClassNameParam = Args[0];
+
+	float DropPercentage = 20.f;
+	if (Args.Num() > 1)
+	{
+		float DropParam = DropPercentage;
+		LexFromString(DropParam, *Args[1]);
+		if (DropParam > 0.0f)
+		{
+			DropPercentage = FMath::Min(DropParam, 100.f);
+		}
+	}
+
+	NetDriver->SendRPCDel.BindLambda([ClassNameParam, DropPercentage](AActor* Actor, UFunction* Function, void* Parameters, FOutParmRec* OutParms, FFrame* Stack, UObject* SubObject, bool& bOutBlockSendRPC)
+	{
+		if ((Function->FunctionFlags & FUNC_NetReliable) == 0)
+		{
+			const float RandValue = FMath::FRand();
+			if (RandValue <= (DropPercentage * 0.01f))
+			{
+				return;
+			}
+
+			bool bMatches = false;
+
+			UClass* Class = Actor->GetClass();
+			while (Class && !bMatches)
+			{
+				bMatches = Class->GetName().Contains(*ClassNameParam);
+				Class = Class->GetSuperClass();
+			}
+
+			if (bMatches)
+			{
+				bOutBlockSendRPC = true;
+
+				if (SubObject)
+				{
+					UE_LOG(LogNetTraffic, Log, TEXT("      Dropped unreliable RPC %s::%s : %s"), *GetFullNameSafe(Actor), *GetNameSafe(SubObject), *GetNameSafe(Function));
+				}
+				else
+				{
+					UE_LOG(LogNetTraffic, Log, TEXT("      Dropped unreliable RPC %s : %s"), *GetFullNameSafe(Actor), *GetNameSafe(Function));
+				}
+			}
+		}
+	});
+
+	UE_LOG(LogNetTraffic, Warning, TEXT("Will start dropping %.2f%% of all unreliable RPCs of actors of class: %s"), DropPercentage, *ClassNameParam);
+}));
+
+FAutoConsoleCommandWithWorldArgsAndOutputDevice NetEmulationDropUnreliableRPCFunction(TEXT("NetEmulation.DropUnreliableRPC"),
+	TEXT("Drop randomly the unreliable RPCs of the given name. "
+		"<RPCName> The name of the RPC (can be a substring). "
+		"(optional)<0-100> to set the drop percentage (default is 20)."),
+FConsoleCommandWithWorldArgsAndOutputDeviceDelegate::CreateStatic([](const TArray<FString>& Args, UWorld* World, FOutputDevice& Output)
+{
+	UNetDriver* NetDriver = World->NetDriver;
+	if (!NetDriver)
+	{
+		return;
+	}
+
+	if (Args.Num() < 1)
+	{
+		UE_LOG(LogNet, Warning, TEXT("No RPC name parameter passed to NetEmulation.DropUnreliableRPC"));
+		return;
+	}
+
+	FString RPCNameParam = Args[0];
+
+	float DropPercentage = 20.f;
+	if (Args.Num() > 1)
+	{
+		float DropParam = DropPercentage;
+		LexFromString(DropParam, *Args[1]);
+		if (DropParam > 0.0f)
+		{
+			DropPercentage = FMath::Min(DropParam, 100.f);
+		}
+	}
+
+	NetDriver->SendRPCDel.BindLambda([RPCNameParam, DropPercentage](AActor* Actor, UFunction* Function, void* Parameters, FOutParmRec* OutParms, FFrame* Stack, UObject* SubObject, bool& bOutBlockSendRPC)
+	{
+		if ((Function->FunctionFlags & FUNC_NetReliable) == 0)
+		{
+			const float RandValue = FMath::FRand();
+			if (RandValue <= (DropPercentage * 0.01f))
+			{
+				return;
+			}
+
+			if (Function->GetName().Contains(*RPCNameParam))
+			{
+				bOutBlockSendRPC = true;
+
+				if (SubObject)
+				{
+					UE_LOG(LogNetTraffic, Log, TEXT("      Dropped unreliable RPC %s::%s : %s"), *GetFullNameSafe(Actor), *GetNameSafe(SubObject), *GetNameSafe(Function));
+				}
+				else
+				{
+					UE_LOG(LogNetTraffic, Log, TEXT("      Dropped unreliable RPC %s : %s"), *GetFullNameSafe(Actor), *GetNameSafe(Function));
+				}
+			}
+		}
+	});
+
+	UE_LOG(LogNetTraffic, Warning, TEXT("Will start dropping %.2f%% of all unreliable RPCs named: %s"), DropPercentage, *RPCNameParam);
+}));
+
+FAutoConsoleCommandWithWorldArgsAndOutputDevice NetEmulationDropUnreliableOfSubObjectClassFunction(TEXT("NetEmulation.DropUnreliableOfSubObjectClass"),
+	TEXT("Drop randomly the unreliable RPCs of a subobject of the given class. "
+		"<SubObjectClassName> The name of the RPC (can be a substring). "
+		"(optional)<0-100> to set the drop percentage (default is 20)."),
+FConsoleCommandWithWorldArgsAndOutputDeviceDelegate::CreateStatic([](const TArray<FString>& Args, UWorld* World, FOutputDevice& Output)
+{
+	UNetDriver* NetDriver = World->NetDriver;
+	if (!NetDriver)
+	{
+		return;
+	}
+
+	if (Args.Num() < 1)
+	{
+		UE_LOG(LogNet, Warning, TEXT("No SubObject name parameter passed to NetEmulation.DropUnreliableOfSubObjectClass"));
+		return;
+	}
+
+	FString SubObjectClassNameParam = Args[0];
+
+	float DropPercentage = 20.f;
+	if (Args.Num() > 1)
+	{
+		float DropParam = DropPercentage;
+		LexFromString(DropParam, *Args[1]);
+		if (DropParam > 0.0f)
+		{
+			DropPercentage = FMath::Min(DropParam, 100.f);
+		}
+	}
+
+	NetDriver->SendRPCDel.BindLambda([SubObjectClassNameParam, DropPercentage](AActor* Actor, UFunction* Function, void* Parameters, FOutParmRec* OutParms, FFrame* Stack, UObject* SubObject, bool& bOutBlockSendRPC)
+	{
+		if (SubObject && (Function->FunctionFlags & FUNC_NetReliable) == 0)
+		{
+			const float RandValue = FMath::FRand();
+			if (RandValue <= (DropPercentage * 0.01f))
+			{
+				return;
+			}
+
+			bool bMatches = false;
+
+			UClass* Class = SubObject->GetClass();
+			while (Class && !bMatches)
+			{
+				bMatches = Class->GetName().Contains(*SubObjectClassNameParam);
+				Class = Class->GetSuperClass();
+			}
+
+			if (bMatches)
+			{
+				bOutBlockSendRPC = true;
+
+				UE_LOG(LogNetTraffic, Log, TEXT("      Dropped unreliable RPC %s::%s : %s"), *GetFullNameSafe(Actor), *GetNameSafe(SubObject), *GetNameSafe(Function));
+			}
+		}
+	});
+
+	UE_LOG(LogNetTraffic, Warning, TEXT("Will start dropping %.2f%% of all unreliable RPCs for subobjects: %s"), DropPercentage, *SubObjectClassNameParam);
+}));
+
+#define BUILD_NETEMULATION_CONSOLE_COMMAND(CommandName, CommandHelp) FAutoConsoleCommandWithWorldAndArgs NetEmulation##CommandName(TEXT("NetEmulation."#CommandName), TEXT(CommandHelp), \
+	FConsoleCommandWithWorldAndArgsDelegate::CreateStatic([](const TArray<FString>& Args, UWorld* World) \
+	{ \
+		if (Args.Num() > 0) \
+		{ \
+			CreatePersistentSimulationSettings(); \
+			FString CmdParams = FString::Printf(TEXT(#CommandName"=%s"), *(Args[0])); \
+			PersistentPacketSimulationSettings.GetValue().ParseSettings(*CmdParams, nullptr); \
+			ApplySimulationSettingsOnNetDrivers(World, PersistentPacketSimulationSettings.GetValue()); \
+		} \
+	}));
+
+	BUILD_NETEMULATION_CONSOLE_COMMAND(PktLoss, "Simulates network packet loss");
+	BUILD_NETEMULATION_CONSOLE_COMMAND(PktOrder, "Simulates network packets received out of order");
+	BUILD_NETEMULATION_CONSOLE_COMMAND(PktDup, "Simulates sending/receiving duplicate network packets");
+	BUILD_NETEMULATION_CONSOLE_COMMAND(PktLag, "Simulates network packet lag");
+	BUILD_NETEMULATION_CONSOLE_COMMAND(PktLagVariance, "Simulates variable network packet lag");
+	BUILD_NETEMULATION_CONSOLE_COMMAND(PktLagMin, "Sets minimum outgoing packet latency");
+	BUILD_NETEMULATION_CONSOLE_COMMAND(PktLagMax, "Sets maximum outgoing packet latency)");
+	BUILD_NETEMULATION_CONSOLE_COMMAND(PktIncomingLagMin, "Sets minimum incoming packet latency");
+	BUILD_NETEMULATION_CONSOLE_COMMAND(PktIncomingLagMax, "Sets maximum incoming packet latency");
+	BUILD_NETEMULATION_CONSOLE_COMMAND(PktIncomingLoss, "Simulates incoming packet loss");
+	BUILD_NETEMULATION_CONSOLE_COMMAND(PktJitter, "Simulates outgoing packet jitter");
+
+
+} // end namespace UE::Net::Private::NetEmulationHelper
+
+#endif //#if DO_ENABLE_NET_TEST
\ No newline at end of file
diff --git a/Engine/Source/Runtime/Engine/Private/Net/NetEmulationHelper.h b/Engine/Source/Runtime/Engine/Private/Net/NetEmulationHelper.h
new file mode 100644
--- /dev/null
+++ b/Engine/Source/Runtime/Engine/Private/Net/NetEmulationHelper.h
@@ -0,0 +1,19 @@
+// Copyright Epic Games, Inc. All Rights Reserved.
+
+#pragma once
+
+#include "Engine/NetDriver.h"
+
+class UWorld;
+
+#if DO_ENABLE_NET_TEST
+
+namespace UE::Net::Private::NetEmulationHelper
+{
+	void CreatePersistentSimulationSettings();
+	void ApplySimulationSettingsOnNetDrivers(UWorld* World, const FPacketSimulationSettings& Settings);
+	bool HasPersistentPacketEmulationSettings();
+	void ApplyPersistentPacketEmulationSettings(UNetDriver* NetDriver);
+} //namespace UE::Net::Private::NetEmulationHelper
+
+#endif //#if DO_ENABLE_NET_TEST
diff --git a/Engine/Source/Runtime/Engine/Private/NetConnection.cpp b/Engine/Source/Runtime/Engine/Private/NetConnection.cpp
--- a/Engine/Source/Runtime/Engine/Private/NetConnection.cpp
+++ b/Engine/Source/Runtime/Engine/Private/NetConnection.cpp
@@ -3694,6 +3694,7 @@ void UNetConnection::DispatchPacket( FBitReader& Reader, int32 PacketId, bool& b
 			}
 			Driver->InBunches++;
 			Driver->InTotalBunches++;
+			Driver->InTotalReliableBunches += Bunch.bReliable ? 1 : 0;
 
 			if (Bunch.IsCriticalError() || Bunch.IsError())
 			{
@@ -4019,6 +4020,7 @@ int32 UNetConnection::SendRawBunch(FOutBunch& Bunch, bool InAllowMerge, const FN
 	check(!Bunch.IsError());
 	Driver->OutBunches++;
 	Driver->OutTotalBunches++;
+	Driver->OutTotalReliableBunches += Bunch.bReliable ? 1 : 0;
 
 	// Build header.
 	SendBunchHeader.Reset();
diff --git a/Engine/Source/Runtime/Engine/Private/NetDriver.cpp b/Engine/Source/Runtime/Engine/Private/NetDriver.cpp
--- a/Engine/Source/Runtime/Engine/Private/NetDriver.cpp
+++ b/Engine/Source/Runtime/Engine/Private/NetDriver.cpp
@@ -43,6 +43,7 @@
 #include "Engine/ReplicationDriver.h"
 #include "Engine/LevelScriptActor.h"
 #include "Engine/NetworkSettings.h"
+#include "Net/NetEmulationHelper.h"
 #include "Net/NetSubObjectRegistryGetter.h"
 #include "Net/NetworkGranularMemoryLogging.h"
 #include "UObject/Stack.h"
@@ -191,35 +192,6 @@ namespace UE::Net::Private
 	static const FString PerfCounter_NumConnections(TEXT("NumConnections"));
 }
 
-namespace UE::Net::EmulationHelper
-{
-#if DO_ENABLE_NET_TEST
-	/** Global that stores the network emulation values outside the NetDriver lifetime */
-	static TOptional<FPacketSimulationSettings> PersistentPacketSimulationSettings;
-
-	void CreatePersistentSimulationSettings()
-	{
-		if (!PersistentPacketSimulationSettings.IsSet())
-		{
-			PersistentPacketSimulationSettings.Emplace(FPacketSimulationSettings());
-		}
-	}
-
-	void ApplySimulationSettingsOnNetDrivers(UWorld* World, const FPacketSimulationSettings& Settings)
-	{
-		// Execute on all active NetDrivers
-		FWorldContext& Context = GEngine->GetWorldContextFromWorldChecked(World);
-		for (const FNamedNetDriver& ActiveNetDriver : Context.ActiveNetDrivers)
-		{
-			if (ActiveNetDriver.NetDriver)
-			{
-				ActiveNetDriver.NetDriver->SetPacketSimulationSettings(Settings);
-			}
-		}
-	}
-#endif //#if DO_ENABLE_NET_TEST
-}
-
 namespace UE::Net::Private
 {
 	struct FAutoDestructProperty
@@ -547,6 +519,8 @@ UNetDriver::UNetDriver(const FObjectInitializer& ObjectInitializer)
 ,	OutBunches(0)
 ,	InTotalBunches(0)
 ,	OutTotalBunches(0)
+,	OutTotalReliableBunches(0)
+,	InTotalReliableBunches(0)
 ,	InPacketsLost(0)
 ,	InTotalPacketsLost(0)
 ,	OutPacketsLost(0)
@@ -591,6 +565,8 @@ PRAGMA_ENABLE_DEPRECATION_WARNINGS
 void UNetDriver::InitPacketSimulationSettings()
 {
 #if DO_ENABLE_NET_TEST
+	using namespace UE::Net::Private::NetEmulationHelper;
+
 	if (bNeverApplyNetworkEmulationSettings)
 	{
 		return;
@@ -601,9 +577,9 @@ void UNetDriver::InitPacketSimulationSettings()
 		return;
 	}
 
-	if (UE::Net::EmulationHelper::PersistentPacketSimulationSettings.IsSet())
+	if (HasPersistentPacketEmulationSettings())
 	{
-		SetPacketSimulationSettings(UE::Net::EmulationHelper::PersistentPacketSimulationSettings.GetValue());
+		ApplyPersistentPacketEmulationSettings(this);
 		return;
 	}
 
@@ -2458,7 +2434,8 @@ void UNetDriver::ProcessRemoteFunctionForChannelPrivate(
 				}
 
 				NETWORK_PROFILER(GNetworkProfiler.TrackSendRPC(Ch->Actor, Function, HeaderBits, ParameterBits, 0, Connection));
-				FPacketIdRange PacketIdRange = Ch->SendBunch(&Bunch, true);
+				constexpr bool bDoMerge = true;
+				FPacketIdRange PacketIdRange = Ch->SendBunch(&Bunch, bDoMerge);
 
 				if (UNLIKELY(PacketIdRange.First == INDEX_NONE && PacketIdRange.Last == INDEX_NONE))
 				{
@@ -5089,7 +5066,8 @@ int64 UNetDriver::SendDestructionInfo(UNetConnection* Connection, FActorDestruct
 					UE_LOG(LogNetTraffic, Log, TEXT("SendDestructionInfo: Channel %d. NetGUID <%s> Path: %s. Bits: %d"), Channel->ChIndex, *DestructionInfo->NetGUID.ToString(), *DestructionInfo->PathName, CloseBunch.GetNumBits());
 					UE_LOG(LogNetDormancy, Verbose, TEXT("SendDestructionInfo: Channel %d. NetGUID <%s> Path: %s. Bits: %d"), Channel->ChIndex, *DestructionInfo->NetGUID.ToString(), *DestructionInfo->PathName, CloseBunch.GetNumBits());
 
-					Channel->SendBunch(&CloseBunch, false);
+					constexpr bool bDontMerge = false;
+					Channel->SendBunch(&CloseBunch, bDontMerge);
 
 					NumBits = CloseBunch.GetNumBits();
 				}
@@ -7859,76 +7837,7 @@ FAutoConsoleCommandWithWorld	DumpRelevantActorsCommand(
 
 #if DO_ENABLE_NET_TEST
 
-FAutoConsoleCommandWithWorldArgsAndOutputDevice NetEmulationPktEmulationProfile(TEXT("NetEmulation.PktEmulationProfile"), TEXT("Apply a preconfigured emulation profile."),
-	FConsoleCommandWithWorldArgsAndOutputDeviceDelegate::CreateStatic([](const TArray<FString>& Args, UWorld* World, FOutputDevice& Output)
-{
-	bool bProfileApplied(false);
-	if (Args.Num() > 0)
-	{
-		FString CmdParams = FString::Printf(TEXT("PktEmulationProfile=%s"), *(Args[0]));
-
-		UE::Net::EmulationHelper::CreatePersistentSimulationSettings();
-
-		bProfileApplied = UE::Net::EmulationHelper::PersistentPacketSimulationSettings.GetValue().ParseSettings(*CmdParams, nullptr);
-		
-		if (bProfileApplied)
-		{
-			UE::Net::EmulationHelper::ApplySimulationSettingsOnNetDrivers(World, UE::Net::EmulationHelper::PersistentPacketSimulationSettings.GetValue());
-		}
-		else
-		{
-			Output.Log(FString::Printf(TEXT("EmulationProfile: %s was not found in Engine.ini"), *(Args[0])));
-		}
-	}
-	else
-	{
-		Output.Log(FString::Printf(TEXT("Missing emulation profile name")));
-	}
-
-	if (!bProfileApplied)
-	{
-		if (const UNetworkSettings* NetworkSettings = GetDefault<UNetworkSettings>())
-		{
-			Output.Log(TEXT("List of some supported emulation profiles:"));
-			for (const FNetworkEmulationProfileDescription& ProfileDesc : NetworkSettings->NetworkEmulationProfiles)
-			{
-				Output.Log(FString::Printf(TEXT("%s"), *ProfileDesc.ProfileName));
-			}
-		}
-	}
-}));
-
-FAutoConsoleCommandWithWorld NetEmulationOff(TEXT("NetEmulation.Off"), TEXT("Turn off network emulation"),
-	FConsoleCommandWithWorldDelegate::CreateStatic([](UWorld* World)
-{
-	UE::Net::EmulationHelper::CreatePersistentSimulationSettings();
-	UE::Net::EmulationHelper::PersistentPacketSimulationSettings.GetValue().ResetSettings();
-	UE::Net::EmulationHelper::ApplySimulationSettingsOnNetDrivers(World, UE::Net::EmulationHelper::PersistentPacketSimulationSettings.GetValue());
-}));
-
-#define BUILD_NETEMULATION_CONSOLE_COMMAND(CommandName, CommandHelp) FAutoConsoleCommandWithWorldAndArgs NetEmulation##CommandName(TEXT("NetEmulation."#CommandName), TEXT(CommandHelp), \
-	FConsoleCommandWithWorldAndArgsDelegate::CreateStatic([](const TArray<FString>& Args, UWorld* World) \
-	{ \
-		if (Args.Num() > 0) \
-		{ \
-			UE::Net::EmulationHelper::CreatePersistentSimulationSettings(); \
-			FString CmdParams = FString::Printf(TEXT(#CommandName"=%s"), *(Args[0])); \
-			UE::Net::EmulationHelper::PersistentPacketSimulationSettings.GetValue().ParseSettings(*CmdParams, nullptr); \
-			UE::Net::EmulationHelper::ApplySimulationSettingsOnNetDrivers(World, UE::Net::EmulationHelper::PersistentPacketSimulationSettings.GetValue()); \
-		} \
-	}));
 
-BUILD_NETEMULATION_CONSOLE_COMMAND(PktLoss, "Simulates network packet loss");
-BUILD_NETEMULATION_CONSOLE_COMMAND(PktOrder, "Simulates network packets received out of order");
-BUILD_NETEMULATION_CONSOLE_COMMAND(PktDup, "Simulates sending/receiving duplicate network packets");
-BUILD_NETEMULATION_CONSOLE_COMMAND(PktLag, "Simulates network packet lag");
-BUILD_NETEMULATION_CONSOLE_COMMAND(PktLagVariance, "Simulates variable network packet lag");
-BUILD_NETEMULATION_CONSOLE_COMMAND(PktLagMin, "Sets minimum outgoing packet latency");
-BUILD_NETEMULATION_CONSOLE_COMMAND(PktLagMax, "Sets maximum outgoing packet latency)");
-BUILD_NETEMULATION_CONSOLE_COMMAND(PktIncomingLagMin, "Sets minimum incoming packet latency");
-BUILD_NETEMULATION_CONSOLE_COMMAND(PktIncomingLagMax, "Sets maximum incoming packet latency");
-BUILD_NETEMULATION_CONSOLE_COMMAND(PktIncomingLoss, "Simulates incoming packet loss");
-BUILD_NETEMULATION_CONSOLE_COMMAND(PktJitter, "Simulates outgoing packet jitter");
 
 #endif //#if DO_ENABLE_NET_TEST
 
diff --git a/Engine/Source/Runtime/Engine/Public/Net/DataReplication.h b/Engine/Source/Runtime/Engine/Public/Net/DataReplication.h
--- a/Engine/Source/Runtime/Engine/Public/Net/DataReplication.h
+++ b/Engine/Source/Runtime/Engine/Public/Net/DataReplication.h
@@ -98,6 +98,9 @@ public:
 		/** Number of bits in buffer */
 		int64 NumBits;
 
+		/** The replication frame number when the RPC was queued. Used to track the latency caused by the delay. */
+		uint32 FrameQueuedAt;
+
 		/** Guids being waited on */
 		TSet<FNetworkGUID> UnmappedGuids;
 
@@ -105,11 +108,13 @@ public:
 			const FFieldNetCache* InRPCField,
 			const FReplicationFlags& InRepFlags,
 			FNetBitReader& InReader,
+			uint32 InReplicationFrame,
 			const TSet<FNetworkGUID>& InUnmappedGuids)
 			: RPCFieldIndex(InRPCField->FieldNetIndex)
 			, RepFlags(InRepFlags)
 			, Buffer(InReader.GetBuffer())
 			, NumBits(InReader.GetNumBits())
+			, FrameQueuedAt(InReplicationFrame)
 			, UnmappedGuids(InUnmappedGuids)
 		{}
 
