Support pose watch when debugging parent or linked anim blueprints. Also prevent toggling pose watch from child anim bps

diff --git a/Engine/Source/Editor/AnimationBlueprintEditor/Private/AnimationBlueprintEditor.cpp b/Engine/Source/Editor/AnimationBlueprintEditor/Private/AnimationBlueprintEditor.cpp
--- a/Engine/Source/Editor/AnimationBlueprintEditor/Private/AnimationBlueprintEditor.cpp
+++ b/Engine/Source/Editor/AnimationBlueprintEditor/Private/AnimationBlueprintEditor.cpp
@@ -880,6 +880,12 @@ bool FAnimationBlueprintEditor::CanTogglePoseWatch()
 	const FGraphPanelSelectionSet SelectedNodes = GetSelectedNodes();
 	UAnimBlueprint* AnimBP = GetAnimBlueprint();
 
+	// Can't add pose watch from a child anim bp
+	if (UAnimBlueprint::FindRootAnimBlueprint(AnimBP))
+	{
+		return false;
+	}
+
 	for (FGraphPanelSelectionSet::TConstIterator NodeIt(SelectedNodes); NodeIt; ++NodeIt)
 	{
 		if (UAnimGraphNode_Base* SelectedNode = Cast<UAnimGraphNode_Base>(*NodeIt))
@@ -920,6 +926,12 @@ void FAnimationBlueprintEditor::OnTogglePoseWatch()
 	const FGraphPanelSelectionSet SelectedNodes = GetSelectedNodes();
 	UAnimBlueprint* AnimBP = GetAnimBlueprint();
 
+	// Can't toggle pose watch from a child anim bp
+	if (UAnimBlueprint::FindRootAnimBlueprint(AnimBP))
+	{
+		return;
+	}
+
 	for (FGraphPanelSelectionSet::TConstIterator NodeIt(SelectedNodes); NodeIt; ++NodeIt)
 	{
 		if (UAnimGraphNode_Base* SelectedNode = Cast<UAnimGraphNode_Base>(*NodeIt))
diff --git a/Engine/Source/Runtime/Engine/Private/Components/SkeletalMeshComponent.cpp b/Engine/Source/Runtime/Engine/Private/Components/SkeletalMeshComponent.cpp
--- a/Engine/Source/Runtime/Engine/Private/Components/SkeletalMeshComponent.cpp
+++ b/Engine/Source/Runtime/Engine/Private/Components/SkeletalMeshComponent.cpp
@@ -4683,31 +4683,35 @@ void USkeletalMeshComponent::GetPrimitiveStats(FPrimitiveStats& PrimitiveStats)
 void USkeletalMeshComponent::UpdatePoseWatches()
 {
 	PoseWatches.Empty();
-	UAnimInstance* AnimInstance = GetAnimInstance();
-	if (AnimInstance)
+	
+	auto CopyPoseWatches = [this](const UAnimInstance* InAnimInstance)
 	{
-		if (AnimInstance->IsBeingDebugged())
+		if (InAnimInstance && InAnimInstance->IsBeingDebugged())
 		{
-			UAnimBlueprintGeneratedClass* AnimBPGenClass = CastChecked<UAnimBlueprintGeneratedClass>(AnimInstance->GetClass());
-			if (AnimBPGenClass)
+			if (const UAnimBlueprintGeneratedClass* AnimBPGenClass = Cast<UAnimBlueprintGeneratedClass>(InAnimInstance->GetClass()))
 			{
-				if (USkeletalMesh* TmpSkeletalMesh = GetSkeletalMeshAsset())
+				if (const UAnimBlueprint* AnimBlueprint = Cast<UAnimBlueprint>(AnimBPGenClass->ClassGeneratedBy))
 				{
-					FAnimBlueprintDebugData& DebugData = AnimBPGenClass->GetAnimBlueprintDebugData();
-					for (FAnimNodePoseWatch& PoseWatch : DebugData.AnimNodePoseWatch)
+					const UAnimBlueprint* RootAnimBlueprint = UAnimBlueprint::FindRootAnimBlueprint(AnimBlueprint);
+					AnimBlueprint = RootAnimBlueprint ? RootAnimBlueprint : AnimBlueprint;
+
+					const FAnimBlueprintDebugData& DebugData = AnimBlueprint->GetAnimBlueprintGeneratedClass()->GetAnimBlueprintDebugData();
+					for (const FAnimNodePoseWatch& PoseWatch : DebugData.AnimNodePoseWatch)
 					{
 						if (const UPoseWatchPoseElement* PoseWatchPoseElement = PoseWatch.PoseWatchPoseElement)
 						{
 							if (PoseWatchPoseElement->GetIsEnabled() && PoseWatchPoseElement->GetIsVisible())
 							{
-								PoseWatches.Add(PoseWatch); // Copy
+								PoseWatches.Add(PoseWatch);
 							}
 						}
-					}
+					}	
 				}
 			}
 		}
-	}
+	};
+	
+	ForEachAnimInstance(CopyPoseWatches);
 }
 #endif // #if WITH_EDITOR
 
