UeWaveComp: fixed warning when trying to parse zero sized data

diff --git a/Engine/Source/Runtime/Engine/Private/SoundWave.cpp b/Engine/Source/Runtime/Engine/Private/SoundWave.cpp
--- a/Engine/Source/Runtime/Engine/Private/SoundWave.cpp
+++ b/Engine/Source/Runtime/Engine/Private/SoundWave.cpp
@@ -4864,6 +4864,11 @@ TFuture<FSharedBuffer> USoundWave::FEditorAudioBulkData::GetPayload() const
 		TArray<int16> scratch_buffer;
 		for (int i = 0; i < SoundWave->ChannelOffsets.Num(); ++i)
 		{
+			// If the channel doesn't exist, skip it.
+			if (!SoundWave->ChannelSizes[i])
+			{
+				continue;
+			}
 			if (!WaveInfo.ReadWaveInfo(Data + SoundWave->ChannelOffsets[i], SoundWave->ChannelSizes[i]))
 			{
 				UE_LOG(LogAudio, Warning, TEXT("Failed to read wave data out of '%s'."), *SoundWave->GetFullName());
@@ -4918,6 +4923,11 @@ void USoundWave::FEditorAudioBulkData::UpdatePayload(FSharedBuffer InPayload, UO
 		TArray<int16> scratch_buffer;
 		for (int i = 0; i < SoundWave->ChannelOffsets.Num(); ++i)
 		{
+			// If the channel doesn't exist, skip it.
+			if (!SoundWave->ChannelSizes[i])
+			{
+				continue;
+			}
 			FWaveModInfo WaveInfo;
 			if (!WaveInfo.ReadWaveInfo(Data + SoundWave->ChannelOffsets[i], SoundWave->ChannelSizes[i]))
 			{
