Added world name to UWorld::CleanupWorld warning to help diagnose related issues faster

#tests editor

diff --git a/Engine/Source/Runtime/Engine/Private/World.cpp b/Engine/Source/Runtime/Engine/Private/World.cpp
--- a/Engine/Source/Runtime/Engine/Private/World.cpp
+++ b/Engine/Source/Runtime/Engine/Private/World.cpp
@@ -5355,7 +5355,7 @@ void UWorld::CleanupWorld(bool bSessionEnded, bool bCleanupResources, UWorld* Ne
 		// streaming sublevels, and they never call InitWorld. (this is done by PrivateDestroyLevel when removing a
 		// streaming level from the Level List in the editor.)
 		bool bIsStreamingSubWorld = PersistentLevel && PersistentLevel->OwningWorld != this;
-		UE_CLOG(!bIsStreamingSubWorld, LogWorld, Warning, TEXT("UWorld::CleanupWorld called twice or called without InitWorld called first."));
+		UE_CLOG(!bIsStreamingSubWorld, LogWorld, Warning, TEXT("UWorld::CleanupWorld called twice or called without InitWorld called first (%s)"), *GetName());
 	}
 	const bool bWorldChanged = NewWorld != this;
 	CleanupWorldInternal(bSessionEnded, bCleanupResources, bWorldChanged);
