Added new "Bounds" material expression to expose object local, instance local and pre-skinned local bounding boxes info with a single node.

diff --git a/Engine/Source/Runtime/Engine/Classes/Materials/MaterialExpressionBounds.h b/Engine/Source/Runtime/Engine/Classes/Materials/MaterialExpressionBounds.h
new file mode 100644
--- /dev/null
+++ b/Engine/Source/Runtime/Engine/Classes/Materials/MaterialExpressionBounds.h
@@ -0,0 +1,46 @@
+// Copyright Epic Games, Inc. All Rights Reserved.
+
+#pragma once
+
+#include "CoreMinimal.h"
+#include "UObject/ObjectMacros.h"
+#include "Materials/MaterialExpression.h"
+#include "MaterialExpressionBounds.generated.h"
+
+UENUM()
+enum EMaterialExpressionBoundsType : uint8
+{
+	/** Instance local */
+	MEILB_InstanceLocal UMETA(DisplayName = "Instance Local"),
+
+	/** Object local */
+	MEILB_ObjectLocal UMETA(DisplayName = "Object Local"),
+
+	/** Skinned local */
+	MEILB_PreSkinnedLocal UMETA(DisplayName = "Pre-Skinned Local"),
+};
+
+UCLASS(collapsecategories, hidecategories = Object)
+class UMaterialExpressionBounds : public UMaterialExpression
+{
+	GENERATED_UCLASS_BODY()
+
+	/** The type of bounds to output */
+	UPROPERTY(EditAnywhere, Category = MaterialExpressionBounds, meta = (DisplayName = "Type", ShowAsInputPin = "Advanced"))
+	TEnumAsByte<EMaterialExpressionBoundsType> Type = MEILB_InstanceLocal;
+
+#if WITH_EDITOR
+	//~ Begin UMaterialExpression Interface
+	virtual int32 Compile(class FMaterialCompiler* Compiler, int32 OutputIndex) override;
+	virtual void GetCaption(TArray<FString>& OutCaptions) const override;
+	virtual void GetConnectorToolTip(int32 InputIndex, int32 OutputIndex, TArray<FString>& OutToolTip) override;
+	virtual void GetExpressionToolTip(TArray<FString>& OutToolTip) override;
+
+	virtual bool GenerateHLSLExpression(FMaterialHLSLGenerator& Generator, UE::HLSLTree::FScope& Scope, int32 OutputIndex, UE::HLSLTree::FExpression const*& OutExpression) const override;
+	//~ End UMaterialExpression Interface
+
+#if WITH_EDITORONLY_DATA
+	TArray<FString> OutputToolTips;
+#endif
+#endif
+};
\ No newline at end of file
diff --git a/Engine/Source/Runtime/Engine/Private/Materials/HLSLMaterialTranslator.cpp b/Engine/Source/Runtime/Engine/Private/Materials/HLSLMaterialTranslator.cpp
--- a/Engine/Source/Runtime/Engine/Private/Materials/HLSLMaterialTranslator.cpp
+++ b/Engine/Source/Runtime/Engine/Private/Materials/HLSLMaterialTranslator.cpp
@@ -6015,6 +6015,25 @@ int32 FHLSLMaterialTranslator::ObjectLocalBounds(int32 OutputIndex)
 	return INDEX_NONE; 
 }
 
+int32 FHLSLMaterialTranslator::InstanceLocalBounds(int32 OutputIndex)
+{
+	switch (OutputIndex)
+	{
+	case 0: // Half extents
+		return AddInlinedCodeChunk(MCT_Float3, TEXT("(GetPrimitiveData(Parameters).InstanceLocalBoundsExtent)"));
+	case 1: // Full extents
+		return AddInlinedCodeChunk(MCT_Float3, TEXT("(GetPrimitiveData(Parameters).InstanceLocalBoundsExtent * 2.0f)"));
+	case 2: // Min point
+		return AddInlinedCodeChunk(MCT_Float3, TEXT("(GetPrimitiveData(Parameters).InstanceLocalBoundsCenter - GetPrimitiveData(Parameters).InstanceLocalBoundsExtent)"));
+	case 3: // Max point
+		return AddInlinedCodeChunk(MCT_Float3, TEXT("(GetPrimitiveData(Parameters).InstanceLocalBoundsCenter + GetPrimitiveData(Parameters).InstanceLocalBoundsExtent)"));
+	default:
+		check(false);
+	}
+
+	return INDEX_NONE;
+}
+
 int32 FHLSLMaterialTranslator::PreSkinnedLocalBounds(int32 OutputIndex)
 {
 	switch (OutputIndex)
diff --git a/Engine/Source/Runtime/Engine/Private/Materials/HLSLMaterialTranslator.h b/Engine/Source/Runtime/Engine/Private/Materials/HLSLMaterialTranslator.h
--- a/Engine/Source/Runtime/Engine/Private/Materials/HLSLMaterialTranslator.h
+++ b/Engine/Source/Runtime/Engine/Private/Materials/HLSLMaterialTranslator.h
@@ -943,6 +943,7 @@ protected:
 	virtual int32 ObjectBounds() override;
 
 	virtual int32 ObjectLocalBounds(int32 OutputIndex) override;
+	virtual int32 InstanceLocalBounds(int32 OutputIndex) override;
 	virtual int32 PreSkinnedLocalBounds(int32 OutputIndex) override;
 
 	virtual int32 DistanceCullFade() override;
diff --git a/Engine/Source/Runtime/Engine/Private/Materials/MaterialExpressionHLSL.cpp b/Engine/Source/Runtime/Engine/Private/Materials/MaterialExpressionHLSL.cpp
--- a/Engine/Source/Runtime/Engine/Private/Materials/MaterialExpressionHLSL.cpp
+++ b/Engine/Source/Runtime/Engine/Private/Materials/MaterialExpressionHLSL.cpp
@@ -125,6 +125,7 @@
 #include "Materials/MaterialExpressionNeuralPostProcessNode.h"
 #include "Materials/MaterialExpressionObjectBounds.h"
 #include "Materials/MaterialExpressionObjectLocalBounds.h"
+#include "Materials/MaterialExpressionBounds.h"
 #include "Materials/MaterialExpressionObjectOrientation.h"
 #include "Materials/MaterialExpressionObjectPositionWS.h"
 #include "Materials/MaterialExpressionObjectRadius.h"
@@ -4220,6 +4221,97 @@ bool UMaterialExpressionObjectBounds::GenerateHLSLExpression(FMaterialHLSLGenera
 	return true;
 }
 
+bool UMaterialExpressionBounds::GenerateHLSLExpression(FMaterialHLSLGenerator& Generator, UE::HLSLTree::FScope& Scope, int32 OutputIndex, UE::HLSLTree::FExpression const*& OutExpression) const
+{
+	using namespace UE::HLSLTree;
+	using namespace UE::Shader;
+
+	const UMaterial* TargetMaterial = Generator.GetTargetMaterial();
+	if (TargetMaterial)
+	{
+		if (TargetMaterial->MaterialDomain == MD_DeferredDecal)
+		{
+			return Generator.Error(TEXT("Expression not available in the deferred decal material domain."));
+		}
+		else if (TargetMaterial->MaterialDomain != MD_Surface && TargetMaterial->MaterialDomain != MD_Volume)
+		{
+			return Generator.Error(TEXT("The material expression 'ObjectLocalBounds' is only supported in the 'Surface' or 'Volume' material domain."));
+		}
+	}
+
+	if (Type == MEILB_ObjectLocal)
+	{
+		switch (OutputIndex)
+		{
+		case 0: // Half extents
+			OutExpression = Generator.GetTree().NewExpression<FExpressionInlineCustomHLSL>(EValueType::Float3,
+				TEXT("((GetPrimitiveData(Parameters).LocalObjectBoundsMax - GetPrimitiveData(Parameters).LocalObjectBoundsMin) / 2.0f)"));
+			return true;
+		case 1: // Full extents
+			OutExpression = Generator.GetTree().NewExpression<FExpressionInlineCustomHLSL>(EValueType::Float3,
+				TEXT("(GetPrimitiveData(Parameters).LocalObjectBoundsMax - GetPrimitiveData(Parameters).LocalObjectBoundsMin)"));
+			return true;
+		case 2: // Min point
+			OutExpression = Generator.GetTree().NewExpression<FExpressionInlineCustomHLSL>(EValueType::Float3,
+				TEXT("GetPrimitiveData(Parameters).LocalObjectBoundsMin"));
+			return true;
+		case 3: // Max point
+			OutExpression = Generator.GetTree().NewExpression<FExpressionInlineCustomHLSL>(EValueType::Float3,
+				TEXT("GetPrimitiveData(Parameters).LocalObjectBoundsMax"));
+			return true;
+		default:
+			break;
+		}
+	}
+	else if (Type == MEILB_InstanceLocal)
+	{
+		switch (OutputIndex)
+		{
+		case 0: // Half extents
+			OutExpression = Generator.GetTree().NewExpression<FExpressionInlineCustomHLSL>(EValueType::Float3,
+				TEXT("(GetPrimitiveData(Parameters).InstanceLocalBoundsExtent)"));
+			return true;
+		case 1: // Full extents
+			OutExpression = Generator.GetTree().NewExpression<FExpressionInlineCustomHLSL>(EValueType::Float3,
+				TEXT("(GetPrimitiveData(Parameters).InstanceLocalBoundsExtent * 2.0f)"));
+			return true;
+		case 2: // Min point
+			OutExpression = Generator.GetTree().NewExpression<FExpressionInlineCustomHLSL>(EValueType::Float3,
+				TEXT("(GetPrimitiveData(Parameters).InstanceLocalBoundsCenter - GetPrimitiveData(Parameters).InstanceLocalBoundsExtent)"));
+			return true;
+		case 3: // Max point
+			OutExpression = Generator.GetTree().NewExpression<FExpressionInlineCustomHLSL>(EValueType::Float3,
+				TEXT("(GetPrimitiveData(Parameters).InstanceLocalBoundsCenter + GetPrimitiveData(Parameters).InstanceLocalBoundsExtent)"));
+			return true;
+		default:
+			break;
+		}
+	}
+	else if (Type == MEILB_PreSkinnedLocal)
+	{
+		FTree& Tree = Generator.GetTree();
+		switch (OutputIndex)
+		{
+		case 0: // Half extents
+			OutExpression = Tree.NewMul(Tree.NewSub(Generator.NewExternalInput(Material::EExternalInput::PreSkinnedLocalBoundsMax), Generator.NewExternalInput(Material::EExternalInput::PreSkinnedLocalBoundsMin)), Tree.NewConstant(0.5f));
+			return true;
+		case 1: // Full extents
+			OutExpression = Tree.NewSub(Generator.NewExternalInput(Material::EExternalInput::PreSkinnedLocalBoundsMax), Generator.NewExternalInput(Material::EExternalInput::PreSkinnedLocalBoundsMin));
+			return true;
+		case 2: // Min point
+			OutExpression = Generator.NewExternalInput(Material::EExternalInput::PreSkinnedLocalBoundsMin);
+			return true;
+		case 3: // Max point
+			OutExpression = Generator.NewExternalInput(Material::EExternalInput::PreSkinnedLocalBoundsMax);
+			return true;
+		default:
+			break;
+		}
+	}
+	checkNoEntry();
+	return false;
+}
+
 bool UMaterialExpressionObjectLocalBounds::GenerateHLSLExpression(FMaterialHLSLGenerator& Generator, UE::HLSLTree::FScope& Scope, int32 OutputIndex, UE::HLSLTree::FExpression const*& OutExpression) const
 {
 	using namespace UE::HLSLTree;
diff --git a/Engine/Source/Runtime/Engine/Private/Materials/MaterialExpressions.cpp b/Engine/Source/Runtime/Engine/Private/Materials/MaterialExpressions.cpp
--- a/Engine/Source/Runtime/Engine/Private/Materials/MaterialExpressions.cpp
+++ b/Engine/Source/Runtime/Engine/Private/Materials/MaterialExpressions.cpp
@@ -149,6 +149,7 @@
 #include "Materials/MaterialExpressionObjectOrientation.h"
 #include "Materials/MaterialExpressionObjectPositionWS.h"
 #include "Materials/MaterialExpressionObjectRadius.h"
+#include "Materials/MaterialExpressionBounds.h"
 #include "Materials/MaterialExpressionOneMinus.h"
 #include "Materials/MaterialExpressionPanner.h"
 #include "Materials/MaterialExpressionParameter.h"
@@ -13410,6 +13411,93 @@ void UMaterialExpressionObjectLocalBounds::GetExpressionToolTip(TArray<FString>&
 }
 #endif // WITH_EDITOR
 
+
+///////////////////////////////////////////////////////////////////////////////
+// UMaterialExpressionBounds
+///////////////////////////////////////////////////////////////////////////////
+UMaterialExpressionBounds::UMaterialExpressionBounds(const FObjectInitializer& ObjectInitializer)
+	: Super(ObjectInitializer)
+{
+#if WITH_EDITORONLY_DATA
+	// Structure to hold one-time initialization
+	struct FConstructorStatics
+	{
+		FText NAME_Vectors;
+		FConstructorStatics()
+			: NAME_Vectors(LOCTEXT("Vectors", "Vectors"))
+		{
+		}
+	};
+	static FConstructorStatics ConstructorStatics;
+
+	MenuCategories.Add(ConstructorStatics.NAME_Vectors);
+
+	bShaderInputData = true;
+	bShowOutputNameOnPin = true;
+
+	Outputs.Reset();
+	Outputs.Add(FExpressionOutput(TEXT("Half Extents"), 1, 1, 1, 1, 0));
+	OutputToolTips.Add("Half the extent (width, depth and height) of the bounding box. In local space.");
+	Outputs.Add(FExpressionOutput(TEXT("Extents"), 1, 1, 1, 1, 0));
+	OutputToolTips.Add("Full extent (width, depth and height) of the bounding box. Same as 2x Half Extents. In local space.");
+	Outputs.Add(FExpressionOutput(TEXT("Min"), 1, 1, 1, 1, 0));
+	OutputToolTips.Add("Minimum 3D point of the bounding box. In local space.");
+	Outputs.Add(FExpressionOutput(TEXT("Max"), 1, 1, 1, 1, 0));
+	OutputToolTips.Add("Maximum 3D point of the bounding box. In local space.");
+#endif
+}
+
+#if WITH_EDITOR
+int32 UMaterialExpressionBounds::Compile(class FMaterialCompiler* Compiler, int32 OutputIndex)
+{
+	if (Material && Material->MaterialDomain == MD_DeferredDecal)
+	{
+		return CompilerError(Compiler, TEXT("Expression not available in the deferred decal material domain."));
+	}
+
+	switch (Type)
+	{
+	case MEILB_ObjectLocal: return Compiler->ObjectLocalBounds(OutputIndex);
+	case MEILB_InstanceLocal: return Compiler->InstanceLocalBounds(OutputIndex);
+	case MEILB_PreSkinnedLocal: return Compiler->PreSkinnedLocalBounds(OutputIndex);
+	default: checkNoEntry();
+	}
+	return INDEX_NONE;
+}
+
+void UMaterialExpressionBounds::GetCaption(TArray<FString>& OutCaptions) const
+{
+	FString Caption;
+	switch (Type)
+	{
+	case MEILB_ObjectLocal: Caption = TEXT("Bounds (Object Local)"); break;
+	case MEILB_InstanceLocal: Caption = TEXT("Bounds (Instance Local)"); break;
+	case MEILB_PreSkinnedLocal: Caption = TEXT("Bounds (Pre-Skinned Local)"); break;
+	default: checkNoEntry();
+	}
+
+	OutCaptions.Add(MoveTemp(Caption));
+}
+
+void UMaterialExpressionBounds::GetConnectorToolTip(int32 InputIndex, int32 OutputIndex, TArray<FString>& OutToolTip)
+{
+#if WITH_EDITORONLY_DATA
+	if (OutputIndex >= 0 && OutputIndex < OutputToolTips.Num())
+	{
+		ConvertToMultilineToolTip(OutputToolTips[OutputIndex], 40, OutToolTip);
+	}
+#endif // WITH_EDITORONLY_DATA
+}
+
+void UMaterialExpressionBounds::GetExpressionToolTip(TArray<FString>& OutToolTip)
+{
+	ConvertToMultilineToolTip(TEXT("Returns bounding box info of the specified type."
+		"Usable in vertex or pixel shader (no need to pipe this through vertex interpolators)."
+		"Hover the output pins for more information."), 40, OutToolTip);
+}
+
+#endif // WITH_EDITOR
+
 ///////////////////////////////////////////////////////////////////////////////
 // UMaterialExpressionPreSkinnedLocalBounds
 ///////////////////////////////////////////////////////////////////////////////
diff --git a/Engine/Source/Runtime/Engine/Public/MaterialCompiler.h b/Engine/Source/Runtime/Engine/Public/MaterialCompiler.h
--- a/Engine/Source/Runtime/Engine/Public/MaterialCompiler.h
+++ b/Engine/Source/Runtime/Engine/Public/MaterialCompiler.h
@@ -239,6 +239,7 @@ public:
 	virtual int32 ObjectRadius() = 0;
 	virtual int32 ObjectBounds() = 0;
 	virtual int32 ObjectLocalBounds(int32 OutputIndex) = 0;
+	virtual int32 InstanceLocalBounds(int32 OutputIndex) = 0;
 	virtual int32 PreSkinnedLocalBounds(int32 OutputIndex) = 0;
 	virtual int32 DistanceCullFade() = 0;
 	virtual int32 ActorWorldPosition(EPositionOrigin OriginType) = 0;
@@ -841,6 +842,7 @@ public:
 	virtual int32 ObjectRadius() override { return Compiler->ObjectRadius(); }
 	virtual int32 ObjectBounds() override { return Compiler->ObjectBounds(); }
 	virtual int32 ObjectLocalBounds(int32 OutputIndex) override { return Compiler->ObjectLocalBounds(OutputIndex); }
+	virtual int32 InstanceLocalBounds(int32 OutputIndex) override { return Compiler->InstanceLocalBounds(OutputIndex); }
 	virtual int32 PreSkinnedLocalBounds(int32 OutputIndex) override { return Compiler->PreSkinnedLocalBounds(OutputIndex); }
 	virtual int32 DistanceCullFade() override { return Compiler->DistanceCullFade(); }
 	virtual int32 ActorWorldPosition(EPositionOrigin OriginType) override { return Compiler->ActorWorldPosition(OriginType); }
