Allow the material editor search function to check against texture names on Texture Object Nodes, for parity with how search currently works with the Texture Sample node.

diff --git a/Engine/Source/Runtime/Engine/Classes/Materials/MaterialExpressionTextureObject.h b/Engine/Source/Runtime/Engine/Classes/Materials/MaterialExpressionTextureObject.h
--- a/Engine/Source/Runtime/Engine/Classes/Materials/MaterialExpressionTextureObject.h
+++ b/Engine/Source/Runtime/Engine/Classes/Materials/MaterialExpressionTextureObject.h
@@ -27,6 +27,7 @@ class UMaterialExpressionTextureObject : public UMaterialExpressionTextureBase
 	virtual int32 Compile(class FMaterialCompiler* Compiler, int32 OutputIndex) override;
 	virtual int32 CompilePreview(class FMaterialCompiler* Compiler, int32 OutputIndex) override;
 	virtual void GetCaption(TArray<FString>& OutCaptions) const override;
+	virtual bool MatchesSearchQuery(const TCHAR* SearchQuery) override;
 	virtual uint32 GetOutputType(int32 OutputIndex) override;
 	virtual bool GenerateHLSLExpression(FMaterialHLSLGenerator& Generator, UE::HLSLTree::FScope& Scope, int32 OutputIndex, UE::HLSLTree::FExpression const*& OutExpression) const override;
 	//~ End UMaterialExpression Interface
diff --git a/Engine/Source/Runtime/Engine/Private/Materials/MaterialExpressions.cpp b/Engine/Source/Runtime/Engine/Private/Materials/MaterialExpressions.cpp
--- a/Engine/Source/Runtime/Engine/Private/Materials/MaterialExpressions.cpp
+++ b/Engine/Source/Runtime/Engine/Private/Materials/MaterialExpressions.cpp
@@ -4069,6 +4069,16 @@ void UMaterialExpressionTextureObject::GetCaption(TArray<FString>& OutCaptions)
 	OutCaptions.Add(TEXT("Texture Object")); 
 }
 
+bool UMaterialExpressionTextureObject::MatchesSearchQuery(const TCHAR* SearchQuery)
+{
+	if (Texture != nullptr && Texture->GetName().Contains(SearchQuery))
+	{
+		return true;
+	}
+
+	return Super::MatchesSearchQuery(SearchQuery);
+}
+
 
 int32 UMaterialExpressionTextureObject::Compile(class FMaterialCompiler* Compiler, int32 OutputIndex)
 {
