[Material] Fixing GetTexturesInPropertyChain ignoring TextureObjects

#rnx

diff --git a/Engine/Source/Runtime/Engine/Private/Materials/Material.cpp b/Engine/Source/Runtime/Engine/Private/Materials/Material.cpp
--- a/Engine/Source/Runtime/Engine/Private/Materials/Material.cpp
+++ b/Engine/Source/Runtime/Engine/Private/Materials/Material.cpp
@@ -6053,7 +6053,7 @@ bool UMaterial::GetTexturesInPropertyChain(EMaterialProperty InProperty, TArray<
 			if (MatExp != NULL)
 			{
 				// Is it a texture sample or texture parameter sample?
-				UMaterialExpressionTextureSample* TextureSampleExp = Cast<UMaterialExpressionTextureSample>(MatExp);
+				UMaterialExpressionTextureBase* TextureSampleExp = Cast<UMaterialExpressionTextureBase>(MatExp);
 				if (TextureSampleExp != NULL)
 				{
 					// Check the default texture...
