Refactor MaterialExpressions TextureSample and TextureObject search functionality up to TextureBase.

diff --git a/Engine/Source/Runtime/Engine/Classes/Materials/MaterialExpressionTextureBase.h b/Engine/Source/Runtime/Engine/Classes/Materials/MaterialExpressionTextureBase.h
--- a/Engine/Source/Runtime/Engine/Classes/Materials/MaterialExpressionTextureBase.h
+++ b/Engine/Source/Runtime/Engine/Classes/Materials/MaterialExpressionTextureBase.h
@@ -41,7 +41,7 @@ class UMaterialExpressionTextureBase : public UMaterialExpression
 	//~ Begin UMaterialExpression Interface
 #if WITH_EDITOR
 	ENGINE_API virtual FString GetDescription() const override;
-
+	ENGINE_API virtual bool MatchesSearchQuery(const TCHAR* SearchQuery) override;
 	ENGINE_API virtual FText GetPreviewOverlayText() const override;
 #endif
 	//~ End UMaterialExpression Interface
diff --git a/Engine/Source/Runtime/Engine/Classes/Materials/MaterialExpressionTextureObject.h b/Engine/Source/Runtime/Engine/Classes/Materials/MaterialExpressionTextureObject.h
--- a/Engine/Source/Runtime/Engine/Classes/Materials/MaterialExpressionTextureObject.h
+++ b/Engine/Source/Runtime/Engine/Classes/Materials/MaterialExpressionTextureObject.h
@@ -27,7 +27,6 @@ class UMaterialExpressionTextureObject : public UMaterialExpressionTextureBase
 	virtual int32 Compile(class FMaterialCompiler* Compiler, int32 OutputIndex) override;
 	virtual int32 CompilePreview(class FMaterialCompiler* Compiler, int32 OutputIndex) override;
 	virtual void GetCaption(TArray<FString>& OutCaptions) const override;
-	virtual bool MatchesSearchQuery(const TCHAR* SearchQuery) override;
 	virtual uint32 GetOutputType(int32 OutputIndex) override;
 	virtual bool GenerateHLSLExpression(FMaterialHLSLGenerator& Generator, UE::HLSLTree::FScope& Scope, int32 OutputIndex, UE::HLSLTree::FExpression const*& OutExpression) const override;
 	//~ End UMaterialExpression Interface
diff --git a/Engine/Source/Runtime/Engine/Classes/Materials/MaterialExpressionTextureSample.h b/Engine/Source/Runtime/Engine/Classes/Materials/MaterialExpressionTextureSample.h
--- a/Engine/Source/Runtime/Engine/Classes/Materials/MaterialExpressionTextureSample.h
+++ b/Engine/Source/Runtime/Engine/Classes/Materials/MaterialExpressionTextureSample.h
@@ -96,7 +96,6 @@ public:
 	virtual int32 GetLabelPadding() override { return 8; }
 	ENGINE_API virtual int32 Compile(class FMaterialCompiler* Compiler, int32 OutputIndex) override;
 	ENGINE_API virtual void GetCaption(TArray<FString>& OutCaptions) const override;
-	ENGINE_API virtual bool MatchesSearchQuery( const TCHAR* SearchQuery ) override;
 	ENGINE_API virtual uint32 GetInputType(int32 InputIndex) override;
 	virtual bool CanIgnoreOutputIndex() { return true; }
 	ENGINE_API virtual bool GenerateHLSLExpression(FMaterialHLSLGenerator& Generator, UE::HLSLTree::FScope& Scope, int32 OutputIndex, UE::HLSLTree::FExpression const*& OutExpression) const override;
diff --git a/Engine/Source/Runtime/Engine/Private/Materials/MaterialExpressions.cpp b/Engine/Source/Runtime/Engine/Private/Materials/MaterialExpressions.cpp
--- a/Engine/Source/Runtime/Engine/Private/Materials/MaterialExpressions.cpp
+++ b/Engine/Source/Runtime/Engine/Private/Materials/MaterialExpressions.cpp
@@ -2473,6 +2473,16 @@ FString UMaterialExpressionTextureBase::GetDescription() const
 	return Result;
 }
 
+bool UMaterialExpressionTextureBase::MatchesSearchQuery(const TCHAR* SearchQuery)
+{
+	if (Texture != nullptr && Texture->GetName().Contains(SearchQuery))
+	{
+		return true;
+	}
+
+	return Super::MatchesSearchQuery(SearchQuery);
+}
+
 FText UMaterialExpressionTextureBase::GetPreviewOverlayText() const
 {
 	if (IsVirtualSamplerType(SamplerType))
@@ -2915,16 +2925,6 @@ void UMaterialExpressionTextureSample::GetCaption(TArray<FString>& OutCaptions)
 	OutCaptions.Add(TEXT("Texture Sample"));
 }
 
-bool UMaterialExpressionTextureSample::MatchesSearchQuery( const TCHAR* SearchQuery )
-{
-	if( Texture!=nullptr && Texture->GetName().Contains(SearchQuery) )
-	{
-		return true;
-	}
-
-	return Super::MatchesSearchQuery(SearchQuery);
-}
-
 // this define is only used for the following function
 #define IF_INPUT_RETURN(Type) if(!InputIndex) return (Type); --InputIndex
 uint32 UMaterialExpressionTextureSample::GetInputType(int32 InputIndex)
@@ -4069,17 +4069,6 @@ void UMaterialExpressionTextureObject::GetCaption(TArray<FString>& OutCaptions)
 	OutCaptions.Add(TEXT("Texture Object")); 
 }
 
-bool UMaterialExpressionTextureObject::MatchesSearchQuery(const TCHAR* SearchQuery)
-{
-	if (Texture != nullptr && Texture->GetName().Contains(SearchQuery))
-	{
-		return true;
-	}
-
-	return Super::MatchesSearchQuery(SearchQuery);
-}
-
-
 int32 UMaterialExpressionTextureObject::Compile(class FMaterialCompiler* Compiler, int32 OutputIndex)
 {
 	if (!Texture)
