Added extra info to an ensure

This info is already being logged, but the log happens after the crash reporter has already captured the log from the ensure, so the crash report is missing the required info.

#rnx

diff --git a/Engine/Source/Runtime/Engine/Private/Components/SceneComponent.cpp b/Engine/Source/Runtime/Engine/Private/Components/SceneComponent.cpp
--- a/Engine/Source/Runtime/Engine/Private/Components/SceneComponent.cpp
+++ b/Engine/Source/Runtime/Engine/Private/Components/SceneComponent.cpp
@@ -1271,7 +1271,7 @@ PRAGMA_ENABLE_DEPRECATION_WARNINGS
 						{
 #if WITH_EDITORONLY_DATA
 							// If we are in the middle of a transaction it isn't entirely unexpected that an AttachParent/AttachChildren pairing is wrong
-							if (!ensure(GIsTransacting))
+							if (!ensureAlwaysMsgf(GIsTransacting, TEXT("Component '%s' has '%s' in its AttachChildren array, however, '%s' believes it is attached to '%s'"), *GetFullName(), *Child->GetFullName(), *Child->GetFullName(), *Child->GetAttachParent()->GetFullName()))
 #endif
 							{
 								// We've gotten in to a bad state where the Child's AttachParent doesn't jive with the AttachChildren array
