Fix InstancedStaticMesh dynamic mesh batch creation.
Use first valid LOD index to account for mesh streaming.
Fixes later crash in RHI with invalid vertex buffer.

diff --git a/Engine/Source/Runtime/Engine/Private/InstancedStaticMesh.cpp b/Engine/Source/Runtime/Engine/Private/InstancedStaticMesh.cpp
--- a/Engine/Source/Runtime/Engine/Private/InstancedStaticMesh.cpp
+++ b/Engine/Source/Runtime/Engine/Private/InstancedStaticMesh.cpp
@@ -1057,8 +1057,9 @@ void FInstancedStaticMeshSceneProxy::GetDynamicMeshElements(const TArray<const F
 
 			for (int32 SelectionGroupIndex = 0; SelectionGroupIndex < NumSelectionGroups; SelectionGroupIndex++)
 			{
-				const int32 LODIndex = GetLOD(View);
-				const FStaticMeshLODResources& LODModel = StaticMesh->GetRenderData()->LODResources[LODIndex];
+				const int32 FirstLODIndex = RenderData->GetFirstValidLODIdx(RenderData->CurrentFirstLODIdx);
+				const int32 LODIndex = FMath::Max(GetLOD(View), FirstLODIndex);
+				const FStaticMeshLODResources& LODModel = RenderData->LODResources[LODIndex];
 
 				FInstancedStaticMeshVFLooseUniformShaderParametersRef LooseUniformBuffer = CreateLooseUniformBuffer(View, PassUserData[SelectionGroupIndex], /*InstancedLODRange=*/0, LODIndex, EUniformBufferUsage::UniformBuffer_SingleFrame);
 
