World Partition: fix regression when computing cells streaming performance for old runtime spatial hash.

#preglight none

diff --git a/Engine/Source/Runtime/Engine/Private/WorldPartition/WorldPartitionRuntimeCellDataSpatialHash.cpp b/Engine/Source/Runtime/Engine/Private/WorldPartition/WorldPartitionRuntimeCellDataSpatialHash.cpp
--- a/Engine/Source/Runtime/Engine/Private/WorldPartition/WorldPartitionRuntimeCellDataSpatialHash.cpp
+++ b/Engine/Source/Runtime/Engine/Private/WorldPartition/WorldPartitionRuntimeCellDataSpatialHash.cpp
@@ -28,6 +28,8 @@ UWorldPartitionRuntimeCellDataSpatialHash::UWorldPartitionRuntimeCellDataSpatial
 	: Super(ObjectInitializer)
 	, Extent(0)
 	, CachedMinSquareDistanceToSource(MAX_dbl)
+	, CachedMinBlockOnSlowStreamingRatio2D(MAX_dbl)
+	, CachedMinSquareDistanceToBlockingSource2D(MAX_dbl)
 {}
 
 void UWorldPartitionRuntimeCellDataSpatialHash::ResetStreamingSourceInfo() const
@@ -38,6 +40,8 @@ void UWorldPartitionRuntimeCellDataSpatialHash::ResetStreamingSourceInfo() const
 	CachedSourcePriorityWeights.Reset();
 	CachedSourceSquaredDistances.Reset();
 	CachedInstersectingShapes.Reset();
+	CachedMinBlockOnSlowStreamingRatio2D = MAX_dbl;
+	CachedMinSquareDistanceToBlockingSource2D = MAX_dbl;
 }
 
 float UWorldPartitionRuntimeCellDataSpatialHash::ComputeSourceToCellAngleFactor(const FSphericalSector& SourceShape) const
@@ -83,15 +87,15 @@ void UWorldPartitionRuntimeCellDataSpatialHash::AppendStreamingSourceInfo(const
 {
 	Super::AppendStreamingSourceInfo(Source, SourceShape);
 
-	const double SquareDistance = FVector::DistSquared2D(SourceShape.GetCenter(), Position);
+	const double SquareDistance2D = FVector::DistSquared2D(SourceShape.GetCenter(), Position);
 
 	// Update cached values based on the 2D distance
 	if (Source.bBlockOnSlowLoading)
 	{
-		CachedMinSquareDistanceToBlockingSource = FMath::Min(SquareDistance, CachedMinSquareDistanceToBlockingSource);
-		CachedMinBlockOnSlowStreamingRatio = FMath::Min(CachedMinBlockOnSlowStreamingRatio, FMath::Sqrt(CachedMinSquareDistanceToBlockingSource) / SourceShape.GetRadius());
+		CachedMinSquareDistanceToBlockingSource2D = FMath::Min(SquareDistance2D, CachedMinSquareDistanceToBlockingSource2D);
+		CachedMinBlockOnSlowStreamingRatio2D = FMath::Min(CachedMinBlockOnSlowStreamingRatio2D, FMath::Sqrt(CachedMinSquareDistanceToBlockingSource2D) / SourceShape.GetRadius());
 	}
-	CachedSourceSquaredDistances.Add(SquareDistance);
+	CachedSourceSquaredDistances.Add(SquareDistance2D);
 	CachedSourcePriorityWeights.Add(1.0f - ((float)Source.Priority / (float)EStreamingSourcePriority::Lowest));
 	CachedInstersectingShapes.Add(SourceShape);
 }
@@ -105,6 +109,8 @@ void UWorldPartitionRuntimeCellDataSpatialHash::MergeStreamingSourceInfo() const
 	check(Count == CachedInstersectingShapes.Num());
 
 	CachedMinSquareDistanceToSource = Count ? FMath::Min(CachedSourceSquaredDistances) : MAX_dbl;
+	CachedMinBlockOnSlowStreamingRatio = CachedMinBlockOnSlowStreamingRatio2D;
+	CachedMinSquareDistanceToBlockingSource = CachedMinSquareDistanceToBlockingSource2D;
 
 	if (Count)
 	{
diff --git a/Engine/Source/Runtime/Engine/Public/WorldPartition/WorldPartitionRuntimeCellDataSpatialHash.h b/Engine/Source/Runtime/Engine/Public/WorldPartition/WorldPartitionRuntimeCellDataSpatialHash.h
--- a/Engine/Source/Runtime/Engine/Public/WorldPartition/WorldPartitionRuntimeCellDataSpatialHash.h
+++ b/Engine/Source/Runtime/Engine/Public/WorldPartition/WorldPartitionRuntimeCellDataSpatialHash.h
@@ -53,4 +53,10 @@ private:
 
 	// Intersecting streaming source shapes
 	mutable TArray<FSphericalSector> CachedInstersectingShapes;
+
+	// 2D version of CachedMinBlockOnSlowStreamingRatio
+	mutable double CachedMinBlockOnSlowStreamingRatio2D;
+
+	// 2D version of CachedMinSquareDistanceToBlockingSource
+	mutable double CachedMinSquareDistanceToBlockingSource2D;
 };
