World Partition Runtime HashSet:
  - Don't show always loaded grid when showing 2d debug view.
  - Fix 2d debug view scale not respecting each grid debug view world bounds.
  - Fix always loaded grid loading range not properly initialized.

diff --git a/Engine/Source/Runtime/Engine/Private/WorldPartition/RuntimeHashSet/WorldPartitionRuntimeHashSetDebugDraw.cpp b/Engine/Source/Runtime/Engine/Private/WorldPartition/RuntimeHashSet/WorldPartitionRuntimeHashSetDebugDraw.cpp
--- a/Engine/Source/Runtime/Engine/Private/WorldPartition/RuntimeHashSet/WorldPartitionRuntimeHashSetDebugDraw.cpp
+++ b/Engine/Source/Runtime/Engine/Private/WorldPartition/RuntimeHashSet/WorldPartitionRuntimeHashSetDebugDraw.cpp
@@ -251,8 +251,20 @@ bool UWorldPartitionRuntimeHashSet::Draw2D(FWorldPartitionDraw2DContext& DrawCon
 			Source.ForEachShape(StreamingDataList[0]->LoadingRange, Name, true, [&GridsShapeBounds](const FSphericalSector& Shape) { GridsShapeBounds += Shape.CalcBounds(); });
 		}
 
-		const FVector2D GridReferenceWorldPos = FVector2D(WorldRegion.GetCenter());
-		const FVector2D WorldRegionExtent = FVector2D(WorldRegion.GetExtent().GetMax());
+		FVector2D GridReferenceWorldPos;
+		FVector2D WorldRegionExtent;
+
+		if (DrawContext.IsDetailedMode())
+		{
+			GridReferenceWorldPos = FVector2D(GridsShapeBounds.GetCenter());
+			WorldRegionExtent = FVector2D(GridsShapeBounds.ExpandBy(GridsShapeBounds.GetExtent() * 0.1f).GetExtent().GetMax());
+		}
+		else
+		{
+			GridReferenceWorldPos = FVector2D(WorldRegion.GetCenter());
+			WorldRegionExtent = FVector2D(WorldRegion.GetExtent().GetMax());
+		}
+
 		const FVector2D GridScreenOffset = GridScreenInitialOffset + ((float)GridIndex * FVector2D(GridMaxScreenWidth, 0.f)) + GridScreenHalfExtent + FVector2D(GridScreenWidthShrinkSize * 0.5f);
 		const FVector2D WorldToScreenScale = GridScreenHalfExtent / WorldRegionExtent;
 		const FBox2D GridScreenBounds(GridScreenOffset - GridScreenHalfExtent, GridScreenOffset + GridScreenHalfExtent);
@@ -279,8 +291,9 @@ bool UWorldPartitionRuntimeHashSet::Draw2D(FWorldPartitionDraw2DContext& DrawCon
 			MultiLineText.Emplace(GridInfoText, FLinearColor::Yellow);
 			FWorldPartitionCanvasMultiLineTextItem Item(GridInfoPos, MultiLineText);
 			DrawContext.PushDrawText(Item);
-			++GridIndex;
 		}
+
+		++GridIndex;
 	}
 
 	FBox2D DesiredWorldBounds(ForceInit);
diff --git a/Engine/Source/Runtime/Engine/Private/WorldPartition/RuntimeHashSet/WorldPartitionRuntimeHashSetStreamingGeneration.cpp b/Engine/Source/Runtime/Engine/Private/WorldPartition/RuntimeHashSet/WorldPartitionRuntimeHashSetStreamingGeneration.cpp
--- a/Engine/Source/Runtime/Engine/Private/WorldPartition/RuntimeHashSet/WorldPartitionRuntimeHashSetStreamingGeneration.cpp
+++ b/Engine/Source/Runtime/Engine/Private/WorldPartition/RuntimeHashSet/WorldPartitionRuntimeHashSetStreamingGeneration.cpp
@@ -155,6 +155,7 @@ bool UWorldPartitionRuntimeHashSet::GenerateStreaming(UWorldPartitionStreamingPo
 	PersistentPartitionDesc.Name = NAME_PersistentLevel;
 	PersistentPartitionDesc.MainLayer = NewObject<URuntimePartition>(this, URuntimePartitionPersistent::StaticClass(), NAME_None);
 	PersistentPartitionDesc.MainLayer->Name = NAME_PersistentLevel;
+	PersistentPartitionDesc.MainLayer->LoadingRange = 0;
 
 	ON_SCOPE_EXIT
 	{
diff --git a/Engine/Source/Runtime/Engine/Private/WorldPartition/WorldPartitionDebugHelper.cpp b/Engine/Source/Runtime/Engine/Private/WorldPartition/WorldPartitionDebugHelper.cpp
--- a/Engine/Source/Runtime/Engine/Private/WorldPartition/WorldPartitionDebugHelper.cpp
+++ b/Engine/Source/Runtime/Engine/Private/WorldPartition/WorldPartitionDebugHelper.cpp
@@ -45,7 +45,7 @@ FAutoConsoleCommand FWorldPartitionDebugHelper::DebugFilterByRuntimeHashGridName
 
 bool FWorldPartitionDebugHelper::IsDebugRuntimeHashGridShown(FName Name)
 {
-	return !DebugRuntimeHashFilter.Num() || DebugRuntimeHashFilter.Contains(Name);
+	return (Name != NAME_PersistentLevel) && (!DebugRuntimeHashFilter.Num() || DebugRuntimeHashFilter.Contains(Name));
 }
 
 TSet<FName> FWorldPartitionDebugHelper::DebugDataLayerFilter;
