[NetPhysics] Physics replication to make sure the object goes to sleep and stay asleep before stopping the replication.

#tests Local testing in PIE with logs and debug draw confirming logic works as intended.
#localization none

diff --git a/Engine/Source/Runtime/Engine/Private/PhysicsEngine/PhysicsReplication.cpp b/Engine/Source/Runtime/Engine/Private/PhysicsEngine/PhysicsReplication.cpp
--- a/Engine/Source/Runtime/Engine/Private/PhysicsEngine/PhysicsReplication.cpp
+++ b/Engine/Source/Runtime/Engine/Private/PhysicsEngine/PhysicsReplication.cpp
@@ -187,6 +187,9 @@ namespace PhysicsReplicationCVars
 		
 		static bool bDrawDebugVectors = false;
 		static FAutoConsoleVariableRef CVarDrawDebugVectors(TEXT("np2.PredictiveInterpolation.DrawDebugVectors"), bDrawDebugVectors, TEXT("Draw replication vectors, target velocity, replicated velocity, velocity change between replication calls etc."));
+		
+		static float SleepSecondsClearTarget = 15.0f;
+		static FAutoConsoleVariableRef CVarSleepSecondsClearTarget(TEXT("np2.PredictiveInterpolation.SleepSecondsClearTarget"), SleepSecondsClearTarget, TEXT("Wait for the object to sleep for this many seconds before clearing the replication target, to ensure nothing wakes up the object just after it goes to sleep on the client."));
 	}
 
 }
@@ -869,6 +872,7 @@ void FPhysicsReplicationAsync::UpdateAsyncTarget(const FPhysicsRepAsyncInputData
 		Target->RepMode = Input.RepMode;
 		Target->FrameOffset = Input.FrameOffset;
 		Target->TickCount = bFirstTarget ? 0 : FMath::Clamp((PrevTickCount - (SendInterval > 0 ? SendInterval : PrevReceiveInterval)), -AverageReceiveInterval, AverageReceiveInterval);
+		Target->AccumulatedSleepSeconds = 0.0f;
 
 		if (Input.RepMode == EPhysicsReplicationMode::PredictiveInterpolation)
 		{
@@ -1310,6 +1314,9 @@ bool FPhysicsReplicationAsync::PredictiveInterpolation(Chaos::FPBDRigidParticleH
 	const bool bIsSleeping = Handle->IsSleeping();
 	const bool bCanSimulate = Handle->IsDynamic() || bIsSleeping;
 
+	// Accumulate sleep time or reset back to 0s if not sleeping
+	Target.AccumulatedSleepSeconds = bIsSleeping ? (Target.AccumulatedSleepSeconds + DeltaSeconds) : 0.0f;
+	
 	// Helper for sleep and target clearing at replication end
 	auto EndReplicationHelper = [RigidsSolver, Handle, bCanSimulate, bIsSleeping](FReplicatedPhysicsTargetAsync& Target, bool bOkToClear) -> bool
 	{
@@ -1322,9 +1329,9 @@ bool FPhysicsReplicationAsync::PredictiveInterpolation(Chaos::FPBDRigidParticleH
 		}
 
 		const bool bClearTarget =
-			(!bCanSimulate 
-			|| (bOkToClear && bIsSleeping && bShouldSleep) // Don't clear the target due to sleeping until the object both should sleep and is sleeping
-			|| (bOkToClear && !bReplicatingPhysics))
+			(!bCanSimulate
+				|| (bOkToClear && bShouldSleep && Target.AccumulatedSleepSeconds >= PhysicsReplicationCVars::PredictiveInterpolationCVars::SleepSecondsClearTarget) // Don't clear the target due to sleeping until the object both should sleep and is sleeping for n seconds
+				|| (bOkToClear && !bReplicatingPhysics))
 			&& !PhysicsReplicationCVars::PredictiveInterpolationCVars::bDontClearTarget;
 	
 		return bClearTarget;
@@ -1530,7 +1537,8 @@ bool FPhysicsReplicationAsync::PredictiveInterpolation(Chaos::FPBDRigidParticleH
 		{
 			const FVector Offset = FVector(0.0f, 0.0f, 50.0f);
 			const FVector StartPos = TargetPos + Offset;
-			Chaos::FDebugDrawQueue::GetInstance().DrawDebugBox(StartPos, FVector(5.0f + Target.TickCount * 0.75f, 5.0f + Target.TickCount * 0.75f, 5.0f + Target.TickCount * 0.75f), Target.TargetState.Quaternion, FColor::MakeRandomSeededColor(Target.ServerFrame), false, CharacterMovementCVars::NetCorrectionLifetime, 0, 1.0f);
+			const int32 SizeMultiplier = FMath::Min(Target.TickCount, 100);
+			Chaos::FDebugDrawQueue::GetInstance().DrawDebugBox(StartPos, FVector(5.0f + SizeMultiplier * 0.75f, 5.0f + SizeMultiplier * 0.75f, 5.0f + SizeMultiplier * 0.75f), Target.TargetState.Quaternion, FColor::MakeRandomSeededColor(Target.ServerFrame), false, CharacterMovementCVars::NetCorrectionLifetime, 0, 1.0f);
 		}
 #endif
 
diff --git a/Engine/Source/Runtime/Engine/Public/PhysicsReplication.h b/Engine/Source/Runtime/Engine/Public/PhysicsReplication.h
--- a/Engine/Source/Runtime/Engine/Public/PhysicsReplication.h
+++ b/Engine/Source/Runtime/Engine/Public/PhysicsReplication.h
@@ -99,6 +99,7 @@ struct FReplicatedPhysicsTargetAsync
 		, AverageReceiveInterval(5.f)
 		, PrevServerFrame(INDEX_NONE)
 		, bWaiting(false)
+		, AccumulatedSleepSeconds(0.0f)
 	{ }
 
 	/** The target state replicated by server */
@@ -136,6 +137,9 @@ struct FReplicatedPhysicsTargetAsync
 
 	/** If this target is waiting for up-to-date data? */
 	bool bWaiting;
+
+	/** Accumulated seconds asleep */
+	float AccumulatedSleepSeconds;
 };
 
 class FPhysicsReplicationAsync : public Chaos::TSimCallbackObject<
