Disable ray tracing for SM5:
* Set r.RayTracing.RequireSM6 to 1 by default.
* Editor will display a notification if shader format is set to SM5 but ray tracing is enabled
* To make the transition easier setting r.RayTracing.RequireSM6=0 will make ray tracing available in SM5 as before.

diff --git a/Engine/Source/Runtime/Engine/Private/RendererSettings.cpp b/Engine/Source/Runtime/Engine/Private/RendererSettings.cpp
--- a/Engine/Source/Runtime/Engine/Private/RendererSettings.cpp
+++ b/Engine/Source/Runtime/Engine/Private/RendererSettings.cpp
@@ -398,7 +398,10 @@ void URendererSettings::CheckForMissingShaderModels()
 {
 	// Don't show the SM6 toasts on non-Windows/Linux platforms to avoid confusion around platform requirements.
 #if PLATFORM_WINDOWS || PLATFORM_LINUX
-	if (GIsEditor && ShadowMapMethod == EShadowMapMethod::VirtualShadowMaps)
+	static IConsoleVariable* RayTracingRequireSM6CVar = IConsoleManager::Get().FindConsoleVariable(TEXT("r.RayTracing.RequireSM6"));
+	const bool bRayTracingRequireSM6 = bEnableRayTracing && RayTracingRequireSM6CVar && RayTracingRequireSM6CVar->GetBool();
+
+	if (GIsEditor && (ShadowMapMethod == EShadowMapMethod::VirtualShadowMaps || bRayTracingRequireSM6))
 	{
 		auto CopySM6Format = [](const TCHAR* ShaderFormatName, const TArray<FString>& SrcArray, TArray<FString>& DstArray)
 		{
@@ -483,18 +486,60 @@ void URendererSettings::CheckForMissingShaderModels()
 			Info.Text = LOCTEXT("NeedProjectSettings", "Missing Project Settings!");
 			Info.HyperlinkText = LOCTEXT("ProjectSettingsHyperlinkText", "Open Project Settings");
 			Info.Hyperlink = FSimpleDelegate::CreateLambda(OpenProjectSettings);
+	
+			const FText FeaturesNeedSM6[] =
+			{
+				LOCTEXT("VirtualShadowMapsNeedsSM6", "Shader Model 6 (SM6) is required to use Virtual Shadow Maps."),
+				LOCTEXT("RayTracingNeedsSM6", "Shader Model 6 (SM6) is required to use Ray Tracing."),
+				LOCTEXT("VirtualShadowMapsAndRayTracingNeedsSM6", "Shader Model 6 (SM6) is required to use Virtual Shadow Maps and Ray Tracing.")
+			};
+
+			const uint32 FeatureNeedsSM6Index = (ShadowMapMethod == EShadowMapMethod::VirtualShadowMaps && bRayTracingRequireSM6) ? 2 : (bRayTracingRequireSM6 ? 1 : 0);
 
 			if (bProjectMissingD3DSM6)
 			{
-				Info.SubText = LOCTEXT("VirtualShadowMapsNeedsSM6Setting", "Shader Model 6 (SM6) is required to use Virtual Shadow Maps. Please enable this in:\n  Project Settings -> Platforms -> Windows -> D3D12 Targeted Shader Formats\nVirtual shadow maps will not work until this is enabled.");
+				const FText FeatureWontWork[] =
+				{
+					LOCTEXT("VirtualShadowMapsWontWork", "Virtual shadow maps will not work until this is enabled."),
+					LOCTEXT("RayTracingWontWork", "Ray tracing will not work until this is enabled."),
+					LOCTEXT("VirtualShadowMapsAndRayTracingWontWork", "Virtual shadow maps and ray tracing will not work until this is enabled."),
+				};
+
+				Info.SubText = FText::Format(
+					LOCTEXT("VirtualShadowMapsNeedsSM6Setting", "{0} Please enable this in:\n  Project Settings -> Platforms -> Windows -> D3D12 Targeted Shader Formats\n{1}"),
+					FeaturesNeedSM6[FeatureNeedsSM6Index],
+					FeatureWontWork[FeatureNeedsSM6Index]
+				);
 			}
 			else if (bProjectMissingWindowsVulkanSM6)
 			{
-				Info.SubText = LOCTEXT("VirtualShadowMapsNeedsVulkanSM6WindowsSetting", "Shader Model 6 (SM6) is required to use Virtual Shadow Maps. Please enable this in:\n  Project Settings -> Platforms -> Windows -> Vulkan Targeted Shader Formats\nVirtual shadow maps will not work in Vulkan on Windows until this is enabled.");
+				const FText FeatureWontWorkWindowsVulkan[] =
+				{
+					LOCTEXT("VirtualShadowMapsWontWorkWindowsVulkan", "Virtual shadow maps will not work in Vulkan on Windows until this is enabled."),
+					LOCTEXT("RayTracingWontWorkWindowsVulkan", "Ray tracing will not work in Vulkan on Windows until this is enabled."),
+					LOCTEXT("VirtualShadowMapsAndRayTracingWontWorkWindowsVulkan", "Virtual shadow maps and ray tracing will not work in Vulkan on Windows until this is enabled."),
+				};
+
+				Info.SubText = FText::Format(
+					LOCTEXT("VirtualShadowMapsNeedsVulkanSM6WindowsSetting", "{0} Please enable this in:\n  Project Settings -> Platforms -> Windows -> Vulkan Targeted Shader Formats\n{1}"),
+					FeaturesNeedSM6[FeatureNeedsSM6Index],
+					FeatureWontWorkWindowsVulkan[FeatureNeedsSM6Index]
+				);
 			}
 			else if (bProjectMissingLinuxVulkanSM6)
 			{
-				Info.SubText = LOCTEXT("VirtualShadowMapsNeedsVulkanSM6LinuxSetting", "Shader Model 6 (SM6) is required to use Virtual Shadow Maps. Please enable this in:\n  Project Settings -> Platforms -> Linux -> Targeted RHIs\nVirtual shadow maps will not work in Vulkan on Linux until this is enabled.");
+				const FText FeatureWontWorkLinuxVulkan[] =
+				{
+					LOCTEXT("VirtualShadowMapsWontWorkLinuxVulkan", "Virtual shadow maps will not work in Vulkan on Linux until this is enabled.."),
+					LOCTEXT("RayTracingWontWorkLinuxVulkan", "Ray tracing will not work in Vulkan on Linux until this is enabled."),
+					LOCTEXT("VirtualShadowMapsAndRayTracingWontWorkLinuxVulkan", "Virtual shadow maps and ray tracing will not work in Vulkan on Linux until this is enabled."),
+				};
+
+				Info.SubText = FText::Format(
+					LOCTEXT("VirtualShadowMapsNeedsVulkanSM6LinuxSetting", "{0} Please enable this in:\n  Project Settings -> Platforms -> Linux -> Targeted RHIs\n{1}"),
+					FeaturesNeedSM6[FeatureNeedsSM6Index],
+					FeatureWontWorkLinuxVulkan[FeatureNeedsSM6Index]
+				);
 			}
 
 			ShaderModelNotificationPtr = FSlateNotificationManager::Get().AddNotification(Info);
diff --git a/Engine/Source/Runtime/RenderCore/Private/RenderUtils.cpp b/Engine/Source/Runtime/RenderCore/Private/RenderUtils.cpp
--- a/Engine/Source/Runtime/RenderCore/Private/RenderUtils.cpp
+++ b/Engine/Source/Runtime/RenderCore/Private/RenderUtils.cpp
@@ -70,11 +70,11 @@ static FAutoConsoleVariableRef CVarRayTracingEnableOnDemand(
 	TEXT(" 1: on"),
 	ECVF_RenderThreadSafe | ECVF_ReadOnly);
 
-static int32 GRayTracingRequireSM6 = 0;
+static int32 GRayTracingRequireSM6 = 1;
 static FAutoConsoleVariableRef CVarRayTracingRequireSM6(
 	TEXT("r.RayTracing.RequireSM6"),
 	GRayTracingRequireSM6,
-	TEXT("Whether ray tracing shaders and features should only be available when targetting and running SM6. If disabled, ray tracing shaders will also be available when running in SM5 mode. (default = 0, allow SM5 and SM6)"),
+	TEXT("Whether ray tracing shaders and features should only be available when targetting and running SM6. If disabled, ray tracing shaders will also be available when running in SM5 mode. (default = 1, allow only SM6)"),
 	ECVF_RenderThreadSafe | ECVF_ReadOnly
 );
 
