Readded universal object locator BP library

diff --git a/Engine/Source/Runtime/Engine/Private/UniversalObjectLocators/UniversalObjectLocatorScriptingExtensions.cpp b/Engine/Source/Runtime/Engine/Private/UniversalObjectLocators/UniversalObjectLocatorScriptingExtensions.cpp
new file mode 100644
--- /dev/null
+++ b/Engine/Source/Runtime/Engine/Private/UniversalObjectLocators/UniversalObjectLocatorScriptingExtensions.cpp
@@ -0,0 +1,60 @@
+// Copyright Epic Games, Inc. All Rights Reserved.
+
+#include "UniversalObjectLocators/UniversalObjectLocatorScriptingExtensions.h"
+#include "UniversalObjectLocator.h"
+#include "UniversalObjectLocatorStringParams.h"
+#include "UObject/Stack.h"
+
+FUniversalObjectLocator UUniversalObjectLocatorScriptingExtensions::MakeUniversalObjectLocator(UObject* Object, const UObject* Context)
+{
+	return FUniversalObjectLocator(Object, Context);
+}
+
+FUniversalObjectLocator UUniversalObjectLocatorScriptingExtensions::UniversalObjectLocatorFromString(const FString& InString)
+{
+	using namespace::UE::UniversalObjectLocator;
+
+	FParseStringParams Params;
+	Params.Flags = EParseStringFlags::ErrorMessaging;
+
+	FUniversalObjectLocator Locator;
+	FParseStringResult Result = Locator.TryParseString(InString, Params);
+	if (!Result)
+	{
+		FFrame::KismetExecutionMessage(
+			*FText::Format(
+				NSLOCTEXT("UOL", "ParseError", "Unable to parse a valid UOL from string '{0}': {1}."),
+				FText::FromStringView(InString),
+				Result.ErrorMessage
+			).ToString()
+		, ELogVerbosity::Error);
+	}
+	return Locator;
+}
+
+bool UUniversalObjectLocatorScriptingExtensions::IsEmpty(const FUniversalObjectLocator& Locator)
+{
+	return Locator.IsEmpty();
+}
+
+FString UUniversalObjectLocatorScriptingExtensions::ToString(const FUniversalObjectLocator& Locator)
+{
+	TStringBuilder<128> String;
+	Locator.ToString(String);
+	return String.ToString();
+}
+
+UObject* UUniversalObjectLocatorScriptingExtensions::SyncFind(const FUniversalObjectLocator& Locator, UObject* Context)
+{
+	return Locator.SyncFind(Context);
+}
+
+UObject* UUniversalObjectLocatorScriptingExtensions::SyncLoad(const FUniversalObjectLocator& Locator, UObject* Context)
+{
+	return Locator.SyncLoad(Context);
+}
+
+void UUniversalObjectLocatorScriptingExtensions::SyncUnload(const FUniversalObjectLocator& Locator, UObject* Context)
+{
+	Locator.SyncUnload(Context);
+}
diff --git a/Engine/Source/Runtime/Engine/Public/UniversalObjectLocators/UniversalObjectLocatorScriptingExtensions.h b/Engine/Source/Runtime/Engine/Public/UniversalObjectLocators/UniversalObjectLocatorScriptingExtensions.h
new file mode 100644
--- /dev/null
+++ b/Engine/Source/Runtime/Engine/Public/UniversalObjectLocators/UniversalObjectLocatorScriptingExtensions.h
@@ -0,0 +1,72 @@
+// Copyright Epic Games, Inc. All Rights Reserved.
+
+#pragma once
+
+#include "Kismet/BlueprintFunctionLibrary.h"
+#include "Engine/LatentActionManager.h"
+
+#include "UniversalObjectLocatorScriptingExtensions.generated.h"
+
+struct FUniversalObjectLocator;
+
+/**
+ * Function library containing methods that should be hoisted onto FUniversalObjectLocators for scripting
+ */
+UCLASS(MinimalAPI)
+class UUniversalObjectLocatorScriptingExtensions : public UBlueprintFunctionLibrary
+{
+public:
+	GENERATED_BODY()
+
+	/**
+	 * Construct a new universal object locator
+	 */
+	UFUNCTION(BlueprintPure, Category="Universal Object Locators")
+	static FUniversalObjectLocator MakeUniversalObjectLocator(UObject* Object, const UObject* Context);
+
+	/**
+	 * Construct a new universal object locator from a string
+	 */
+	UFUNCTION(BlueprintPure, Category="Universal Object Locators")
+	static FUniversalObjectLocator UniversalObjectLocatorFromString(const FString& InString);
+
+	/**
+	 * Check whether the specified locator is empty; not equivalent to Resolve() != None.
+	 * An empty locator will never resolve to a valid object.
+	 */
+	UFUNCTION(BlueprintPure, Category="Universal Object Locators", meta=(ScriptMethod))
+	static bool IsEmpty(const FUniversalObjectLocator& Locator);
+
+	/**
+	 * Convert the specified locator to its string representation
+	 */
+	UFUNCTION(BlueprintPure, Category="Universal Object Locators", meta=(ScriptMethod))
+	static FString ToString(const FUniversalObjectLocator& Locator);
+
+	/**
+	 * Attempt to resolve the object locator by finding the object. If it is not currently loaded or created, 
+	 * 
+	 * @param Context    (Optional) Context object to use for resolving the object. This should usually be the object that owns or created the locator.
+	 * @return The resolve object pointer, or null if it was not found.
+	 */
+	UFUNCTION(BlueprintPure, Category="Universal Object Locators", meta=(ScriptMethod))
+	static UObject* SyncFind(const FUniversalObjectLocator& Locator, UObject* Context = nullptr);
+
+	/**
+	 * Attempt to resolve the object locator by finding or loading the object.
+	 * 
+	 * @param Context    (Optional) Context object to use for resolving the object. This should usually be the object that owns or created the locator.
+	 * @return The resolve object pointer, or null if it was not found.
+	 */
+	UFUNCTION(BlueprintPure, Category="Universal Object Locators", meta=(ScriptMethod))
+	static UObject* SyncLoad(const FUniversalObjectLocator& Locator, UObject* Context = nullptr);
+
+	/**
+	 * Attempt to resolve the object locator by unloading the object if possible.
+	 * 
+	 * @param Context    (Optional) Context object to use for resolving the object. This should usually be the object that owns or created the locator.
+	 * @return The resolve object pointer, or null if it was not found.
+	 */
+	UFUNCTION(BlueprintPure, Category="Universal Object Locators", meta=(ScriptMethod))
+	static void SyncUnload(const FUniversalObjectLocator& Locator, UObject* Context = nullptr);
+};
\ No newline at end of file
diff --git a/Engine/Source/Runtime/UniversalObjectLocator/Public/UniversalObjectLocator.h b/Engine/Source/Runtime/UniversalObjectLocator/Public/UniversalObjectLocator.h
--- a/Engine/Source/Runtime/UniversalObjectLocator/Public/UniversalObjectLocator.h
+++ b/Engine/Source/Runtime/UniversalObjectLocator/Public/UniversalObjectLocator.h
@@ -20,7 +20,7 @@
  *
  * This type is 16 bytes.
  */
-USTRUCT(BlueprintType, Category=GameFramework)
+USTRUCT(BlueprintType, Category=GameFramework, meta=(HasNativeMake="/Script/Engine.UniversalObjectLocatorScriptingExtensions.MakeUniversalObjectLocator"))
 struct FUniversalObjectLocator
 {
 	GENERATED_BODY()
