Block compilation of UI materials when expressions relying on the primitive data uniform buffer are required, as this UB data is not available for slate quads.

diff --git a/Engine/Source/Runtime/Engine/Private/Materials/HLSLMaterialTranslator.cpp b/Engine/Source/Runtime/Engine/Private/Materials/HLSLMaterialTranslator.cpp
--- a/Engine/Source/Runtime/Engine/Private/Materials/HLSLMaterialTranslator.cpp
+++ b/Engine/Source/Runtime/Engine/Private/Materials/HLSLMaterialTranslator.cpp
@@ -5975,6 +5975,11 @@ int32 FHLSLMaterialTranslator::WorldPosition(EWorldPositionIncludedOffsets World
 
 int32 FHLSLMaterialTranslator::ObjectWorldPosition(EPositionOrigin OriginType)
 {
+	if (Material->GetMaterialDomain() == MD_UI)
+	{
+		return Errorf(TEXT("ObjectWorldPosition not available for UI materials"));
+	}
+
 	if (OriginType == EPositionOrigin::CameraRelative)
 	{
 		return AddInlinedCodeChunkZeroDeriv(MCT_Float3,TEXT("GetObjectTranslatedWorldPosition(Parameters)"));
@@ -5993,11 +5998,20 @@ int32 FHLSLMaterialTranslator::ObjectRadius()
 
 int32 FHLSLMaterialTranslator::ObjectBounds()
 {
+	if (Material->GetMaterialDomain() == MD_UI)
+	{
+		return Errorf(TEXT("ObjectBounds not available for UI materials"));
+	}
 	return AddInlinedCodeChunk(MCT_Float3, TEXT("float3(GetPrimitiveData(Parameters).ObjectBoundsX, GetPrimitiveData(Parameters).ObjectBoundsY, GetPrimitiveData(Parameters).ObjectBoundsZ)"));
 }
 
 int32 FHLSLMaterialTranslator::ObjectLocalBounds(int32 OutputIndex)
 {
+	if (Material->GetMaterialDomain() == MD_UI)
+	{
+		return Errorf(TEXT("ObjectLocalBounds not available for UI materials"));
+	}
+
 	switch (OutputIndex)
 	{
 	case 0: // Half extents
@@ -6036,6 +6050,11 @@ int32 FHLSLMaterialTranslator::InstanceLocalBounds(int32 OutputIndex)
 
 int32 FHLSLMaterialTranslator::PreSkinnedLocalBounds(int32 OutputIndex)
 {
+	if (Material->GetMaterialDomain() == MD_UI)
+	{
+		return Errorf(TEXT("PreSkinnedLocalBounds not available for UI materials"));
+	}
+
 	switch (OutputIndex)
 	{
 	case 0: // Half extents
@@ -10000,6 +10019,11 @@ int32 FHLSLMaterialTranslator::TransformBase(EMaterialCommonBasis SourceCoordBas
 		return INDEX_NONE;
 	}
 
+	if(Material->GetMaterialDomain() == MD_UI)
+	{
+		return Errorf(TEXT("Transform not available for UI materials"));
+	}
+
 	const EMaterialValueType SourceType = GetParameterType(A);
 	const bool bIsPositionTranform = AWComponent != 0;
 		
@@ -10478,6 +10502,11 @@ int32 FHLSLMaterialTranslator::VirtualTextureOutputReplace(int32 Default, int32
 
 int32 FHLSLMaterialTranslator::ObjectOrientation()
 { 
+	if (Material->GetMaterialDomain() == MD_UI)
+	{
+		return Errorf(TEXT("ObjectOrientation not available for UI materials"));
+	}
+
 	return AddInlinedCodeChunkZeroDeriv(MCT_Float3,TEXT("GetObjectOrientation(Parameters)"));
 }
 
@@ -11175,6 +11204,11 @@ int32 FHLSLMaterialTranslator::GetCloudEmptySpaceSkippingSphereRadius()
 
 int32 FHLSLMaterialTranslator::CustomPrimitiveData(int32 OutputIndex, EMaterialValueType Type)
 {
+	if (Material->GetMaterialDomain() == MD_UI)
+	{
+		return Errorf(TEXT("CustomPrimitiveData not available for UI materials"));
+	}
+
 	check(OutputIndex < FCustomPrimitiveData::NumCustomPrimitiveDataFloats);
 
 	const int32 NumComponents = GetNumComponents(Type);
