Material Editor - Turned some checks into error messages / ensures in FunctionCall expression code to make sure the editor doesn't crash when loading a material containing invalid function calls.

diff --git a/Engine/Source/Runtime/Engine/Private/Materials/MaterialExpressions.cpp b/Engine/Source/Runtime/Engine/Private/Materials/MaterialExpressions.cpp
--- a/Engine/Source/Runtime/Engine/Private/Materials/MaterialExpressions.cpp
+++ b/Engine/Source/Runtime/Engine/Private/Materials/MaterialExpressions.cpp
@@ -15448,7 +15448,10 @@ bool UMaterialFunction::IsDependent(UMaterialFunctionInterface* OtherFunction)
 bool UMaterialFunction::IterateDependentFunctions(TFunctionRef<bool(UMaterialFunctionInterface*)> Predicate) const
 {
 #if WITH_EDITOR
-	check(!HasAnyFlags(RF_NeedPostLoad));
+	if (!ensure(!HasAnyFlags(RF_NeedPostLoad)))
+	{
+		return false;
+	}
 	for (UMaterialExpressionMaterialFunctionCall* MaterialFunctionExpression : DependentFunctionExpressionCandidates)
 	{
 		if(!MaterialFunctionExpression->IterateDependentFunctions(Predicate))
@@ -16749,12 +16752,18 @@ int32 UMaterialExpressionMaterialFunctionCall::Compile(class FMaterialCompiler*
 	// Verify that all function inputs and outputs are in a valid state to be linked into this material for compiling
 	for (int32 i = 0; i < FunctionInputs.Num(); i++)
 	{
-		check(FunctionInputs[i].ExpressionInput);
+		if (!FunctionInputs[i].ExpressionInput)
+		{
+			return Compiler->Errorf(TEXT("Function call input with index %d is unset."), i);
+		}
 	}
 
 	for (int32 i = 0; i < FunctionOutputs.Num(); i++)
 	{
-		check(FunctionOutputs[i].ExpressionOutput);
+		if (!FunctionOutputs[i].ExpressionOutput)
+		{
+			return Compiler->Errorf(TEXT("Function call output with index %d is unset."), i);
+		}
 	}
 
 	if (!FunctionOutputs.IsValidIndex(OutputIndex))
