Improve Global Shaders error logging to make it easier to diagnose and triage shader errors.

# Unifying global shader error logging
- There is now _one_ location in the code (ProcessCompiledGlobalShaders) where global shader errors are reported (instead of three total).
- We process all global shader errors at a higher level across all jobs, instead of the individual job level.
- Remove error logging from ProcessCompiledJob
- Remove error logging from FinishCompileShader
- HandlePotentialRetryOnError shouldn't also log errors too.  When you are running with `-unattended` we don't attempt to retry errors.  This means cooks won't print every global shader error message twice.

# Structured Logging
- Global shaders now log with structured logging (UE_LOGFMT_NSLOC)

# Misc
- De-dupe shader error messages and only report them once.
- Print all global shader errors as `Error` then print a `Fatal` error after all of them to kill the process.
- Refactor unique error gathering into a function `GatherUniqueErrors` which is called by both `HandlePotentialRetryOnError` and `ProcessCompiledGlobalShaders`.

#rnx

diff --git a/Engine/Source/Runtime/Engine/Private/ShaderCompiler/ShaderCompiler.cpp b/Engine/Source/Runtime/Engine/Private/ShaderCompiler/ShaderCompiler.cpp
--- a/Engine/Source/Runtime/Engine/Private/ShaderCompiler/ShaderCompiler.cpp
+++ b/Engine/Source/Runtime/Engine/Private/ShaderCompiler/ShaderCompiler.cpp
@@ -23,6 +23,7 @@
 #include "Interfaces/ITargetPlatform.h"
 #include "Interfaces/ITargetPlatformManagerModule.h"
 #include "Internationalization/LocKeyFuncs.h"
+#include "Logging/StructuredLog.h"
 #include "MaterialShared.h"
 #include "Materials/Material.h"
 #include "Materials/MaterialInstance.h"
@@ -3275,6 +3276,15 @@ FSCWErrorCode::ECode FShaderCompileUtilities::DoReadTaskResults(const TArray<FSh
 	return FSCWErrorCode::Success;
 }
 
+struct FShaderErrorInfo
+{
+	TArray<FShaderCommonCompileJob*> ErrorJobs;
+	TArray<FString> UniqueErrors;
+	TArray<FString> UniqueWarnings;
+	TArray<EShaderPlatform> ErrorPlatforms;
+	FString TargetShaderPlatformString;
+};
+
 #if WITH_EDITOR
 static bool CheckSingleJob(const FShaderCompileJob& SingleJob, TArray<FString>& OutErrors)
 {
@@ -3319,10 +3329,11 @@ static bool CheckSingleJob(const FShaderCompileJob& SingleJob, TArray<FString>&
 
 	return bSucceeded;
 };
+#endif // WITH_EDITOR
 
-static void AddErrorsForFailedJob(FShaderCompileJob& CurrentJob, TArray<EShaderPlatform>& ErrorPlatforms, TArray<FString>& UniqueErrors, TArray<FShaderCommonCompileJob*>& ErrorJobs)
+static void AddErrorsForFailedJob(FShaderCompileJob& CurrentJob, FShaderErrorInfo& OutShaderErrorInfo)
 {
-	ErrorPlatforms.AddUnique((EShaderPlatform)CurrentJob.Input.Target.Platform);
+	OutShaderErrorInfo.ErrorPlatforms.AddUnique((EShaderPlatform)CurrentJob.Input.Target.Platform);
 
 	if (CurrentJob.Output.Errors.Num() == 0)
 	{
@@ -3338,12 +3349,23 @@ static void AddErrorsForFailedJob(FShaderCompileJob& CurrentJob, TArray<EShaderP
 		// Include warnings if LogShaders is unsuppressed, otherwise only include errors
 		if (UE_LOG_ACTIVE(LogShaders, Log) || CurrentError.StrippedErrorMessage.Contains(TEXT("error")))
 		{
-			UniqueErrors.AddUnique(CurrentJob.Output.Errors[ErrorIndex].GetErrorString());
-			ErrorJobs.AddUnique(&CurrentJob);
+			OutShaderErrorInfo.UniqueErrors.AddUnique(CurrentJob.Output.Errors[ErrorIndex].GetErrorString());
+			OutShaderErrorInfo.ErrorJobs.AddUnique(&CurrentJob);
+		}
+	}
+}
+
+static void AddWarningsForJob(const FShaderCompileJob& CurrentJob, FShaderErrorInfo& OutShaderErrorInfo)
+{
+	if (GShowShaderWarnings && CurrentJob.bSucceeded)
+	{
+		for (int32 ErrorIndex = 0; ErrorIndex < CurrentJob.Output.Errors.Num(); ErrorIndex++)
+		{
+			// If the job succeeded the Errors array will contain warnings.
+			OutShaderErrorInfo.UniqueWarnings.AddUnique(CurrentJob.Output.Errors[ErrorIndex].GetErrorString());
 		}
 	}
 }
-#endif // WITH_EDITOR
 
 /** Information tracked for each shader compile worker process instance. */
 struct FShaderCompileWorkerInfo
@@ -6699,8 +6721,70 @@ void FShaderCompilingManager::GetLocalStats(FShaderCompilerStats& OutStats) cons
 	}
 }
 
+static bool GatherUniqueErrors(const TArray<FShaderCommonCompileJobPtr>& CompleteJobs, FShaderErrorInfo& OutShaderErrorInfo)
+{
+	// Gather unique errors
+	for (int32 JobIndex = 0; JobIndex < CompleteJobs.Num(); JobIndex++)
+	{
+		FShaderCommonCompileJob& CurrentJob = *CompleteJobs[JobIndex];
+		if (!CurrentJob.bSucceeded)
+		{
+			FShaderCompileJob* SingleJob = CurrentJob.GetSingleShaderJob();
+			if (SingleJob)
+			{
+				AddErrorsForFailedJob(*SingleJob, OutShaderErrorInfo);
+			}
+			else
+			{
+				FShaderPipelineCompileJob* PipelineJob = CurrentJob.GetShaderPipelineJob();
+				check(PipelineJob);
+				for (TRefCountPtr<FShaderCompileJob>& CommonJob : PipelineJob->StageJobs)
+				{
+					AddErrorsForFailedJob(*CommonJob, OutShaderErrorInfo);
+				}
+			}
+		}
+		else if (GShowShaderWarnings)
+		{
+			const FShaderCompileJob* SingleJob = CurrentJob.GetSingleShaderJob();
+			if (SingleJob)
+			{
+				AddWarningsForJob(*SingleJob, OutShaderErrorInfo);
+			}
+			else
+			{
+				const FShaderPipelineCompileJob* PipelineJob = CurrentJob.GetShaderPipelineJob();
+				check(PipelineJob);
+				for (const TRefCountPtr<FShaderCompileJob>& CommonJob : PipelineJob->StageJobs)
+				{
+					AddWarningsForJob(*CommonJob, OutShaderErrorInfo);
+				}
+			}
+		}
+	}
+
+	for (int32 PlatformIndex = 0; PlatformIndex < OutShaderErrorInfo.ErrorPlatforms.Num(); PlatformIndex++)
+	{
+		if (OutShaderErrorInfo.TargetShaderPlatformString.IsEmpty())
+		{
+			OutShaderErrorInfo.TargetShaderPlatformString = FDataDrivenShaderPlatformInfo::GetName(OutShaderErrorInfo.ErrorPlatforms[PlatformIndex]).ToString();
+		}
+		else
+		{
+			OutShaderErrorInfo.TargetShaderPlatformString += FString(TEXT(", ")) + FDataDrivenShaderPlatformInfo::GetName(OutShaderErrorInfo.ErrorPlatforms[PlatformIndex]).ToString();
+		}
+	}
+
+	return OutShaderErrorInfo.UniqueErrors.Num() > 0;
+}
+
 bool FShaderCompilingManager::HandlePotentialRetryOnError(TMap<int32, FShaderMapFinalizeResults>& CompletedShaderMaps)
 {
+	if (FApp::IsUnattended())
+	{
+		return false;
+	}
+
 	bool bRetryCompile = false;
 
 #if WITH_EDITORONLY_DATA
@@ -6732,49 +6816,11 @@ bool FShaderCompilingManager::HandlePotentialRetryOnError(TMap<int32, FShaderMap
 				|| It.Key() == GlobalShaderMapId)
 			{
 				TArray<FShaderCommonCompileJobPtr>& CompleteJobs = Results.FinishedJobs;
-				TArray<FShaderCommonCompileJob*> ErrorJobs;
-				TArray<FString> UniqueErrors;
-				TArray<EShaderPlatform> ErrorPlatforms;
-
-				// Gather unique errors
-				for (int32 JobIndex = 0; JobIndex < CompleteJobs.Num(); JobIndex++)
-				{
-					FShaderCommonCompileJob& CurrentJob = *CompleteJobs[JobIndex];
-					if (!CurrentJob.bSucceeded)
-					{
-						FShaderCompileJob* SingleJob = CurrentJob.GetSingleShaderJob();
-						if (SingleJob)
-						{
-							AddErrorsForFailedJob(*SingleJob, ErrorPlatforms, UniqueErrors, ErrorJobs);
-						}
-						else
-						{
-							FShaderPipelineCompileJob* PipelineJob = CurrentJob.GetShaderPipelineJob();
-							check(PipelineJob);
-							for (auto CommonJob : PipelineJob->StageJobs)
-							{
-								AddErrorsForFailedJob(*CommonJob, ErrorPlatforms, UniqueErrors, ErrorJobs);
-							}
-						}
-					}
-				}
-
-				FString TargetShaderPlatformString;
-
-				for (int32 PlatformIndex = 0; PlatformIndex < ErrorPlatforms.Num(); PlatformIndex++)
-				{
-					if (TargetShaderPlatformString.IsEmpty())
-					{
-						TargetShaderPlatformString = FDataDrivenShaderPlatformInfo::GetName(ErrorPlatforms[PlatformIndex]).ToString();
-					}
-					else
-					{
-						TargetShaderPlatformString += FString(TEXT(", ")) + FDataDrivenShaderPlatformInfo::GetName(ErrorPlatforms[PlatformIndex]).ToString();
-					}
-				}
+				FShaderErrorInfo ShaderErrorInfo;
+				GatherUniqueErrors(CompleteJobs, ShaderErrorInfo);
 
 				const TCHAR* MaterialName = ShaderMap ? ShaderMap->GetFriendlyName() : TEXT("global shaders");
-				FString ErrorString = FString::Printf(TEXT("%i Shader compiler errors compiling %s for platform %s:"), UniqueErrors.Num(), MaterialName, *TargetShaderPlatformString);
+				FString ErrorString = FString::Printf(TEXT("%i Shader compiler errors compiling %s for platform %s:"), ShaderErrorInfo.UniqueErrors.Num(), MaterialName, *ShaderErrorInfo.TargetShaderPlatformString);
 				UE_LOG(LogShaderCompilers, Warning, TEXT("%s"), *ErrorString);
 				ErrorString += TEXT("\n");
 				bool bAnyErrorLikelyToBeCodeError = false;
@@ -6788,7 +6834,7 @@ bool FShaderCompilingManager::HandlePotentialRetryOnError(TMap<int32, FShaderMap
 						const auto* SingleJob = CurrentJob.GetSingleShaderJob();
 						if (SingleJob)
 						{
-							ProcessErrors(*SingleJob, UniqueErrors, ErrorString);
+							ProcessErrors(*SingleJob, ShaderErrorInfo.UniqueErrors, ErrorString);
 						}
 						else
 						{
@@ -6796,7 +6842,7 @@ bool FShaderCompilingManager::HandlePotentialRetryOnError(TMap<int32, FShaderMap
 							check(PipelineJob);
 							for (auto CommonJob : PipelineJob->StageJobs)
 							{
-								ProcessErrors(*CommonJob, UniqueErrors, ErrorString);
+								ProcessErrors(*CommonJob, ShaderErrorInfo.UniqueErrors, ErrorString);
 							}
 						}
 					}
@@ -8888,26 +8934,6 @@ FShader* FGlobalShaderTypeCompiler::FinishCompileShader(const FGlobalShaderType*
 		CurrentJob.Output.ParameterMap.VerifyBindingsAreComplete(ShaderType->GetName(), CurrentJob.Output.Target, CurrentJob.Key.VFType);
 	}
 
-	if (CurrentJob.Output.Errors.Num() > 0)
-	{
-		if (CurrentJob.bSucceeded == false)
-		{
-			UE_LOG(LogShaderCompilers, Error, TEXT("Errors compiling global shader %s %s %s:\n"), CurrentJob.Key.ShaderType->GetName(), ShaderPipelineType ? TEXT("ShaderPipeline") : TEXT(""), ShaderPipelineType ? ShaderPipelineType->GetName() : TEXT(""));
-			for (int32 ErrorIndex = 0; ErrorIndex < CurrentJob.Output.Errors.Num(); ErrorIndex++)
-			{
-				UE_LOG(LogShaderCompilers, Display, TEXT("%s"), *CurrentJob.Output.Errors[ErrorIndex].GetErrorStringWithLineMarker());
-			}
-		}
-		else if (GShowShaderWarnings)
-		{
-			UE_LOG(LogShaderCompilers, Warning, TEXT("Warnings compiling global shader %s %s %s:\n"), CurrentJob.Key.ShaderType->GetName(), ShaderPipelineType ? TEXT("ShaderPipeline") : TEXT(""), ShaderPipelineType ? ShaderPipelineType->GetName() : TEXT(""));
-			for (int32 ErrorIndex = 0; ErrorIndex < CurrentJob.Output.Errors.Num(); ErrorIndex++)
-			{
-				UE_LOG(LogShaderCompilers, Display, TEXT("%s"), *CurrentJob.Output.Errors[ErrorIndex].GetErrorStringWithLineMarker());
-			}
-		}
-	}
-
 	return Shader;
 }
 
@@ -10151,23 +10177,6 @@ static inline FShader* ProcessCompiledJob(FShaderCompileJob* SingleJob, const FS
 		}
 		ShaderPlatformsProcessed.AddUnique(Platform);
 	}
-	else
-	{
-		if (AreShaderErrorsFatal())
-		{
-			UE_LOG(LogShaders, Fatal, TEXT("Failed to compile global shader %s %s %s.  Enable 'r.ShaderDevelopmentMode' in ConsoleVariables.ini for retries."),
-				GlobalShaderType->GetName(),
-				Pipeline ? TEXT("for pipeline") : TEXT(""),
-				Pipeline ? Pipeline->GetName() : TEXT(""));
-		}
-		else
-		{
-			UE_LOG(LogShaders, Error, TEXT("Failed to compile global shader %s %s %s.  Enable 'r.ShaderDevelopmentMode' in ConsoleVariables.ini for retries."),
-				GlobalShaderType->GetName(),
-				Pipeline ? TEXT("for pipeline") : TEXT(""),
-				Pipeline ? Pipeline->GetName() : TEXT(""));
-		}
-	}
 
 	return Shader;
 };
@@ -10178,6 +10187,44 @@ void ProcessCompiledGlobalShaders(const TArray<FShaderCommonCompileJobPtr>& Comp
 
 	UE_LOG(LogShaders, Verbose, TEXT("Compiled %u global shaders"), CompilationResults.Num());
 
+	FShaderErrorInfo ShaderErrorInfo;
+	GatherUniqueErrors(CompilationResults, ShaderErrorInfo);
+
+	// Report unique errors for global shaders.
+	for (const FString& ErrorString : ShaderErrorInfo.UniqueErrors)
+	{
+		UE_LOGFMT_NSLOC(LogShaders, Error, "Shaders", "GlobalShaderCompileError", "{ErrorMessage}", 
+			("ErrorMessage", ErrorString));
+	}
+
+	if (GShowShaderWarnings)
+	{
+		for (const FString& WarningString : ShaderErrorInfo.UniqueWarnings)
+		{
+			UE_LOGFMT_NSLOC(LogShaders, Warning, "Shaders", "GlobalShaderCompileWarning", "{WarningMessage}",
+				("WarningMessage", WarningString));
+		}
+	}
+
+	const int32 UniqueErrorCount = ShaderErrorInfo.UniqueErrors.Num();
+	if (UniqueErrorCount)
+	{
+		if (AreShaderErrorsFatal())
+		{
+			UE_LOGFMT_NSLOC(LogShaders, Fatal, "Shaders", "GlobalShadersCompilationFailed", "{NumErrors} Shader compiler errors compiling GlobalShaders for platform {Platform}.  Enable 'r.ShaderDevelopmentMode' in ConsoleVariables.ini for retries.",
+				("NumErrors", UniqueErrorCount),
+				("Platform", ShaderErrorInfo.TargetShaderPlatformString)
+			);
+		}
+		else
+		{
+			UE_LOGFMT_NSLOC(LogShaders, Error, "Shaders", "GlobalShadersCompilationFailed", "{NumErrors} Shader compiler errors compiling GlobalShaders for platform {Platform}.  Enable 'r.ShaderDevelopmentMode' in ConsoleVariables.ini for retries.",
+				("NumErrors", UniqueErrorCount),
+				("Platform", ShaderErrorInfo.TargetShaderPlatformString)
+			);
+		}
+	}
+
 	TArray<EShaderPlatform> ShaderPlatformsProcessed;
 	TArray<const FShaderPipelineType*> SharedPipelines;
 
