Moved FRayTracingGeometryManager out of RenderCore
- This will enable extending it to track which proxies use each FRayTracingGeometry/UStaticMesh in follow up changes.
- Added new IRayTracingGeometryManagerInterface to RenderCore since FRayTracingGeometry also needs to interact with FRayTracingGeometryManager.

diff --git a/Engine/Plugins/Experimental/GPULightmass/Source/GPULightmass/Private/LightmapRenderer.cpp b/Engine/Plugins/Experimental/GPULightmass/Source/GPULightmass/Private/LightmapRenderer.cpp
--- a/Engine/Plugins/Experimental/GPULightmass/Source/GPULightmass/Private/LightmapRenderer.cpp
+++ b/Engine/Plugins/Experimental/GPULightmass/Source/GPULightmass/Private/LightmapRenderer.cpp
@@ -22,7 +22,7 @@
 #include "LightmapDenoising.h"
 #include "EngineModule.h"
 #include "PostProcess/PostProcessing.h"
-#include "RayTracingGeometryManager.h"
+#include "RayTracingGeometryManagerInterface.h"
 #include "RayTracingInstanceBufferUtil.h"
 #include "ScreenPass.h"
 #include "RayTracingDynamicGeometryCollection.h"
@@ -813,7 +813,7 @@ bool FSceneRenderState::SetupRayTracingScene(FRDGBuilder& GraphBuilder, FSceneUn
 #if RHI_RAYTRACING
 	// Force build all the open build requests
 	bool bBuildAll = true;
-	GRayTracingGeometryManager.ProcessBuildRequests(RHICmdList, bBuildAll);
+	GRayTracingGeometryManager->ProcessBuildRequests(RHICmdList, bBuildAll);
 #endif // RHI_RAYTRACING
 
 	if (CachedRayTracingScene.IsValid())
diff --git a/Engine/Source/Runtime/Engine/Private/Rendering/RayTracingGeometryManager.cpp b/Engine/Source/Runtime/Engine/Private/Rendering/RayTracingGeometryManager.cpp
new file mode 100644
--- /dev/null
+++ b/Engine/Source/Runtime/Engine/Private/Rendering/RayTracingGeometryManager.cpp
@@ -0,0 +1,270 @@
+// Copyright Epic Games, Inc. All Rights Reserved.
+
+#include "Rendering/RayTracingGeometryManager.h"
+
+#include "RHIResources.h"
+#include "RHICommandList.h"
+
+#include "RayTracingGeometry.h"
+#include "RenderUtils.h"
+
+#if RHI_RAYTRACING
+
+static int32 GRayTracingMaxBuiltPrimitivesPerFrame = -1;
+static FAutoConsoleVariableRef CVarRayTracingMaxBuiltPrimitivesPerFrame(
+	TEXT("r.RayTracing.Geometry.MaxBuiltPrimitivesPerFrame"),
+	GRayTracingMaxBuiltPrimitivesPerFrame,
+	TEXT("Sets the ray tracing acceleration structure build budget in terms of maximum number of triangles per frame (<= 0 then disabled and all acceleration structures are build immediatly - default)"),
+	ECVF_RenderThreadSafe
+);
+
+static float GRayTracingPendingBuildPriorityBoostPerFrame = 0.001f;
+static FAutoConsoleVariableRef CVarRayTracingPendingBuildPriorityBoostPerFrame(
+	TEXT("r.RayTracing.Geometry.PendingBuildPriorityBoostPerFrame"),
+	GRayTracingPendingBuildPriorityBoostPerFrame,
+	TEXT("Increment the priority for all pending build requests which are not scheduled that frame (0.001 - default)"),
+	ECVF_RenderThreadSafe
+);
+
+DECLARE_DWORD_ACCUMULATOR_STAT(TEXT("Ray tracing pending builds"), STAT_RayTracingPendingBuilds, STATGROUP_SceneRendering);
+DECLARE_DWORD_ACCUMULATOR_STAT(TEXT("Ray tracing pending build primitives"), STAT_RayTracingPendingBuildPrimitives, STATGROUP_SceneRendering);
+
+FRayTracingGeometryManager::~FRayTracingGeometryManager()
+{
+	check(GeometryBuildRequests.IsEmpty());
+	check(RegisteredGeometries.IsEmpty());
+}
+
+static float GetInitialBuildPriority(ERTAccelerationStructureBuildPriority InBuildPriority)
+{
+	switch (InBuildPriority)
+	{
+	case ERTAccelerationStructureBuildPriority::Immediate:	return 1.0f;
+	case ERTAccelerationStructureBuildPriority::High:		return 0.5f;
+	case ERTAccelerationStructureBuildPriority::Normal:		return 0.24f;
+	case ERTAccelerationStructureBuildPriority::Low:		return 0.01f;
+	case ERTAccelerationStructureBuildPriority::Skip:
+	default:
+	{
+		// should not get here
+		check(false);
+		return 0.0f;
+	}
+	}
+}
+
+FRayTracingGeometryManager::BuildRequestIndex FRayTracingGeometryManager::RequestBuildAccelerationStructure(FRHICommandList& RHICmdList, FRayTracingGeometry* InGeometry, ERTAccelerationStructureBuildPriority InPriority, EAccelerationStructureBuildMode InBuildMode)
+{
+	// If immediate then enqueue command directly on the immediate command list
+	if (GRayTracingMaxBuiltPrimitivesPerFrame <= 0 || InPriority == ERTAccelerationStructureBuildPriority::Immediate)
+	{
+		check(InBuildMode == EAccelerationStructureBuildMode::Build);
+		RHICmdList.BuildAccelerationStructure(InGeometry->RayTracingGeometryRHI);
+		return INDEX_NONE;
+	}
+	else
+	{
+		BuildRequest Request;
+		Request.BuildPriority = GetInitialBuildPriority(InPriority);
+		Request.Owner = InGeometry;
+		Request.BuildMode = EAccelerationStructureBuildMode::Build;
+
+		FScopeLock ScopeLock(&RequestCS);
+		BuildRequestIndex RequestIndex = GeometryBuildRequests.Add(Request);
+		GeometryBuildRequests[RequestIndex].RequestIndex = RequestIndex;
+
+		INC_DWORD_STAT(STAT_RayTracingPendingBuilds);
+		INC_DWORD_STAT_BY(STAT_RayTracingPendingBuildPrimitives, InGeometry->Initializer.TotalPrimitiveCount);
+
+		return RequestIndex;
+	}
+}
+
+void FRayTracingGeometryManager::RemoveBuildRequest(BuildRequestIndex InRequestIndex)
+{
+	FScopeLock ScopeLock(&RequestCS);
+
+	DEC_DWORD_STAT(STAT_RayTracingPendingBuilds);
+	DEC_DWORD_STAT_BY(STAT_RayTracingPendingBuildPrimitives, GeometryBuildRequests[InRequestIndex].Owner->Initializer.TotalPrimitiveCount);
+
+	GeometryBuildRequests.RemoveAt(InRequestIndex);
+}
+
+FRayTracingGeometryManager::RayTracingGeometryHandle FRayTracingGeometryManager::RegisterRayTracingGeometry(FRayTracingGeometry* InGeometry)
+{	
+	if (GetRayTracingMode() == ERayTracingMode::Dynamic)
+	{
+		check(InGeometry);
+
+		FScopeLock ScopeLock(&RequestCS);
+		RayTracingGeometryHandle Handle = RegisteredGeometries.Add(InGeometry);
+		return Handle;
+	}
+	return INDEX_NONE;
+}
+
+void FRayTracingGeometryManager::ReleaseRayTracingGeometryHandle(RayTracingGeometryHandle Handle)
+{
+	if (GetRayTracingMode() == ERayTracingMode::Dynamic)
+	{
+		check(Handle != INDEX_NONE);
+		FScopeLock ScopeLock(&RequestCS);
+		RegisteredGeometries.RemoveAt(Handle);
+	}	
+}
+
+void FRayTracingGeometryManager::Tick(FRHICommandList& RHICmdList, bool bHasRayTracingEnableChanged)
+{
+	if (GetRayTracingMode() != ERayTracingMode::Dynamic)
+	{
+		return;
+	}
+
+	TRACE_CPUPROFILER_EVENT_SCOPE(FRayTracingGeometryManager::Tick);
+	QUICK_SCOPE_CYCLE_COUNTER(STAT_FRayTracingGeometryManager_Tick);
+
+	if (!bHasRayTracingEnableChanged)
+	{
+// If the code below triggers a check then dynamic ray tracing is not going to work as expected ie. not all memory will be released or we'll be missing geometry.
+#if DO_CHECK
+		if (IsRayTracingEnabled())
+		{
+			FScopeLock ScopeLock(&RequestCS);
+			for (FRayTracingGeometry* Geometry : RegisteredGeometries)
+			{
+				checkf(!Geometry->IsEvicted(), TEXT("Ray tracing geometry should not be evicted when ray tracing is enabled."));
+			}
+		}
+		else
+		{
+			FScopeLock ScopeLock(&RequestCS);
+			for (FRayTracingGeometry* Geometry : RegisteredGeometries)
+			{
+				checkf(Geometry->IsEvicted() || Geometry->RayTracingGeometryRHI == nullptr, TEXT("Ray tracing geometry should be evicted when ray tracing is disabled."));
+			}
+		}
+#endif
+		return;
+	}
+
+	if (IsRayTracingEnabled())
+	{
+		FScopeLock ScopeLock(&RequestCS);
+		for (FRayTracingGeometry* Geometry : RegisteredGeometries)
+		{
+			if (Geometry->IsEvicted())
+			{
+				Geometry->MakeResident(RHICmdList);
+			}
+		}
+	}
+	else
+	{
+		FScopeLock ScopeLock(&RequestCS);
+		for (FRayTracingGeometry* Geometry : RegisteredGeometries)
+		{
+			if (Geometry->RayTracingGeometryRHI != nullptr)
+			{
+				Geometry->Evict();
+			}
+		}
+	}
+}
+
+void FRayTracingGeometryManager::BoostPriority(BuildRequestIndex InRequestIndex, float InBoostValue)
+{
+	FScopeLock ScopeLock(&RequestCS);
+	GeometryBuildRequests[InRequestIndex].BuildPriority += InBoostValue;
+}
+
+void FRayTracingGeometryManager::ForceBuildIfPending(FRHIComputeCommandList& InCmdList, const TArrayView<const FRayTracingGeometry*> InGeometries)
+{
+	FScopeLock ScopeLock(&RequestCS);
+
+	BuildParams.Empty(FMath::Max(BuildParams.Max(), InGeometries.Num()));
+	for (const FRayTracingGeometry* Geometry : InGeometries)
+	{
+		if (Geometry->HasPendingBuildRequest())
+		{
+			SetupBuildParams(GeometryBuildRequests[Geometry->RayTracingBuildRequestIndex], BuildParams);
+		}
+	}
+
+	if (BuildParams.Num())
+	{
+		InCmdList.BuildAccelerationStructures(BuildParams);
+	}
+}
+
+void FRayTracingGeometryManager::ProcessBuildRequests(FRHIComputeCommandList& InCmdList, bool bInBuildAll)
+{
+	TRACE_CPUPROFILER_EVENT_SCOPE(FRayTracingGeometryManager::ProcessBuildRequests);
+
+	FScopeLock ScopeLock(&RequestCS);
+
+	if (GeometryBuildRequests.Num() == 0)
+	{
+		return;
+	}
+
+	SortedRequests.Empty(FMath::Max(SortedRequests.Max(), GeometryBuildRequests.Num()));
+
+	{
+		TRACE_CPUPROFILER_EVENT_SCOPE(SortRequests);
+
+		// Is there a fast way to extract all entries from sparse array?
+		for (const BuildRequest& Request : GeometryBuildRequests)
+		{
+			SortedRequests.Add(Request);
+		}
+		SortedRequests.Sort([](const BuildRequest& InLHS, const BuildRequest& InRHS)
+			{				
+				return InLHS.BuildPriority > InRHS.BuildPriority;
+			});
+	}
+
+	BuildParams.Empty(FMath::Max(BuildParams.Max(), SortedRequests.Num()));
+
+	// process n requests each 'frame'
+	uint64 PrimitivesBuild = 0;
+	bool bAddBuildRequest = true;
+	for (BuildRequest& Request : SortedRequests)
+	{
+		if (bAddBuildRequest)
+		{
+			SetupBuildParams(Request, BuildParams);
+
+			// Requested enough?
+			PrimitivesBuild += Request.Owner->Initializer.TotalPrimitiveCount;
+			if (!bInBuildAll && PrimitivesBuild > GRayTracingMaxBuiltPrimitivesPerFrame)
+				bAddBuildRequest = false;
+		}
+		else
+		{
+			// Increment priority to make sure requests don't starve
+			Request.BuildPriority += GRayTracingPendingBuildPriorityBoostPerFrame;
+		}
+	}
+
+	// kick actual build request to RHI command list
+	InCmdList.BuildAccelerationStructures(BuildParams);
+}
+
+void FRayTracingGeometryManager::SetupBuildParams(const BuildRequest& InBuildRequest, TArray<FRayTracingGeometryBuildParams>& InBuildParams)
+{
+	// Setup the actual build params
+	FRayTracingGeometryBuildParams BuildParam;
+	BuildParam.Geometry = InBuildRequest.Owner->RayTracingGeometryRHI;
+	BuildParam.BuildMode = InBuildRequest.BuildMode;
+	InBuildParams.Add(BuildParam);
+
+	// Remove from pending array and update the geometry that data is valid
+	check(InBuildRequest.RequestIndex != INDEX_NONE);
+	GeometryBuildRequests.RemoveAt(InBuildRequest.RequestIndex);
+	InBuildRequest.Owner->RayTracingBuildRequestIndex = INDEX_NONE;
+
+	DEC_DWORD_STAT(STAT_RayTracingPendingBuilds);
+	DEC_DWORD_STAT_BY(STAT_RayTracingPendingBuildPrimitives, InBuildRequest.Owner->Initializer.TotalPrimitiveCount);
+}
+
+#endif // RHI_RAYTRACING
diff --git a/Engine/Source/Runtime/Engine/Public/Rendering/RayTracingGeometryManager.h b/Engine/Source/Runtime/Engine/Public/Rendering/RayTracingGeometryManager.h
new file mode 100644
--- /dev/null
+++ b/Engine/Source/Runtime/Engine/Public/Rendering/RayTracingGeometryManager.h
@@ -0,0 +1,58 @@
+// Copyright Epic Games, Inc. All Rights Reserved.
+#pragma once
+
+#include "RayTracingGeometryManagerInterface.h"
+
+#include "Containers/SparseArray.h"
+
+#if RHI_RAYTRACING
+
+class FRayTracingGeometry;
+class FRHIComputeCommandList;
+enum class EAccelerationStructureBuildMode;
+enum class ERTAccelerationStructureBuildPriority;
+
+class FRayTracingGeometryManager : public IRayTracingGeometryManager
+{
+public:
+
+	ENGINE_API virtual ~FRayTracingGeometryManager();
+
+	ENGINE_API virtual BuildRequestIndex RequestBuildAccelerationStructure(FRHICommandList& RHICmdList, FRayTracingGeometry* InGeometry, ERTAccelerationStructureBuildPriority InPriority, EAccelerationStructureBuildMode InBuildMode) override;
+
+	ENGINE_API virtual void RemoveBuildRequest(BuildRequestIndex InRequestIndex) override;
+	ENGINE_API virtual void BoostPriority(BuildRequestIndex InRequestIndex, float InBoostValue) override;
+	ENGINE_API virtual void ForceBuildIfPending(FRHIComputeCommandList& InCmdList, const TArrayView<const FRayTracingGeometry*> InGeometries) override;
+	ENGINE_API virtual void ProcessBuildRequests(FRHIComputeCommandList& InCmdList, bool bInBuildAll = false) override;
+
+	ENGINE_API virtual RayTracingGeometryHandle RegisterRayTracingGeometry(FRayTracingGeometry* InGeometry) override;
+	ENGINE_API virtual void ReleaseRayTracingGeometryHandle(RayTracingGeometryHandle Handle) override;
+
+	ENGINE_API virtual void Tick(FRHICommandList& RHICmdList, bool bHasRayTracingEnableChanged) override;
+
+private:
+
+	struct BuildRequest
+	{
+		BuildRequestIndex RequestIndex = INDEX_NONE;
+
+		float BuildPriority = 0.0f;
+		FRayTracingGeometry* Owner;
+		EAccelerationStructureBuildMode BuildMode;
+	};
+
+	void SetupBuildParams(const BuildRequest& InBuildRequest, TArray<FRayTracingGeometryBuildParams>& InBuildParams);
+
+	FCriticalSection RequestCS;
+
+	TSparseArray<BuildRequest> GeometryBuildRequests;
+
+	// Used for keeping track of geometries when ray tracing is dynamic
+	TSparseArray<FRayTracingGeometry*> RegisteredGeometries;
+
+	// Working array with all active build build params in the RHI
+	TArray<BuildRequest> SortedRequests;
+	TArray<FRayTracingGeometryBuildParams> BuildParams;
+};
+
+#endif // RHI_RAYTRACING
diff --git a/Engine/Source/Runtime/RenderCore/Private/RayTracingGeometry.cpp b/Engine/Source/Runtime/RenderCore/Private/RayTracingGeometry.cpp
--- a/Engine/Source/Runtime/RenderCore/Private/RayTracingGeometry.cpp
+++ b/Engine/Source/Runtime/RenderCore/Private/RayTracingGeometry.cpp
@@ -3,11 +3,17 @@
 #include "RayTracingGeometry.h"
 #include "RHICommandList.h"
 #include "HAL/IConsoleManager.h"
-#include "RayTracingGeometryManager.h"
+#include "RayTracingGeometryManagerInterface.h"
 #include "RenderUtils.h"
 #include "RHIResourceUpdates.h"
 #include "RHITextureReference.h" // IWYU pragma: keep
 
+#if RHI_RAYTRACING
+
+IRayTracingGeometryManager* GRayTracingGeometryManager = nullptr;
+
+#endif
+
 static TAutoConsoleVariable<int32> CVarDebugForceRuntimeBLAS(
 	TEXT("r.Raytracing.DebugForceRuntimeBLAS"),
 	0,
@@ -83,7 +89,7 @@ void FRayTracingGeometry::RequestBuildIfNeeded(FRHICommandList& RHICmdList, ERTA
 
 	if (GetRequiresBuild())
 	{
-		RayTracingBuildRequestIndex = GRayTracingGeometryManager.RequestBuildAccelerationStructure(RHICmdList, this, InBuildPriority);
+		RayTracingBuildRequestIndex = GRayTracingGeometryManager->RequestBuildAccelerationStructure(RHICmdList, this, InBuildPriority);
 		SetRequiresBuild(false);
 	}
 }
@@ -202,7 +208,7 @@ void FRayTracingGeometry::CreateRayTracingGeometry(FRHICommandList& RHICmdList,
 			{
 				if (RayTracingGeometryRHI)
 				{
-					RayTracingBuildRequestIndex = GRayTracingGeometryManager.RequestBuildAccelerationStructure(RHICmdList, this, InBuildPriority);
+					RayTracingBuildRequestIndex = GRayTracingGeometryManager->RequestBuildAccelerationStructure(RHICmdList, this, InBuildPriority);
 				}
 				SetRequiresBuild(false);
 			}
@@ -215,7 +221,7 @@ void FRayTracingGeometry::CreateRayTracingGeometry(FRHICommandList& RHICmdList,
 		{
 			if (RayTracingGeometryRHI && RayTracingGeometryRHI->IsCompressed())
 			{
-				RayTracingBuildRequestIndex = GRayTracingGeometryManager.RequestBuildAccelerationStructure(RHICmdList, this, InBuildPriority);
+				RayTracingBuildRequestIndex = GRayTracingGeometryManager->RequestBuildAccelerationStructure(RHICmdList, this, InBuildPriority);
 			}
 
 			SetRequiresBuild(false);
@@ -284,7 +290,7 @@ void FRayTracingGeometry::RemoveBuildRequest()
 {
 	if (HasPendingBuildRequest())
 	{
-		GRayTracingGeometryManager.RemoveBuildRequest(RayTracingBuildRequestIndex);
+		GRayTracingGeometryManager->RemoveBuildRequest(RayTracingBuildRequestIndex);
 		RayTracingBuildRequestIndex = INDEX_NONE;
 	}
 }
@@ -297,7 +303,7 @@ void FRayTracingGeometry::InitResource(FRHICommandListBase& RHICmdList)
 
 	if (RayTracingGeometryHandle == INDEX_NONE)
 	{
-		RayTracingGeometryHandle = GRayTracingGeometryManager.RegisterRayTracingGeometry(this);
+		RayTracingGeometryHandle = GRayTracingGeometryManager->RegisterRayTracingGeometry(this);
 	}
 }
 
@@ -307,7 +313,7 @@ void FRayTracingGeometry::ReleaseResource()
 
 	if (RayTracingGeometryHandle != INDEX_NONE)
 	{
-		GRayTracingGeometryManager.ReleaseRayTracingGeometryHandle(RayTracingGeometryHandle);
+		GRayTracingGeometryManager->ReleaseRayTracingGeometryHandle(RayTracingGeometryHandle);
 		RayTracingGeometryHandle = INDEX_NONE;
 	}
 
@@ -321,7 +327,7 @@ void FRayTracingGeometry::ReleaseResource()
 void FRayTracingGeometry::BoostBuildPriority(float InBoostValue) const
 {
 	check(HasPendingBuildRequest());
-	GRayTracingGeometryManager.BoostPriority(RayTracingBuildRequestIndex, InBoostValue);
+	GRayTracingGeometryManager->BoostPriority(RayTracingBuildRequestIndex, InBoostValue);
 }
 
 #endif // RHI_RAYTRACING
diff --git a/Engine/Source/Runtime/RenderCore/Private/RayTracingGeometryManager.cpp b/Engine/Source/Runtime/RenderCore/Private/RayTracingGeometryManager.cpp
deleted file mode 100644
diff --git a/Engine/Source/Runtime/RenderCore/Public/RayTracingGeometryManager.h b/Engine/Source/Runtime/RenderCore/Public/RayTracingGeometryManager.h
deleted file mode 100644
diff --git a/Engine/Source/Runtime/RenderCore/Public/RayTracingGeometryManagerInterface.h b/Engine/Source/Runtime/RenderCore/Public/RayTracingGeometryManagerInterface.h
new file mode 100644
--- /dev/null
+++ b/Engine/Source/Runtime/RenderCore/Public/RayTracingGeometryManagerInterface.h
@@ -0,0 +1,58 @@
+// Copyright Epic Games, Inc. All Rights Reserved.
+#pragma once
+
+#include "Containers/ArrayView.h"
+
+#include "RHI.h"
+#include "RHIResources.h"
+#include "RHICommandList.h"
+
+#if RHI_RAYTRACING
+
+class FRayTracingGeometry;
+class FRHIComputeCommandList;
+enum class EAccelerationStructureBuildMode;
+enum class ERTAccelerationStructureBuildPriority;
+
+class IRayTracingGeometryManager
+{
+public:
+
+	using BuildRequestIndex = int32;
+	using RayTracingGeometryHandle = int32;
+
+	virtual ~IRayTracingGeometryManager() = default;
+
+	RENDERCORE_API virtual BuildRequestIndex RequestBuildAccelerationStructure(FRHICommandList& RHICmdList, FRayTracingGeometry* InGeometry, ERTAccelerationStructureBuildPriority InPriority, EAccelerationStructureBuildMode InBuildMode) = 0;
+
+	BuildRequestIndex RequestBuildAccelerationStructure(FRHICommandList& RHICmdList, FRayTracingGeometry* InGeometry, ERTAccelerationStructureBuildPriority InPriority)
+	{
+		return RequestBuildAccelerationStructure(RHICmdList, InGeometry, InPriority, EAccelerationStructureBuildMode::Build);
+	}
+
+	UE_DEPRECATED(5.4, "RequestBuildAccelerationStructure requires a command list.")
+	BuildRequestIndex RequestBuildAccelerationStructure(FRayTracingGeometry* InGeometry, ERTAccelerationStructureBuildPriority InPriority)
+	{
+		return RequestBuildAccelerationStructure(FRHICommandListImmediate::Get(), InGeometry, InPriority, EAccelerationStructureBuildMode::Build);
+	}
+
+	UE_DEPRECATED(5.4, "RequestBuildAccelerationStructure requires a command list.")
+	BuildRequestIndex RequestBuildAccelerationStructure(FRayTracingGeometry* InGeometry, ERTAccelerationStructureBuildPriority InPriority, EAccelerationStructureBuildMode InBuildMode)
+	{
+		return RequestBuildAccelerationStructure(FRHICommandListImmediate::Get(), InGeometry, InPriority, InBuildMode);
+	}
+
+	RENDERCORE_API virtual void RemoveBuildRequest(BuildRequestIndex InRequestIndex) = 0;
+	RENDERCORE_API virtual void BoostPriority(BuildRequestIndex InRequestIndex, float InBoostValue) = 0;
+	RENDERCORE_API virtual void ForceBuildIfPending(FRHIComputeCommandList& InCmdList, const TArrayView<const FRayTracingGeometry*> InGeometries) = 0;
+	RENDERCORE_API virtual void ProcessBuildRequests(FRHIComputeCommandList& InCmdList, bool bInBuildAll = false) = 0;
+
+	RENDERCORE_API virtual RayTracingGeometryHandle RegisterRayTracingGeometry(FRayTracingGeometry* InGeometry) = 0;
+	RENDERCORE_API virtual void ReleaseRayTracingGeometryHandle(RayTracingGeometryHandle Handle) = 0;
+
+	RENDERCORE_API virtual void Tick(FRHICommandList& RHICmdList, bool bHasRayTracingEnableChanged) = 0;
+};
+
+extern RENDERCORE_API IRayTracingGeometryManager* GRayTracingGeometryManager;
+
+#endif // RHI_RAYTRACING
diff --git a/Engine/Source/Runtime/Renderer/Private/DeferredShadingRenderer.cpp b/Engine/Source/Runtime/Renderer/Private/DeferredShadingRenderer.cpp
--- a/Engine/Source/Runtime/Renderer/Private/DeferredShadingRenderer.cpp
+++ b/Engine/Source/Runtime/Renderer/Private/DeferredShadingRenderer.cpp
@@ -66,7 +66,7 @@
 #include "Substrate/Substrate.h"
 #include "Lumen/Lumen.h"
 #include "Experimental/Containers/SherwoodHashTable.h"
-#include "RayTracingGeometryManager.h"
+#include "Rendering/RayTracingGeometryManager.h"
 #include "InstanceCulling/InstanceCullingManager.h"
 #include "InstanceCulling/InstanceCullingOcclusionQuery.h"
 #include "ProfilingDebugging/CpuProfilerTrace.h"
@@ -2014,7 +2014,7 @@ bool FDeferredShadingSceneRenderer::DispatchRayTracingWorldUpdates(FRDGBuilder&
 	{
 		// This needs to happen even when ray tracing is not enabled
 		// - importers might batch BVH creation requests that need to be resolved in any case
-		GRayTracingGeometryManager.ProcessBuildRequests(GraphBuilder.RHICmdList);
+		GRayTracingGeometryManager->ProcessBuildRequests(GraphBuilder.RHICmdList);
 		// - Nanite ray tracing instances are already pointing at the new BLASes and RayTracingDataOffsets in GPUScene have been updated
 		Nanite::GRayTracingManager.ProcessBuildRequests(GraphBuilder);
 		return false;
@@ -2033,7 +2033,7 @@ bool FDeferredShadingSceneRenderer::DispatchRayTracingWorldUpdates(FRDGBuilder&
 		RayTracingSkinnedGeometryUpdateQueue->Commit(GraphBuilder);
 	}
 
-	GRayTracingGeometryManager.ProcessBuildRequests(GraphBuilder.RHICmdList);
+	GRayTracingGeometryManager->ProcessBuildRequests(GraphBuilder.RHICmdList);
 
 	const int32 ReferenceViewIndex = 0;
 	FViewInfo& ReferenceView = Views[ReferenceViewIndex];
@@ -2042,7 +2042,7 @@ bool FDeferredShadingSceneRenderer::DispatchRayTracingWorldUpdates(FRDGBuilder&
 	if (RayTracingScene.GeometriesToBuild.Num() > 0)
 	{
 		// Force update all the collected geometries (use stack allocator?)
-		GRayTracingGeometryManager.ForceBuildIfPending(GraphBuilder.RHICmdList, RayTracingScene.GeometriesToBuild);
+		GRayTracingGeometryManager->ForceBuildIfPending(GraphBuilder.RHICmdList, RayTracingScene.GeometriesToBuild);
 	}
 
 	FTaskGraphInterface::Get().WaitUntilTaskCompletes(ReferenceView.RayTracingSceneInitTask, ENamedThreads::GetRenderThread_Local());
@@ -2783,7 +2783,7 @@ void FDeferredShadingSceneRenderer::Render(FRDGBuilder& GraphBuilder)
 		// Initialize ray tracing flags, in case they weren't initialized in the CreateSceneRenderers code path
 		InitializeRayTracingFlags_RenderThread();
 
-		GRayTracingGeometryManager.Tick(GraphBuilder.RHICmdList, bHasRayTracingEnableChanged);
+		GRayTracingGeometryManager->Tick(GraphBuilder.RHICmdList, bHasRayTracingEnableChanged);
 
 		if ((GetRayTracingMode() == ERayTracingMode::Dynamic) && bHasRayTracingEnableChanged)
 		{
diff --git a/Engine/Source/Runtime/Renderer/Private/Renderer.cpp b/Engine/Source/Runtime/Renderer/Private/Renderer.cpp
--- a/Engine/Source/Runtime/Renderer/Private/Renderer.cpp
+++ b/Engine/Source/Runtime/Renderer/Private/Renderer.cpp
@@ -45,6 +45,7 @@
 #include "PrimitiveSceneShaderData.h"
 #include "MeshDrawCommandStats.h"
 #include "LocalFogVolumeRendering.h"
+#include "Rendering/RayTracingGeometryManager.h"
 #include "PathTracing.h"
 
 DEFINE_LOG_CATEGORY(LogRenderer);
@@ -71,6 +72,10 @@ void FRendererModule::StartupModule()
 
 	GScreenSpaceDenoiser = IScreenSpaceDenoiser::GetDefaultDenoiser();
 
+#if RHI_RAYTRACING
+	GRayTracingGeometryManager = new FRayTracingGeometryManager();
+#endif
+
 	FRendererOnScreenNotification::Get();
 	FVirtualTextureSystem::Initialize();
 
@@ -109,6 +114,11 @@ void FRendererModule::ShutdownModule()
 	FVirtualTextureSystem::Shutdown();
 	FRendererOnScreenNotification::TearDown();
 
+#if RHI_RAYTRACING
+	delete GRayTracingGeometryManager;
+	GRayTracingGeometryManager = nullptr;
+#endif
+
 	// Free up the memory of the default denoiser. Responsibility of the plugin to free up theirs.
 	delete IScreenSpaceDenoiser::GetDefaultDenoiser();
 
