- Updated SetEvaluateWorldPositionOffset and SetWorldPositionOffsetDisableDistance in UStaticMeshComponent to set the bool regardless of whether or not it's valid yet.
 - Added UpdateInitialEvaluateWorldPositionOffset to UStaticMeshComponent, which will update bInitialEvaluateWorldPositionOffset to the current value of bEvaluateWorldPositionOffset.  This is useful, as it is sometimes used to restore the original value.
 - Updated ABuildingWall to set evalute WPO to false for it's door components by default.
 - Updated how ABuildingSMActor sets evaluate WPO on it's static mesh components to not exclude some of it's meshes by default.  This fixes many issues where some meshes always had evaluate WPO active.  Child classes can override GetExcludedWPOStaticMeshes to exclude specific mesh components from this.

diff --git a/Engine/Source/Runtime/Engine/Classes/Components/StaticMeshComponent.h b/Engine/Source/Runtime/Engine/Classes/Components/StaticMeshComponent.h
--- a/Engine/Source/Runtime/Engine/Classes/Components/StaticMeshComponent.h
+++ b/Engine/Source/Runtime/Engine/Classes/Components/StaticMeshComponent.h
@@ -400,10 +400,16 @@ public:
 	UFUNCTION(BlueprintCallable, Category = "Rendering|LOD")
 	ENGINE_API void SetWorldPositionOffsetDisableDistance(int32 NewValue);
 
-	/** Get the initial value of bEvaluateWorldPositionOffset. This is the value when BeginPlay() was last called. */
+	/** Get the initial value of bEvaluateWorldPositionOffset. This is the value when BeginPlay() was last called, or if UpdateInitialEvaluateWorldPositionOffset is called. */
 	UFUNCTION(BlueprintCallable, Category = "Rendering|LOD")
 	bool GetInitialEvaluateWorldPositionOffset() { return bInitialEvaluateWorldPositionOffset; }
 
+	/** This manually updates the initial value of bEvaluateWorldPositionOffset to be the current value.
+	 *	This is useful if the default value of bEvaluateWorldPositionOffset is changed after constructing
+	 *	the component. */
+	UFUNCTION(BlueprintCallable, Category = "Rendering|LOD")
+	void UpdateInitialEvaluateWorldPositionOffset() { bInitialEvaluateWorldPositionOffset = bEvaluateWorldPositionOffset; }
+
 	/** 
 	 * Get Local bounds
 	 */
diff --git a/Engine/Source/Runtime/Engine/Private/Components/StaticMeshComponent.cpp b/Engine/Source/Runtime/Engine/Private/Components/StaticMeshComponent.cpp
--- a/Engine/Source/Runtime/Engine/Private/Components/StaticMeshComponent.cpp
+++ b/Engine/Source/Runtime/Engine/Private/Components/StaticMeshComponent.cpp
@@ -2495,35 +2495,35 @@ void UStaticMeshComponent::SetEvaluateWorldPositionOffsetInRayTracing(bool NewVa
 
 void UStaticMeshComponent::SetEvaluateWorldPositionOffset(bool NewValue)
 {
-	// Skip when this doesn't have a valid static mesh or a valid scene
-	if (!GetStaticMesh() || GetScene() == nullptr || SceneProxy == nullptr)
-	{
-		return;
-	}
-
 	if (bEvaluateWorldPositionOffset != NewValue)
 	{
 		// Update game thread data
 		bEvaluateWorldPositionOffset = NewValue;
-		// Update render thread data
-		SceneProxy->SetEvaluateWorldPositionOffset_GameThread(NewValue);
+
+		// make sure this has a valid static mesh and a valid scene
+		if (GetStaticMesh() && GetScene() && SceneProxy)
+		{
+			// Update render thread data
+			SceneProxy->SetEvaluateWorldPositionOffset_GameThread(NewValue);
+		}
 	}
 }
 
 void UStaticMeshComponent::SetWorldPositionOffsetDisableDistance(int32 NewValue)
 {
-	// Skip when this doesn't have a valid static mesh or a valid scene
-	if (!GetStaticMesh() || GetScene() == nullptr || SceneProxy == nullptr)
-	{
-		return;
-	}
+	
 
 	if (WorldPositionOffsetDisableDistance != NewValue)
 	{
 		// Update game thread data
 		WorldPositionOffsetDisableDistance = NewValue;
-		// Update render thread data
-		SceneProxy->SetWorldPositionOffsetDisableDistance_GameThread(NewValue);
+
+		// make sure this has a valid static mesh and a valid scene
+		if (GetStaticMesh() && GetScene() && SceneProxy)
+		{
+			// Update render thread data
+			SceneProxy->SetWorldPositionOffsetDisableDistance_GameThread(NewValue);
+		}
 	}
 }
 
