[Backout] - CL26308339

Original CL Desc
-----------------------------------------------------------------
[Backout] - CL26290283
Original CL Desc
-----------------------------------------------------------------
[Backout] - CL26287803 for CIS error.
#rnx
Original CL Desc
-----------------------------------------------------------------
World Partition
* Add new unsaved actors into the DirtyActors set. Useful to see unsaved LocationVolume on the WP Editor.
* Make sure we clean DirtyActors when CB content is removed.

diff --git a/Engine/Source/Editor/WorldPartitionEditor/Private/WorldPartition/SWorldPartitionEditorGrid2D.cpp b/Engine/Source/Editor/WorldPartitionEditor/Private/WorldPartition/SWorldPartitionEditorGrid2D.cpp
--- a/Engine/Source/Editor/WorldPartitionEditor/Private/WorldPartition/SWorldPartitionEditorGrid2D.cpp
+++ b/Engine/Source/Editor/WorldPartitionEditor/Private/WorldPartition/SWorldPartitionEditorGrid2D.cpp
@@ -1265,20 +1265,17 @@ void SWorldPartitionEditorGrid2D::Tick(const FGeometry& AllottedGeometry, const
 
 	// Also include transient actor loader adapters that might have been spawned by blutilities, etc. Since these actors can't be saved because they are transient,
 	// they will never get an actor descriptor so they will never appear in the world partition editor. Also include unsaved, newly created actors for convenience.
-	for (auto& [DirtyActorGuid, DirtyActor] : GetWorldPartition()->GetDirtyActors())
+	for (auto& [Reference, Actor] : GetWorldPartition()->GetDirtyActors())
 	{
-		if (DirtyActor.ActorPtr.IsValid())
+		if (!GetWorldPartition()->GetActorDesc(Actor->GetActorGuid()) && Actor->Implements<UWorldPartitionActorLoaderInterface>())
 		{
-			if (!GetWorldPartition()->GetActorDesc(DirtyActor.ActorPtr->GetActorGuid()) && DirtyActor.ActorPtr->Implements<UWorldPartitionActorLoaderInterface>())
+			if (IWorldPartitionActorLoaderInterface::ILoaderAdapter* LoaderAdapter = Cast<IWorldPartitionActorLoaderInterface>(Actor)->GetLoaderAdapter())
 			{
-				if (IWorldPartitionActorLoaderInterface::ILoaderAdapter* LoaderAdapter = Cast<IWorldPartitionActorLoaderInterface>(DirtyActor.ActorPtr)->GetLoaderAdapter())
-				{
-					ShownLoaderInterfaces.Add(DirtyActor.ActorPtr.Get());
-				}
+				ShownLoaderInterfaces.Add(Actor);
 			}
-
-			DirtyActorGuids.Add(DirtyActorGuid);
 		}
+
+		DirtyActorGuids.Add(Actor->GetActorGuid());
 	}
 
 	FLoaderInterfaceSet LastHoveredLoaderInterfaces = MoveTemp(HoveredLoaderInterfaces);
diff --git a/Engine/Source/Runtime/Engine/Private/WorldPartition/WorldPartition.cpp b/Engine/Source/Runtime/Engine/Private/WorldPartition/WorldPartition.cpp
--- a/Engine/Source/Runtime/Engine/Private/WorldPartition/WorldPartition.cpp
+++ b/Engine/Source/Runtime/Engine/Private/WorldPartition/WorldPartition.cpp
@@ -45,8 +45,6 @@
 #include "WorldPartition/LoaderAdapter/LoaderAdapterPinnedActors.h"
 #include "WorldPartition/HLOD/HLODLayer.h"
 #include "WorldPartition/Cook/WorldPartitionCookPackageContextInterface.h"
-#include "WorldPartition/ContentBundle/ContentBundleEditorSubsystemInterface.h"
-#include "WorldPartition/ContentBundle/ContentBundleEditor.h"
 #include "WorldPartition/ErrorHandling/WorldPartitionStreamingGenerationLogErrorHandler.h"
 #include "WorldPartition/ErrorHandling/WorldPartitionStreamingGenerationMapCheckErrorHandler.h"
 #include "Modules/ModuleManager.h"
@@ -354,9 +352,15 @@ void UWorldPartition::OnObjectsReplaced(const TMap<UObject*, UObject*>& OldToNew
 	{
 		if (AActor* OldActor = Cast<AActor>(OldObject))
 		{
-			if (FDirtyActor* DirtyActor = DirtyActors.Find(OldActor->GetActorGuid()))
+			if (AActor* NewActor = Cast<AActor>(NewObject))
 			{
-				DirtyActor->ActorPtr = CastChecked<AActor>(NewObject);
+				for (auto& [ActorReference, Actor] : DirtyActors)
+				{
+					if (Actor == OldActor)
+					{
+						Actor = NewActor;
+					}
+				}
 			}
 		}
 	}
@@ -382,23 +386,11 @@ void UWorldPartition::OnPackageDirtyStateChanged(UPackage* Package)
 		{
 			if (Package->IsDirty())
 			{
-				DirtyActors.Add(Actor->GetActorGuid(), FDirtyActor(ActorHandle.ToReference(), Actor));
-			}
-		}
-		else if (DirtyActors.Contains(Actor->GetActorGuid()))
-		{
-			if (!Package->IsDirty())
-			{
-				// Remove newly created, unsaved actor that dissapeared after an undo
-				DirtyActors.Remove(Actor->GetActorGuid());
+				DirtyActors.Add(ActorHandle.ToReference(), Actor);
 			}
 		}
-		else
-		{
-			// Add newly created, unsaved actor after editor placement or a redo
-			DirtyActors.Add(Actor->GetActorGuid(), FDirtyActor(Actor));
-		}
 	}
+
 }
 
 // Returns whether the memory package is part of the known/valid package names
@@ -929,7 +921,6 @@ void UWorldPartition::RegisterDelegates()
 			GEditor->OnPostBugItGoCalled().AddUObject(this, &UWorldPartition::OnPostBugItGoCalled);
 			GEditor->OnEditorClose().AddUObject(this, &UWorldPartition::SavePerUserSettings);
 			FWorldDelegates::OnPostWorldRename.AddUObject(this, &UWorldPartition::OnWorldRenamed);
-			IContentBundleEditorSubsystemInterface::Get()->OnContentBundleRemovedContent().AddUObject(this, &UWorldPartition::OnContentBundleRemovedContent);
 		}
 
 		if (!IsRunningCommandlet())
@@ -961,11 +952,6 @@ void UWorldPartition::UnregisterDelegates()
 	{
 		if (IsMainWorldPartition())
 		{
-			if (IContentBundleEditorSubsystemInterface* ContentBundleEditorSubsystem = IContentBundleEditorSubsystemInterface::Get())
-			{
-				ContentBundleEditorSubsystem->OnContentBundleRemovedContent().RemoveAll(this);
-			}
-
 			FWorldDelegates::OnPostWorldRename.RemoveAll(this);
 			FEditorDelegates::PreBeginPIE.RemoveAll(this);
 			FEditorDelegates::PrePIEEnded.RemoveAll(this);
@@ -1269,16 +1255,7 @@ void UWorldPartition::OnActorDescAdded(FWorldPartitionActorDesc* NewActorDesc)
 
 	if (AActor* NewActor = NewActorDesc->GetActor())
 	{
-		if (FDirtyActor* ExistingDirtyActor = DirtyActors.Find(NewActorDesc->GetGuid()))
-		{
-			check(!ExistingDirtyActor->WorldPartitionRef.IsSet());
-			check(!ExistingDirtyActor->ActorPtr.IsValid() || ExistingDirtyActor->ActorPtr == NewActor);			
-			ExistingDirtyActor->WorldPartitionRef = FWorldPartitionReference(NewActorDesc->GetContainer(), NewActorDesc->GetGuid());
-		}
-		else
-		{
-			DirtyActors.Add(NewActorDesc->GetGuid(), FDirtyActor(FWorldPartitionReference(NewActorDesc->GetContainer(), NewActorDesc->GetGuid()), NewActor));
-		}
+		DirtyActors.Add(FWorldPartitionReference(NewActorDesc->GetContainer(), NewActorDesc->GetGuid()), NewActor);
 	}
 
 	if (ForceLoadedActors)
@@ -1469,33 +1446,6 @@ void UWorldPartition::UnhashActorDesc(FWorldPartitionActorDesc* ActorDesc)
 	EditorHash->UnhashActor(ActorHandle);
 }
 
-void UWorldPartition::OnContentBundleRemovedContent(const FContentBundleEditor* ContentBundle)
-{
-	check(ContentBundle);
-	
-	const TWeakObjectPtr<UActorDescContainer>& ContentBundleActorDescContainer = ContentBundle->GetActorDescContainer();
-	if (!ContentBundleActorDescContainer.IsValid())
-	{
-		return;
-	}
-
-	if (Contains(ContentBundleActorDescContainer.Get()->GetContainerPackage()))
-	{
-		// This is handling a new actor (unsaved) from Content Bundle.
-		for (TMap<FGuid, FDirtyActor>::TIterator DirtyActorIt(DirtyActors); DirtyActorIt; ++DirtyActorIt)
-		{
-			const FDirtyActor& DirtyActor = DirtyActorIt.Value();
-			if (!DirtyActor.WorldPartitionRef.IsSet())
-			{
-				if (DirtyActor.ActorPtr.IsValid() && DirtyActor.ActorPtr->GetContentBundleGuid() == ContentBundleActorDescContainer.Get()->GetContentBundleGuid())
-				{
-					DirtyActorIt.RemoveCurrent();
-				}
-			}
-		}
-	}
-}
-
 bool UWorldPartition::IsStreamingEnabledInEditor() const
 {
 	return bOverrideEnableStreamingInEditor.IsSet() ? *bOverrideEnableStreamingInEditor : IsStreamingEnabled();
@@ -1615,9 +1565,9 @@ void UWorldPartition::AddReferencedObjects(UObject* InThis, FReferenceCollector&
 	if (IsGarbageCollecting())
 	{
 		Collector.AllowEliminatingReferences(false);
-		for (auto& [DirtyActorGuid, DirtyActor] : This->DirtyActors)
+		for (auto& [ActorReference, Actor] : This->DirtyActors)
 		{
-			Collector.AddReferencedObject(DirtyActor.ActorPtr);
+			Collector.AddReferencedObject(Actor);
 		}
 		Collector.AllowEliminatingReferences(true);
 	}
@@ -1639,55 +1589,23 @@ void UWorldPartition::Tick(float DeltaSeconds)
 		EditorHash->Tick(DeltaSeconds);
 	}
 
-	for (TMap<FGuid, FDirtyActor>::TIterator DirtyActorIt(DirtyActors); DirtyActorIt; ++DirtyActorIt)
+		for (TMap<FWorldPartitionReference, TObjectPtr<AActor>>::TIterator DirtyActorIt(DirtyActors); DirtyActorIt; ++DirtyActorIt)
 	{
-		FDirtyActor& DirtyActor = DirtyActorIt.Value();
-		if (DirtyActor.WorldPartitionRef.IsSet())
+		if (!DirtyActorIt.Key().IsValid())
 		{
-			const FWorldPartitionReference& Ref = DirtyActor.WorldPartitionRef.GetValue();
-
-			if (!Ref.IsValid())
-			{
-				DirtyActorIt.RemoveCurrent();
-			}
-			else if (DirtyActor.ActorPtr.IsValid() && !DirtyActor.ActorPtr->GetPackage()->IsDirty())
-			{
-				// Transfer ownership of the last ref if actor can be pinned
-				if (Ref->GetHardRefCount() <= 1 && PinnedActors && FLoaderAdapterPinnedActors::SupportsPinning(Ref.Get()))
-				{
-					PinnedActors->AddActors({ Ref.ToHandle() });
-				}
-
-				DirtyActorIt.RemoveCurrent();
-			}
+			DirtyActorIt.RemoveCurrent();
 		}
-		else
+		else if (!DirtyActorIt.Value()->GetPackage()->IsDirty())
 		{
-			// This is handling a new actor (unsaved).
-			if (DirtyActor.ActorPtr.IsValid())
+			// Transfer ownership of the last ref if actor can be pinned
+			if (DirtyActorIt.Key()->GetHardRefCount() <= 1 && PinnedActors && FLoaderAdapterPinnedActors::SupportsPinning(DirtyActorIt.Key().Get()))
 			{
-				if (!DirtyActor.ActorPtr->GetPackage()->IsDirty())
-				{
-					DirtyActorIt.RemoveCurrent();
-				}
-			}
-			else
+				PinnedActors->AddActors({ DirtyActorIt.Key().ToHandle() });
+				DirtyActorIt.RemoveCurrent();
+			} // Clean up if a loader took a reference on the actor (ILoaderAdapter::RefreshLoadedState was called since)
+			else if (DirtyActorIt.Key()->GetHardRefCount() > 1)
 			{
-				if (DirtyActor.ActorPtr.IsValid(true))
-				{
-					// If the actor is pending kill, check if there's another version of it in the same package and replace the pointer with the new version
-					// (since UEditorEngine::ReplaceActors don't always broadcast FCoreUObjectDelegates::OnObjectsReplaced).
-					if (AActor* NewActor = AActor::FindActorInPackage(DirtyActor.ActorPtr.Get(true)->GetPackage(), false); NewActor && NewActor != DirtyActor.ActorPtr.Get(true))
-					{
-						DirtyActor.ActorPtr = NewActor;
-					}
-				}
-				else
-				{
-					// In this case, we know that the dirty actor is not in the transaction buffer anymore and this is fine removing it (actor was added, then
-					// deleted and gargage collection happened).
-					DirtyActorIt.RemoveCurrent();
-				}
+				DirtyActorIt.RemoveCurrent();
 			}
 		}
 	}
@@ -2146,15 +2064,11 @@ bool UWorldPartition::UnregisterActorDescContainer(UActorDescContainer* InActorD
 			{
 				ActorGuids.Add(It->GetGuid());
 
-				for (TMap<FGuid, FDirtyActor>::TIterator DirtyActorIt(DirtyActors); DirtyActorIt; ++DirtyActorIt)
+				for (TMap<FWorldPartitionReference, TObjectPtr<AActor>>::TIterator DirtyActorIt(DirtyActors); DirtyActorIt; ++DirtyActorIt)
 				{
-					const FDirtyActor& DirtyActor = DirtyActorIt.Value();
-					if (DirtyActor.WorldPartitionRef.IsSet())
+					if (DirtyActorIt.Key() == ActorHandle)
 					{
-						if (DirtyActor.WorldPartitionRef.GetValue() == ActorHandle)
-						{
-							DirtyActorIt.RemoveCurrent();
-						}
+						DirtyActorIt.RemoveCurrent();
 					}
 				}
 			}
diff --git a/Engine/Source/Runtime/Engine/Public/WorldPartition/WorldPartition.h b/Engine/Source/Runtime/Engine/Public/WorldPartition/WorldPartition.h
--- a/Engine/Source/Runtime/Engine/Public/WorldPartition/WorldPartition.h
+++ b/Engine/Source/Runtime/Engine/Public/WorldPartition/WorldPartition.h
@@ -105,40 +105,6 @@ public:
 	virtual bool Delete(UPackage* Package) const =0;
 	virtual bool Save(UPackage* Package) const =0;
 };
-
-struct ENGINE_API FDirtyActor
-{
-	FDirtyActor()
-		: ActorPtr(nullptr)
-	{}
-
-	FDirtyActor(AActor* InActor)
-		: ActorPtr(InActor)
-	{}
-
-	FDirtyActor(const FWorldPartitionReference& InWorldPartitionRef, AActor* InActor)
-		: WorldPartitionRef(InWorldPartitionRef)
-		, ActorPtr(InActor)
-	{}
-
-	TOptional<FWorldPartitionReference> WorldPartitionRef;
-	TWeakObjectPtr<AActor> ActorPtr; // TWeakObjectPtr is for undo support.
-
-	bool operator==(const FDirtyActor& InDirtyActor) const
-	{
-		return WorldPartitionRef == InDirtyActor.WorldPartitionRef && ActorPtr == InDirtyActor.ActorPtr;
-	}
-
-	friend uint32 GetTypeHash(const FDirtyActor& InDirtyActor)
-	{
-		uint32 Hash = GetTypeHash(InDirtyActor.ActorPtr);
-		if (InDirtyActor.WorldPartitionRef.IsSet())
-		{
-			Hash = HashCombine(Hash, GetTypeHash(InDirtyActor.WorldPartitionRef.GetValue()));
-		}
-		return Hash;
-	}
-};
 #endif
 
 UCLASS(AutoExpandCategories=(WorldPartition), MinimalAPI)
@@ -393,7 +359,7 @@ public:
 
 	bool IsHLODsInEditorAllowed() const { return bAllowShowingHLODsInEditor; }
 
-	const TMap<FGuid, FDirtyActor>& GetDirtyActors() const { return DirtyActors; }
+	const TMap<FWorldPartitionReference, AActor*>& GetDirtyActors() const { return ObjectPtrDecay(DirtyActors); }
 #endif
 
 public:
@@ -521,7 +487,7 @@ private:
 
 	TArray<FWorldPartitionReference> LoadedSubobjects;
 
-	TMap<FGuid, FDirtyActor> DirtyActors;
+	TMap<FWorldPartitionReference, TObjectPtr<AActor>> DirtyActors;
 
 	TSet<FString> GeneratedStreamingPackageNames;
 
