Fix  nullptr deref when changing virtuatl texture scalability settings.

diff --git a/Engine/Source/Runtime/Engine/Private/VT/VirtualTextureScalability.cpp b/Engine/Source/Runtime/Engine/Private/VT/VirtualTextureScalability.cpp
--- a/Engine/Source/Runtime/Engine/Private/VT/VirtualTextureScalability.cpp
+++ b/Engine/Source/Runtime/Engine/Private/VT/VirtualTextureScalability.cpp
@@ -200,16 +200,16 @@ namespace VirtualTextureScalability
 		if (bUpdate)
 		{
 			// Temporarily release runtime virtual textures
-			for (TObjectIterator<URuntimeVirtualTexture> It; It; ++It)
+			for (TObjectIterator<URuntimeVirtualTexture> It(RF_ClassDefaultObject, false, EInternalObjectFlags::Garbage); It; ++It)
 			{
 				It->Release();
 			}
 
 			// Release streaming virtual textures
 			TArray<UTexture2D*> ReleasedVirtualTextures;
-			for (TObjectIterator<UTexture2D> It; It; ++It)
+			for (TObjectIterator<UTexture2D> It(RF_ClassDefaultObject, false, EInternalObjectFlags::Garbage); It; ++It)
 			{
-				if (It->IsCurrentlyVirtualTextured())
+				if (It->IsCurrentlyVirtualTextured() && It->GetResource() != nullptr)
 				{
 					ReleasedVirtualTextures.Add(*It);
 					BeginReleaseResource(It->GetResource());
@@ -230,7 +230,7 @@ namespace VirtualTextureScalability
 			}
 
 			// Reinit runtime virtual textures
-			for (TObjectIterator<URuntimeVirtualTextureComponent> It; It; ++It)
+			for (TObjectIterator<URuntimeVirtualTextureComponent> It(RF_ClassDefaultObject, false, EInternalObjectFlags::Garbage); It; ++It)
 			{
 				It->MarkRenderStateDirty();
 			}
