Level Instance:
  - Exposed the runtime behavior so we can now place level instances referencing partitioned worlds that will be streamed at runtime instead of embedded in the container world.
  - Don't allow edit in place for partitioned level instance worlds that have streaming disabled, as we load the entire level when doing so and could result in out of memory situations.
  - Added a message box when trying to edit a level instance that can't be edited in place with the reason.
  - Can be reenabled by setting LevelInstanceEditorSettings.bEnableEditInPlaceStreamingEnabled to true.

diff --git a/Engine/Source/Editor/LevelInstanceEditor/Private/LevelInstanceEditorModule.cpp b/Engine/Source/Editor/LevelInstanceEditor/Private/LevelInstanceEditorModule.cpp
--- a/Engine/Source/Editor/LevelInstanceEditor/Private/LevelInstanceEditorModule.cpp
+++ b/Engine/Source/Editor/LevelInstanceEditor/Private/LevelInstanceEditorModule.cpp
@@ -941,8 +941,11 @@ void FLevelInstanceEditorModule::ExtendContextMenu()
 			}
 		}), FToolMenuInsert(NAME_None, EToolMenuInsertType::Default));
 	}
-	
 }
 
-#undef LOCTEXT_NAMESPACE
+bool FLevelInstanceEditorModule::IsEditInPlaceStreamingEnabled() const
+{
+	return GetDefault<ULevelInstanceEditorSettings>()->bIsEditInPlaceStreamingEnabled;
+}
 
+#undef LOCTEXT_NAMESPACE
\ No newline at end of file
diff --git a/Engine/Source/Editor/LevelInstanceEditor/Private/LevelInstanceEditorModule.h b/Engine/Source/Editor/LevelInstanceEditor/Private/LevelInstanceEditorModule.h
--- a/Engine/Source/Editor/LevelInstanceEditor/Private/LevelInstanceEditorModule.h
+++ b/Engine/Source/Editor/LevelInstanceEditor/Private/LevelInstanceEditorModule.h
@@ -38,6 +38,8 @@ public:
 	DECLARE_DERIVED_EVENT(FLevelInstanceEditorModule, ILevelInstanceEditorModule::FTryExitEditorModeEvent, FTryExitEditorModeEvent);
 	virtual FTryExitEditorModeEvent& OnTryExitEditorMode() override { return TryExitEditorModeEvent; }
 
+	virtual bool IsEditInPlaceStreamingEnabled() const override;
+
 private:
 	void OnEditorModeIDChanged(const FEditorModeID& InModeID, bool bIsEnteringMode);
 	void OnLevelActorDeleted(AActor* Actor);
diff --git a/Engine/Source/Editor/LevelInstanceEditor/Private/LevelInstanceEditorSettings.cpp b/Engine/Source/Editor/LevelInstanceEditor/Private/LevelInstanceEditorSettings.cpp
--- a/Engine/Source/Editor/LevelInstanceEditor/Private/LevelInstanceEditorSettings.cpp
+++ b/Engine/Source/Editor/LevelInstanceEditor/Private/LevelInstanceEditorSettings.cpp
@@ -7,6 +7,7 @@ ULevelInstanceEditorSettings::ULevelInstanceEditorSettings()
 {
 	LevelInstanceClassName = ALevelInstance::StaticClass()->GetPathName();
 	bEnableStreaming = false;
+	bIsEditInPlaceStreamingEnabled = false;
 }
 
 ULevelInstanceEditorPerProjectUserSettings::ULevelInstanceEditorPerProjectUserSettings()
diff --git a/Engine/Source/Editor/LevelInstanceEditor/Public/LevelInstanceEditorSettings.h b/Engine/Source/Editor/LevelInstanceEditor/Public/LevelInstanceEditorSettings.h
--- a/Engine/Source/Editor/LevelInstanceEditor/Public/LevelInstanceEditorSettings.h
+++ b/Engine/Source/Editor/LevelInstanceEditor/Public/LevelInstanceEditorSettings.h
@@ -27,6 +27,9 @@ public:
 
 	UPROPERTY(config, EditAnywhere, Category="World Partition", meta = (ToolTip="Create World Partition Level Instances with Streaming Enabled/Disabled by default"))
 	bool bEnableStreaming;
+
+	UPROPERTY(config, EditAnywhere, Category="World Partition", meta = (ToolTip="Allow Editing Level Instances with Streaming Enabled"))
+	bool bIsEditInPlaceStreamingEnabled;
 };
 
 UCLASS(config = EditorPerProjectUserSettings, meta = (DisplayName = "Level Instance"))
diff --git a/Engine/Source/Runtime/Engine/Private/LevelInstance/LevelInstanceActor.cpp b/Engine/Source/Runtime/Engine/Private/LevelInstance/LevelInstanceActor.cpp
--- a/Engine/Source/Runtime/Engine/Private/LevelInstance/LevelInstanceActor.cpp
+++ b/Engine/Source/Runtime/Engine/Private/LevelInstance/LevelInstanceActor.cpp
@@ -13,6 +13,7 @@
 #include "UObject/ObjectSaveContext.h"
 #include "WorldPartition/LevelInstance/LevelInstanceActorDesc.h"
 #include "LevelInstance/LevelInstanceEditorPivotActor.h"
+#include "Misc/MessageDialog.h"
 #endif
 
 #define LOCTEXT_NAMESPACE "LevelInstanceActor"
@@ -309,7 +310,13 @@ bool ALevelInstance::GetSoftReferencedContentObjects(TArray<FSoftObjectPath>& So
 
 bool ALevelInstance::OpenAssetEditor()
 {
-	return CanEnterEdit() ? EnterEdit() : false;
+	FText Reason;
+	if (!CanEnterEdit(&Reason))
+	{
+		FMessageDialog::Open(EAppMsgType::Ok, Reason);
+		return false;
+	}
+	return EnterEdit();
 }
 
 bool ALevelInstance::EditorCanAttachFrom(const AActor* InChild, FText& OutReason) const
diff --git a/Engine/Source/Runtime/Engine/Private/LevelInstance/LevelInstanceSubsystem.cpp b/Engine/Source/Runtime/Engine/Private/LevelInstance/LevelInstanceSubsystem.cpp
--- a/Engine/Source/Runtime/Engine/Private/LevelInstance/LevelInstanceSubsystem.cpp
+++ b/Engine/Source/Runtime/Engine/Private/LevelInstance/LevelInstanceSubsystem.cpp
@@ -1903,6 +1903,19 @@ bool ULevelInstanceSubsystem::CanEditLevelInstance(const ILevelInstanceInterface
 			}
 			return false;
 		}
+
+		if (ULevel::GetIsLevelPartitionedFromPackage(*LevelInstance->GetWorldAssetPackage()) && !ULevel::GetIsStreamingDisabledFromPackage(*LevelInstance->GetWorldAssetPackage()))
+		{
+			ILevelInstanceEditorModule& EditorModule = FModuleManager::GetModuleChecked<ILevelInstanceEditorModule>("LevelInstanceEditor");
+			if (!EditorModule.IsEditInPlaceStreamingEnabled())
+			{
+				if (OutReason)
+				{
+					*OutReason = FText::Format(LOCTEXT("CanEditLevelInstanceWithStreamingEnabled", "Level Instance can't be edited in place because it has streaming enabled ({0})"), FText::FromString(LevelInstance->GetWorldAssetPackage()));
+				}
+				return false;
+			}
+		}
 	}
 
 	return true;
diff --git a/Engine/Source/Runtime/Engine/Private/WorldPartition/LevelInstance/LevelInstanceActorDesc.cpp b/Engine/Source/Runtime/Engine/Private/WorldPartition/LevelInstance/LevelInstanceActorDesc.cpp
--- a/Engine/Source/Runtime/Engine/Private/WorldPartition/LevelInstance/LevelInstanceActorDesc.cpp
+++ b/Engine/Source/Runtime/Engine/Private/WorldPartition/LevelInstance/LevelInstanceActorDesc.cpp
@@ -209,13 +209,6 @@ void FLevelInstanceActorDesc::CheckForErrors(IStreamingGenerationErrorHandler* E
 			ErrorHandler->OnLevelInstanceInvalidWorldAsset(this, GetContainerPackage(), IStreamingGenerationErrorHandler::ELevelInstanceInvalidReason::WorldAssetNotUsingExternalActors);
 		}
 	}
-	else if (ULevel::GetIsLevelPartitionedFromPackage(GetContainerPackage()))
-	{
-		if (DesiredRuntimeBehavior != ELevelInstanceRuntimeBehavior::Partitioned)
-		{
-			ErrorHandler->OnLevelInstanceInvalidWorldAsset(this, GetContainerPackage(), IStreamingGenerationErrorHandler::ELevelInstanceInvalidReason::WorldAssetImcompatiblePartitioned);
-		}
-	}
 }
 
 void FLevelInstanceActorDesc::TransferFrom(const FWorldPartitionActorDesc* From)
diff --git a/Engine/Source/Runtime/Engine/Public/LevelInstance/ILevelInstanceEditorModule.h b/Engine/Source/Runtime/Engine/Public/LevelInstance/ILevelInstanceEditorModule.h
--- a/Engine/Source/Runtime/Engine/Public/LevelInstance/ILevelInstanceEditorModule.h
+++ b/Engine/Source/Runtime/Engine/Public/LevelInstance/ILevelInstanceEditorModule.h
@@ -23,4 +23,6 @@ public:
 
 	DECLARE_EVENT(ILevelInstanceEditorModule, FTryExitEditorModeEvent);
 	virtual FTryExitEditorModeEvent& OnTryExitEditorMode() = 0;
+
+	virtual bool IsEditInPlaceStreamingEnabled() const = 0;
 };
\ No newline at end of file
diff --git a/Engine/Source/Runtime/Engine/Public/LevelInstance/LevelInstanceActor.h b/Engine/Source/Runtime/Engine/Public/LevelInstance/LevelInstanceActor.h
--- a/Engine/Source/Runtime/Engine/Public/LevelInstance/LevelInstanceActor.h
+++ b/Engine/Source/Runtime/Engine/Public/LevelInstance/LevelInstanceActor.h
@@ -24,7 +24,7 @@ public:
 protected:
 #if WITH_EDITORONLY_DATA
 	/** Level LevelInstance */
-	UPROPERTY(EditAnywhere, Category = Level, meta = (NoCreate, DisplayName="Level"))
+	UPROPERTY(EditAnywhere, Category = Level, Meta = (NoCreate, DisplayName="Level"))
 	TSoftObjectPtr<UWorld> WorldAsset;
 #endif
 	UPROPERTY(VisibleAnywhere, Category = Default)
@@ -41,7 +41,7 @@ protected:
 
 public:
 #if WITH_EDITORONLY_DATA
-	UPROPERTY()
+	UPROPERTY(EditAnywhere, Category = Level, AdvancedDisplay, Meta = (DisplayName="Level Behavior"))
 	ELevelInstanceRuntimeBehavior DesiredRuntimeBehavior;
 #endif
 
diff --git a/Engine/Source/Runtime/Engine/Public/LevelInstance/LevelInstanceTypes.h b/Engine/Source/Runtime/Engine/Public/LevelInstance/LevelInstanceTypes.h
--- a/Engine/Source/Runtime/Engine/Public/LevelInstance/LevelInstanceTypes.h
+++ b/Engine/Source/Runtime/Engine/Public/LevelInstance/LevelInstanceTypes.h
@@ -58,10 +58,10 @@ enum class ELevelInstanceRuntimeBehavior : uint8
 	None UMETA(Hidden),
 	// Deprecated exists only to avoid breaking Actor Desc serialization
 	Embedded_Deprecated UMETA(Hidden),
-	// Default behavior is to move Level Instance level actors to the main world partition using World Partition clustering rules
-	Partitioned,
-	// Behavior only supported through Conversion Commandlet or on non OFPA Level Instances
-	LevelStreaming UMETA(Hidden)
+	// Move level instance actors to the main world partition
+	Partitioned UMETA(DisplayName = "Embedded"),
+	// Use level streaming to load level instance actors
+	LevelStreaming UMETA(DisplayName = "Standalone"),
 };
 
 UENUM()
