[BodyInstance] Partial rollback of 29331694 (option to skip actors and shapes creation for all instances if the body setup doesn't have any geometry) and replaced existing log by a warning so problematic assets could be detected and fixed.

diff --git a/Engine/Source/Runtime/Engine/Private/PhysicsEngine/BodyInstance.cpp b/Engine/Source/Runtime/Engine/Private/PhysicsEngine/BodyInstance.cpp
--- a/Engine/Source/Runtime/Engine/Private/PhysicsEngine/BodyInstance.cpp
+++ b/Engine/Source/Runtime/Engine/Private/PhysicsEngine/BodyInstance.cpp
@@ -90,12 +90,6 @@ FAutoConsoleVariableRef CVarbEnableOverrideSolverDeltaTime(
 	bEnableOverrideSolverDeltaTime,
 	TEXT("If true, setting for override solver delta time can be used.  False will disable this feature."));
 
-bool bSkipShapeCreationForEmptyBodySetup = false;
-FAutoConsoleVariableRef CVarSkipShapeCreationForEmptyBodySetup(
-	TEXT("p.SkipShapeCreationForEmptyBodySetup"),
-	bSkipShapeCreationForEmptyBodySetup,
-	TEXT("If true, CreateShapesAndActors will not try to create actors and shapes for all instances if the body setup doesn't have any geometry."));
-
 using namespace PhysicsInterfaceTypes;
 
 bool IsRigidBodyKinematic_AssumesLocked(const FPhysicsActorHandle& InActorRef)
@@ -1333,24 +1327,6 @@ bool FInitBodiesHelperBase::CreateShapesAndActors()
 	// Ensure we have the AggGeom inside the body setup so we can calculate the number of shapes
 	BodySetup->CreatePhysicsMeshes();
 
-	if (bSkipShapeCreationForEmptyBodySetup)
-	{
-		if (BodySetup->TriMeshGeometries.IsEmpty() && BodySetup->AggGeom.GetElementCount() == 0)
-		{
-#if WITH_EDITOR
-			// In the editor we may have ended up here because of world trace ignoring our EnableCollision.
-			// Since we can't get at the data in that function we check for it here
-			if (PrimitiveComp && PrimitiveComp->IsCollisionEnabled())
-#endif
-			{
-				UE_LOG(LogPhysics, Log, TEXT("Init of %d instances of Primitive Component %s failed. Does it have collision data available?"),
-					NumBodies, PrimitiveComp ? *PrimitiveComp->GetReadableName() : TEXT("null"));
-			}
-
-			return false;
-		}
-	}
-
 	for (int32 BodyIdx = NumBodies - 1; BodyIdx >= 0; BodyIdx--)   // iterate in reverse since list might shrink
 	{
 		FBodyInstance* Instance = Bodies[BodyIdx];
@@ -1438,7 +1414,7 @@ bool FInitBodiesHelperBase::CreateShapesAndActors()
 			if(!PrimitiveComp || PrimitiveComp->IsCollisionEnabled())
 #endif
 			{
-				UE_LOG(LogPhysics, Log, TEXT("Init Instance %d of Primitive Component %s failed. Does it have collision data available?"), BodyIdx, *PrimitiveComp->GetReadableName());
+				UE_LOG(LogPhysics, Warning, TEXT("Init Instance %d of Primitive Component %s failed. Does it have collision data available?"), BodyIdx, *PrimitiveComp->GetReadableName());
 			}
 
 			FPhysicsInterface::ReleaseActor(Instance->ActorHandle, PhysScene);
