Added option for request forward behaviour in Inertialization node and removed request forwarding from DeadBlending node.

diff --git a/Engine/Source/Runtime/Engine/Classes/Animation/AnimNode_Inertialization.h b/Engine/Source/Runtime/Engine/Classes/Animation/AnimNode_Inertialization.h
--- a/Engine/Source/Runtime/Engine/Classes/Animation/AnimNode_Inertialization.h
+++ b/Engine/Source/Runtime/Engine/Classes/Animation/AnimNode_Inertialization.h
@@ -421,6 +421,15 @@ private:
 	UPROPERTY(EditAnywhere, Category = Blending)
 	bool bResetOnBecomingRelevant = false;
 
+	/**
+	* When enabled this option will forward inertialization requests through any downstream UseCachedPose nodes which 
+	* have had their update skipped (e.g. because they have already been updated in another location). This can be
+	* useful in the case where the same cached pose is used in multiple places, and having an inertialization request 
+	* that goes with it caught in only one of those places would create popping.
+	*/
+	UPROPERTY(EditAnywhere, Category = Requests)
+	bool bForwardRequestsThroughSkippedCachedPoseNodes = true;
+
 public: // FAnimNode_Inertialization
 
 	ENGINE_API FAnimNode_Inertialization();
diff --git a/Engine/Source/Runtime/Engine/Private/Animation/AnimNode_DeadBlending.cpp b/Engine/Source/Runtime/Engine/Private/Animation/AnimNode_DeadBlending.cpp
--- a/Engine/Source/Runtime/Engine/Private/Animation/AnimNode_DeadBlending.cpp
+++ b/Engine/Source/Runtime/Engine/Private/Animation/AnimNode_DeadBlending.cpp
@@ -712,41 +712,11 @@ void FAnimNode_DeadBlending::Update_AnyThread(const FAnimationUpdateContext& Con
 	{
 		Deactivate();
 	}
-	UpdateCounter.SynchronizeWith(Context.AnimInstanceProxy->GetUpdateCounter());
-
-	const int32 NodeId = Context.GetCurrentNodeId();
-	const FAnimInstanceProxy& Proxy = *Context.AnimInstanceProxy;
 
-	// Allow nodes further towards the leaves to inertialize using this node
-	UE::Anim::TScopedGraphMessage<UE::Anim::FDeadBlendingRequester> Inertialization(Context, Context, this);
+	UpdateCounter.SynchronizeWith(Context.AnimInstanceProxy->GetUpdateCounter());
 
-	// Handle skipped updates for cached poses by forwarding to inertialization nodes in those residual stacks
-	UE::Anim::TScopedGraphMessage<UE::Anim::FCachedPoseSkippedUpdateHandler> CachedPoseSkippedUpdate(Context, [this, NodeId, &Proxy](TArrayView<const UE::Anim::FMessageStack> InSkippedUpdates)
-	{
-		// If we have a pending request forward the request to other Inertialization nodes
-		// that were skipped due to pose caching.
-		if (RequestQueue.Num() > 0)
-		{
-			// Cached poses have their Update function called once even though there may be multiple UseCachedPose nodes for the same pose.
-			// Because of this, there may be Inertialization ancestors of the UseCachedPose nodes that missed out on requests.
-			// So here we forward 'this' node's requests to the ancestors of those skipped UseCachedPose nodes.
-			// Note that in some cases, we may be forwarding the requests back to this same node.  Those duplicate requests will ultimately
-			// be ignored by the 'AddUnique' in the body of FAnimNode_DeadBlending::RequestInertialization.
-			for (const UE::Anim::FMessageStack& Stack : InSkippedUpdates)
-			{
-				Stack.ForEachMessage<UE::Anim::IInertializationRequester>([this, NodeId, &Proxy](UE::Anim::IInertializationRequester& InMessage)
-				{
-					for (const FInertializationRequest& Request : RequestQueue)
-					{
-						InMessage.RequestInertialization(Request);
-					}
-					InMessage.AddDebugRecord(Proxy, NodeId);
-
-					return UE::Anim::FMessageStack::EEnumerate::Stop;
-				});
-			}
-		}
-	});
+	// Catch the inertialization request message and call the node's RequestInertialization function with the request
+	UE::Anim::TScopedGraphMessage<UE::Anim::FDeadBlendingRequester> InertializationMessage(Context, Context, this);
 
 	Source.Update(Context);
 
diff --git a/Engine/Source/Runtime/Engine/Private/Animation/AnimNode_Inertialization.cpp b/Engine/Source/Runtime/Engine/Private/Animation/AnimNode_Inertialization.cpp
--- a/Engine/Source/Runtime/Engine/Private/Animation/AnimNode_Inertialization.cpp
+++ b/Engine/Source/Runtime/Engine/Private/Animation/AnimNode_Inertialization.cpp
@@ -212,41 +212,45 @@ void FAnimNode_Inertialization::Update_AnyThread(const FAnimationUpdateContext&
 	{
 		Deactivate();
 	}
-	UpdateCounter.SynchronizeWith(Context.AnimInstanceProxy->GetUpdateCounter());
 
-	const int32 NodeId = Context.GetCurrentNodeId();
-	const FAnimInstanceProxy& Proxy = *Context.AnimInstanceProxy;
+	UpdateCounter.SynchronizeWith(Context.AnimInstanceProxy->GetUpdateCounter());
 
-	// Allow nodes further towards the leaves to inertialize using this node
-	UE::Anim::TScopedGraphMessage<UE::Anim::FInertializationRequester> Inertialization(Context, Context, this);
+	// Catch the inertialization request message and call the node's RequestInertialization function with the request
+	UE::Anim::TScopedGraphMessage<UE::Anim::FInertializationRequester> InertializationMessage(Context, Context, this);
 
-	// Handle skipped updates for cached poses by forwarding to inertialization nodes in those residual stacks
-	UE::Anim::TScopedGraphMessage<UE::Anim::FCachedPoseSkippedUpdateHandler> CachedPoseSkippedUpdate(Context, [this, NodeId, &Proxy](TArrayView<const UE::Anim::FMessageStack> InSkippedUpdates)
+	if (bForwardRequestsThroughSkippedCachedPoseNodes)
 	{
-		// If we have a pending request forward the request to other Inertialization nodes
-		// that were skipped due to pose caching.
-		if(RequestQueue.Num() > 0)
+		const int32 NodeId = Context.GetCurrentNodeId();
+		const FAnimInstanceProxy& Proxy = *Context.AnimInstanceProxy;
+
+		// Handle skipped updates for cached poses by forwarding to inertialization nodes in those residual stacks
+		UE::Anim::TScopedGraphMessage<UE::Anim::FCachedPoseSkippedUpdateHandler> CachedPoseSkippedUpdate(Context, [this, NodeId, &Proxy](TArrayView<const UE::Anim::FMessageStack> InSkippedUpdates)
 		{
-			// Cached poses have their Update function called once even though there may be multiple UseCachedPose nodes for the same pose.
-			// Because of this, there may be Inertialization ancestors of the UseCachedPose nodes that missed out on requests.
-			// So here we forward 'this' node's requests to the ancestors of those skipped UseCachedPose nodes.
-			// Note that in some cases, we may be forwarding the requests back to this same node.  Those duplicate requests will ultimately
-			// be ignored by the 'AddUnique' in the body of FAnimNode_Inertialization::RequestInertialization.
-			for (const UE::Anim::FMessageStack& Stack : InSkippedUpdates)
+			// If we have a pending request forward the request to other Inertialization nodes
+			// that were skipped due to pose caching.
+			if (RequestQueue.Num() > 0)
 			{
-				Stack.ForEachMessage<UE::Anim::IInertializationRequester>([this, NodeId, &Proxy](UE::Anim::IInertializationRequester& InMessage)
+				// Cached poses have their Update function called once even though there may be multiple UseCachedPose nodes for the same pose.
+				// Because of this, there may be Inertialization ancestors of the UseCachedPose nodes that missed out on requests.
+				// So here we forward 'this' node's requests to the ancestors of those skipped UseCachedPose nodes.
+				// Note that in some cases, we may be forwarding the requests back to this same node.  Those duplicate requests will ultimately
+				// be ignored by the 'AddUnique' in the body of FAnimNode_Inertialization::RequestInertialization.
+				for (const UE::Anim::FMessageStack& Stack : InSkippedUpdates)
 				{
-					for (const FInertializationRequest& Request : RequestQueue)
+					Stack.ForEachMessage<UE::Anim::IInertializationRequester>([this, NodeId, &Proxy](UE::Anim::IInertializationRequester& InMessage)
 					{
-						InMessage.RequestInertialization(Request);
-					}
- 					InMessage.AddDebugRecord(Proxy, NodeId);
-
-					return UE::Anim::FMessageStack::EEnumerate::Stop;
-				});
+						for (const FInertializationRequest& Request : RequestQueue)
+						{
+							InMessage.RequestInertialization(Request);
+						}
+						InMessage.AddDebugRecord(Proxy, NodeId);
+
+						return UE::Anim::FMessageStack::EEnumerate::Stop;
+					});
+				}
 			}
-		}
-	});
+		});
+	}
 
 	Source.Update(Context);
 
