Fix race condition on FKAggregateGeom::GetGeom when rendering in parallel.

diff --git a/Engine/Source/Developer/MeshMergeUtilities/Private/MeshMergeHelpers.cpp b/Engine/Source/Developer/MeshMergeUtilities/Private/MeshMergeHelpers.cpp
--- a/Engine/Source/Developer/MeshMergeUtilities/Private/MeshMergeHelpers.cpp
+++ b/Engine/Source/Developer/MeshMergeUtilities/Private/MeshMergeHelpers.cpp
@@ -957,7 +957,7 @@ void FMeshMergeHelpers::ExtractPhysicsGeometry(UBodySetup* InBodySetup, const FT
 	OutAggGeom.BoxElems.Empty();
 
 	// we are not owner of this stuff
-	OutAggGeom.RenderInfo = nullptr;
+	OutAggGeom.FreeRenderInfo();
 	for (FKConvexElem& Elem : OutAggGeom.ConvexElems)
 	{
 		Elem.ResetChaosConvexMesh();
diff --git a/Engine/Source/Runtime/Engine/Classes/PhysicsEngine/AggregateGeom.h b/Engine/Source/Runtime/Engine/Classes/PhysicsEngine/AggregateGeom.h
--- a/Engine/Source/Runtime/Engine/Classes/PhysicsEngine/AggregateGeom.h
+++ b/Engine/Source/Runtime/Engine/Classes/PhysicsEngine/AggregateGeom.h
@@ -11,6 +11,7 @@
 #include "PhysicsEngine/SphereElem.h"
 #include "PhysicsEngine/SphylElem.h"
 #include "PhysicsEngine/TaperedCapsuleElem.h"
+#include "Async/Mutex.h"
 #include "AggregateGeom.generated.h"
 
 class FMaterialRenderProxy;
@@ -42,15 +43,13 @@ struct FKAggregateGeom
 	UPROPERTY(EditAnywhere, editfixedsize, Category = "Aggregate Geometry", meta = (DisplayName = "(Experimental) Skinned Level Sets"), Experimental)
 	TArray<FKSkinnedLevelSetElem> SkinnedLevelSetElems;
 
-	class FKConvexGeomRenderInfo* RenderInfo;
-
 	FKAggregateGeom()
-		: RenderInfo(NULL)
+		: RenderInfoPtr(nullptr)
 	{
 	}
 
 	FKAggregateGeom(const FKAggregateGeom& Other)
-		: RenderInfo(nullptr)
+		: RenderInfoPtr(nullptr)
 	{
 		CloneAgg(Other);
 	}
@@ -293,4 +292,8 @@ private:
 			});
 		return FoundIndex;
 	}
+
+	// NOTE: RenderInfo is generated concurrently and lazily (hence being mutable)
+	mutable std::atomic<class FKConvexGeomRenderInfo*> RenderInfoPtr;
+	mutable UE::FMutex RenderInfoLock;
 };
diff --git a/Engine/Source/Runtime/Engine/Private/PhysicsEngine/PhysDrawing.cpp b/Engine/Source/Runtime/Engine/Private/PhysicsEngine/PhysDrawing.cpp
--- a/Engine/Source/Runtime/Engine/Private/PhysicsEngine/PhysDrawing.cpp
+++ b/Engine/Source/Runtime/Engine/Private/PhysicsEngine/PhysDrawing.cpp
@@ -20,6 +20,7 @@
 #include "Chaos/Levelset.h"
 #include "Chaos/UniformGrid.h"
 #include "Chaos/WeightedLatticeImplicitObject.h"
+#include "Misc/ScopeLock.h"
 
 static const int32 DrawCollisionSides = 32;
 static const int32 DrawConeLimitSides = 40;
@@ -1096,33 +1097,40 @@ void FKAggregateGeom::GetAggGeom(const FTransform& Transform, const FColor Color
 		if(bDrawSolid)
 		{
 			// Cache collision vertex/index buffer
+			auto RenderInfo = RenderInfoPtr.load(std::memory_order_relaxed);
 			if(!RenderInfo)
 			{
-				//@todo - parallel rendering, remove const cast
-				FKAggregateGeom& ThisGeom = const_cast<FKAggregateGeom&>(*this);
-				ThisGeom.RenderInfo = new FKConvexGeomRenderInfo();
-				ThisGeom.RenderInfo->VertexBuffers = new FStaticMeshVertexBuffers();
-				ThisGeom.RenderInfo->IndexBuffer = new FDynamicMeshIndexBuffer32();
-
-				TArray<FDynamicMeshVertex> OutVerts;
-				for(int32 i=0; i<ConvexElems.Num(); i++)
-				{
-					// Get vertices/triangles from this hull.
-					ConvexElems[i].AddCachedSolidConvexGeom(OutVerts, ThisGeom.RenderInfo->IndexBuffer->Indices, FColor::White);
-				}
+				UE::TScopeLock Lock(RenderInfoLock);
 
-				// Only continue if we actually got some valid geometry
-				// Will crash if we try to init buffers with no data
-				if(ThisGeom.RenderInfo->VertexBuffers
-					&& ThisGeom.RenderInfo->IndexBuffer
-					&& OutVerts.Num() > 0
-					&& ThisGeom.RenderInfo->IndexBuffer->Indices.Num() > 0)
+				// After obtaining the lock, we might now have render info, in which case skip re-creating the data
+				RenderInfo = RenderInfoPtr.load(std::memory_order_acquire);
+				if (!RenderInfo)
 				{
-					ThisGeom.RenderInfo->IndexBuffer->InitResource(RHICmdList);
+					RenderInfo = new FKConvexGeomRenderInfo();
+					RenderInfo->VertexBuffers = new FStaticMeshVertexBuffers();
+					RenderInfo->IndexBuffer = new FDynamicMeshIndexBuffer32();
 
-					ThisGeom.RenderInfo->CollisionVertexFactory = new FLocalVertexFactory(Collector.GetFeatureLevel(), "FKAggregateGeom");
-					ThisGeom.RenderInfo->VertexBuffers->InitFromDynamicVertex(RHICmdList, ThisGeom.RenderInfo->CollisionVertexFactory, OutVerts);
+					TArray<FDynamicMeshVertex> OutVerts;
+					for(int32 i=0; i<ConvexElems.Num(); i++)
+					{
+						// Get vertices/triangles from this hull.
+						ConvexElems[i].AddCachedSolidConvexGeom(OutVerts, RenderInfo->IndexBuffer->Indices, FColor::White);
+					}
 
+					// Only continue if we actually got some valid geometry
+					// Will crash if we try to init buffers with no data
+					if(RenderInfo->VertexBuffers
+						&& RenderInfo->IndexBuffer
+						&& OutVerts.Num() > 0
+						&& RenderInfo->IndexBuffer->Indices.Num() > 0)
+					{
+						RenderInfo->IndexBuffer->InitResource(RHICmdList);
+
+						RenderInfo->CollisionVertexFactory = new FLocalVertexFactory(Collector.GetFeatureLevel(), "FKAggregateGeom");
+						RenderInfo->VertexBuffers->InitFromDynamicVertex(RHICmdList, RenderInfo->CollisionVertexFactory, OutVerts);
+					}
+
+					RenderInfoPtr.store(RenderInfo, std::memory_order_release);
 				}
 			}
 
@@ -1214,10 +1222,11 @@ void FKAggregateGeom::GetAggGeom(const FTransform& Transform, const FColor Color
 	}
 }
 
-/** Release the RenderInfo (if its there) and safely clean up any resources. Not thread safe, but can be called from any thread (conditionally safe). */
+/** Release the RenderInfo (if its there) and safely clean up any resources. Can be called from any thread. */
 void FKAggregateGeom::FreeRenderInfo()
 {
 	// See if we have rendering resources to free
+	auto RenderInfo = RenderInfoPtr.exchange(nullptr, std::memory_order_acq_rel);
 	if (RenderInfo)
 	{
 		// Should always have these if RenderInfo exists
@@ -1251,9 +1260,6 @@ void FKAggregateGeom::FreeRenderInfo()
 				delete RenderInfoToRelease;
 			}
 		);
-
-		// Reset the pointer as it's been given to the render thread
-		RenderInfo = nullptr;
 	}
 }
 
