Convert warning to display, and display to (valid) warning

I accidentally changed the wrong warning

diff --git a/Engine/Source/Runtime/Engine/Private/AudioStreamingCache.cpp b/Engine/Source/Runtime/Engine/Private/AudioStreamingCache.cpp
--- a/Engine/Source/Runtime/Engine/Private/AudioStreamingCache.cpp
+++ b/Engine/Source/Runtime/Engine/Private/AudioStreamingCache.cpp
@@ -548,7 +548,7 @@ FAudioChunkHandle FCachedAudioStreamingManager::GetLoadedChunk(const FSoundWaveP
 
 		if (!FAudioChunkCache::DoesKeyContainValidChunkIndex(ChunkKey, *SoundWaveData))
 		{
-			UE_LOG(LogAudioStreamCaching, Display, TEXT("Invalid Chunk Index %d Requested for Wave %s!"), ChunkIndex, *SoundWave->GetFName().ToString());
+			UE_LOG(LogAudioStreamCaching, Warning, TEXT("Invalid Chunk Index %d Requested for Wave %s!"), ChunkIndex, *SoundWave->GetFName().ToString());
 			return FAudioChunkHandle();
 		}
 
@@ -751,7 +751,7 @@ FAudioChunkCache::FAudioChunkCache(uint32 InMaxChunkSize, uint32 NumChunks, uint
 
 	if (MemoryLimitBytes == 0)
 	{
-		UE_LOG(LogAudioStreamCaching, Warning, TEXT("Audio stream cache size is 0 bytes. Audio will not play. To resolve this issue, set \'CacheSizeKB\' to a non-zero value."));
+		UE_LOG(LogAudioStreamCaching, Display, TEXT("Audio stream cache size is 0 bytes. Audio will not play. To resolve this issue, set \'CacheSizeKB\' to a non-zero value."));
 	}
 }
 
