Improving texture source encoding override tooltip.

#rnx

diff --git a/Engine/Source/Runtime/Engine/Classes/Engine/TextureDefines.h b/Engine/Source/Runtime/Engine/Classes/Engine/TextureDefines.h
--- a/Engine/Source/Runtime/Engine/Classes/Engine/TextureDefines.h
+++ b/Engine/Source/Runtime/Engine/Classes/Engine/TextureDefines.h
@@ -402,7 +402,7 @@ enum TextureCompressionSettings : int
 UENUM()
 enum class ETextureSourceEncoding : uint8
 {
-	TSE_None		= 0 UMETA(DisplayName = "Default", ToolTip = "The source encoding is assumed to match the state of the sRGB checkbox parameter."),
+	TSE_None		= 0 UMETA(DisplayName = "Default", ToolTip = "The source encoding is not overridden."),
 	TSE_Linear		= 1 UMETA(DisplayName = "Linear", ToolTip = "The source encoding is considered linear (before optional sRGB encoding is applied)."),
 	TSE_sRGB		= 2 UMETA(DisplayName = "sRGB", ToolTip = "sRGB source encoding to be linearized (before optional sRGB encoding is applied)."),
 	TSE_ST2084		= 3 UMETA(DisplayName = "ST 2084/PQ", ToolTip = "SMPTE ST 2084/PQ source encoding to be linearized (before optional sRGB encoding is applied)."),
