If project settings sets Stream Cache Size to 0, use default value instead so that audio still plays.

Log a "Display" message that informs the user if the .ini file is setting the value to 0. Makes it clear that they can update their settings to use a non-default value.

Also, when creating the Stream Cache Params, make sure we don't pass in 0 for the stream cache size as well!

#tests PIE

diff --git a/Engine/Source/Runtime/Engine/Private/AudioCompressionSettingsUtils.cpp b/Engine/Source/Runtime/Engine/Private/AudioCompressionSettingsUtils.cpp
--- a/Engine/Source/Runtime/Engine/Private/AudioCompressionSettingsUtils.cpp
+++ b/Engine/Source/Runtime/Engine/Private/AudioCompressionSettingsUtils.cpp
@@ -119,11 +119,16 @@ void CacheAudioCookOverrides(FPlatformAudioCookOverrides& OutOverrides, const TC
 
 	/** Memory Load On Demand Settings */
 	// Cache size:
-	const int32 DefaultCacheSize = 64 * 1024;
-	int32 RetrievedCacheSize = DefaultCacheSize;
-	if (!PlatformFile->GetInt(*CategoryName, TEXT("CacheSizeKB"), RetrievedCacheSize))
+	constexpr int32 DefaultCacheSizeKB = FAudioStreamCachingSettings::DefaultCacheSize;
+	int32 RetrievedCacheSize = DefaultCacheSizeKB;
+	if (PlatformFile->GetInt(*CategoryName, TEXT("CacheSizeKB"), RetrievedCacheSize))
 	{
-		RetrievedCacheSize = DefaultCacheSize;
+		if (RetrievedCacheSize == 0)
+		{
+			UE_LOG(LogConfig, Display, TEXT("Audio Stream Cache \"Max Cache Size KB\" set to 0 by config: \"%s%s.ini\". Default value of %d KB will be used. You can update Project Settings here: Project Settings->Platforms->%s->Audio->Cook Overrides->Stream Caching->Max Cache Size (KB)"),
+				*PlatformFile->SourceProjectConfigDir, *PlatformFile->Name.ToString(), DefaultCacheSizeKB, *PlatformFile->PlatformName);
+			RetrievedCacheSize = DefaultCacheSizeKB;
+		}
 	}
 
 	OutOverrides.StreamCachingSettings.CacheSizeKB = RetrievedCacheSize;
@@ -354,7 +359,7 @@ FCachedAudioStreamingManagerParams FPlatformCompressionUtilities::BuildCachedStr
 
 	// Primary cache defined here:
 	CacheDimensions.MaxChunkSize = 256 * 1024; // max possible chunk size (hard coded for legacy streaming path)
-	CacheDimensions.MaxMemoryInBytes = CacheSettings.CacheSizeKB * 1024;
+	CacheDimensions.MaxMemoryInBytes = CacheSettings.CacheSizeKB > 0 ? CacheSettings.CacheSizeKB * 1024 : FAudioStreamCachingSettings::DefaultCacheSize * 1024;
 	CacheDimensions.NumElements = FMath::Max(NumElements, 1); // force at least a single cache element to avoid crashes
 	Params.Caches.Add(CacheDimensions);
 
