Ensure that DormantConnections and RecentlyDormantConnections are cleared for seamlessly travelling Actors

diff --git a/Engine/Source/Runtime/Engine/Private/NetworkObjectList.cpp b/Engine/Source/Runtime/Engine/Private/NetworkObjectList.cpp
--- a/Engine/Source/Runtime/Engine/Private/NetworkObjectList.cpp
+++ b/Engine/Source/Runtime/Engine/Private/NetworkObjectList.cpp
@@ -450,6 +450,10 @@ void FNetworkObjectList::OnPostSeamlessTravel()
 			ObjectsDormantOnAllConnections.Remove(TravelingActor);
 		}
 
+		//Ensure the actor's dormant connections lists are empty when seamlessly traveling
+		SeamlessTravelingActorInfo->DormantConnections.Empty();
+		SeamlessTravelingActorInfo->RecentlyDormantConnections.Empty();
+
 		// Add the saved object info back into the active list
 		AllNetworkObjects.Add(SeamlessTravelingActorInfo);
 		ActiveNetworkObjects.Add(SeamlessTravelingActorInfo);
