Fix bug where starting Play In Editor before the asset registry scan finishes confuses the Asset Manager into thinking the scan has completed
This would result in failing to properly load DataRegistries or other systems that listen for the completion callback

diff --git a/Engine/Source/Runtime/Engine/Private/AssetManager.cpp b/Engine/Source/Runtime/Engine/Private/AssetManager.cpp
--- a/Engine/Source/Runtime/Engine/Private/AssetManager.cpp
+++ b/Engine/Source/Runtime/Engine/Private/AssetManager.cpp
@@ -3753,6 +3753,12 @@ void UAssetManager::RefreshPrimaryAssetDirectory(bool bForceRefresh)
 {
 	WarningInvalidAssets.Reset();
 
+	// Do not refresh before the initial scan has completed
+	if (!HasInitialScanCompleted())
+	{
+		return;
+	}
+
 	if (bForceRefresh || !bIsPrimaryAssetDirectoryCurrent)
 	{
 		PushBulkScanning();
@@ -4618,7 +4624,18 @@ bool UAssetManager::GetPrimaryAssetSetChunkIds(const TSet<FPrimaryAssetId>& Prim
 
 void UAssetManager::PreBeginPIE(bool bStartSimulate)
 {
-	RefreshPrimaryAssetDirectory();
+	if (HasInitialScanCompleted())
+	{
+		// If the scan has finished, we need to refresh in case there have been in-editor changes
+		RefreshPrimaryAssetDirectory();
+	}
+	else
+	{
+		// If the scan is still in progress, we need to finish it now which will call PostInitialAssetScan
+		GetAssetRegistry().WaitForCompletion();
+
+		ensure(HasInitialScanCompleted());
+	}
 
 	// Cache asset state
 	GetPrimaryAssetBundleStateMap(PrimaryAssetStateBeforePIE, false);
