Fix for potential integer overflow in FTextureData::DeduplicateTiles.

#rnx
#11027

diff --git a/Engine/Source/Runtime/Engine/Private/SparseVolumeTexture/SparseVolumeTextureData.cpp b/Engine/Source/Runtime/Engine/Private/SparseVolumeTexture/SparseVolumeTextureData.cpp
--- a/Engine/Source/Runtime/Engine/Private/SparseVolumeTexture/SparseVolumeTextureData.cpp
+++ b/Engine/Source/Runtime/Engine/Private/SparseVolumeTexture/SparseVolumeTextureData.cpp
@@ -706,11 +706,11 @@ bool FTextureData::DeduplicateTiles()
 
 		if (!MipMaps[MipLevel].PhysicalTileDataA.IsEmpty())
 		{
-			MipMaps[MipLevel].PhysicalTileDataA.SetNum(TileWritePos * TileByteSize[0]);
+			MipMaps[MipLevel].PhysicalTileDataA.SetNum(DeduplicatedTileSizeBytes[0]);
 		}
 		if (!MipMaps[MipLevel].PhysicalTileDataB.IsEmpty())
 		{
-			MipMaps[MipLevel].PhysicalTileDataB.SetNum(TileWritePos * TileByteSize[1]);
+			MipMaps[MipLevel].PhysicalTileDataB.SetNum(DeduplicatedTileSizeBytes[1]);
 		}
 	}
 
