Introduced Material Translation DDC caching.

This change implements a caching mechanism for material translation results. After translating a material, the intermediate results are now stored in the DDC (Derived Data Cache) for future retrieval. When the material translator is requested to translate a material, it first checks the DDC to see if th combination of specified base material, parameter sets, and platform-specific data is already present in the cache. If this key is found, the intermediate results are pulled and returned, avoiding expensive material translation and decoupling material translation from material shader generation.

The significant advantage of this change is that modifications to the material template source file or any dependent shader file no longer trigger a material retranslation. The material shader can still be generated correctly, using the cached translation results, which ensures that material string generation always uses up-to-date input files.

diff --git a/Engine/Source/Editor/MaterialEditor/Private/MaterialEditor.h b/Engine/Source/Editor/MaterialEditor/Private/MaterialEditor.h
--- a/Engine/Source/Editor/MaterialEditor/Private/MaterialEditor.h
+++ b/Engine/Source/Editor/MaterialEditor/Private/MaterialEditor.h
@@ -122,6 +122,7 @@ public:
 
 	virtual EMaterialDomain GetMaterialDomain() const override { return MD_Surface; }
 	virtual FString GetMaterialUsageDescription() const override { return FString::Printf(TEXT("FMatExpressionPreview %s"), Expression.IsValid() ? *Expression->GetName() : TEXT("NULL")); }
+	virtual bool IsPreview() const override { return true; }
 	virtual bool IsTwoSided() const override { return false; }
 	virtual bool IsThinSurface() const override { return false; }
 	virtual bool IsDitheredLODTransition() const override { return false; }
diff --git a/Engine/Source/Runtime/Engine/Private/Materials/HLSLMaterialTranslator.cpp b/Engine/Source/Runtime/Engine/Private/Materials/HLSLMaterialTranslator.cpp
--- a/Engine/Source/Runtime/Engine/Private/Materials/HLSLMaterialTranslator.cpp
+++ b/Engine/Source/Runtime/Engine/Private/Materials/HLSLMaterialTranslator.cpp
@@ -5,7 +5,6 @@
 =============================================================================*/
 
 #include "HLSLMaterialTranslator.h"
-#include "Containers/LazyPrintf.h"
 #include "VT/VirtualTextureScalability.h"
 #include "Engine/Engine.h"
 #include "MaterialDomain.h"
@@ -54,6 +53,13 @@
 #if WITH_EDITORONLY_DATA
 #include "Materials/MaterialExpressionSubstrate.h"
 #include "ShaderPlatformCachedIniValue.h"
+#include "Serialization/ObjectWriter.h"
+#include "Serialization/ObjectReader.h"
+#include "Serialization/BufferArchive.h"
+#include "Serialization/MemoryReader.h"
+#include "DerivedDataCache.h"
+#include "DerivedDataRequestOwner.h"
+#include "MaterialCachedData.h"
 #endif
 
 #if ENABLE_COOK_STATS && STATS
@@ -110,14 +116,11 @@ struct FCsvLogFile
 
 };
 
-#include "ProfilingDebugging/ScopedTimers.h"
-namespace MaterialTranslatorCookStats
-{
-	static FCookStatsManager::FAutoRegisterCallback RegisterCookStats([](FCookStatsManager::AddStatFuncRef AddStat)
-		{
-			FCsvLogFile::Get().Save();
-		});
-}
+static FCookStatsManager::FAutoRegisterCallback RegisterCookStats([](FCookStatsManager::AddStatFuncRef AddStat)
+	{
+		FCsvLogFile::Get().Save();
+	});
+
 #endif
 
 static int32 GetLWCTruncateMode()
@@ -186,6 +189,215 @@ static FAutoConsoleVariableRef CVarPedanticErrorChecksEnabled(
 	GPedanticErrorChecksEnabled,
 	TEXT("Enables material compilation pedantic error checking"));
 
+/* Controls whether to use the new GetMaterialShaderCode() and GetMaterialEnvironment() implementations. */
+static bool GUseMaterialTranslationResultsGrouping = false;
+
+/* Controls whether DDC caching of material translation results is enabled. */
+static bool GJobMaterialTranslateDDCEnable = false;
+static FAutoConsoleVariableRef CVarJobMaterialTranslateDDCEnable(
+	TEXT("r.Material.TranslateDDCEnable"),
+	GJobMaterialTranslateDDCEnable,
+	TEXT("Whether to enable material translation DDC caching.\n"));
+
+UE::DerivedData::FCacheBucket MaterialTranslationDDCBucket = UE::DerivedData::FCacheBucket(TEXT("MaterialTranslation"));
+UE::DerivedData::FValueId MaterialCompilationOutputId = UE::DerivedData::FValueId::FromName("FHLSLMaterialTranslator_MaterialCompilationOutput");
+UE::DerivedData::FValueId MaterialResultsOutputId = UE::DerivedData::FValueId::FromName("FHLSLMaterialTranslator_Results");
+UE::DerivedData::FValueId EnvironmentDefinesId = UE::DerivedData::FValueId::FromName("FHLSLMaterialTranslator_EnvironmentDefines");
+
+/* This version number models the layout of the data stored on the DDC after a material translation (e.g. FEnvironmentDefines)
+ * It must be bumped whenever a change is made requires re-translation of all materials. */
+static constexpr int MaterialTranslationDDCVersion = 1;
+
+/** Data structure used to cache a part of material translation results. It contains all the generated
+ *  defines that will be declared during the compilation of the generated material shader.
+ *  Note: if you make a change to this structure, remember to bumpMaterialTranslationDDCVersion.
+ */
+struct FHLSLMaterialTranslator::FEnvironmentDefines
+{
+	enum EVirtualPageType
+	{
+		TABLE0 = 0,
+		TABLE1 = 1,
+		TABLE_INDIRECTION = 2
+	};
+
+	bool bNeedsParticlePosition;
+	bool bNeedsParticleVelocity;
+	bool bUseDynamicParameters;
+	uint32 DynamicParametersMask;
+	bool bNeedsParticleTime;
+	bool bUsesParticleMotionBlur;
+	bool bNeedsParticleRandom;
+	bool bSphericalParticleOpacity;
+	bool bUseParticleSubUVs;
+	bool bLightmapUVAccess;
+	bool bUsesAOMaterialMask;
+	bool bUsesSpeedtree;
+	bool bNeedsWorldPositionExcludingShaderOffsets;
+	bool bNeedsParticleSize;
+	bool bNeedsParticleSpriteRotation;
+	bool bNeedsSceneTextures;
+	bool bUsesEyeAdaptation;
+	bool bVirtualTextureOutput;
+	bool bUsesPerInstanceCustomData;
+	bool bUsesPerInstanceRandom;
+	bool bUsesPerInstanceFadeAmount;
+	bool bUsesVertexInterpolator;
+	bool bUsesSkyAtmosphere;
+	bool bUsesVertexColor;
+	bool bUsesParticleColor;
+	bool bUsesParticleLocalToWorld;
+	bool bUsesParticleWorldToLocal;
+	bool bUsesInstanceLocalToWorldPS;
+	bool bUsesInstanceWorldToLocalPS;
+	bool bUsesTransformVector;
+	bool bUsesPixelDepthOffset;
+	bool bUsesWorldPositionOffset;
+	bool bUsesDisplacement;
+	bool bUsesEmissiveColor;
+	bool bUsesDistortion;
+	bool bDistortionAccountForCoverage;
+	bool bMaterialEnableTranslucencyFogging;
+	bool bMaterialEnableTranslucencyCloudFogging;
+	bool bMaterialIsSky;
+	bool bMaterialComputeFogPerPixel;
+	bool bMaterialFullyRough;
+	bool bMaterialUsesAnisotropy;
+	uint8 MaterialDecalReadMask;
+	bool bMaterialUsesDecalLookup;
+	uint8 MaterialPathTracingBufferRead;
+	bool MaterialNeuralPostProcess;
+	uint32 NumVirtualTextureSamples;
+	bool bMaterialVirtualTextureFeedback;
+	TArray<uint8> VirtualPageTypes;
+	bool bSingleLayerWaterShadingQuality;
+	bool bShadingModelsIsLit;
+	uint32 MaterialShadingModelEnabled;
+	bool bMaterialSubsurfaceProfileUseCurvature;
+	bool bMaterialShadingModelEyeUseCurvature;
+	bool bDisableForwardLocalLights;
+	bool bSingleLayerWaterSeparatedMainLight;
+	bool bMaterialSingleShadingModel;
+	bool bMaterialVolumetricAdvanced;
+	bool bMaterialVolumetricAdvancedPhasePerSample;
+	bool bMaterialVolumetricAdvancedGreyscaleMaterial;
+	bool bMaterialVolumetricAdvancedRaymarchVolumeShadow;
+	bool bMaterialVolumetricAdvancedClampMultiscatteringContribution;
+	uint32 MaterialVolumetricAdvancedMultiscatteringOctaveCount;
+	bool bMaterialVolumetricAdvancedConservativeDensity;
+	bool bMaterialVolumetricAdvancedOverrideAmbientOcclusion;
+	bool bMaterialVolumetricAdvancedAdvancedGroundContribution;
+	bool bMaterialVolumetricCloudEmptySpaceSkippingOutput;
+	bool bMaterialIsSubstrate;
+	bool bDualSourceColorBlendingEnabled;
+	bool bSubstratePremultipliedAlphaOpacityOverridden;
+	bool bSubstrateUsesConversionFromLegacy;
+	bool bSubstrateMaterialOutputOpaqueRoughRefractions;
+	int32 SubstrateMaterialExportType;
+	int32 SubstrateMaterialExportContext;
+	int32 SubstrateMaterialExportLegacyBlendMode;
+	bool bSubstrateOptimizedUnlit;
+	bool bSubstrateSinglePath;
+	bool bSubstrateFastPath;
+	uint32 SubstrateClampedClosureCount;
+	bool SubstrateIsComplexSpecialPath;
+	bool bSubstrateComplexSpecialPath;
+	bool bTextureSampleDebug;
+	TArray<TPair<FString, int32>> SubstrateDefines;
+
+	bool HasShadingModel(EMaterialShadingModel model) const
+	{
+		return (MaterialShadingModelEnabled & (1 << model)) != 0;
+	}
+
+	void Serialize(FArchive& Ar)
+	{
+		Ar << bNeedsParticlePosition;
+		Ar << bNeedsParticleVelocity;
+		Ar << bUseDynamicParameters;
+		Ar << DynamicParametersMask;
+		Ar << bNeedsParticleTime;
+		Ar << bUsesParticleMotionBlur;
+		Ar << bNeedsParticleRandom;
+		Ar << bSphericalParticleOpacity;
+		Ar << bUseParticleSubUVs;
+		Ar << bLightmapUVAccess;
+		Ar << bUsesAOMaterialMask;
+		Ar << bUsesSpeedtree;
+		Ar << bNeedsWorldPositionExcludingShaderOffsets;
+		Ar << bNeedsParticleSize;
+		Ar << bNeedsParticleSpriteRotation;
+		Ar << bNeedsSceneTextures;
+		Ar << bUsesEyeAdaptation;
+		Ar << bVirtualTextureOutput;
+		Ar << bUsesPerInstanceCustomData;
+		Ar << bUsesPerInstanceRandom;
+		Ar << bUsesPerInstanceFadeAmount;
+		Ar << bUsesVertexInterpolator;
+		Ar << bUsesSkyAtmosphere;
+		Ar << bUsesVertexColor;
+		Ar << bUsesParticleColor;
+		Ar << bUsesParticleLocalToWorld;
+		Ar << bUsesParticleWorldToLocal;
+		Ar << bUsesInstanceLocalToWorldPS;
+		Ar << bUsesInstanceWorldToLocalPS;
+		Ar << bUsesTransformVector;
+		Ar << bUsesPixelDepthOffset;
+		Ar << bUsesWorldPositionOffset;
+		Ar << bUsesDisplacement;
+		Ar << bUsesEmissiveColor;
+		Ar << bUsesDistortion;
+		Ar << bDistortionAccountForCoverage;
+		Ar << bMaterialEnableTranslucencyFogging;
+		Ar << bMaterialEnableTranslucencyCloudFogging;
+		Ar << bMaterialIsSky;
+		Ar << bMaterialComputeFogPerPixel;
+		Ar << bMaterialFullyRough;
+		Ar << bMaterialUsesAnisotropy;
+		Ar << MaterialDecalReadMask;
+		Ar << bMaterialUsesDecalLookup;
+		Ar << MaterialPathTracingBufferRead;
+		Ar << MaterialNeuralPostProcess;
+		Ar << NumVirtualTextureSamples;
+		Ar << bMaterialVirtualTextureFeedback;
+		Ar << VirtualPageTypes;
+		Ar << bSingleLayerWaterShadingQuality;
+		Ar << bShadingModelsIsLit;
+		Ar << MaterialShadingModelEnabled;
+		Ar << bMaterialSubsurfaceProfileUseCurvature;
+		Ar << bMaterialShadingModelEyeUseCurvature;
+		Ar << bDisableForwardLocalLights;
+		Ar << bSingleLayerWaterSeparatedMainLight;
+		Ar << bMaterialSingleShadingModel;
+		Ar << bMaterialVolumetricAdvanced;
+		Ar << bMaterialVolumetricAdvancedPhasePerSample;
+		Ar << bMaterialVolumetricAdvancedGreyscaleMaterial;
+		Ar << bMaterialVolumetricAdvancedRaymarchVolumeShadow;
+		Ar << bMaterialVolumetricAdvancedClampMultiscatteringContribution;
+		Ar << MaterialVolumetricAdvancedMultiscatteringOctaveCount;
+		Ar << bMaterialVolumetricAdvancedConservativeDensity;
+		Ar << bMaterialVolumetricAdvancedOverrideAmbientOcclusion;
+		Ar << bMaterialVolumetricAdvancedAdvancedGroundContribution;
+		Ar << bMaterialVolumetricCloudEmptySpaceSkippingOutput;
+		Ar << bMaterialIsSubstrate;
+		Ar << bDualSourceColorBlendingEnabled;
+		Ar << bSubstratePremultipliedAlphaOpacityOverridden;
+		Ar << bSubstrateUsesConversionFromLegacy;
+		Ar << bSubstrateMaterialOutputOpaqueRoughRefractions;
+		Ar << SubstrateMaterialExportType;
+		Ar << SubstrateMaterialExportContext;
+		Ar << SubstrateMaterialExportLegacyBlendMode;
+		Ar << bSubstrateOptimizedUnlit;
+		Ar << bSubstrateSinglePath;
+		Ar << bSubstrateFastPath;
+		Ar << SubstrateClampedClosureCount;
+		Ar << bSubstrateComplexSpecialPath;
+		Ar << bTextureSampleDebug;
+		Ar << SubstrateDefines;
+	}
+};
+
+
 static inline bool IsAnalyticDerivEnabled()
 {
 	return GAnalyticDerivEnabled != 0;
@@ -367,6 +579,13 @@ void FHLSLMaterialTranslator::FSubstrateCompilationContext::Initialise()
 	SubstrateThicknessStack.SetNum(0);
 }
 
+void FHLSLMaterialTranslator::AppendVersionString(FString& Output, EShaderPlatform Platform)
+{
+	Output.Appendf(TEXT("_MatTransl_%s_%d_"),
+		*FMaterialSourceTemplate::Get().GetTemplateHashString(Platform),
+		MaterialTranslationDDCVersion);
+}
+
 FHLSLMaterialTranslator::FHLSLMaterialTranslator(FMaterial* InMaterial,
 	FMaterialCompilationOutput& InMaterialCompilationOutput,
 	const FStaticParameterSet& InStaticParameters,
@@ -374,7 +593,8 @@ FHLSLMaterialTranslator::FHLSLMaterialTranslator(FMaterial* InMaterial,
 	EMaterialQualityLevel::Type InQualityLevel,
 	ERHIFeatureLevel::Type InFeatureLevel,
 	const ITargetPlatform* InTargetPlatform, //if InTargetPlatform is nullptr, we use the current active
-	const FSubstrateCompilationConfig* InSubstrateCompilationConfig)
+	const FSubstrateCompilationConfig* InSubstrateCompilationConfig,
+	FString MaterialTranslationDDCKeyString)
 :	ShaderFrequency(SF_Pixel)
 ,	MaterialProperty(MP_EmissiveColor)
 ,	CurrentScopeChunks(nullptr)
@@ -497,6 +717,11 @@ FHLSLMaterialTranslator::FHLSLMaterialTranslator(FMaterial* InMaterial,
 	{
 		SubstrateCompilationConfig = *InSubstrateCompilationConfig;
 	}
+
+	EnvironmentDefines.Reset(new FEnvironmentDefines);
+
+	// Hash the DDC key string and store the result for DDC retrieval.
+	DDCKeyHash = FIoHash::HashBuffer(MakeMemoryView(FTCHARToUTF8(MaterialTranslationDDCKeyString)));
 }
 
 FHLSLMaterialTranslator::~FHLSLMaterialTranslator()
@@ -874,6 +1099,7 @@ UE_TRACE_EVENT_BEGIN(Cpu, FHLSLMaterialTranslatorTranslate, NoSync)
 	UE_TRACE_EVENT_FIELD(UE::Trace::WideString, MaterialName)
 UE_TRACE_EVENT_END()
  
+
 bool FHLSLMaterialTranslator::Translate()
 {
 #if CPUPROFILERTRACE_ENABLED
@@ -886,252 +1112,308 @@ bool FHLSLMaterialTranslator::Translate()
 		<< FHLSLMaterialTranslatorTranslate.MaterialName(*TraceMaterialName);
 #endif
 
+	static const bool bVerbose = FParse::Param(FCommandLine::Get(), TEXT("verbosematerialtranslation"));
+	if (bVerbose)
+	{
+		UE_LOG(LogMaterial, Display, TEXT("Translating '%s'"), *Material->GetMaterialInterface()->GetFullName());
+	}
 
-#if ENABLE_COOK_STATS && STATS
-	FDateTime TranslationDateTime = FDateTime::Now();
-#endif
-
+	STAT(FDateTime TranslationDateTime = FDateTime::Now());
 	STAT(double HLSLTranslateTime = 0);
-	{
-		SCOPE_SECONDS_COUNTER(HLSLTranslateTime);
+	STAT(double DDCQueryTime = 0);
 
-		check(ScopeStack.Num() == 0);
-		bSuccess = true;
+	bSuccess = true;
+	bool bInCache = false;
 
-		// WARNING: No compile outputs should be stored on the UMaterial / FMaterial / FMaterialResource, unless they are transient editor-only data (like error expressions)
-		// Compile outputs that need to be saved must be stored in MaterialCompilationOutput, which will be saved to the DDC.
+	{
+		SCOPE_SECONDS_COUNTER(HLSLTranslateTime);
 
-		Material->CompileErrors.Empty();
-		Material->ErrorExpressions.Empty();
+		if (!Material->IsPreview())
+		{
+			SCOPE_SECONDS_COUNTER(DDCQueryTime);
+			bInCache = QueryDDCCachedTranslationResults();
+		}
 
-		// If pedantic error checking is enabled, log out the pre-compilation errors
-		if (GPedanticErrorChecksEnabled)
+		// If material not in cache, we need to translate the material now.
+		if (!bInCache)
 		{
-			bSuccess = Material->CheckInValidStateForCompilation(this);
+			DoTranslate();
+
+			// If translation was succesful, finalize the results and push them to the DDC
+			if (bSuccess)
+			{
+				PrepareEnvironmentDefines();
+				PrepareMaterialSourceStringParameters();
+				PushResultsToDDCCache(); 
+				ClearAllFunctionStacks();
+			}
 		}
+	}
 
-		bEnableExecutionFlow = Material->IsUsingControlFlow();
-		bCompileForComputeShader = Material->IsLightFunction();
-		
-		// Verify for the absence of loops.
-		UMaterialInterface* Interface = Material->GetMaterialInterface();
-		if (!Interface)
+	// Report timings to Material Cook Stats
+#if STATS
+	GShaderCompilerStats->IncrementMaterialTranslated(HLSLTranslateTime);
+	if (bInCache)
+	{
+		STAT(GShaderCompilerStats->IncrementMaterialCacheHit(DDCQueryTime));
+	}
+#endif
+
+	INC_FLOAT_STAT_BY(STAT_ShaderCompiling_HLSLTranslation, (float)HLSLTranslateTime);
+
+#if ENABLE_COOK_STATS && STATS
+	// Write out a CSV file MaterialTranslationLog.txt containing info about all material translations.
+	FCsvLogFile::Get().AddEntry(Material->GetMaterialInterface()->GetFullName(), TranslationDateTime, HLSLTranslateTime);
+#endif
+
+	return bSuccess;
+}
+
+void FHLSLMaterialTranslator::DoTranslate()
+{
+	// No cache hit, continue translating the material
+	check(ScopeStack.Num() == 0);
+
+	// WARNING: No compile outputs should be stored on the UMaterial / FMaterial / FMaterialResource, unless they are transient editor-only data (like error expressions)
+	// Compile outputs that need to be saved must be stored in MaterialCompilationOutput, which will be saved to the DDC.
+
+	bEnableExecutionFlow = Material->IsUsingControlFlow();
+	bCompileForComputeShader = Material->IsLightFunction();
+
+	Material->CompileErrors.Empty();
+	Material->ErrorExpressions.Empty();
+
+	// If pedantic error checking is enabled, log out the pre-compilation errors
+	if (GPedanticErrorChecksEnabled)
+	{
+		bSuccess = Material->CheckInValidStateForCompilation(this);
+		if (!bSuccess)
 		{
-			return false;
+			return;
 		}
-		
-		if (UMaterial* UMaterial = Interface->GetMaterial())
+	}
+
+	// Verify for the absence of loops.
+	UMaterialInterface* Interface = Material->GetMaterialInterface();
+	if (!Interface)
+	{
+		bSuccess = false;
+		return;
+	}
+
+	if (UMaterial* UMaterial = Interface->GetMaterial())
+	{
+		TSet<UMaterialExpression*> VisitedExpressions;
+		for (UMaterialExpression* Expression : UMaterial->GetExpressions())
 		{
-			TSet<UMaterialExpression*> VisitedExpressions;
-			for (UMaterialExpression* Expression : UMaterial->GetExpressions())
+			if (Expression->ContainsInputLoop(VisitedExpressions))
 			{
-				if (Expression->ContainsInputLoop(VisitedExpressions))
-				{
-					AppendExpressionError(Expression, TEXT("Expression is part of a cycle. Please make sure the material graph is acyclic."));
-					return false;
-				}
+				AppendExpressionError(Expression, TEXT("Expression is part of a cycle. Please make sure the material graph is acyclic."));
+				bSuccess = false;
+				return;
 			}
 		}
+	}
 
-		//
-		// Process the Substrate tree representing the material topology.
-		//
-		const bool bSubstrateEnabled = Substrate::IsSubstrateEnabled();
-		UMaterialExpression* FrontMaterialExpr = nullptr;
-		int32 FrontMaterialOutputIndex = INDEX_NONE;
+	//
+	// Process the Substrate tree representing the material topology.
+	//
+	const bool bSubstrateEnabled = Substrate::IsSubstrateEnabled();
+	UMaterialExpression* FrontMaterialExpr = nullptr;
+	int32 FrontMaterialOutputIndex = INDEX_NONE;
 
-		bSubstrateWritesEmissive = false;
-		bSubstrateWritesAmbientOcclusion = false;
+	bSubstrateWritesEmissive = false;
+	bSubstrateWritesAmbientOcclusion = false;
 
-		UMaterialExpression* ExpressionToPreview = Material->GetMaterialGraphNodePreviewExpression();
-		if (ExpressionToPreview)
+	UMaterialExpression* ExpressionToPreview = Material->GetMaterialGraphNodePreviewExpression();
+	if (ExpressionToPreview)
+	{
+		// If this is a compilation used for a preview of a node from am aterial graph, then we need ot use that node as the root of the Substrate tree.
+		// Then the resulting SubstrateData will be converted to a color using SubstrateCompilePreview translator function when compiling the emissive color channel the node has been short circuited into.
+		const uint32 ExpressionPreviewOutputIndex = 0;
+		if (ExpressionToPreview && ExpressionToPreview->IsResultSubstrateMaterial(ExpressionPreviewOutputIndex))
 		{
-			// If this is a compilation used for a preview of a node from am aterial graph, then we need ot use that node as the root of the Substrate tree.
-			// Then the resulting SubstrateData will be converted to a color using SubstrateCompilePreview translator function when compiling the emissive color channel the node has been short circuited into.
-			const uint32 ExpressionPreviewOutputIndex = 0;
-			if (ExpressionToPreview && ExpressionToPreview->IsResultSubstrateMaterial(ExpressionPreviewOutputIndex))
-			{
-				FrontMaterialExpr = ExpressionToPreview;
-				FrontMaterialOutputIndex = ExpressionPreviewOutputIndex;
-			}
-			else
-			{
-				FrontMaterialExpr = nullptr;	// This is not a Substrate input that is connected there so we cannot create a Substrate tree.
-			}
+			FrontMaterialExpr = ExpressionToPreview;
+			FrontMaterialOutputIndex = ExpressionPreviewOutputIndex;
 		}
 		else
 		{
-			FSubstrateMaterialInput* FrontMaterialInput = Material->GetMaterialInterface() ? &Material->GetMaterialInterface()->GetMaterial()->GetEditorOnlyData()->FrontMaterial : nullptr;
-			FrontMaterialExpr = FrontMaterialInput ? FrontMaterialInput->GetTracedInput().Expression : nullptr;
-			FrontMaterialOutputIndex = FrontMaterialInput ? FrontMaterialInput->OutputIndex : INDEX_NONE;
+			FrontMaterialExpr = nullptr;	// This is not a Substrate input that is connected there so we cannot create a Substrate tree.
 		}
-		if (bSubstrateEnabled && FrontMaterialExpr)
-		{
-			// Temp code chunk scope (e.g.needed for the creation of static booleans from static switch parameter node, see UMaterialExpressionStaticSwitch::GetEffectiveInput).
-			TArray<FShaderCodeChunk> TempChunks;
-			AssignTempScope(TempChunks);
+	}
+	else
+	{
+		FSubstrateMaterialInput* FrontMaterialInput = Material->GetMaterialInterface() ? &Material->GetMaterialInterface()->GetMaterial()->GetEditorOnlyData()->FrontMaterial : nullptr;
+		FrontMaterialExpr = FrontMaterialInput ? FrontMaterialInput->GetTracedInput().Expression : nullptr;
+		FrontMaterialOutputIndex = FrontMaterialInput ? FrontMaterialInput->OutputIndex : INDEX_NONE;
+	}
+	if (bSubstrateEnabled && FrontMaterialExpr)
+	{
+		// Temp code chunk scope (e.g.needed for the creation of static booleans from static switch parameter node, see UMaterialExpressionStaticSwitch::GetEffectiveInput).
+		TArray<FShaderCodeChunk> TempChunks;
+		AssignTempScope(TempChunks);
 
-		#if DEBUG_SUBSTRATE_TREE_STACK
-			UE_LOG(LogMaterial, Display, TEXT(" SubstrateTreeStack: SubstrateGenerateMaterialTopologyTree"));
-		#endif
-			for(uint32 SubstrateCompilationContextIndex = 0; SubstrateCompilationContextIndex < ESubstrateCompilationContext::SCC_MAX; ++SubstrateCompilationContextIndex)
-			{
-				// This is needed because SubstrateGenerateMaterialTopologyTree will call some compiler context though material expressions.
-				CurrentSubstrateCompilationContext = ESubstrateCompilationContext(SubstrateCompilationContextIndex);
+	#if DEBUG_SUBSTRATE_TREE_STACK
+		UE_LOG(LogMaterial, Display, TEXT(" SubstrateTreeStack: SubstrateGenerateMaterialTopologyTree"));
+	#endif
+		for (uint32 SubstrateCompilationContextIndex = 0; SubstrateCompilationContextIndex < ESubstrateCompilationContext::SCC_MAX; ++SubstrateCompilationContextIndex)
+		{
+			// This is needed because SubstrateGenerateMaterialTopologyTree will call some compiler context though material expressions.
+			CurrentSubstrateCompilationContext = ESubstrateCompilationContext(SubstrateCompilationContextIndex);
 
-				FSubstrateCompilationContext& SubstrateCtx = SubstrateCompilationContext[SubstrateCompilationContextIndex];
-				SubstrateCtx = FSubstrateCompilationContext(ESubstrateCompilationContext(SubstrateCompilationContextIndex));
-			
-				FExpressionInput* SurfaceThickness = Material->IsThinSurface() && Material->GetMaterialInterface() ? &Material->GetMaterialInterface()->GetMaterial()->GetEditorOnlyData()->SurfaceThickness : nullptr;
-				SubstrateThicknessStackPush(nullptr, SurfaceThickness);
-				FrontMaterialExpr->SubstrateGenerateMaterialTopologyTree(this, nullptr, FrontMaterialOutputIndex);
-				SubstrateThicknessStackPop();
+			FSubstrateCompilationContext& SubstrateCtx = SubstrateCompilationContext[SubstrateCompilationContextIndex];
+			SubstrateCtx = FSubstrateCompilationContext(ESubstrateCompilationContext(SubstrateCompilationContextIndex));
 
-				check(SubstrateCtx.SubstrateThicknessStack.Num() == 0);
+			FExpressionInput* SurfaceThickness = Material->IsThinSurface() && Material->GetMaterialInterface() ? &Material->GetMaterialInterface()->GetMaterial()->GetEditorOnlyData()->SurfaceThickness : nullptr;
+			SubstrateThicknessStackPush(nullptr, SurfaceThickness);
+			FrontMaterialExpr->SubstrateGenerateMaterialTopologyTree(this, nullptr, FrontMaterialOutputIndex);
+			SubstrateThicknessStackPop();
 
-				if (SubstrateCtx.bSubstrateTreeOutOfStackDepthOccurred)
-				{
-					Errorf(TEXT(" %s [%s]: Substrate - Cyclic graph detected when we only support acyclic graph."), *Material->GetDebugName(), *Material->GetAssetPath().ToString());
-				}
-				if (!SubstrateCtx.SubstrateGenerateDerivedMaterialOperatorData(this))
-				{
-					Errorf(TEXT("Substrate material errors encountered."));
-				}
+			check(SubstrateCtx.SubstrateThicknessStack.Num() == 0);
 
-				bSubstrateWritesEmissive |= SubstrateCtx.bSubstrateWritesEmissive;
-				bSubstrateWritesAmbientOcclusion |= SubstrateCtx.bSubstrateWritesAmbientOcclusion;
+			if (SubstrateCtx.bSubstrateTreeOutOfStackDepthOccurred)
+			{
+				Errorf(TEXT(" %s [%s]: Substrate - Cyclic graph detected when we only support acyclic graph."), *Material->GetDebugName(), *Material->GetAssetPath().ToString());
 			}
-			CurrentSubstrateCompilationContext = ESubstrateCompilationContext::SCC_Default;
+			if (!SubstrateCtx.SubstrateGenerateDerivedMaterialOperatorData(this))
+			{
+				Errorf(TEXT("Substrate material errors encountered."));
+			}
+
+			bSubstrateWritesEmissive |= SubstrateCtx.bSubstrateWritesEmissive;
+			bSubstrateWritesAmbientOcclusion |= SubstrateCtx.bSubstrateWritesAmbientOcclusion;
 		}
+		CurrentSubstrateCompilationContext = ESubstrateCompilationContext::SCC_Default;
+	}
 
-		// Generate code:
-		// Normally one would expect the generator to emit something like
-		//		float Local0 = ...
-		//		...
-		//		float Local3= ...
-		//		...
-		//		float Localn= ...
-		//		PixelMaterialInputs.EmissiveColor = Local0 + ...
-		//		PixelMaterialInputs.Normal = Local3 * ...
-		// However because the Normal can be used in the middle of generating other Locals (which happens when using a node like PixelNormalWS)
-		// instead we generate this:
-		//		float Local0 = ...
-		//		...
-		//		float Local3= ...
-		//		PixelMaterialInputs.Normal = Local3 * ...
-		//		...
-		//		float Localn= ...
-		//		PixelMaterialInputs.EmissiveColor = Local0 + ...
-		// in other words, compile Normal first, then emit all the expressions up to the last one Normal requires;
-		// assign the normal into the shared struct, then emit the remaining expressions; finally assign the rest of the shared struct inputs.
-		// Inputs that are not shared, have false in the SharedPixelProperties array, and those ones will emit the full code.
-
-		int32 NormalCodeChunkEnd = -1;
-		int32 Chunk[CompiledMP_MAX];
-		int32 VertexAttributesChunk = INDEX_NONE;
-		int32 PixelAttributesChunk = INDEX_NONE;
-			
-		memset(Chunk, INDEX_NONE, sizeof(Chunk));
+	// Generate code:
+	// Normally one would expect the generator to emit something like
+	//		float Local0 = ...
+	//		...
+	//		float Local3= ...
+	//		...
+	//		float Localn= ...
+	//		PixelMaterialInputs.EmissiveColor = Local0 + ...
+	//		PixelMaterialInputs.Normal = Local3 * ...
+	// However because the Normal can be used in the middle of generating other Locals (which happens when using a node like PixelNormalWS)
+	// instead we generate this:
+	//		float Local0 = ...
+	//		...
+	//		float Local3= ...
+	//		PixelMaterialInputs.Normal = Local3 * ...
+	//		...
+	//		float Localn= ...
+	//		PixelMaterialInputs.EmissiveColor = Local0 + ...
+	// in other words, compile Normal first, then emit all the expressions up to the last one Normal requires;
+	// assign the normal into the shared struct, then emit the remaining expressions; finally assign the rest of the shared struct inputs.
+	// Inputs that are not shared, have false in the SharedPixelProperties array, and those ones will emit the full code.
 
-		// Translate all custom vertex interpolators before main attributes so type information is available
-		{
-			CustomVertexInterpolators.Empty();
-			CurrentCustomVertexInterpolatorOffset = 0;
-			NextVertexInterpolatorIndex = 0;
-			MaterialProperty = MP_MAX;
-			ShaderFrequency = SF_Vertex;
+	int32 NormalCodeChunkEnd = -1;
+	int32 Chunk[CompiledMP_MAX];
+	int32 VertexAttributesChunk = INDEX_NONE;
+	int32 PixelAttributesChunk = INDEX_NONE;
 
-			TArray<UMaterialExpression*> Expressions;
-			Material->GatherExpressionsForCustomInterpolators(Expressions);
-			GatherCustomVertexInterpolators(Expressions);
+	memset(Chunk, INDEX_NONE, sizeof(Chunk));
 
-			// Reset shared stack data
-			while (FunctionStacks[SF_Vertex].Num() > 1)
-			{
-				FMaterialFunctionCompileState* Stack = FunctionStacks[SF_Vertex].Pop(false);
-				delete Stack;
-			}
-			FunctionStacks[SF_Vertex][0]->Reset();
+	// Translate all custom vertex interpolators before main attributes so type information is available
+	{
+		CustomVertexInterpolators.Empty();
+		CurrentCustomVertexInterpolatorOffset = 0;
+		NextVertexInterpolatorIndex = 0;
+		MaterialProperty = MP_MAX;
+		ShaderFrequency = SF_Vertex;
+
+		TArray<UMaterialExpression*> Expressions;
+		Material->GatherExpressionsForCustomInterpolators(Expressions);
+		GatherCustomVertexInterpolators(Expressions);
 
-			// Whilst expression list is available, apply node count limits
-			int32 NumMaterialLayersAttributes = 0;
-			for (UMaterialExpression* Expression : Expressions)
+		// Reset shared stack data
+		while (FunctionStacks[SF_Vertex].Num() > 1)
+		{
+			FMaterialFunctionCompileState* Stack = FunctionStacks[SF_Vertex].Pop(false);
+			delete Stack;
+		}
+		FunctionStacks[SF_Vertex][0]->Reset();
+
+		// Whilst expression list is available, apply node count limits
+		int32 NumMaterialLayersAttributes = 0;
+		for (UMaterialExpression* Expression : Expressions)
+		{
+			if (UMaterialExpressionMaterialAttributeLayers* Layers = Cast<UMaterialExpressionMaterialAttributeLayers>(Expression))
 			{
-				if (UMaterialExpressionMaterialAttributeLayers* Layers = Cast<UMaterialExpressionMaterialAttributeLayers>(Expression))
-				{
-					++NumMaterialLayersAttributes;
+				++NumMaterialLayersAttributes;
 
-					if (NumMaterialLayersAttributes > 1)
-					{
-						Errorf(TEXT("Materials can contain only one Material Attribute Layers node."));
-						break;
-					}
+				if (NumMaterialLayersAttributes > 1)
+				{
+					Errorf(TEXT("Materials can contain only one Material Attribute Layers node."));
+					break;
 				}
 			}
 		}
+	}
 
-		const EShaderFrequency NormalShaderFrequency = FMaterialAttributeDefinitionMap::GetShaderFrequency(MP_Normal);
-		const EMaterialDomain Domain = Material->GetMaterialDomain();
-		const EBlendMode BlendMode = Material->GetBlendMode();
+	const EShaderFrequency NormalShaderFrequency = FMaterialAttributeDefinitionMap::GetShaderFrequency(MP_Normal);
+	const EMaterialDomain Domain = Material->GetMaterialDomain();
+	const EBlendMode BlendMode = Material->GetBlendMode();
 
-		const EShaderFrequency FrontMaterialShaderFrequency = FMaterialAttributeDefinitionMap::GetShaderFrequency(MP_FrontMaterial);
-		int32 FullySimplifiedFrontMaterialCodeChunkStart = -1;
-		int32 FullySimplifiedFrontMaterialCodeChunkEnd = -1;
+	const EShaderFrequency FrontMaterialShaderFrequency = FMaterialAttributeDefinitionMap::GetShaderFrequency(MP_FrontMaterial);
+	int32 FullySimplifiedFrontMaterialCodeChunkStart = -1;
+	int32 FullySimplifiedFrontMaterialCodeChunkEnd = -1;
 
-		// Gather the implementation for any custom output expressions
-		TArray<UMaterialExpressionCustomOutput*> CustomOutputExpressions;
-		Material->GatherCustomOutputExpressions(CustomOutputExpressions);
-		TSet<UClass*> SeenCustomOutputExpressionsClasses;
+	// Gather the implementation for any custom output expressions
+	TArray<UMaterialExpressionCustomOutput*> CustomOutputExpressions;
+	Material->GatherCustomOutputExpressions(CustomOutputExpressions);
+	TSet<UClass*> SeenCustomOutputExpressionsClasses;
 
-		// Some custom outputs must be pre-compiled so they can be re-used as shared inputs
-		CompileCustomOutputs(CustomOutputExpressions, SeenCustomOutputExpressionsClasses, true);			
+	// Some custom outputs must be pre-compiled so they can be re-used as shared inputs
+	CompileCustomOutputs(CustomOutputExpressions, SeenCustomOutputExpressionsClasses, true);
 
-		// Validate some things on the VT system. Since generated code for expressions shared between multiple properties
-		// (e.g. a texture sample connected to both diffuse and opacity mask) is reused we can't check based on the MaterialProperty
-		// variable inside the actual code generation pass. So we do a pre-pass over it here.
-		if (UseVirtualTexturing(Platform))
-		{
-			ValidateVtPropertyLimits();
-		}
+	// Validate some things on the VT system. Since generated code for expressions shared between multiple properties
+	// (e.g. a texture sample connected to both diffuse and opacity mask) is reused we can't check based on the MaterialProperty
+	// variable inside the actual code generation pass. So we do a pre-pass over it here.
+	if (UseVirtualTexturing(Platform))
+	{
+		ValidateVtPropertyLimits();
+	}
 
-		if (bEnableExecutionFlow)
-		{
-			PixelAttributesChunk = Material->CompilePropertyAndSetMaterialProperty(MP_MaterialAttributes, this, SF_Pixel);
-			VertexAttributesChunk = Material->CompilePropertyAndSetMaterialProperty(MP_MaterialAttributes, this, SF_Vertex);
-		}
-		else
-		{
+	if (bEnableExecutionFlow)
+	{
+		PixelAttributesChunk = Material->CompilePropertyAndSetMaterialProperty(MP_MaterialAttributes, this, SF_Pixel);
+		VertexAttributesChunk = Material->CompilePropertyAndSetMaterialProperty(MP_MaterialAttributes, this, SF_Vertex);
+	}
+	else
+	{
 		// Normal must always be compiled first; this will ensure its chunk calculations are the first to be added
 		{
 			// Verify that start chunk is 0
 			check(SharedPropertyCodeChunks[NormalShaderFrequency].Num() == 0);
-			Chunk[MP_Normal]					= Material->CompilePropertyAndSetMaterialProperty(MP_Normal					,this);
+			Chunk[MP_Normal] = Material->CompilePropertyAndSetMaterialProperty(MP_Normal					,this);
 			NormalCodeChunkEnd = SharedPropertyCodeChunks[NormalShaderFrequency].Num();
 		}
 
 		// Rest of properties
-		Chunk[MP_EmissiveColor]					= Material->CompilePropertyAndSetMaterialProperty(MP_EmissiveColor			,this);
-		Chunk[MP_DiffuseColor]					= Material->CompilePropertyAndSetMaterialProperty(MP_DiffuseColor			,this);
-		Chunk[MP_SpecularColor]					= Material->CompilePropertyAndSetMaterialProperty(MP_SpecularColor			,this);
-		Chunk[MP_BaseColor]						= Material->CompilePropertyAndSetMaterialProperty(MP_BaseColor				,this);
-		Chunk[MP_Metallic]						= Material->CompilePropertyAndSetMaterialProperty(MP_Metallic				,this);
-		Chunk[MP_Specular]						= Material->CompilePropertyAndSetMaterialProperty(MP_Specular				,this);
-		Chunk[MP_Roughness]						= Material->CompilePropertyAndSetMaterialProperty(MP_Roughness				,this);
-		Chunk[MP_Anisotropy]					= Material->CompilePropertyAndSetMaterialProperty(MP_Anisotropy				,this);
-		Chunk[MP_Opacity]						= Material->CompilePropertyAndSetMaterialProperty(MP_Opacity				,this);
-		Chunk[MP_OpacityMask]					= Material->CompilePropertyAndSetMaterialProperty(MP_OpacityMask			,this);
-		Chunk[MP_Tangent]						= Material->CompilePropertyAndSetMaterialProperty(MP_Tangent				,this);
-		Chunk[MP_WorldPositionOffset]			= Material->CompilePropertyAndSetMaterialProperty(MP_WorldPositionOffset	,this);
+		Chunk[MP_EmissiveColor] = Material->CompilePropertyAndSetMaterialProperty(MP_EmissiveColor			,this);
+		Chunk[MP_DiffuseColor] = Material->CompilePropertyAndSetMaterialProperty(MP_DiffuseColor			,this);
+		Chunk[MP_SpecularColor] = Material->CompilePropertyAndSetMaterialProperty(MP_SpecularColor			,this);
+		Chunk[MP_BaseColor] = Material->CompilePropertyAndSetMaterialProperty(MP_BaseColor				,this);
+		Chunk[MP_Metallic] = Material->CompilePropertyAndSetMaterialProperty(MP_Metallic				,this);
+		Chunk[MP_Specular] = Material->CompilePropertyAndSetMaterialProperty(MP_Specular				,this);
+		Chunk[MP_Roughness] = Material->CompilePropertyAndSetMaterialProperty(MP_Roughness				,this);
+		Chunk[MP_Anisotropy] = Material->CompilePropertyAndSetMaterialProperty(MP_Anisotropy				,this);
+		Chunk[MP_Opacity] = Material->CompilePropertyAndSetMaterialProperty(MP_Opacity				,this);
+		Chunk[MP_OpacityMask] = Material->CompilePropertyAndSetMaterialProperty(MP_OpacityMask			,this);
+		Chunk[MP_Tangent] = Material->CompilePropertyAndSetMaterialProperty(MP_Tangent				,this);
+		Chunk[MP_WorldPositionOffset] = Material->CompilePropertyAndSetMaterialProperty(MP_WorldPositionOffset	,this);
 
 		// Make sure to compile this property before using ShadingModelsFromCompilation
-		Chunk[MP_ShadingModel]					= Material->CompilePropertyAndSetMaterialProperty(MP_ShadingModel			,this);
-		
-		
+		Chunk[MP_ShadingModel] = Material->CompilePropertyAndSetMaterialProperty(MP_ShadingModel			,this);
+
+
 	#if DEBUG_SUBSTRATE_TREE_STACK
 		UE_LOG(LogMaterial, Display, TEXT(" SubstrateTreeStack: Material->CompilePropertyAndSetMaterialProperty(MP_FrontMaterial)"));
 	#endif
-		Chunk[MP_SurfaceThickness]				= Material->CompilePropertyAndSetMaterialProperty(MP_SurfaceThickness		,this);
-		Chunk[MP_FrontMaterial]					= Material->CompilePropertyAndSetMaterialProperty(MP_FrontMaterial			,this);
+		Chunk[MP_SurfaceThickness] = Material->CompilePropertyAndSetMaterialProperty(MP_SurfaceThickness		,this);
+		Chunk[MP_FrontMaterial] = Material->CompilePropertyAndSetMaterialProperty(MP_FrontMaterial			,this);
 
 		// Now generate the code for the fully simplified MP_FrontMaterial right after MP_FrontMaterial
 		FullySimplifiedSubstrateFrontMaterialCodeChunk = SubstrateCreateAndRegisterNullMaterial();
@@ -1145,852 +1427,835 @@ bool FHLSLMaterialTranslator::Translate()
 		}
 		}
 
-		// Get shading models from compilation (or material).
-		FMaterialShadingModelField MaterialShadingModels = GetCompiledShadingModels(); 
-		
-		ValidateShadingModelsForFeatureLevel(MaterialShadingModels);
-		
-		if (!bEnableExecutionFlow)
-		{
-			if (Domain == MD_Volume || (Domain == MD_Surface && IsSubsurfaceShadingModel(MaterialShadingModels)))
-			{
-				// Note we don't test for the blend mode as you can have a translucent material using the subsurface shading model
+	// Get shading models from compilation (or material).
+	FMaterialShadingModelField MaterialShadingModels = GetCompiledShadingModels();
 
-				// another ForceCast as CompilePropertyAndSetMaterialProperty() can return MCT_Float which we don't want here
-				int32 SubsurfaceColor = Material->CompilePropertyAndSetMaterialProperty(MP_SubsurfaceColor, this);
-				SubsurfaceColor = ForceCast(SubsurfaceColor, FMaterialAttributeDefinitionMap::GetValueType(MP_SubsurfaceColor), MFCF_ExactMatch | MFCF_ReplicateValue);
+	ValidateShadingModelsForFeatureLevel(MaterialShadingModels);
 
-				static FName NameSubsurfaceProfile(TEXT("__SubsurfaceProfile"));
+	if (!bEnableExecutionFlow)
+	{
+		if (Domain == MD_Volume || (Domain == MD_Surface && IsSubsurfaceShadingModel(MaterialShadingModels)))
+		{
+			// Note we don't test for the blend mode as you can have a translucent material using the subsurface shading model
 
-				// 1.0f is is a not used profile - later this gets replaced with the actual profile
-				int32 CodeSubsurfaceProfile = ForceCast(ScalarParameter(NameSubsurfaceProfile, 1.0f), MCT_Float1);
+			// another ForceCast as CompilePropertyAndSetMaterialProperty() can return MCT_Float which we don't want here
+			int32 SubsurfaceColor = Material->CompilePropertyAndSetMaterialProperty(MP_SubsurfaceColor, this);
+			SubsurfaceColor = ForceCast(SubsurfaceColor, FMaterialAttributeDefinitionMap::GetValueType(MP_SubsurfaceColor), MFCF_ExactMatch | MFCF_ReplicateValue);
 
-				Chunk[MP_SubsurfaceColor] = AppendVector(SubsurfaceColor, CodeSubsurfaceProfile);
-			}
+			static FName NameSubsurfaceProfile(TEXT("__SubsurfaceProfile"));
 
-			Chunk[MP_CustomData0] = Material->CompilePropertyAndSetMaterialProperty(MP_CustomData0, this);
-			Chunk[MP_CustomData1] = Material->CompilePropertyAndSetMaterialProperty(MP_CustomData1, this);
-			Chunk[MP_AmbientOcclusion] = Material->CompilePropertyAndSetMaterialProperty(MP_AmbientOcclusion, this);
+			// 1.0f is is a not used profile - later this gets replaced with the actual profile
+			int32 CodeSubsurfaceProfile = ForceCast(ScalarParameter(NameSubsurfaceProfile, 1.0f), MCT_Float1);
 
-			if (IsTranslucentBlendMode(BlendMode) || MaterialShadingModels.HasShadingModel(MSM_SingleLayerWater))
-			{
-				// Cast to exact match is needed for float parameter to be correctly cast to float2.
-				int32 UserRefraction = ForceCast(Material->CompilePropertyAndSetMaterialProperty(MP_Refraction, this), MCT_Float2, MFCF_ExactMatch);
-				int32 RefractionDepthBias = ForceCast(ScalarParameter(FName(TEXT("RefractionDepthBias")), Material->GetRefractionDepthBiasValue()), MCT_Float1);
+			Chunk[MP_SubsurfaceColor] = AppendVector(SubsurfaceColor, CodeSubsurfaceProfile);
+		}
 
-				Chunk[MP_Refraction] = AppendVector(UserRefraction, RefractionDepthBias);
-			}
+		Chunk[MP_CustomData0] = Material->CompilePropertyAndSetMaterialProperty(MP_CustomData0, this);
+		Chunk[MP_CustomData1] = Material->CompilePropertyAndSetMaterialProperty(MP_CustomData1, this);
+		Chunk[MP_AmbientOcclusion] = Material->CompilePropertyAndSetMaterialProperty(MP_AmbientOcclusion, this);
 
-			if (bCompileForComputeShader)
-			{
-				Chunk[CompiledMP_EmissiveColorCS] = Material->CompilePropertyAndSetMaterialProperty(MP_EmissiveColor, this, SF_Compute);
-			}
+		if (IsTranslucentBlendMode(BlendMode) || MaterialShadingModels.HasShadingModel(MSM_SingleLayerWater))
+		{
+			// Cast to exact match is needed for float parameter to be correctly cast to float2.
+			int32 UserRefraction = ForceCast(Material->CompilePropertyAndSetMaterialProperty(MP_Refraction, this), MCT_Float2, MFCF_ExactMatch);
+			int32 RefractionDepthBias = ForceCast(ScalarParameter(FName(TEXT("RefractionDepthBias")), Material->GetRefractionDepthBiasValue()), MCT_Float1);
 
-			if (Chunk[MP_WorldPositionOffset] != INDEX_NONE)
-			{
-				// Only calculate previous WPO if there is a current WPO
-				Chunk[CompiledMP_PrevWorldPositionOffset] = Material->CompilePropertyAndSetMaterialProperty(MP_WorldPositionOffset, this, SF_Vertex, true);
-			}
+			Chunk[MP_Refraction] = AppendVector(UserRefraction, RefractionDepthBias);
+		}
 
-			Chunk[MP_PixelDepthOffset] = Material->CompilePropertyAndSetMaterialProperty(MP_PixelDepthOffset, this);
+		if (bCompileForComputeShader)
+		{
+			Chunk[CompiledMP_EmissiveColorCS] = Material->CompilePropertyAndSetMaterialProperty(MP_EmissiveColor, this, SF_Compute);
+		}
 
-			Chunk[MP_Displacement] = Material->CompilePropertyAndSetMaterialProperty(MP_Displacement, this);
+		if (Chunk[MP_WorldPositionOffset] != INDEX_NONE)
+		{
+			// Only calculate previous WPO if there is a current WPO
+			Chunk[CompiledMP_PrevWorldPositionOffset] = Material->CompilePropertyAndSetMaterialProperty(MP_WorldPositionOffset, this, SF_Vertex, true);
 		}
 
-		ResourcesString = TEXT("");
+		Chunk[MP_PixelDepthOffset] = Material->CompilePropertyAndSetMaterialProperty(MP_PixelDepthOffset, this);
+
+		Chunk[MP_Displacement] = Material->CompilePropertyAndSetMaterialProperty(MP_Displacement, this);
+	}
+
+	ResourcesString = TEXT("");
 
 #if HANDLE_CUSTOM_OUTPUTS_AS_MATERIAL_ATTRIBUTES
-		// Handle custom outputs when using material attribute output
-		if (Material->HasMaterialAttributesConnected())
+	// Handle custom outputs when using material attribute output
+	if (Material->HasMaterialAttributesConnected())
+	{
+		TArray<FMaterialCustomOutputAttributeDefintion> CustomAttributeList;
+		FMaterialAttributeDefinitionMap::GetCustomAttributeList(CustomAttributeList);
+		TArray<FShaderCodeChunk> CustomExpressionChunks;
+
+		for (FMaterialCustomOutputAttributeDefintion& Attribute : CustomAttributeList)
 		{
-			TArray<FMaterialCustomOutputAttributeDefintion> CustomAttributeList;
-			FMaterialAttributeDefinitionMap::GetCustomAttributeList(CustomAttributeList);
-			TArray<FShaderCodeChunk> CustomExpressionChunks;
+			// Compile all outputs for attribute
+			bool bValidResultCompiled = false;
+			int32 NumOutputs = 1;//CustomOutput->GetNumOutputs();
 
-			for (FMaterialCustomOutputAttributeDefintion& Attribute : CustomAttributeList)
+			for (int32 OutputIndex = 0; OutputIndex < NumOutputs; ++OutputIndex)
 			{
-				// Compile all outputs for attribute
-				bool bValidResultCompiled = false;
-				int32 NumOutputs = 1;//CustomOutput->GetNumOutputs();
+				MaterialProperty = Attribute.Property;
+				ShaderFrequency = Attribute.ShaderFrequency;
+				FunctionStacks[ShaderFrequency].Empty();
+				FunctionStacks[ShaderFrequency].Add(FMaterialFunctionCompileState(nullptr));
 
-				for (int32 OutputIndex = 0; OutputIndex < NumOutputs; ++OutputIndex)
-				{
-					MaterialProperty = Attribute.Property;
-					ShaderFrequency = Attribute.ShaderFrequency;
-					FunctionStacks[ShaderFrequency].Empty();
-					FunctionStacks[ShaderFrequency].Add(FMaterialFunctionCompileState(nullptr));
+				CustomExpressionChunks.Empty();
+				AssignTempScope(CustomExpressionChunks);
+				int32 Result = Material->CompileCustomAttribute(Attribute.AttributeID, this);
 
-					CustomExpressionChunks.Empty();
-					AssignTempScope(CustomExpressionChunks);
-					int32 Result = Material->CompileCustomAttribute(Attribute.AttributeID, this);
+				// Consider attribute used if varies from default value
+				if (Result != INDEX_NONE)
+				{
+					bool bValueNonDefault = true;
 
-					// Consider attribute used if varies from default value
-					if (Result != INDEX_NONE)
+					if (FMaterialUniformExpression* Expression = GetParameterUniformExpression(Result))
 					{
-						bool bValueNonDefault = true;
-
-						if (FMaterialUniformExpression* Expression = GetParameterUniformExpression(Result))
-						{
-							FLinearColor Value;
-							FMaterialRenderContext DummyContext(nullptr, *Material, nullptr);
-							Expression->GetNumberValue(DummyContext, Value);
-
-							bool bEqualValue = Value.R == Attribute.DefaultValue.X;
-							bEqualValue &= Value.G == Attribute.DefaultValue.Y || Attribute.ValueType < MCT_Float2;
-							bEqualValue &= Value.B == Attribute.DefaultValue.Z || Attribute.ValueType < MCT_Float3;
-							bEqualValue &= Value.A == Attribute.DefaultValue.W || Attribute.ValueType < MCT_Float4;
+						FLinearColor Value;
+						FMaterialRenderContext DummyContext(nullptr, *Material, nullptr);
+						Expression->GetNumberValue(DummyContext, Value);
 
-							if (Expression->IsConstant() && bEqualValue)
-							{
-								bValueNonDefault = false;
-							}
-						}
+						bool bEqualValue = Value.R == Attribute.DefaultValue.X;
+						bEqualValue &= Value.G == Attribute.DefaultValue.Y || Attribute.ValueType < MCT_Float2;
+						bEqualValue &= Value.B == Attribute.DefaultValue.Z || Attribute.ValueType < MCT_Float3;
+						bEqualValue &= Value.A == Attribute.DefaultValue.W || Attribute.ValueType < MCT_Float4;
 
-						// Valid, non-default value so generate shader code
-						if (bValueNonDefault)
+						if (Expression->IsConstant() && bEqualValue)
 						{
-							GenerateCustomAttributeCode(OutputIndex, Result, Attribute.ValueType, Attribute.FunctionName);
-							bValidResultCompiled = true;
+							bValueNonDefault = false;
 						}
 					}
-				}
 
-				// If used, add compile data
-				if (bValidResultCompiled)
-				{
-					ResourcesString += FString::Printf(TEXT("#define NUM_MATERIAL_OUTPUTS_%s %d\n"), *Attribute.FunctionName.ToUpper(), NumOutputs);
+					// Valid, non-default value so generate shader code
+					if (bValueNonDefault)
+					{
+						GenerateCustomAttributeCode(OutputIndex, Result, Attribute.ValueType, Attribute.FunctionName);
+						bValidResultCompiled = true;
+					}
 				}
 			}
+
+			// If used, add compile data
+			if (bValidResultCompiled)
+			{
+				ResourcesString += FString::Printf(TEXT("#define NUM_MATERIAL_OUTPUTS_%s %d\n"), *Attribute.FunctionName.ToUpper(), NumOutputs);
+			}
 		}
-		else
+	}
+	else
 #endif // #if HANDLE_CUSTOM_OUTPUTS_AS_MATERIAL_ATTRIBUTES
 		{
 			CompileCustomOutputs(CustomOutputExpressions, SeenCustomOutputExpressionsClasses, false);
 		}
 
-		// No more calls to non-vertex shader CompilePropertyAndSetMaterialProperty beyond this point
-		const uint32 SavedNumUserTexCoords = GetNumUserTexCoords();
+	// No more calls to non-vertex shader CompilePropertyAndSetMaterialProperty beyond this point
+	const uint32 SavedNumUserTexCoords = GetNumUserTexCoords();
 
-		if (!bEnableExecutionFlow)
+	if (!bEnableExecutionFlow)
+	{
+		for (uint32 CustomUVIndex = MP_CustomizedUVs0; CustomUVIndex <= MP_CustomizedUVs7; CustomUVIndex++)
 		{
-			for (uint32 CustomUVIndex = MP_CustomizedUVs0; CustomUVIndex <= MP_CustomizedUVs7; CustomUVIndex++)
+			// Only compile custom UV inputs for UV channels requested by the pixel shader inputs
+			// Any unconnected inputs will have a texcoord generated for them in Material->CompileProperty, which will pass through the vertex (uncustomized) texture coordinates
+			// Note: this is using NumUserTexCoords, which is set by translating all the pixel properties above
+			if (CustomUVIndex - MP_CustomizedUVs0 < SavedNumUserTexCoords)
 			{
-				// Only compile custom UV inputs for UV channels requested by the pixel shader inputs
-				// Any unconnected inputs will have a texcoord generated for them in Material->CompileProperty, which will pass through the vertex (uncustomized) texture coordinates
-				// Note: this is using NumUserTexCoords, which is set by translating all the pixel properties above
-				if (CustomUVIndex - MP_CustomizedUVs0 < SavedNumUserTexCoords)
-				{
-					Chunk[CustomUVIndex] = Material->CompilePropertyAndSetMaterialProperty((EMaterialProperty)CustomUVIndex, this);
-				}
+				Chunk[CustomUVIndex] = Material->CompilePropertyAndSetMaterialProperty((EMaterialProperty)CustomUVIndex, this);
 			}
 		}
+	}
 
-		// Output the implementation for any custom expressions we will call below.
-		for (int32 ExpressionIndex = 0; ExpressionIndex < CustomExpressions.Num(); ExpressionIndex++)
-		{
-			ResourcesString += CustomExpressions[ExpressionIndex].Implementation + "\n\n";
-		}
+	// Output the implementation for any custom expressions we will call below.
+	for (int32 ExpressionIndex = 0; ExpressionIndex < CustomExpressions.Num(); ExpressionIndex++)
+	{
+		ResourcesString += CustomExpressions[ExpressionIndex].Implementation + "\n\n";
+	}
+
+	// Translation is designed to have a code chunk generation phase followed by several passes that only has readonly access to the code chunks.
+	// At this point we mark the code chunk generation complete.
+	bAllowCodeChunkGeneration = false;
+
+	bUsesEmissiveColor = bSubstrateWritesEmissive || IsMaterialPropertyUsed(MP_EmissiveColor, Chunk[MP_EmissiveColor], FLinearColor(0, 0, 0, 0), 3);
+	bUsesPixelDepthOffset = (AllowPixelDepthOffset(Platform) && IsMaterialPropertyUsed(MP_PixelDepthOffset, Chunk[MP_PixelDepthOffset], FLinearColor(0, 0, 0, 0), 1));
 
-		// Translation is designed to have a code chunk generation phase followed by several passes that only has readonly access to the code chunks.
-		// At this point we mark the code chunk generation complete.
-		bAllowCodeChunkGeneration = false;
+	bool bUsesWorldPositionOffsetCurrent = IsMaterialPropertyUsed(MP_WorldPositionOffset, Chunk[MP_WorldPositionOffset], FLinearColor(0, 0, 0, 0), 3);
+	bool bUsesWorldPositionOffsetPrevious = IsMaterialPropertyUsed(MP_WorldPositionOffset, Chunk[CompiledMP_PrevWorldPositionOffset], FLinearColor(0, 0, 0, 0), 3);
+	bUsesWorldPositionOffset = bUsesWorldPositionOffsetCurrent || bUsesWorldPositionOffsetPrevious;
 
-		bUsesEmissiveColor = bSubstrateWritesEmissive || IsMaterialPropertyUsed(MP_EmissiveColor, Chunk[MP_EmissiveColor], FLinearColor(0, 0, 0, 0), 3);
-		bUsesPixelDepthOffset = (AllowPixelDepthOffset(Platform) && IsMaterialPropertyUsed(MP_PixelDepthOffset, Chunk[MP_PixelDepthOffset], FLinearColor(0, 0, 0, 0), 1));
+	bUsesDisplacement = DoesPlatformSupportNanite(Platform) && IsMaterialPropertyUsed(MP_Displacement, Chunk[MP_Displacement], FLinearColor(0, 0, 0, 0), 1);
 
-		bool bUsesWorldPositionOffsetCurrent = IsMaterialPropertyUsed(MP_WorldPositionOffset, Chunk[MP_WorldPositionOffset], FLinearColor(0, 0, 0, 0), 3);
-		bool bUsesWorldPositionOffsetPrevious = IsMaterialPropertyUsed(MP_WorldPositionOffset, Chunk[CompiledMP_PrevWorldPositionOffset], FLinearColor(0, 0, 0, 0), 3);
-		bUsesWorldPositionOffset = bUsesWorldPositionOffsetCurrent || bUsesWorldPositionOffsetPrevious;
+	MaterialCompilationOutput.bModifiesMeshPosition = bUsesPixelDepthOffset || bUsesWorldPositionOffset || bUsesDisplacement;
+	MaterialCompilationOutput.bUsesWorldPositionOffset = bUsesWorldPositionOffset;
+	MaterialCompilationOutput.bUsesPixelDepthOffset = bUsesPixelDepthOffset;
+	MaterialCompilationOutput.bUsesDisplacement = bUsesDisplacement;
 
-		bUsesDisplacement = DoesPlatformSupportNanite(Platform) && IsMaterialPropertyUsed(MP_Displacement, Chunk[MP_Displacement], FLinearColor(0, 0, 0, 0), 1);
+	// Fully rough if we have a roughness code chunk and it's constant and evaluates to 1.
+	bIsFullyRough = Chunk[MP_Roughness] != INDEX_NONE && IsMaterialPropertyUsed(MP_Roughness, Chunk[MP_Roughness], FLinearColor(1, 0, 0, 0), 1) == false;
+
+	bUsesAnisotropy = IsMaterialPropertyUsed(MP_Anisotropy, Chunk[MP_Anisotropy], FLinearColor(0, 0, 0, 0), 1);
+	MaterialCompilationOutput.bUsesAnisotropy = bUsesAnisotropy;
+
+	EMaterialDecalResponse MDR = (EMaterialDecalResponse)Material->GetMaterialDecalResponse();
+	if (MDR == MDR_Color || MDR == MDR_ColorNormal || MDR == MDR_ColorRoughness || MDR == MDR_ColorNormalRoughness)
+	{
+		MaterialCompilationOutput.SetIsDBufferTextureUsed(0);
+		AddEstimatedTextureSample(1);
+	}
+	if (MDR == MDR_Normal || MDR == MDR_ColorNormal || MDR == MDR_NormalRoughness || MDR == MDR_ColorNormalRoughness)
+	{
+		MaterialCompilationOutput.SetIsDBufferTextureUsed(1);
+		AddEstimatedTextureSample(1);
+	}
+	if (MDR == MDR_Roughness || MDR == MDR_ColorRoughness || MDR == MDR_NormalRoughness || MDR == MDR_ColorNormalRoughness)
+	{
+		MaterialCompilationOutput.SetIsDBufferTextureUsed(2);
+		AddEstimatedTextureSample(1);
+	}
 
-		MaterialCompilationOutput.bModifiesMeshPosition = bUsesPixelDepthOffset || bUsesWorldPositionOffset || bUsesDisplacement;
-		MaterialCompilationOutput.bUsesWorldPositionOffset = bUsesWorldPositionOffset;
-		MaterialCompilationOutput.bUsesPixelDepthOffset = bUsesPixelDepthOffset;
-		MaterialCompilationOutput.bUsesDisplacement = bUsesDisplacement;
+	bOpacityPropertyIsUsed = IsMaterialPropertyUsed(MP_Opacity, Chunk[MP_Opacity], FLinearColor(1, 0, 0, 0), 1);
 
-		// Fully rough if we have a roughness code chunk and it's constant and evaluates to 1.
-		bIsFullyRough = Chunk[MP_Roughness] != INDEX_NONE && IsMaterialPropertyUsed(MP_Roughness, Chunk[MP_Roughness], FLinearColor(1, 0, 0, 0), 1) == false;
+	bUsesCurvature = FeatureLevel == ERHIFeatureLevel::ES3_1 &&
+					((MaterialShadingModels.HasShadingModel(MSM_SubsurfaceProfile) && IsMaterialPropertyUsed(MP_CustomData0, Chunk[MP_CustomData0], FLinearColor(1, 0, 0, 0), 1))
+					|| (MaterialShadingModels.HasShadingModel(MSM_Eye) && bOpacityPropertyIsUsed));
 
-		bUsesAnisotropy = IsMaterialPropertyUsed(MP_Anisotropy, Chunk[MP_Anisotropy], FLinearColor(0, 0, 0, 0), 1);
-		MaterialCompilationOutput.bUsesAnisotropy = bUsesAnisotropy;
+	// If Substrate is enabled or if this is a Substrate material cooked/used in non-Substrate mode, 
+	// we disable this warning as Substrate supports 'colored transmittance only' mode (i.e, modulate).
+	if (!bSubstrateEnabled && !FrontMaterialExpr && IsModulateBlendMode(BlendMode) && MaterialShadingModels.IsLit() && !Material->IsDeferredDecal())
+	{
+		Errorf(TEXT("Dynamically lit translucency is not supported for BLEND_Modulate materials."));
+	}
 
-		EMaterialDecalResponse MDR = (EMaterialDecalResponse)Material->GetMaterialDecalResponse();
-		if (MDR == MDR_Color || MDR == MDR_ColorNormal || MDR == MDR_ColorRoughness || MDR == MDR_ColorNormalRoughness)
+	if (Domain == MD_Surface)
+	{
+		if (IsModulateBlendMode(BlendMode) && Material->IsTranslucencyAfterDOFEnabled() && !RHISupportsDualSourceBlending(Platform))
 		{
-			MaterialCompilationOutput.SetIsDBufferTextureUsed(0);
-			AddEstimatedTextureSample(1);
+			Errorf(TEXT("Translucency after DOF with BLEND_Modulate is only allowed on platforms that support dual-blending. Consider using BLEND_Translucent with black emissive"));
 		}
-		if (MDR == MDR_Normal || MDR == MDR_ColorNormal || MDR == MDR_NormalRoughness || MDR == MDR_ColorNormalRoughness)
+	}
+
+	// Don't allow opaque and masked materials to scene depth as the results are undefined
+	if (bUsesSceneDepth && Domain != MD_PostProcess && !IsTranslucentBlendMode(BlendMode))
+	{
+		Errorf(TEXT("Only transparent or postprocess materials can read from scene depth."));
+	}
+
+	if (bUsesSceneDepth)
+	{
+		MaterialCompilationOutput.SetIsSceneTextureUsed(PPI_SceneDepth);
+	}
+
+	MaterialCompilationOutput.bUsesDistanceCullFade = bUsesDistanceCullFade;
+
+	if (MaterialCompilationOutput.RequiresSceneColorCopy())
+	{
+		if (Domain != MD_Surface)
 		{
-			MaterialCompilationOutput.SetIsDBufferTextureUsed(1);
-			AddEstimatedTextureSample(1);
+			Errorf(TEXT("Only 'surface' material domain can use the scene color node."));
 		}
-		if (MDR == MDR_Roughness || MDR == MDR_ColorRoughness || MDR == MDR_NormalRoughness || MDR == MDR_ColorNormalRoughness)
+		else if (!IsTranslucentBlendMode(BlendMode))
 		{
-			MaterialCompilationOutput.SetIsDBufferTextureUsed(2);
-			AddEstimatedTextureSample(1);
+			Errorf(TEXT("Only translucent materials can use the scene color node."));
 		}
+	}
+
+	if (IsAlphaHoldoutBlendMode(BlendMode) && !MaterialShadingModels.IsUnlit())
+	{
+		Errorf(TEXT("Alpha Holdout blend mode must use unlit shading model."));
+	}
+
+	if (Domain == MD_Volume && BlendMode != BLEND_Additive)
+	{
+		Errorf(TEXT("Volume materials must use an Additive blend mode."));
+	}
+	if (Domain == MD_Volume && Material->IsUsedWithSkeletalMesh())
+	{
+		Errorf(TEXT("Volume materials are not compatible with skinned meshes: they are voxelised as boxes anyway. Please disable UsedWithSkeletalMesh on the material."));
+	}
+
+	if (Material->IsLightFunction() && !IsOpaqueBlendMode(BlendMode))
+	{
+		Errorf(TEXT("Light function materials must be opaque."));
+	}
+
+	if (Material->IsLightFunction() && MaterialShadingModels.IsLit())
+	{
+		Errorf(TEXT("Light function materials must use unlit."));
+	}
 
-		bOpacityPropertyIsUsed = IsMaterialPropertyUsed(MP_Opacity, Chunk[MP_Opacity], FLinearColor(1, 0, 0, 0), 1);
+	if (Domain == MD_PostProcess && MaterialShadingModels.IsLit())
+	{
+		Errorf(TEXT("Post process materials must use unlit."));
+	}
+
+	if (Material->AllowNegativeEmissiveColor() && MaterialShadingModels.IsLit())
+	{
+		Errorf(TEXT("Only unlit materials can output negative emissive color."));
+	}
 
-		bUsesCurvature = FeatureLevel == ERHIFeatureLevel::ES3_1 && 
-						((MaterialShadingModels.HasShadingModel(MSM_SubsurfaceProfile) && IsMaterialPropertyUsed(MP_CustomData0, Chunk[MP_CustomData0], FLinearColor(1, 0, 0, 0), 1))
-						|| (MaterialShadingModels.HasShadingModel(MSM_Eye) && bOpacityPropertyIsUsed));
+	if (Material->IsSky() && (!MaterialShadingModels.IsUnlit() || !(IsOpaqueOrMaskedBlendMode(BlendMode))))
+	{
+		Errorf(TEXT("Sky materials must be opaque or masked, and unlit. They are expected to completely replace the background."));
+	}
 
-		// If Substrate is enabled or if this is a Substrate material cooked/used in non-Substrate mode, 
-		// we disable this warning as Substrate supports 'colored transmittance only' mode (i.e, modulate).
-		if (!bSubstrateEnabled && !FrontMaterialExpr && IsModulateBlendMode(BlendMode) && MaterialShadingModels.IsLit() && !Material->IsDeferredDecal())
+	if (MaterialShadingModels.HasShadingModel(MSM_SingleLayerWater))
+	{
+		if (!IsOpaqueOrMaskedBlendMode(BlendMode))
 		{
-			Errorf(TEXT("Dynamically lit translucency is not supported for BLEND_Modulate materials."));
+			Errorf(TEXT("SingleLayerWater materials must be opaque or masked."));
 		}
-
-		if (Domain == MD_Surface)
+		if (!MaterialShadingModels.HasOnlyShadingModel(MSM_SingleLayerWater))
 		{
-			if (IsModulateBlendMode(BlendMode) && Material->IsTranslucencyAfterDOFEnabled() && !RHISupportsDualSourceBlending(Platform))
-			{
-				Errorf(TEXT("Translucency after DOF with BLEND_Modulate is only allowed on platforms that support dual-blending. Consider using BLEND_Translucent with black emissive"));
-			}
+			Errorf(TEXT("SingleLayerWater materials cannot be combined with other shading models.")); // Simply untested for now
 		}
 
-		// Don't allow opaque and masked materials to scene depth as the results are undefined
-		if (bUsesSceneDepth && Domain != MD_PostProcess && !IsTranslucentBlendMode(BlendMode))
+		if (Material->GetMaterialInterface() && !Material->GetMaterialInterface()->GetMaterial()->HasAnyExpressionsInMaterialAndFunctionsOfType<UMaterialExpressionSingleLayerWaterMaterialOutput>()
+			&& !bSubstrateEnabled
+			)
 		{
-			Errorf(TEXT("Only transparent or postprocess materials can read from scene depth."));
+			Errorf(TEXT("SingleLayerWater materials requires the use of SingleLayerWaterMaterial output node."));
 		}
+	}
 
-		if (bUsesSceneDepth)
+	if (MaterialShadingModels.HasShadingModel(MSM_ThinTranslucent))
+	{
+		if (!IsTranslucentOnlyBlendMode(BlendMode))
 		{
-			MaterialCompilationOutput.SetIsSceneTextureUsed(PPI_SceneDepth);
+			Errorf(TEXT("ThinTranslucent materials must be translucent."));
 		}
 
-		MaterialCompilationOutput.bUsesDistanceCullFade = bUsesDistanceCullFade;
+		const ETranslucencyLightingMode TranslucencyLightingMode = Material->GetTranslucencyLightingMode();
 
-		if (MaterialCompilationOutput.RequiresSceneColorCopy())
+		if (TranslucencyLightingMode != ETranslucencyLightingMode::TLM_SurfacePerPixelLighting)
 		{
-			if (Domain != MD_Surface)
-			{
-				Errorf(TEXT("Only 'surface' material domain can use the scene color node."));
-			}
-			else if (!IsTranslucentBlendMode(BlendMode))
-			{
-				Errorf(TEXT("Only translucent materials can use the scene color node."));
-			}
+			Errorf(TEXT("ThinTranslucent materials must use Surface Per Pixel Lighting (Translucency->LightingMode=Surface ForwardShading).\n"));
 		}
-
-		if (IsAlphaHoldoutBlendMode(BlendMode) && !MaterialShadingModels.IsUnlit())
+		if (!MaterialShadingModels.HasOnlyShadingModel(MSM_ThinTranslucent))
 		{
-			Errorf(TEXT("Alpha Holdout blend mode must use unlit shading model."));
+			Errorf(TEXT("ThinTranslucent materials cannot be combined with other shading models.")); // Simply untested for now
 		}
-
-		if (Domain == MD_Volume && BlendMode != BLEND_Additive)
+		if (Material->GetMaterialInterface() && !Material->GetMaterialInterface()->GetMaterial()->HasAnyExpressionsInMaterialAndFunctionsOfType<UMaterialExpressionThinTranslucentMaterialOutput>())
 		{
-			Errorf(TEXT("Volume materials must use an Additive blend mode."));
+			Errorf(TEXT("ThinTranslucent materials requires the use of ThinTranslucentMaterial output node."));
 		}
-		if (Domain == MD_Volume && Material->IsUsedWithSkeletalMesh())
+		if (Material->IsTranslucencyAfterMotionBlurEnabled())
 		{
-			Errorf(TEXT("Volume materials are not compatible with skinned meshes: they are voxelised as boxes anyway. Please disable UsedWithSkeletalMesh on the material."));
-		}
-
-		if (Material->IsLightFunction() && !IsOpaqueBlendMode(BlendMode))
-		{
-			Errorf(TEXT("Light function materials must be opaque."));
-		}
-
-		if (Material->IsLightFunction() && MaterialShadingModels.IsLit())
-		{
-			Errorf(TEXT("Light function materials must use unlit."));
+			// We don't currently have a separate translucency modulation pass for After Motion Blur
+			Errorf(TEXT("ThinTranslucent materials are not currently supported in the \"After Motion Blur\" translucency pass."));
 		}
+	}
 
-		if (Domain == MD_PostProcess && MaterialShadingModels.IsLit())
+	if (Material->GetMaterialInterface() && Material->GetMaterialInterface()->GetMaterial()->HasAnyExpressionsInMaterialAndFunctionsOfType<UMaterialExpressionAbsorptionMediumMaterialOutput>())
+	{
+		// User placed an AbsorptionMedium node, make sure default lit is being used
+		if (!MaterialShadingModels.HasShadingModel(MSM_DefaultLit))
 		{
-			Errorf(TEXT("Post process materials must use unlit."));
+			Errorf(TEXT("AbsorptionMedium custom output requires Default Lit shading model."));
 		}
-
-		if (Material->AllowNegativeEmissiveColor() && MaterialShadingModels.IsLit())
+		if (Material->GetRefractionMode() != RM_IndexOfRefraction)
 		{
-			Errorf(TEXT("Only unlit materials can output negative emissive color."));
+			Errorf(TEXT("AbsorptionMedium custom output requires \"Index of Refraction\" as the \"Refraction Mode\"."));
 		}
+	}
 
-		if (Material->IsSky() && (!MaterialShadingModels.IsUnlit() || !(IsOpaqueOrMaskedBlendMode(BlendMode))))
-		{
-			Errorf(TEXT("Sky materials must be opaque or masked, and unlit. They are expected to completely replace the background."));
-		}
+	if (Domain == MD_DeferredDecal && !(IsTranslucentOnlyBlendMode(BlendMode) || BlendMode == BLEND_AlphaComposite || IsModulateBlendMode(BlendMode)))
+	{
+		// We could make the change for the user but it would be confusing when going to DeferredDecal and back
+		// or we would have to pay a performance cost to make the change more transparently.
+		// The change saves performance as with translucency we don't need to test for MeshDecals in all opaque rendering passes
+		Errorf(TEXT("Material using the DeferredDecal domain can only use the BlendModes Translucent, AlphaComposite or Modulate"));
+	}
 
-		if (MaterialShadingModels.HasShadingModel(MSM_SingleLayerWater))
+	if (MaterialCompilationOutput.bNeedsSceneTextures)
+	{
+		if (Domain != MD_DeferredDecal && Domain != MD_PostProcess)
 		{
-			if (!IsOpaqueOrMaskedBlendMode(BlendMode))
+			if (!MaterialShadingModels.HasShadingModel(MSM_SingleLayerWater) && IsOpaqueOrMaskedBlendMode(BlendMode))
 			{
-				Errorf(TEXT("SingleLayerWater materials must be opaque or masked."));
+				// In opaque pass, none of the textures are available
+				Errorf(TEXT("SceneTexture expressions cannot be used in opaque materials except if used with the Single Layer Water shading model."));
 			}
-			if (!MaterialShadingModels.HasOnlyShadingModel(MSM_SingleLayerWater))
+			else if (bNeedsSceneTexturePostProcessInputs)
 			{
-				Errorf(TEXT("SingleLayerWater materials cannot be combined with other shading models.")); // Simply untested for now
-			}
-
-			if (Material->GetMaterialInterface() && !Material->GetMaterialInterface()->GetMaterial()->HasAnyExpressionsInMaterialAndFunctionsOfType<UMaterialExpressionSingleLayerWaterMaterialOutput>()
-				&& !bSubstrateEnabled
-				)
-			{
-				Errorf(TEXT("SingleLayerWater materials requires the use of SingleLayerWaterMaterial output node."));
+				Errorf(TEXT("SceneTexture expressions cannot use post process inputs or scene color in non post process domain materials"));
 			}
 		}
+	}
 
-		if (MaterialShadingModels.HasShadingModel(MSM_ThinTranslucent))
-		{
-			if (!IsTranslucentOnlyBlendMode(BlendMode))
-			{
-				Errorf(TEXT("ThinTranslucent materials must be translucent."));
-			}
-
-			const ETranslucencyLightingMode TranslucencyLightingMode = Material->GetTranslucencyLightingMode();
-
-			if (TranslucencyLightingMode != ETranslucencyLightingMode::TLM_SurfacePerPixelLighting)
-			{
-				Errorf(TEXT("ThinTranslucent materials must use Surface Per Pixel Lighting (Translucency->LightingMode=Surface ForwardShading).\n"));
-			}
-			if (!MaterialShadingModels.HasOnlyShadingModel(MSM_ThinTranslucent))
-			{
-				Errorf(TEXT("ThinTranslucent materials cannot be combined with other shading models.")); // Simply untested for now
-			}
-			if (Material->GetMaterialInterface() && !Material->GetMaterialInterface()->GetMaterial()->HasAnyExpressionsInMaterialAndFunctionsOfType<UMaterialExpressionThinTranslucentMaterialOutput>())
-			{
-				Errorf(TEXT("ThinTranslucent materials requires the use of ThinTranslucentMaterial output node."));
-			}
-			if (Material->IsTranslucencyAfterMotionBlurEnabled())
-			{
-				// We don't currently have a separate translucency modulation pass for After Motion Blur
-				Errorf(TEXT("ThinTranslucent materials are not currently supported in the \"After Motion Blur\" translucency pass."));
-			}
-		}
+	if (IsModulateBlendMode(BlendMode) && Material->IsTranslucencyAfterMotionBlurEnabled())
+	{
+		// We don't currently have a separate translucency modulation pass for After Motion Blur
+		Errorf(TEXT("Blend Mode \"Modulate\" materials are not currently supported in the \"After Motion Blur\" translucency pass."));
+	}
 
-		if (Material->GetMaterialInterface() && Material->GetMaterialInterface()->GetMaterial()->HasAnyExpressionsInMaterialAndFunctionsOfType<UMaterialExpressionAbsorptionMediumMaterialOutput>())
-		{
-			// User placed an AbsorptionMedium node, make sure default lit is being used
-			if (!MaterialShadingModels.HasShadingModel(MSM_DefaultLit))
-			{
-				Errorf(TEXT("AbsorptionMedium custom output requires Default Lit shading model."));
-			}
-			if (Material->GetRefractionMode() != RM_IndexOfRefraction)
-			{
-				Errorf(TEXT("AbsorptionMedium custom output requires \"Index of Refraction\" as the \"Refraction Mode\"."));
-			}
-		}
+	// Catch any modifications to NumUserTexCoords that will not seen by customized UVs
+	check(SavedNumUserTexCoords == GetNumUserTexCoords());
 
-		if (Domain == MD_DeferredDecal && !(IsTranslucentOnlyBlendMode(BlendMode) || BlendMode == BLEND_AlphaComposite || IsModulateBlendMode(BlendMode)))
-		{
-			// We could make the change for the user but it would be confusing when going to DeferredDecal and back
-			// or we would have to pay a performance cost to make the change more transparently.
-			// The change saves performance as with translucency we don't need to test for MeshDecals in all opaque rendering passes
-			Errorf(TEXT("Material using the DeferredDecal domain can only use the BlendModes Translucent, AlphaComposite or Modulate"));
-		}
+	FString InterpolatorsOffsetsDefinitionCode;
+	TBitArray<> FinalAllocatedCoords = GetVertexInterpolatorsOffsets(InterpolatorsOffsetsDefinitionCode);
 
-		if (MaterialCompilationOutput.bNeedsSceneTextures)
-		{
-			if (Domain != MD_DeferredDecal && Domain != MD_PostProcess)
-			{
-				if (!MaterialShadingModels.HasShadingModel(MSM_SingleLayerWater) && IsOpaqueOrMaskedBlendMode(BlendMode))
-				{
-					// In opaque pass, none of the textures are available
-					Errorf(TEXT("SceneTexture expressions cannot be used in opaque materials except if used with the Single Layer Water shading model."));
-				}
-				else if (bNeedsSceneTexturePostProcessInputs)
-				{
-					Errorf(TEXT("SceneTexture expressions cannot use post process inputs or scene color in non post process domain materials"));
-				}
-			}
-		}
+	// Finished compilation, verify final interpolator count restrictions
+	if (CurrentCustomVertexInterpolatorOffset > 0)
+	{
+		const int32 MaxNumScalars = 8 * 2;
+		const int32 TotalUsedScalars = FinalAllocatedCoords.FindLast(true) + 1;
 
-		if (IsModulateBlendMode(BlendMode) && Material->IsTranslucencyAfterMotionBlurEnabled())
+		if (TotalUsedScalars > MaxNumScalars)
 		{
-			// We don't currently have a separate translucency modulation pass for After Motion Blur
-			Errorf(TEXT("Blend Mode \"Modulate\" materials are not currently supported in the \"After Motion Blur\" translucency pass."));
+			Errorf(TEXT("Maximum number of custom vertex interpolators exceeded. (%i / %i scalar values) (TexCoord: %i scalars, Custom: %i scalars)"),
+				TotalUsedScalars, MaxNumScalars, GetNumUserTexCoords() * 2, CurrentCustomVertexInterpolatorOffset);
 		}
+	}
 
-		// Catch any modifications to NumUserTexCoords that will not seen by customized UVs
-		check(SavedNumUserTexCoords == GetNumUserTexCoords());
-
-		FString InterpolatorsOffsetsDefinitionCode;
-		TBitArray<> FinalAllocatedCoords = GetVertexInterpolatorsOffsets(InterpolatorsOffsetsDefinitionCode);
+	MaterialCompilationOutput.NumUsedUVScalars = GetNumUserTexCoords() * 2;
+	MaterialCompilationOutput.NumUsedCustomInterpolatorScalars = CurrentCustomVertexInterpolatorOffset;
 
-		// Finished compilation, verify final interpolator count restrictions
-		if (CurrentCustomVertexInterpolatorOffset > 0)
+	if (bEnableExecutionFlow)
+	{
+		if (VertexAttributesChunk != INDEX_NONE)
 		{
-			const int32 MaxNumScalars = 8 * 2;
-			const int32 TotalUsedScalars = FinalAllocatedCoords.FindLast(true) + 1;
+			LinkParentScopes(SharedPropertyCodeChunks[SF_Vertex]);
 
- 			if (TotalUsedScalars > MaxNumScalars)
-			{
-				Errorf(TEXT("Maximum number of custom vertex interpolators exceeded. (%i / %i scalar values) (TexCoord: %i scalars, Custom: %i scalars)"),
-					TotalUsedScalars, MaxNumScalars, GetNumUserTexCoords() * 2, CurrentCustomVertexInterpolatorOffset);
-			}
+			TSet<int32> EmittedChunks;
+			GetScopeCode(1, VertexAttributesChunk, SharedPropertyCodeChunks[SF_Vertex], EmittedChunks, TranslatedAttributesCodeChunks[SF_Vertex]);
 		}
 
-		MaterialCompilationOutput.NumUsedUVScalars = GetNumUserTexCoords() * 2;
-		MaterialCompilationOutput.NumUsedCustomInterpolatorScalars = CurrentCustomVertexInterpolatorOffset;
-
-		if (bEnableExecutionFlow)
+		if (PixelAttributesChunk != INDEX_NONE)
 		{
-			if (VertexAttributesChunk != INDEX_NONE)
-			{
-				LinkParentScopes(SharedPropertyCodeChunks[SF_Vertex]);
+			LinkParentScopes(SharedPropertyCodeChunks[SF_Pixel]);
 
-				TSet<int32> EmittedChunks;
-				GetScopeCode(1, VertexAttributesChunk, SharedPropertyCodeChunks[SF_Vertex], EmittedChunks, TranslatedAttributesCodeChunks[SF_Vertex]);
-			}
-
-			if (PixelAttributesChunk != INDEX_NONE)
-			{
-				LinkParentScopes(SharedPropertyCodeChunks[SF_Pixel]);
-
-				TSet<int32> EmittedChunks;
-				GetScopeCode(1, PixelAttributesChunk, SharedPropertyCodeChunks[SF_Pixel], EmittedChunks, TranslatedAttributesCodeChunks[SF_Pixel]);
-			}
+			TSet<int32> EmittedChunks;
+			GetScopeCode(1, PixelAttributesChunk, SharedPropertyCodeChunks[SF_Pixel], EmittedChunks, TranslatedAttributesCodeChunks[SF_Pixel]);
 		}
-		else
-		{
-			for (int32 VariationIter = 0; VariationIter < CompiledPDV_MAX; VariationIter++)
-			{
-				ECompiledPartialDerivativeVariation Variation = (ECompiledPartialDerivativeVariation)VariationIter;
-
-		// Do Normal Chunk first
+	}
+	else
+	{
+		for (int32 VariationIter = 0; VariationIter < CompiledPDV_MAX; VariationIter++)
 		{
-			GetFixedParameterCode(
-				0,
-				NormalCodeChunkEnd,
-				Chunk[MP_Normal],
-				SharedPropertyCodeChunks[NormalShaderFrequency],
-						DerivativeVariations[Variation].TranslatedCodeChunkDefinitions[MP_Normal],
-						DerivativeVariations[Variation].TranslatedCodeChunks[MP_Normal],
-						Variation);
-
-			// Always gather MP_Normal definitions as they can be shared by other properties
-			if (DerivativeVariations[Variation].TranslatedCodeChunkDefinitions[MP_Normal].IsEmpty())
-			{
-				DerivativeVariations[Variation].TranslatedCodeChunkDefinitions[MP_Normal] = GetDefinitions(SharedPropertyCodeChunks[NormalShaderFrequency], 0, NormalCodeChunkEnd, Variation);
-			}
-		}
+			ECompiledPartialDerivativeVariation Variation = (ECompiledPartialDerivativeVariation)VariationIter;
 
-		// Now the rest, skipping Normal
-		for (uint32 PropertyId = 0; PropertyId < MP_MAX; ++PropertyId)
-		{
-			if (PropertyId == MP_MaterialAttributes || PropertyId == MP_Normal || PropertyId == MP_CustomOutput)
+			// Do Normal Chunk first
 			{
-				continue;
+				GetFixedParameterCode(
+					0,
+					NormalCodeChunkEnd,
+					Chunk[MP_Normal],
+					SharedPropertyCodeChunks[NormalShaderFrequency],
+							DerivativeVariations[Variation].TranslatedCodeChunkDefinitions[MP_Normal],
+							DerivativeVariations[Variation].TranslatedCodeChunks[MP_Normal],
+							Variation);
+
+				// Always gather MP_Normal definitions as they can be shared by other properties
+				if (DerivativeVariations[Variation].TranslatedCodeChunkDefinitions[MP_Normal].IsEmpty())
+				{
+					DerivativeVariations[Variation].TranslatedCodeChunkDefinitions[MP_Normal] = GetDefinitions(SharedPropertyCodeChunks[NormalShaderFrequency], 0, NormalCodeChunkEnd, Variation);
+				}
 			}
 
-			const EShaderFrequency PropertyShaderFrequency = FMaterialAttributeDefinitionMap::GetShaderFrequency((EMaterialProperty)PropertyId);
-
-			int32 StartChunk = 0;
-			if (PropertyShaderFrequency == NormalShaderFrequency && SharedPixelProperties[PropertyId])
+			// Now the rest, skipping Normal
+			for (uint32 PropertyId = 0; PropertyId < MP_MAX; ++PropertyId)
 			{
-				// When processing shared properties, do not generate the code before the Normal was generated as those are already handled
-				StartChunk = NormalCodeChunkEnd;
-			}
+				if (PropertyId == MP_MaterialAttributes || PropertyId == MP_Normal || PropertyId == MP_CustomOutput)
+				{
+					continue;
+				}
 
-			// Reduce definition statements that don't contribute to the function's return value.
-			// @todo-lh: This should be expanded to a general reduction, but is currently only intended to fix an FXC internal compiler error reported in UE-117831
-			const bool bReduceAfterReturnValue = (PropertyId == MP_WorldPositionOffset || PropertyId == CompiledMP_PrevWorldPositionOffset || PropertyId == MP_Displacement);
+				const EShaderFrequency PropertyShaderFrequency = FMaterialAttributeDefinitionMap::GetShaderFrequency((EMaterialProperty)PropertyId);
 
-			if (bSubstrateEnabled && PropertyId >= MP_FrontMaterial && PropertyShaderFrequency == FrontMaterialShaderFrequency)
-			{
-				int32 PropertyIdCodeChunk = Chunk[PropertyId];
-				bool bPropertyIsConstant = false;
-				if (PropertyIdCodeChunk != INDEX_NONE)
+				int32 StartChunk = 0;
+				if (PropertyShaderFrequency == NormalShaderFrequency && SharedPixelProperties[PropertyId])
 				{
-					TArray<FShaderCodeChunk>& CodeChunks = SharedPropertyCodeChunks[PropertyShaderFrequency];
-					bPropertyIsConstant = CodeChunks[PropertyIdCodeChunk].UniformExpression && CodeChunks[PropertyIdCodeChunk].UniformExpression->IsConstant();
+					// When processing shared properties, do not generate the code before the Normal was generated as those are already handled
+					StartChunk = NormalCodeChunkEnd;
 				}
 
-				// We must have all properties generated the same way since it is the last one that is used to append to MaterialTemplate.ush
-				FString Temp;
-
-				// Initialise definitions to TranslatedCodeChunkDefinitions or constant value directly assigned to PixelMaterialInputs to TranslatedCodeChunks.
-				GetFixedParameterCode(
-					StartChunk,
-					FullySimplifiedFrontMaterialCodeChunkStart,
-					Chunk[PropertyId],
-					SharedPropertyCodeChunks[PropertyShaderFrequency],
-					DerivativeVariations[Variation].TranslatedCodeChunkDefinitions[PropertyId],
-					DerivativeVariations[Variation].TranslatedCodeChunks[PropertyId],
-					Variation,
-					bReduceAfterReturnValue);
+				// Reduce definition statements that don't contribute to the function's return value.
+				// @todo-lh: This should be expanded to a general reduction, but is currently only intended to fix an FXC internal compiler error reported in UE-117831
+				const bool bReduceAfterReturnValue = (PropertyId == MP_WorldPositionOffset || PropertyId == CompiledMP_PrevWorldPositionOffset || PropertyId == MP_Displacement);
 
-				// If bPropertyIsConstant, then this means that the property have not generated any code definition in TranslatedCodeChunkDefinitions 
-				// and instead will be assigned as a constant on PixelMaterialInputs using TranslatedCodeChunks.
-				// We need to not add any definitions such as SUBSTRATE_USE_FULLYSIMPLIFIED_MATERIAL to the TranslatedCodeChunkDefinitions string so we need to skip the remaining, now useless, operations.
-				// With that, TranslatedCodeChunkDefinitions[PropertyId] will be empty and this property will be skipped to not end up being the "LastProperty" representing the final code definition.
-				if (!bPropertyIsConstant)
+				if (bSubstrateEnabled && PropertyId >= MP_FrontMaterial && PropertyShaderFrequency == FrontMaterialShaderFrequency)
 				{
-					DerivativeVariations[Variation].TranslatedCodeChunkDefinitions[PropertyId] += TEXT("\t#if SUBSTRATE_USE_FULLYSIMPLIFIED_MATERIAL == 1\n");
+					int32 PropertyIdCodeChunk = Chunk[PropertyId];
+					bool bPropertyIsConstant = false;
+					if (PropertyIdCodeChunk != INDEX_NONE)
+					{
+						TArray<FShaderCodeChunk>& CodeChunks = SharedPropertyCodeChunks[PropertyShaderFrequency];
+						bPropertyIsConstant = CodeChunks[PropertyIdCodeChunk].UniformExpression && CodeChunks[PropertyIdCodeChunk].UniformExpression->IsConstant();
+					}
 
+					// We must have all properties generated the same way since it is the last one that is used to append to MaterialTemplate.ush
+					FString Temp;
+
+					// Initialise definitions to TranslatedCodeChunkDefinitions or constant value directly assigned to PixelMaterialInputs to TranslatedCodeChunks.
 					GetFixedParameterCode(
+						StartChunk,
 						FullySimplifiedFrontMaterialCodeChunkStart,
-						FullySimplifiedFrontMaterialCodeChunkEnd,
 						Chunk[PropertyId],
 						SharedPropertyCodeChunks[PropertyShaderFrequency],
-						Temp,
+						DerivativeVariations[Variation].TranslatedCodeChunkDefinitions[PropertyId],
 						DerivativeVariations[Variation].TranslatedCodeChunks[PropertyId],
 						Variation,
 						bReduceAfterReturnValue);
 
-					DerivativeVariations[Variation].TranslatedCodeChunkDefinitions[PropertyId] += Temp + TEXT("\t#endif // SUBSTRATE_USE_FULLYSIMPLIFIED_MATERIAL\n");
-
+					// If bPropertyIsConstant, then this means that the property have not generated any code definition in TranslatedCodeChunkDefinitions 
+					// and instead will be assigned as a constant on PixelMaterialInputs using TranslatedCodeChunks.
+					// We need to not add any definitions such as SUBSTRATE_USE_FULLYSIMPLIFIED_MATERIAL to the TranslatedCodeChunkDefinitions string so we need to skip the remaining, now useless, operations.
+					// With that, TranslatedCodeChunkDefinitions[PropertyId] will be empty and this property will be skipped to not end up being the "LastProperty" representing the final code definition.
+					if (!bPropertyIsConstant)
+					{
+						DerivativeVariations[Variation].TranslatedCodeChunkDefinitions[PropertyId] += TEXT("\t#if SUBSTRATE_USE_FULLYSIMPLIFIED_MATERIAL == 1\n");
+
+						GetFixedParameterCode(
+							FullySimplifiedFrontMaterialCodeChunkStart,
+							FullySimplifiedFrontMaterialCodeChunkEnd,
+							Chunk[PropertyId],
+							SharedPropertyCodeChunks[PropertyShaderFrequency],
+							Temp,
+							DerivativeVariations[Variation].TranslatedCodeChunks[PropertyId],
+							Variation,
+							bReduceAfterReturnValue);
+
+						DerivativeVariations[Variation].TranslatedCodeChunkDefinitions[PropertyId] += Temp + TEXT("\t#endif // SUBSTRATE_USE_FULLYSIMPLIFIED_MATERIAL\n");
+
+						GetFixedParameterCode(
+							FullySimplifiedFrontMaterialCodeChunkEnd,
+							SharedPropertyCodeChunks[PropertyShaderFrequency].Num(),
+							Chunk[PropertyId],
+							SharedPropertyCodeChunks[PropertyShaderFrequency],
+							Temp,
+							DerivativeVariations[Variation].TranslatedCodeChunks[PropertyId],
+							Variation,
+							bReduceAfterReturnValue);
+
+						DerivativeVariations[Variation].TranslatedCodeChunkDefinitions[PropertyId] += Temp;
+					}
+				}
+				else
+				{
 					GetFixedParameterCode(
-						FullySimplifiedFrontMaterialCodeChunkEnd,
+						StartChunk,
 						SharedPropertyCodeChunks[PropertyShaderFrequency].Num(),
 						Chunk[PropertyId],
 						SharedPropertyCodeChunks[PropertyShaderFrequency],
-						Temp,
+						DerivativeVariations[Variation].TranslatedCodeChunkDefinitions[PropertyId],
 						DerivativeVariations[Variation].TranslatedCodeChunks[PropertyId],
 						Variation,
 						bReduceAfterReturnValue);
-
-					DerivativeVariations[Variation].TranslatedCodeChunkDefinitions[PropertyId] += Temp;
 				}
 			}
-			else
+
+			// The code chunk corresponding to FullySimplifiedSubstrateFrontMaterialCodeChunk have already been written as part of MP_FrontMaterial.
+			// Here we get the FullySimplifiedSubstrateFrontMaterialTranslatedCodeChunks representing the variable storing the final fully simplified SubstrateData.
+			if (bSubstrateEnabled)
 			{
+				uint32 PropertyId = MP_FrontMaterial;
+
+				if (PropertyId == MP_MaterialAttributes || PropertyId == MP_Normal || PropertyId == MP_CustomOutput)
+				{
+					continue;
+				}
+
+				const EShaderFrequency PropertyShaderFrequency = FMaterialAttributeDefinitionMap::GetShaderFrequency((EMaterialProperty)PropertyId);
+
+				int32 StartChunk = 0;
+				if (PropertyShaderFrequency == NormalShaderFrequency && SharedPixelProperties[PropertyId])
+				{
+					// When processing shared properties, do not generate the code before the Normal was generated as those are already handled
+					StartChunk = NormalCodeChunkEnd;
+				}
+
+				// Reduce definition statements that don't contribute to the function's return value.
+				// @todo-lh: This should be expanded to a general reduction, but is currently only intended to fix an FXC internal compiler error reported in UE-117831
+				const bool bReduceAfterReturnValue = (PropertyId == MP_WorldPositionOffset || PropertyId == CompiledMP_PrevWorldPositionOffset || PropertyId == MP_Displacement);
+
 				GetFixedParameterCode(
 					StartChunk,
 					SharedPropertyCodeChunks[PropertyShaderFrequency].Num(),
-					Chunk[PropertyId],
+					FullySimplifiedSubstrateFrontMaterialCodeChunk,								//Chunk[PropertyId],
 					SharedPropertyCodeChunks[PropertyShaderFrequency],
-					DerivativeVariations[Variation].TranslatedCodeChunkDefinitions[PropertyId],
-					DerivativeVariations[Variation].TranslatedCodeChunks[PropertyId],
+					FullySimplifiedSubstrateFrontMaterialTranslatedCodeChunkDefinitions,
+					FullySimplifiedSubstrateFrontMaterialTranslatedCodeChunks,
 					Variation,
 					bReduceAfterReturnValue);
 			}
-		}
 
-		// The code chunk corresponding to FullySimplifiedSubstrateFrontMaterialCodeChunk have already been written as part of MP_FrontMaterial.
-		// Here we get the FullySimplifiedSubstrateFrontMaterialTranslatedCodeChunks representing the variable storing the final fully simplified SubstrateData.
-		if(bSubstrateEnabled)
-		{
-			uint32 PropertyId = MP_FrontMaterial;
-		
-			if (PropertyId == MP_MaterialAttributes || PropertyId == MP_Normal || PropertyId == MP_CustomOutput)
+			for (uint32 PropertyId = MP_MAX; PropertyId < CompiledMP_MAX; ++PropertyId)
 			{
-				continue;
-			}
-		
-			const EShaderFrequency PropertyShaderFrequency = FMaterialAttributeDefinitionMap::GetShaderFrequency((EMaterialProperty)PropertyId);
-		
-			int32 StartChunk = 0;
-			if (PropertyShaderFrequency == NormalShaderFrequency && SharedPixelProperties[PropertyId])
-			{
-				// When processing shared properties, do not generate the code before the Normal was generated as those are already handled
-				StartChunk = NormalCodeChunkEnd;
+						switch (PropertyId)
+				{
+				case CompiledMP_EmissiveColorCS:
+					if (bCompileForComputeShader)
+					{
+								GetFixedParameterCode(Chunk[PropertyId], SharedPropertyCodeChunks[SF_Compute], DerivativeVariations[Variation].TranslatedCodeChunkDefinitions[PropertyId], DerivativeVariations[Variation].TranslatedCodeChunks[PropertyId], Variation);
+					}
+					break;
+				case CompiledMP_PrevWorldPositionOffset:
+					{
+							GetFixedParameterCode(Chunk[PropertyId], SharedPropertyCodeChunks[SF_Vertex], DerivativeVariations[Variation].TranslatedCodeChunkDefinitions[PropertyId], DerivativeVariations[Variation].TranslatedCodeChunks[PropertyId], Variation);
+					}
+					break;
+				default: check(0);
+					break;
+				}
 			}
-		
-			// Reduce definition statements that don't contribute to the function's return value.
-			// @todo-lh: This should be expanded to a general reduction, but is currently only intended to fix an FXC internal compiler error reported in UE-117831
-			const bool bReduceAfterReturnValue = (PropertyId == MP_WorldPositionOffset || PropertyId == CompiledMP_PrevWorldPositionOffset || PropertyId == MP_Displacement);
-		
-			GetFixedParameterCode(
-				StartChunk,
-				SharedPropertyCodeChunks[PropertyShaderFrequency].Num(),
-				FullySimplifiedSubstrateFrontMaterialCodeChunk,								//Chunk[PropertyId],
-				SharedPropertyCodeChunks[PropertyShaderFrequency],
-				FullySimplifiedSubstrateFrontMaterialTranslatedCodeChunkDefinitions,
-				FullySimplifiedSubstrateFrontMaterialTranslatedCodeChunks,
-				Variation,
-				bReduceAfterReturnValue); 
-		}
 
-		for (uint32 PropertyId = MP_MAX; PropertyId < CompiledMP_MAX; ++PropertyId)
-		{
-					switch (PropertyId)
+			// Output the implementation for any custom output expressions
+					for (int32 ExpressionIndex = 0; ExpressionIndex < DerivativeVariations[Variation].CustomOutputImplementations.Num(); ExpressionIndex++)
 			{
-			case CompiledMP_EmissiveColorCS:
-			    if (bCompileForComputeShader)
-				{
-							GetFixedParameterCode(Chunk[PropertyId], SharedPropertyCodeChunks[SF_Compute], DerivativeVariations[Variation].TranslatedCodeChunkDefinitions[PropertyId], DerivativeVariations[Variation].TranslatedCodeChunks[PropertyId], Variation);
-				}
-				break;
-			case CompiledMP_PrevWorldPositionOffset:
-				{
-						GetFixedParameterCode(Chunk[PropertyId], SharedPropertyCodeChunks[SF_Vertex], DerivativeVariations[Variation].TranslatedCodeChunkDefinitions[PropertyId], DerivativeVariations[Variation].TranslatedCodeChunks[PropertyId], Variation);
+						ResourcesString += DerivativeVariations[Variation].CustomOutputImplementations[ExpressionIndex] + "\n\n";
+					}
 				}
-				break;
-			default: check(0);
-				break;
 			}
-		}
 
-		// Output the implementation for any custom output expressions
-				for (int32 ExpressionIndex = 0; ExpressionIndex < DerivativeVariations[Variation].CustomOutputImplementations.Num(); ExpressionIndex++)
+	// Store the number of float4s
+	MaterialCompilationOutput.UniformExpressionSet.UniformPreshaderBufferSize = (UniformPreshaderOffset + 3u) / 4u;
+
+	for (uint32 TypeIndex = 0u; TypeIndex < NumMaterialTextureParameterTypes; ++TypeIndex)
+	{
+		MaterialCompilationOutput.UniformExpressionSet.UniformTextureParameters[TypeIndex].Empty(UniformTextureExpressions[TypeIndex].Num());
+		for (FMaterialUniformExpressionTexture* TextureExpression : UniformTextureExpressions[TypeIndex])
 		{
-					ResourcesString += DerivativeVariations[Variation].CustomOutputImplementations[ExpressionIndex] + "\n\n";
-				}
-			}
+			TextureExpression->GetTextureParameterInfo(MaterialCompilationOutput.UniformExpressionSet.UniformTextureParameters[TypeIndex].AddDefaulted_GetRef());
 		}
+	}
+	MaterialCompilationOutput.UniformExpressionSet.UniformExternalTextureParameters.Empty(UniformExternalTextureExpressions.Num());
+	for (FMaterialUniformExpressionExternalTexture* TextureExpression : UniformExternalTextureExpressions)
+	{
+		TextureExpression->GetExternalTextureParameterInfo(MaterialCompilationOutput.UniformExpressionSet.UniformExternalTextureParameters.AddDefaulted_GetRef());
+	}
 
-		// Store the number of float4s
-		MaterialCompilationOutput.UniformExpressionSet.UniformPreshaderBufferSize = (UniformPreshaderOffset + 3u) / 4u;
+	for (uint32 SubstrateCompilationContextIndex = 0; SubstrateCompilationContextIndex < ESubstrateCompilationContext::SCC_MAX; ++SubstrateCompilationContextIndex)
+	{
+		FSubstrateCompilationContext& SubstrateCtx = SubstrateCompilationContext[SubstrateCompilationContextIndex];
 
-		for (uint32 TypeIndex = 0u; TypeIndex < NumMaterialTextureParameterTypes; ++TypeIndex)
+		FString TreeFunctionPostFix = TEXT("ERROR");
+		switch (SubstrateCompilationContextIndex)
 		{
-			MaterialCompilationOutput.UniformExpressionSet.UniformTextureParameters[TypeIndex].Empty(UniformTextureExpressions[TypeIndex].Num());
-			for (FMaterialUniformExpressionTexture* TextureExpression : UniformTextureExpressions[TypeIndex])
-			{
-				TextureExpression->GetTextureParameterInfo(MaterialCompilationOutput.UniformExpressionSet.UniformTextureParameters[TypeIndex].AddDefaulted_GetRef());
-			}
+		case ESubstrateCompilationContext::SCC_Default:
+			TreeFunctionPostFix = TEXT("");
+			break;
+		case ESubstrateCompilationContext::SCC_FullySimplified:
+			TreeFunctionPostFix = TEXT("_FullySimplified");
+			break;
+		default:
+			check(false);
 		}
-		MaterialCompilationOutput.UniformExpressionSet.UniformExternalTextureParameters.Empty(UniformExternalTextureExpressions.Num());
-		for (FMaterialUniformExpressionExternalTexture* TextureExpression : UniformExternalTextureExpressions)
+
+		const bool bSubstrateFrontMaterialProvided = Chunk[MP_FrontMaterial] != INDEX_NONE;
+		bool bSubstrateFrontMaterialIsValid = bSubstrateFrontMaterialProvided;
+		if (bSubstrateFrontMaterialIsValid)
 		{
-			TextureExpression->GetExternalTextureParameterInfo(MaterialCompilationOutput.UniformExpressionSet.UniformExternalTextureParameters.AddDefaulted_GetRef());
+			// The material can be null when some entries are automatically generated, for instance in the material layer blending system
+			bSubstrateFrontMaterialIsValid &= SubstrateCtx.SubstrateMaterialEffectiveClosureCount > 0;
 		}
 
-		for (uint32 SubstrateCompilationContextIndex = 0; SubstrateCompilationContextIndex < ESubstrateCompilationContext::SCC_MAX; ++SubstrateCompilationContextIndex)
+		if (bSubstrateFrontMaterialIsValid)
 		{
-			FSubstrateCompilationContext& SubstrateCtx = SubstrateCompilationContext[SubstrateCompilationContextIndex];
-
-			FString TreeFunctionPostFix = TEXT("ERROR");
-			switch (SubstrateCompilationContextIndex)
-			{
-			case ESubstrateCompilationContext::SCC_Default:
-				TreeFunctionPostFix = TEXT("");
-				break;
-			case ESubstrateCompilationContext::SCC_FullySimplified:
-				TreeFunctionPostFix = TEXT("_FullySimplified");
-				break;
-			default:
-				check(false);
-			}
+			bMaterialIsSubstrate = true;
 
-			const bool bSubstrateFrontMaterialProvided = Chunk[MP_FrontMaterial] != INDEX_NONE;
-			bool bSubstrateFrontMaterialIsValid = bSubstrateFrontMaterialProvided;
-			if (bSubstrateFrontMaterialIsValid)
+			if (SubstrateCtx.SubstrateMaterialRootOperator)
 			{
-				// The material can be null when some entries are automatically generated, for instance in the material layer blending system
-				bSubstrateFrontMaterialIsValid &= SubstrateCtx.SubstrateMaterialEffectiveClosureCount > 0;
-			}
-
-			if (bSubstrateFrontMaterialIsValid)
-			{
-				bMaterialIsSubstrate = true;
+				// Now implement the functions needed to process the material topology
 
-				if (SubstrateCtx.SubstrateMaterialRootOperator)
+				// Pre-Update the slab BSDF with operators (like thin film coating, which can alter F0/F90)
 				{
-					// Now implement the functions needed to process the material topology
+					// Update the coverage/transmittance of each node in the graph
+					check(SubstrateCtx.SubstrateMaterialRootOperator);
+					int32 RootMaximumDistanceToLeaves = SubstrateCtx.SubstrateMaterialRootOperator->MaxDistanceFromLeaves;
 
-					// Pre-Update the slab BSDF with operators (like thin film coating, which can alter F0/F90)
+					ResourcesString += FString::Printf(TEXT("void  FSubstratePixelHeader::PreUpdateAllBSDFWithBottomUpOperatorVisit%s(float3 V)\n"), *TreeFunctionPostFix);
+					ResourcesString += "{\n";
+					for (uint32 ClosureIndex = 0; ClosureIndex < SubstrateCtx.SubstrateMaterialEffectiveClosureCount; ++ClosureIndex)
 					{
-						// Update the coverage/transmittance of each node in the graph
-						check(SubstrateCtx.SubstrateMaterialRootOperator);
-						int32 RootMaximumDistanceToLeaves = SubstrateCtx.SubstrateMaterialRootOperator->MaxDistanceFromLeaves;
-
-						ResourcesString += FString::Printf(TEXT("void  FSubstratePixelHeader::PreUpdateAllBSDFWithBottomUpOperatorVisit%s(float3 V)\n"), *TreeFunctionPostFix);
-						ResourcesString += "{\n";
-						for (uint32 ClosureIndex = 0; ClosureIndex < SubstrateCtx.SubstrateMaterialEffectiveClosureCount; ++ClosureIndex)
+						ResourcesString += "\t{\n";
+						for (auto& It : SubstrateCtx.SubstrateMaterialExpressionRegisteredOperators)
 						{
-							ResourcesString += "\t{\n";
-							for (auto& It : SubstrateCtx.SubstrateMaterialExpressionRegisteredOperators)
+							if (!It.IsDiscarded() && It.BSDFIndex == ClosureIndex)
 							{
-								if (!It.IsDiscarded() && It.BSDFIndex == ClosureIndex)
+								// Walk up the graph to the root node and apply weight factors
+								std::function<void(const FSubstrateOperator&, int32)> WalkOperatorsUp = [&](const FSubstrateOperator& CurrentOperator, int32 PreviousOperatorIndex) -> void
 								{
-									// Walk up the graph to the root node and apply weight factors
-									std::function<void(const FSubstrateOperator&, int32)> WalkOperatorsUp = [&](const FSubstrateOperator& CurrentOperator, int32 PreviousOperatorIndex) -> void
+									switch (CurrentOperator.OperatorType)
+									{
+									case SUBSTRATE_OPERATOR_WEIGHT:
+									{
+										break; // NOP
+									}
+									case SUBSTRATE_OPERATOR_HORIZONTAL:
 									{
-										switch (CurrentOperator.OperatorType)
-										{
-										case SUBSTRATE_OPERATOR_WEIGHT:
-										{
-											break; // NOP
-										}
-										case SUBSTRATE_OPERATOR_HORIZONTAL:
-										{
-											break; // NOP
-										}
-										case SUBSTRATE_OPERATOR_VERTICAL:
-										{
-											// example ResourcesString += FString::Printf(TEXT("\t PreUpdateAllBSDFWithBottomUpOperatorVisit_Vertical(this, SubstrateTree, SubstrateTree.BSDFs[%d], V, %d /*Op index*/, %d /*PreviousIsInputA*/);\n"), BSDFIndex, CurrentOperator.Index, CurrentOperator.LeftIndex == PreviousOperatorIndex ? 1 : 0);
-											break; // NOP
-										}
-										case SUBSTRATE_OPERATOR_ADD:
-										{
-											break; // NOP
-										}
-										default:
-										case SUBSTRATE_OPERATOR_BSDF:
-										{
-											check(false);
-										}
-										}
-
-										if (CurrentOperator.ParentIndex != INDEX_NONE)
-										{
-											WalkOperatorsUp(SubstrateCtx.SubstrateMaterialExpressionRegisteredOperators[CurrentOperator.ParentIndex], CurrentOperator.Index);
-										}
-									};
-
-									const int32 BSDFOperatorIndex = It.Index;
-									const FSubstrateOperator& BSDFOperator = SubstrateCtx.SubstrateMaterialExpressionRegisteredOperators[BSDFOperatorIndex];
-
-									// Start visiting node up from the BSDF leaf only if it has a parent.
-									if (BSDFOperator.ParentIndex != INDEX_NONE)
+										break; // NOP
+									}
+									case SUBSTRATE_OPERATOR_VERTICAL:
+									{
+										// example ResourcesString += FString::Printf(TEXT("\t PreUpdateAllBSDFWithBottomUpOperatorVisit_Vertical(this, SubstrateTree, SubstrateTree.BSDFs[%d], V, %d /*Op index*/, %d /*PreviousIsInputA*/);\n"), BSDFIndex, CurrentOperator.Index, CurrentOperator.LeftIndex == PreviousOperatorIndex ? 1 : 0);
+										break; // NOP
+									}
+									case SUBSTRATE_OPERATOR_ADD:
+									{
+										break; // NOP
+									}
+									default:
+									case SUBSTRATE_OPERATOR_BSDF:
+									{
+										check(false);
+									}
+									}
+
+									if (CurrentOperator.ParentIndex != INDEX_NONE)
 									{
-										WalkOperatorsUp(SubstrateCtx.SubstrateMaterialExpressionRegisteredOperators[BSDFOperator.ParentIndex], BSDFOperator.Index);
+										WalkOperatorsUp(SubstrateCtx.SubstrateMaterialExpressionRegisteredOperators[CurrentOperator.ParentIndex], CurrentOperator.Index);
 									}
+								};
+
+								const int32 BSDFOperatorIndex = It.Index;
+								const FSubstrateOperator& BSDFOperator = SubstrateCtx.SubstrateMaterialExpressionRegisteredOperators[BSDFOperatorIndex];
+
+								// Start visiting node up from the BSDF leaf only if it has a parent.
+								if (BSDFOperator.ParentIndex != INDEX_NONE)
+								{
+									WalkOperatorsUp(SubstrateCtx.SubstrateMaterialExpressionRegisteredOperators[BSDFOperator.ParentIndex], BSDFOperator.Index);
 								}
 							}
-							ResourcesString += "\t}\n";
 						}
-						ResourcesString += "}\n";
+						ResourcesString += "\t}\n";
 					}
+					ResourcesString += "}\n";
+				}
 
-					// Update the coverage/transmittance of each leaves (==BSDFs) of the Substrate tree.
+				// Update the coverage/transmittance of each leaves (==BSDFs) of the Substrate tree.
+				{
+					ResourcesString += FString::Printf(TEXT("void FSubstratePixelHeader::UpdateAllBSDFsOperatorCoverageTransmittance%s(FSubstrateIntegrationSettings Settings, float3 V)\n"), *TreeFunctionPostFix);
+					ResourcesString += "{\n";
+					for (uint32 ClosureIndex = 0; ClosureIndex < SubstrateCtx.SubstrateMaterialEffectiveClosureCount; ++ClosureIndex)
 					{
-						ResourcesString += FString::Printf(TEXT("void FSubstratePixelHeader::UpdateAllBSDFsOperatorCoverageTransmittance%s(FSubstrateIntegrationSettings Settings, float3 V)\n"), *TreeFunctionPostFix);
-						ResourcesString += "{\n";
-						for (uint32 ClosureIndex = 0; ClosureIndex < SubstrateCtx.SubstrateMaterialEffectiveClosureCount; ++ClosureIndex)
-						{
-							ResourcesString += FString::Printf(TEXT("\t SubstrateTree.UpdateSingleBSDFOperatorCoverageTransmittance(this, %d, Settings, V);\n"), ClosureIndex);
-						}
-						ResourcesString += "}\n";
+						ResourcesString += FString::Printf(TEXT("\t SubstrateTree.UpdateSingleBSDFOperatorCoverageTransmittance(this, %d, Settings, V);\n"), ClosureIndex);
 					}
+					ResourcesString += "}\n";
+				}
 
-					// Propagate up the coverage/transmittance of each node in the Substrate tree.
-					// For that we visit all the operator according to their distance from the Substrate tree leaves, from small to large.
-					{
-						check(SubstrateCtx.SubstrateMaterialRootOperator);
-						int32 RootMaximumDistanceToLeaves = SubstrateCtx.SubstrateMaterialRootOperator->MaxDistanceFromLeaves;
+				// Propagate up the coverage/transmittance of each node in the Substrate tree.
+				// For that we visit all the operator according to their distance from the Substrate tree leaves, from small to large.
+				{
+					check(SubstrateCtx.SubstrateMaterialRootOperator);
+					int32 RootMaximumDistanceToLeaves = SubstrateCtx.SubstrateMaterialRootOperator->MaxDistanceFromLeaves;
 
-						ResourcesString += FString::Printf(TEXT("void FSubstratePixelHeader::UpdateAllOperatorsCoverageTransmittance%s()\n"), *TreeFunctionPostFix);
-						ResourcesString += "{\n";
-						for (int32 DistanceToLeaves = 1; DistanceToLeaves <= RootMaximumDistanceToLeaves; ++DistanceToLeaves)
+					ResourcesString += FString::Printf(TEXT("void FSubstratePixelHeader::UpdateAllOperatorsCoverageTransmittance%s()\n"), *TreeFunctionPostFix);
+					ResourcesString += "{\n";
+					for (int32 DistanceToLeaves = 1; DistanceToLeaves <= RootMaximumDistanceToLeaves; ++DistanceToLeaves)
+					{
+						ResourcesString += FString::Printf(TEXT("\t// MaxDistanceFromLeaves = %d \n"), DistanceToLeaves);
+						for (auto& It : SubstrateCtx.SubstrateMaterialExpressionRegisteredOperators)
 						{
-							ResourcesString += FString::Printf(TEXT("\t// MaxDistanceFromLeaves = %d \n"), DistanceToLeaves);
-							for (auto& It : SubstrateCtx.SubstrateMaterialExpressionRegisteredOperators)
+							if (!It.IsDiscarded() && It.MaxDistanceFromLeaves == DistanceToLeaves)
 							{
-								if (!It.IsDiscarded() && It.MaxDistanceFromLeaves == DistanceToLeaves)
-								{
-									ResourcesString += FString::Printf(TEXT("\t SubstrateTree.UpdateSingleOperatorCoverageTransmittance(%d /*operator index*/);\n"), It.Index);
-								}
+								ResourcesString += FString::Printf(TEXT("\t SubstrateTree.UpdateSingleOperatorCoverageTransmittance(%d /*operator index*/);\n"), It.Index);
 							}
 						}
-						ResourcesString += "}\n";
 					}
+					ResourcesString += "}\n";
+				}
 
-					// Update the luminance weight of each BSDF according to the operators it has to traverse bottom-up to the Substrate tree root node.
-					{
-						// Update the coverage/transmittance of each node in the graph
-						check(SubstrateCtx.SubstrateMaterialRootOperator);
-						int32 RootMaximumDistanceToLeaves = SubstrateCtx.SubstrateMaterialRootOperator->MaxDistanceFromLeaves;
+				// Update the luminance weight of each BSDF according to the operators it has to traverse bottom-up to the Substrate tree root node.
+				{
+					// Update the coverage/transmittance of each node in the graph
+					check(SubstrateCtx.SubstrateMaterialRootOperator);
+					int32 RootMaximumDistanceToLeaves = SubstrateCtx.SubstrateMaterialRootOperator->MaxDistanceFromLeaves;
 
-						ResourcesString += FString::Printf(TEXT("void FSubstratePixelHeader::UpdateAllBSDFWithBottomUpOperatorVisit%s()\n"), *TreeFunctionPostFix);
-						ResourcesString += "{\n";
-						for (uint32 ClosureIndex = 0; ClosureIndex < SubstrateCtx.SubstrateMaterialEffectiveClosureCount; ++ClosureIndex)
+					ResourcesString += FString::Printf(TEXT("void FSubstratePixelHeader::UpdateAllBSDFWithBottomUpOperatorVisit%s()\n"), *TreeFunctionPostFix);
+					ResourcesString += "{\n";
+					for (uint32 ClosureIndex = 0; ClosureIndex < SubstrateCtx.SubstrateMaterialEffectiveClosureCount; ++ClosureIndex)
+					{
+						for (auto& It : SubstrateCtx.SubstrateMaterialExpressionRegisteredOperators)
 						{
-							for (auto& It : SubstrateCtx.SubstrateMaterialExpressionRegisteredOperators)
+							if (!It.IsDiscarded() && It.BSDFIndex == ClosureIndex)
 							{
-								if (!It.IsDiscarded() && It.BSDFIndex == ClosureIndex)
+								// Walk up the graph to the root node and apply weight factors
+								std::function<void(const FSubstrateOperator&, int32)> WalkOperatorsUp = [&](const FSubstrateOperator& CurrentOperator, int32 PreviousOperatorIndex) -> void
 								{
-									// Walk up the graph to the root node and apply weight factors
-									std::function<void(const FSubstrateOperator&, int32)> WalkOperatorsUp = [&](const FSubstrateOperator& CurrentOperator, int32 PreviousOperatorIndex) -> void
+									switch (CurrentOperator.OperatorType)
+									{
+									case SUBSTRATE_OPERATOR_WEIGHT:
+									{
+										ResourcesString += FString::Printf(TEXT("\t SubstrateTree.UpdateAllBSDFWithBottomUpOperatorVisit_Weight(%d /*BSDFIndex*/, %d /*Op index*/, %d /*PreviousIsInputA*/);\n"), ClosureIndex, CurrentOperator.Index, 1);
+										break;
+									}
+									case SUBSTRATE_OPERATOR_HORIZONTAL:
 									{
-										switch (CurrentOperator.OperatorType)
-										{
-										case SUBSTRATE_OPERATOR_WEIGHT:
-										{
-											ResourcesString += FString::Printf(TEXT("\t SubstrateTree.UpdateAllBSDFWithBottomUpOperatorVisit_Weight(%d /*BSDFIndex*/, %d /*Op index*/, %d /*PreviousIsInputA*/);\n"), ClosureIndex, CurrentOperator.Index, 1);
-											break;
-										}
-										case SUBSTRATE_OPERATOR_HORIZONTAL:
-										{
-											ResourcesString += FString::Printf(TEXT("\t SubstrateTree.UpdateAllBSDFWithBottomUpOperatorVisit_Horizontal(%d /*BSDFIndex*/, %d /*Op index*/, %d /*PreviousIsInputA*/);\n"), ClosureIndex, CurrentOperator.Index, CurrentOperator.LeftIndex == PreviousOperatorIndex ? 1 : 0);
-											break;
-										}
-										case SUBSTRATE_OPERATOR_VERTICAL:
-										{
-											ResourcesString += FString::Printf(TEXT("\t SubstrateTree.UpdateAllBSDFWithBottomUpOperatorVisit_Vertical(%d /*BSDFIndex*/, %d /*Op index*/, %d /*PreviousIsInputA*/);\n"), ClosureIndex, CurrentOperator.Index, CurrentOperator.LeftIndex == PreviousOperatorIndex ? 1 : 0);
-											break;
-										}
-										case SUBSTRATE_OPERATOR_ADD:
-										{
-											break; // NOP
-										}
-										default:
-										case SUBSTRATE_OPERATOR_BSDF:
-										{
-											check(false);
-										}
-										}
-
-										if (CurrentOperator.ParentIndex != INDEX_NONE)
-										{
-											WalkOperatorsUp(SubstrateCtx.SubstrateMaterialExpressionRegisteredOperators[CurrentOperator.ParentIndex], CurrentOperator.Index);
-										}
-									};
-
-									const int32 BSDFOperatorIndex = It.Index;
-									const FSubstrateOperator& BSDFOperator = SubstrateCtx.SubstrateMaterialExpressionRegisteredOperators[BSDFOperatorIndex];
-
-									// Start visiting node up from the BSDF leaf only if it has a parent.
-									if (BSDFOperator.ParentIndex != INDEX_NONE)
+										ResourcesString += FString::Printf(TEXT("\t SubstrateTree.UpdateAllBSDFWithBottomUpOperatorVisit_Horizontal(%d /*BSDFIndex*/, %d /*Op index*/, %d /*PreviousIsInputA*/);\n"), ClosureIndex, CurrentOperator.Index, CurrentOperator.LeftIndex == PreviousOperatorIndex ? 1 : 0);
+										break;
+									}
+									case SUBSTRATE_OPERATOR_VERTICAL:
+									{
+										ResourcesString += FString::Printf(TEXT("\t SubstrateTree.UpdateAllBSDFWithBottomUpOperatorVisit_Vertical(%d /*BSDFIndex*/, %d /*Op index*/, %d /*PreviousIsInputA*/);\n"), ClosureIndex, CurrentOperator.Index, CurrentOperator.LeftIndex == PreviousOperatorIndex ? 1 : 0);
+										break;
+									}
+									case SUBSTRATE_OPERATOR_ADD:
+									{
+										break; // NOP
+									}
+									default:
+									case SUBSTRATE_OPERATOR_BSDF:
+									{
+										check(false);
+									}
+									}
+
+									if (CurrentOperator.ParentIndex != INDEX_NONE)
 									{
-										WalkOperatorsUp(SubstrateCtx.SubstrateMaterialExpressionRegisteredOperators[BSDFOperator.ParentIndex], BSDFOperator.Index);
+										WalkOperatorsUp(SubstrateCtx.SubstrateMaterialExpressionRegisteredOperators[CurrentOperator.ParentIndex], CurrentOperator.Index);
 									}
+								};
+
+								const int32 BSDFOperatorIndex = It.Index;
+								const FSubstrateOperator& BSDFOperator = SubstrateCtx.SubstrateMaterialExpressionRegisteredOperators[BSDFOperatorIndex];
+
+								// Start visiting node up from the BSDF leaf only if it has a parent.
+								if (BSDFOperator.ParentIndex != INDEX_NONE)
+								{
+									WalkOperatorsUp(SubstrateCtx.SubstrateMaterialExpressionRegisteredOperators[BSDFOperator.ParentIndex], BSDFOperator.Index);
 								}
 							}
 						}
-						ResourcesString += "}\n";
 					}
+					ResourcesString += "}\n";
 				}
+			}
 
-				// Check if normal/tangent basis are valid
+			// Check if normal/tangent basis are valid
+			{
+				SubstrateCtx.FinalUsedSharedLocalBasesCount = 0;
+				uint8 RequestedSharedLocalBasesCount = 0;
+				SubstrateCtx.SubstrateEvaluateSharedLocalBases(this, RequestedSharedLocalBasesCount, nullptr);
+				if (RequestedSharedLocalBasesCount > SUBSTRATE_MAX_SHAREDLOCALBASES_REGISTERS)
 				{
-					SubstrateCtx.FinalUsedSharedLocalBasesCount = 0;
-					uint8 RequestedSharedLocalBasesCount = 0;
-					SubstrateCtx.SubstrateEvaluateSharedLocalBases(this, RequestedSharedLocalBasesCount, nullptr);
-					if (RequestedSharedLocalBasesCount > SUBSTRATE_MAX_SHAREDLOCALBASES_REGISTERS)
-					{
-						Errorf(TEXT(" %s [%s]: Substrate - Material has more unique normal/tangent basis than the allowed limit %d/%d."), *Material->GetDebugName(), *Material->GetAssetPath().ToString(), RequestedSharedLocalBasesCount, SUBSTRATE_MAX_SHAREDLOCALBASES_REGISTERS);
-					}
+					Errorf(TEXT(" %s [%s]: Substrate - Material has more unique normal/tangent basis than the allowed limit %d/%d."), *Material->GetDebugName(), *Material->GetAssetPath().ToString(), RequestedSharedLocalBasesCount, SUBSTRATE_MAX_SHAREDLOCALBASES_REGISTERS);
 				}
 			}
-			else
-			{
-				MaterialCompilationOutput.SubstrateMaterialCompilationOutput.SubstrateMaterialDescription = "";
-				ResourcesString += "// No Substrate material provided\n";
-
-				// Adde default Substrate functions
-				ResourcesString += "#if TEMPLATE_USES_SUBSTRATE\n";
-				ResourcesString += "void PreUpdateAllBSDFWithBottomUpOperatorVisit(float3 V) {}\n";
-				ResourcesString += "void UpdateAllBSDFsOperatorCoverageTransmittance(FSubstrateIntegrationSettings Settings, float3 V) {}\n";
-				ResourcesString += "void UpdateAllOperatorsCoverageTransmittance(inout FSubstrateTree SubstrateTree) {}\n";
-				ResourcesString += "void UpdateAllBSDFWithBottomUpOperatorVisit(inout FSubstrateTree SubstrateTree) {}\n";
-				ResourcesString += "void PreUpdateAllBSDFWithBottomUpOperatorVisit_FullySimplified(float3 V) {}\n";
-				ResourcesString += "void UpdateAllBSDFsOperatorCoverageTransmittance_FullySimplified(FSubstrateIntegrationSettings Settings, float3 V) {}\n";
-				ResourcesString += "void UpdateAllOperatorsCoverageTransmittance_FullySimplified(inout FSubstrateTree SubstrateTree) {}\n";
-				ResourcesString += "void UpdateAllBSDFWithBottomUpOperatorVisit_FullySimplified(inout FSubstrateTree SubstrateTree) {}\n";
-				ResourcesString += "#endif\n";
-			}
 		}
+		else
+		{
+			MaterialCompilationOutput.SubstrateMaterialCompilationOutput.SubstrateMaterialDescription = "";
+			ResourcesString += "// No Substrate material provided\n";
 
-		MaterialCompilationOutput.UniformExpressionSet.SetParameterCollections(ParameterCollections);
-
-		// Store the number of unique VT samples
-		MaterialCompilationOutput.EstimatedNumVirtualTextureLookups = NumVtSamples;
+			// Adde default Substrate functions
+			ResourcesString += "#if TEMPLATE_USES_SUBSTRATE\n";
+			ResourcesString += "void PreUpdateAllBSDFWithBottomUpOperatorVisit(float3 V) {}\n";
+			ResourcesString += "void UpdateAllBSDFsOperatorCoverageTransmittance(FSubstrateIntegrationSettings Settings, float3 V) {}\n";
+			ResourcesString += "void UpdateAllOperatorsCoverageTransmittance(inout FSubstrateTree SubstrateTree) {}\n";
+			ResourcesString += "void UpdateAllBSDFWithBottomUpOperatorVisit(inout FSubstrateTree SubstrateTree) {}\n";
+			ResourcesString += "void PreUpdateAllBSDFWithBottomUpOperatorVisit_FullySimplified(float3 V) {}\n";
+			ResourcesString += "void UpdateAllBSDFsOperatorCoverageTransmittance_FullySimplified(FSubstrateIntegrationSettings Settings, float3 V) {}\n";
+			ResourcesString += "void UpdateAllOperatorsCoverageTransmittance_FullySimplified(inout FSubstrateTree SubstrateTree) {}\n";
+			ResourcesString += "void UpdateAllBSDFWithBottomUpOperatorVisit_FullySimplified(inout FSubstrateTree SubstrateTree) {}\n";
+			ResourcesString += "#endif\n";
+		}
 	}
 
-	ClearAllFunctionStacks();
-	
-#if STATS
-	GShaderCompilerStats->IncrementMaterialTranslateTime(HLSLTranslateTime);
-#endif // STATS
-	GShaderCompilerStats->IncrementMaterialsTranslated();
-	
-	INC_FLOAT_STAT_BY(STAT_ShaderCompiling_HLSLTranslation,(float)HLSLTranslateTime);
-
-#if ENABLE_COOK_STATS && STATS
-	// Write out a CSV file MaterialTranslationLog.txt containing info about all material translations.
-	FCsvLogFile::Get().AddEntry(Material->GetMaterialInterface()->GetFullName(), TranslationDateTime, HLSLTranslateTime);
-#endif
+	MaterialCompilationOutput.UniformExpressionSet.SetParameterCollections(ParameterCollections);
 
-	return bSuccess;
+	// Store the number of unique VT samples
+	MaterialCompilationOutput.EstimatedNumVirtualTextureLookups = NumVtSamples;
 }
 
 void FHLSLMaterialTranslator::ValidateShadingModelsForFeatureLevel(const FMaterialShadingModelField& ShadingModels)
@@ -2020,40 +2285,360 @@ void FHLSLMaterialTranslator::ValidateShadingModelsForFeatureLevel(const FMateri
 
 void FHLSLMaterialTranslator::GetMaterialEnvironment(EShaderPlatform InPlatform, FShaderCompilerEnvironment& OutEnvironment)
 {
+	// Use the old way of getting material environment if the setting is disabled.
+	if (!GUseMaterialTranslationResultsGrouping)
+	{
+		GetMaterialEnvironmentOld(InPlatform, OutEnvironment);
+		return;
+	}
+
+	check(InPlatform == GetShaderPlatform());
+
 	const bool bSubstrateEnabled = Substrate::IsSubstrateEnabled();
+	bool bMaterialRequestsDualSourceBlending = false;
 
-	if (bNeedsParticlePosition || Material->ShouldGenerateSphericalParticleNormals() || bUsesSphericalParticleOpacity)
+	if (EnvironmentDefines->bNeedsParticlePosition)
 	{
 		OutEnvironment.SetDefine(TEXT("NEEDS_PARTICLE_POSITION"), 1);
 	}
 
-	if (bNeedsParticleVelocity || Material->IsUsedWithNiagaraMeshParticles())
+	if (EnvironmentDefines->bNeedsParticleVelocity)
 	{
 		OutEnvironment.SetDefine(TEXT("NEEDS_PARTICLE_VELOCITY"), 1);
 	}
 
-	if (DynamicParticleParameterMask)
+	if (EnvironmentDefines->bUseDynamicParameters)
 	{
 		OutEnvironment.SetDefine(TEXT("USE_DYNAMIC_PARAMETERS"), 1);
-		OutEnvironment.SetDefine(TEXT("DYNAMIC_PARAMETERS_MASK"), DynamicParticleParameterMask);
+		OutEnvironment.SetDefine(TEXT("DYNAMIC_PARAMETERS_MASK"), EnvironmentDefines->DynamicParametersMask);
 	}
 
-	if (bNeedsParticleTime)
+	if (EnvironmentDefines->bNeedsParticleTime)
 	{
 		OutEnvironment.SetDefine(TEXT("NEEDS_PARTICLE_TIME"), 1);
 	}
 
-	if (bUsesParticleMotionBlur)
+	if (EnvironmentDefines->bUsesParticleMotionBlur)
 	{
 		OutEnvironment.SetDefine(TEXT("USES_PARTICLE_MOTION_BLUR"), 1);
 	}
 
-	if (bNeedsParticleRandom)
+	if (EnvironmentDefines->bNeedsParticleRandom)
 	{
 		OutEnvironment.SetDefine(TEXT("NEEDS_PARTICLE_RANDOM"), 1);
 	}
 
-	if (bUsesSphericalParticleOpacity)
+	if (EnvironmentDefines->bSphericalParticleOpacity)
+	{
+		OutEnvironment.SetDefine(TEXT("SPHERICAL_PARTICLE_OPACITY"), TEXT("1"));
+	}
+
+	if (EnvironmentDefines->bUseParticleSubUVs)
+	{
+		OutEnvironment.SetDefine(TEXT("USE_PARTICLE_SUBUVS"), TEXT("1"));
+	}
+
+	if (EnvironmentDefines->bLightmapUVAccess)
+	{
+		OutEnvironment.SetDefine(TEXT("LIGHTMAP_UV_ACCESS"), TEXT("1"));
+	}
+
+	if (EnvironmentDefines->bUsesAOMaterialMask)
+	{
+		OutEnvironment.SetDefine(TEXT("USES_AO_MATERIAL_MASK"), TEXT("1"));
+	}
+
+	if (EnvironmentDefines->bUsesSpeedtree)
+	{
+		OutEnvironment.SetDefine(TEXT("USES_SPEEDTREE"), TEXT("1"));
+	}
+
+	if (EnvironmentDefines->bNeedsWorldPositionExcludingShaderOffsets)
+	{
+		OutEnvironment.SetDefine(TEXT("NEEDS_WORLD_POSITION_EXCLUDING_SHADER_OFFSETS"), TEXT("1"));
+	}
+
+	if (EnvironmentDefines->bNeedsParticleSize)
+	{
+		OutEnvironment.SetDefine(TEXT("NEEDS_PARTICLE_SIZE"), TEXT("1"));
+	}
+
+	if (EnvironmentDefines->bNeedsParticleSpriteRotation)
+	{
+		OutEnvironment.SetDefine(TEXT("NEEDS_PARTICLE_SPRITE_ROTATION"), TEXT("1"));
+	}
+
+	if (EnvironmentDefines->bNeedsSceneTextures)
+	{
+		OutEnvironment.SetDefine(TEXT("NEEDS_SCENE_TEXTURES"), TEXT("1"));
+	}
+
+	if (EnvironmentDefines->bUsesEyeAdaptation)
+	{
+		OutEnvironment.SetDefine(TEXT("USES_EYE_ADAPTATION"), TEXT("1"));
+	}
+
+	if (EnvironmentDefines->bVirtualTextureOutput)
+	{
+		OutEnvironment.SetDefine(TEXT("VIRTUAL_TEXTURE_OUTPUT"), 1);
+	}
+
+	OutEnvironment.SetDefine(TEXT("USES_PER_INSTANCE_CUSTOM_DATA"), EnvironmentDefines->bUsesPerInstanceCustomData);
+	OutEnvironment.SetDefine(TEXT("USES_PER_INSTANCE_RANDOM"), EnvironmentDefines->bUsesPerInstanceRandom);
+	OutEnvironment.SetDefine(TEXT("USES_PER_INSTANCE_FADE_AMOUNT"), EnvironmentDefines->bUsesPerInstanceFadeAmount);
+	OutEnvironment.SetDefine(TEXT("USES_VERTEX_INTERPOLATOR"), EnvironmentDefines->bUsesVertexInterpolator);
+
+	OutEnvironment.SetDefine(TEXT("MATERIAL_SKY_ATMOSPHERE"), EnvironmentDefines->bUsesSkyAtmosphere);
+	OutEnvironment.SetDefine(TEXT("INTERPOLATE_VERTEX_COLOR"), EnvironmentDefines->bUsesVertexColor);
+	OutEnvironment.SetDefine(TEXT("NEEDS_PARTICLE_COLOR"), EnvironmentDefines->bUsesParticleColor);
+	OutEnvironment.SetDefine(TEXT("NEEDS_PARTICLE_LOCAL_TO_WORLD"), EnvironmentDefines->bUsesParticleLocalToWorld);
+	OutEnvironment.SetDefine(TEXT("NEEDS_PARTICLE_WORLD_TO_LOCAL"), EnvironmentDefines->bUsesParticleWorldToLocal);
+	OutEnvironment.SetDefine(TEXT("NEEDS_INSTANCE_LOCAL_TO_WORLD_PS"), EnvironmentDefines->bUsesInstanceLocalToWorldPS);
+	OutEnvironment.SetDefine(TEXT("NEEDS_INSTANCE_WORLD_TO_LOCAL_PS"), EnvironmentDefines->bUsesInstanceWorldToLocalPS);
+	OutEnvironment.SetDefine(TEXT("USES_TRANSFORM_VECTOR"), EnvironmentDefines->bUsesTransformVector);
+	OutEnvironment.SetDefine(TEXT("WANT_PIXEL_DEPTH_OFFSET"), EnvironmentDefines->bUsesPixelDepthOffset);
+
+	// we want USES_WORLD_POSITION_OFFSET to be readable as a bool compile argument, hence the != 0 comparison 
+	// (bUsesWorldPositionOffset is actually a 1-bit uint32 bitfield member)
+	OutEnvironment.SetDefineAndCompileArgument(TEXT("USES_WORLD_POSITION_OFFSET"), EnvironmentDefines->bUsesWorldPositionOffset);
+	OutEnvironment.SetDefineAndCompileArgument(TEXT("USES_DISPLACEMENT"), EnvironmentDefines->bUsesDisplacement);
+
+	OutEnvironment.SetDefine(TEXT("USES_EMISSIVE_COLOR"), EnvironmentDefines->bUsesEmissiveColor);
+	// Distortion uses tangent space transform 
+	OutEnvironment.SetDefine(TEXT("USES_DISTORTION"), EnvironmentDefines->bUsesDistortion);
+	OutEnvironment.SetDefine(TEXT("DISTORTION_ACCOUNT_FOR_COVERAGE"), EnvironmentDefines->bDistortionAccountForCoverage);
+
+	OutEnvironment.SetDefine(TEXT("MATERIAL_ENABLE_TRANSLUCENCY_FOGGING"), EnvironmentDefines->bMaterialEnableTranslucencyFogging);
+	OutEnvironment.SetDefine(TEXT("MATERIAL_ENABLE_TRANSLUCENCY_CLOUD_FOGGING"), EnvironmentDefines->bMaterialEnableTranslucencyCloudFogging);
+	OutEnvironment.SetDefine(TEXT("MATERIAL_IS_SKY"), EnvironmentDefines->bMaterialIsSky);
+	OutEnvironment.SetDefine(TEXT("MATERIAL_COMPUTE_FOG_PER_PIXEL"), EnvironmentDefines->bMaterialComputeFogPerPixel);
+	OutEnvironment.SetDefine(TEXT("MATERIAL_FULLY_ROUGH"), EnvironmentDefines->bMaterialFullyRough);
+
+	OutEnvironment.SetDefine(TEXT("MATERIAL_USES_ANISOTROPY"), EnvironmentDefines->bMaterialUsesAnisotropy);
+
+	OutEnvironment.SetDefine(TEXT("MATERIAL_DECAL_READ_MASK"), EnvironmentDefines->MaterialDecalReadMask);
+	OutEnvironment.SetDefine(TEXT("MATERIAL_USES_DECAL_LOOKUP"), EnvironmentDefines->bMaterialUsesDecalLookup);
+	OutEnvironment.SetDefine(TEXT("MATERIAL_PATH_TRACING_BUFFER_READ"), EnvironmentDefines->MaterialPathTracingBufferRead);
+	OutEnvironment.SetDefine(TEXT("MATERIAL_NEURAL_POST_PROCESS"), EnvironmentDefines->MaterialNeuralPostProcess);
+
+	OutEnvironment.SetDefine(TEXT("NUM_VIRTUALTEXTURE_SAMPLES"), EnvironmentDefines->NumVirtualTextureSamples);
+	OutEnvironment.SetDefine(TEXT("MATERIAL_VIRTUALTEXTURE_FEEDBACK"), EnvironmentDefines->bMaterialVirtualTextureFeedback);
+
+	// Setup defines to map each VT stack to either 1 or 2 page table textures, depending on how many layers it uses
+	for (int i = 0; i < EnvironmentDefines->VirtualPageTypes.Num(); ++i)
+	{
+		const FMaterialVirtualTextureStack& Stack = MaterialCompilationOutput.UniformExpressionSet.VTStacks[i];
+		FString PageTableValue = FString::Printf(TEXT("Material.VirtualTexturePageTable0_%d"), i);
+		if (EnvironmentDefines->VirtualPageTypes[i] & FEnvironmentDefines::TABLE1)
+		{
+			PageTableValue += FString::Printf(TEXT(", Material.VirtualTexturePageTable1_%d"), i);
+		}
+		if (EnvironmentDefines->VirtualPageTypes[i] & FEnvironmentDefines::TABLE_INDIRECTION)
+		{
+			PageTableValue += FString::Printf(TEXT(", Material.VirtualTexturePageTableIndirection_%d"), i);
+		}
+		OutEnvironment.SetDefine(*FString::Printf(TEXT("VIRTUALTEXTURE_PAGETABLE_%d"), i), *PageTableValue);
+	}
+
+	for (int32 CollectionIndex = 0; CollectionIndex < ParameterCollections.Num(); CollectionIndex++)
+	{
+		// Add uniform buffer declarations for any parameter collections referenced
+		const FString CollectionName = FString::Printf(TEXT("MaterialCollection%u"), CollectionIndex);
+		// This can potentially become an issue for MaterialCollection Uniform Buffers if they ever get non-numeric resources (eg Textures), as
+		// OutEnvironment.ResourceTableMap has a map by name, and the N ParameterCollection Uniform Buffers ALL are names "MaterialCollection"
+		// (and the hlsl cbuffers are named MaterialCollection0, etc, so the names don't match the layout)
+		FShaderUniformBufferParameter::ModifyCompilationEnvironment(*CollectionName, ParameterCollections[CollectionIndex]->GetUniformBufferStruct(), InPlatform, OutEnvironment);
+	}
+
+	OutEnvironment.SetDefine(TEXT("IS_MATERIAL_SHADER"), TEXT("1"));
+
+	if (EnvironmentDefines->bSingleLayerWaterShadingQuality)
+	{
+		// Value must match SINGLE_LAYER_WATER_SHADING_QUALITY_MOBILE_WITH_DEPTH_TEXTURE in SingleLayerWaterCommon.ush!
+		OutEnvironment.SetDefine(TEXT("SINGLE_LAYER_WATER_SHADING_QUALITY"), TEXT("1"));
+	}
+
+	if (EnvironmentDefines->bShadingModelsIsLit)
+	{
+		if (EnvironmentDefines->HasShadingModel(MSM_DefaultLit))
+		{
+			OutEnvironment.SetDefine(TEXT("MATERIAL_SHADINGMODEL_DEFAULT_LIT"), TEXT("1"));
+		}
+
+		if (EnvironmentDefines->HasShadingModel(MSM_Subsurface))
+		{
+			OutEnvironment.SetDefine(TEXT("MATERIAL_SHADINGMODEL_SUBSURFACE"), TEXT("1"));
+		}
+
+		if (EnvironmentDefines->HasShadingModel(MSM_PreintegratedSkin))
+		{
+			OutEnvironment.SetDefine(TEXT("MATERIAL_SHADINGMODEL_PREINTEGRATED_SKIN"), TEXT("1"));
+		}
+
+		if (EnvironmentDefines->HasShadingModel(MSM_SubsurfaceProfile))
+		{
+			OutEnvironment.SetDefine(TEXT("MATERIAL_SHADINGMODEL_SUBSURFACE_PROFILE"), TEXT("1"));
+		}
+
+		if (EnvironmentDefines->bMaterialSubsurfaceProfileUseCurvature)
+		{
+			OutEnvironment.SetDefine(TEXT("MATERIAL_SUBSURFACE_PROFILE_USE_CURVATURE"), TEXT("1"));
+		}
+
+		if (EnvironmentDefines->HasShadingModel(MSM_ClearCoat))
+		{
+			OutEnvironment.SetDefine(TEXT("MATERIAL_SHADINGMODEL_CLEAR_COAT"), TEXT("1"));
+		}
+
+		if (EnvironmentDefines->HasShadingModel(MSM_TwoSidedFoliage))
+		{
+			OutEnvironment.SetDefine(TEXT("MATERIAL_SHADINGMODEL_TWOSIDED_FOLIAGE"), TEXT("1"));
+		}
+
+		if (EnvironmentDefines->HasShadingModel(MSM_Hair))
+		{
+			OutEnvironment.SetDefine(TEXT("MATERIAL_SHADINGMODEL_HAIR"), TEXT("1"));
+		}
+
+		if (EnvironmentDefines->HasShadingModel(MSM_Cloth))
+		{
+			OutEnvironment.SetDefine(TEXT("MATERIAL_SHADINGMODEL_CLOTH"), TEXT("1"));
+		}
+
+		if (EnvironmentDefines->HasShadingModel(MSM_Eye))
+		{
+			OutEnvironment.SetDefine(TEXT("MATERIAL_SHADINGMODEL_EYE"), TEXT("1"));
+		}
+
+		if (EnvironmentDefines->bMaterialShadingModelEyeUseCurvature)
+		{
+			OutEnvironment.SetDefine(TEXT("MATERIAL_SHADINGMODEL_EYE_USE_CURVATURE"), TEXT("1"));
+		}
+
+		if (EnvironmentDefines->HasShadingModel(MSM_SingleLayerWater))
+		{
+			OutEnvironment.SetDefine(TEXT("MATERIAL_SHADINGMODEL_SINGLELAYERWATER"), TEXT("1"));
+		}
+
+		if (EnvironmentDefines->HasShadingModel(MSM_ThinTranslucent))
+		{
+			OutEnvironment.SetDefine(TEXT("MATERIAL_SHADINGMODEL_THIN_TRANSLUCENT"), TEXT("1"));
+			bMaterialRequestsDualSourceBlending = true;
+		}
+
+		if (EnvironmentDefines->bDisableForwardLocalLights)
+		{
+			OutEnvironment.SetDefine(TEXT("DISABLE_FORWARD_LOCAL_LIGHTS"), TEXT("1"));
+		}
+
+		if (EnvironmentDefines->bSingleLayerWaterSeparatedMainLight)
+		{
+			OutEnvironment.SetDefine(TEXT("SINGLE_LAYER_WATER_SEPARATED_MAIN_LIGHT"), TEXT("1"));
+		}
+
+		if (EnvironmentDefines->bMaterialSingleShadingModel)
+		{
+			OutEnvironment.SetDefine(TEXT("MATERIAL_SINGLE_SHADINGMODEL"), TEXT("1"));
+		}
+	}
+	else
+	{
+		OutEnvironment.SetDefine(TEXT("MATERIAL_SINGLE_SHADINGMODEL"), TEXT("1"));
+		OutEnvironment.SetDefine(TEXT("MATERIAL_SHADINGMODEL_UNLIT"), TEXT("1"));
+	}
+
+	if (EnvironmentDefines->bMaterialVolumetricAdvanced)
+	{
+		OutEnvironment.SetDefine(TEXT("MATERIAL_VOLUMETRIC_ADVANCED"), TEXT("1"));
+
+		if (EnvironmentDefines->bMaterialVolumetricAdvancedPhasePerSample)
+		{
+			OutEnvironment.SetDefine(TEXT("MATERIAL_VOLUMETRIC_ADVANCED_PHASE_PERSAMPLE"), TEXT("1"));
+		}
+		else
+		{
+			OutEnvironment.SetDefine(TEXT("MATERIAL_VOLUMETRIC_ADVANCED_PHASE_PERPIXEL"), TEXT("1"));
+		}
+
+		OutEnvironment.SetDefine(TEXT("MATERIAL_VOLUMETRIC_ADVANCED_GRAYSCALE_MATERIAL"), EnvironmentDefines->bMaterialVolumetricAdvancedGreyscaleMaterial);
+		OutEnvironment.SetDefine(TEXT("MATERIAL_VOLUMETRIC_ADVANCED_RAYMARCH_VOLUME_SHADOW"), EnvironmentDefines->bMaterialVolumetricAdvancedRaymarchVolumeShadow);
+		OutEnvironment.SetDefine(TEXT("MATERIAL_VOLUMETRIC_ADVANCED_CLAMP_MULTISCATTERING_CONTRIBUTION"), EnvironmentDefines->bMaterialVolumetricAdvancedClampMultiscatteringContribution);
+		OutEnvironment.SetDefine(TEXT("MATERIAL_VOLUMETRIC_ADVANCED_MULTISCATTERING_OCTAVE_COUNT"), EnvironmentDefines->MaterialVolumetricAdvancedMultiscatteringOctaveCount);
+		OutEnvironment.SetDefine(TEXT("MATERIAL_VOLUMETRIC_ADVANCED_CONSERVATIVE_DENSITY"), EnvironmentDefines->bMaterialVolumetricAdvancedConservativeDensity);
+		OutEnvironment.SetDefine(TEXT("MATERIAL_VOLUMETRIC_ADVANCED_OVERRIDE_AMBIENT_OCCLUSION"), EnvironmentDefines->bMaterialVolumetricAdvancedOverrideAmbientOcclusion);
+		OutEnvironment.SetDefine(TEXT("MATERIAL_VOLUMETRIC_ADVANCED_GROUND_CONTRIBUTION"), EnvironmentDefines->bMaterialVolumetricAdvancedAdvancedGroundContribution);
+	}
+
+	if (EnvironmentDefines->bMaterialVolumetricCloudEmptySpaceSkippingOutput)
+	{
+		OutEnvironment.SetDefine(TEXT("MATERIAL_VOLUMETRIC_CLOUD_EMPTY_SPACE_SKIPPING_OUTPUT"), TEXT("1"));
+	}
+
+	OutEnvironment.SetDefine(TEXT("MATERIAL_IS_STRATA"), EnvironmentDefines->bMaterialIsSubstrate);
+	OutEnvironment.SetDefine(TEXT("DUAL_SOURCE_COLOR_BLENDING_ENABLED"), EnvironmentDefines->bDualSourceColorBlendingEnabled);
+	OutEnvironment.SetDefine(TEXT("SUBSTRATE_PREMULTIPLIED_ALPHA_OPACITY_OVERRIDEN"), EnvironmentDefines->bSubstratePremultipliedAlphaOpacityOverridden);
+
+	if (EnvironmentDefines->bMaterialIsSubstrate)
+	{
+		OutEnvironment.SetDefine(TEXT("SUBSTRATE_USES_CONVERSION_FROM_LEGACY"), EnvironmentDefines->bSubstrateUsesConversionFromLegacy);
+		OutEnvironment.SetDefine(TEXT("SUBSTRATE_MATERIAL_OUTPUT_OPAQUE_ROUGH_REFRACTIONS"), EnvironmentDefines->bSubstrateMaterialOutputOpaqueRoughRefractions);
+		OutEnvironment.SetDefine(TEXT("SUBSTRATE_MATERIAL_EXPORT_TYPE"), EnvironmentDefines->SubstrateMaterialExportType);
+		OutEnvironment.SetDefine(TEXT("SUBSTRATE_MATERIAL_EXPORT_CONTEXT"), EnvironmentDefines->SubstrateMaterialExportContext);
+		OutEnvironment.SetDefine(TEXT("SUBSTRATE_MATERIAL_EXPORT_LEGACY_BLEND_MODE"), EnvironmentDefines->SubstrateMaterialExportLegacyBlendMode);
+		OutEnvironment.SetDefine(TEXT("SUBSTRATE_OPTIMIZED_UNLIT"), EnvironmentDefines->bSubstrateOptimizedUnlit);
+
+		{
+			OutEnvironment.SetDefine(TEXT("SUBSTRATE_SINGLEPATH"), EnvironmentDefines->bSubstrateSinglePath);
+			OutEnvironment.SetDefine(TEXT("SUBSTRATE_FASTPATH"), EnvironmentDefines->bSubstrateFastPath);
+			OutEnvironment.SetDefine(TEXT("SUBSTRATE_CLAMPED_CLOSURE_COUNT"), EnvironmentDefines->SubstrateClampedClosureCount);
+			OutEnvironment.SetDefine(TEXT("SUBSTRATE_COMPLEXSPECIALPATH"), EnvironmentDefines->bSubstrateComplexSpecialPath);
+		}
+	}
+
+	for (int i = 0; i < EnvironmentDefines->SubstrateDefines.Num(); ++i)
+	{
+		OutEnvironment.SetDefine(*EnvironmentDefines->SubstrateDefines[i].Get<0>(), EnvironmentDefines->SubstrateDefines[i].Get<1>());
+	}
+
+	OutEnvironment.SetDefine(TEXT("TEXTURE_SAMPLE_DEBUG"), EnvironmentDefines->bTextureSampleDebug);
+}
+
+void FHLSLMaterialTranslator::GetMaterialEnvironmentOld(EShaderPlatform InPlatform, FShaderCompilerEnvironment& OutEnvironment)
+{
+	const bool bSubstrateEnabled = Substrate::IsSubstrateEnabled();
+
+	if (bNeedsParticlePosition || Material->ShouldGenerateSphericalParticleNormals() || bUsesSphericalParticleOpacity)
+	{
+		OutEnvironment.SetDefine(TEXT("NEEDS_PARTICLE_POSITION"), 1);
+	}
+
+	if (bNeedsParticleVelocity || Material->IsUsedWithNiagaraMeshParticles())
+	{
+		OutEnvironment.SetDefine(TEXT("NEEDS_PARTICLE_VELOCITY"), 1);
+	}
+
+	if (DynamicParticleParameterMask)
+	{
+		OutEnvironment.SetDefine(TEXT("USE_DYNAMIC_PARAMETERS"), 1);
+		OutEnvironment.SetDefine(TEXT("DYNAMIC_PARAMETERS_MASK"), DynamicParticleParameterMask);
+	}
+
+	if (bNeedsParticleTime)
+	{
+		OutEnvironment.SetDefine(TEXT("NEEDS_PARTICLE_TIME"), 1);
+	}
+
+	if (bUsesParticleMotionBlur)
+	{
+		OutEnvironment.SetDefine(TEXT("USES_PARTICLE_MOTION_BLUR"), 1);
+	}
+
+	if (bNeedsParticleRandom)
+	{
+		OutEnvironment.SetDefine(TEXT("NEEDS_PARTICLE_RANDOM"), 1);
+	}
+
+	if (bUsesSphericalParticleOpacity)
 	{
 		OutEnvironment.SetDefine(TEXT("SPHERICAL_PARTICLE_OPACITY"), TEXT("1"));
 	}
@@ -2568,91 +3153,105 @@ TBitArray<> FHLSLMaterialTranslator::GetVertexInterpolatorsOffsets(FString& Vert
 
 void FHLSLMaterialTranslator::GetSharedInputsMaterialCode(FString& PixelMembersDeclaration, FString& NormalAssignment, FString& PixelMembersInitializationEpilog, ECompiledPartialDerivativeVariation DerivativeVariation)
 {
-	{
-		int32 LastProperty = -1;
-
-		FString PixelInputInitializerValues;
-		FString NormalInitializerValue;
+	int32 LastProperty = -1;
+	FString PixelInputInitializerValues;
+	FString NormalInitializerValue;
 
-		for (int32 PropertyIndex = 0; PropertyIndex < MP_MAX; ++PropertyIndex)
+	for (int32 PropertyIndex = 0; PropertyIndex < MP_MAX; ++PropertyIndex)
+	{
+		// Skip non-shared properties
+		if (!SharedPixelProperties[PropertyIndex])
 		{
-			// Skip non-shared properties
-			if (!SharedPixelProperties[PropertyIndex])
-			{
-				continue;
-			}
+			continue;
+		}
 
-			const EMaterialProperty Property = (EMaterialProperty)PropertyIndex;
-			check(FMaterialAttributeDefinitionMap::GetShaderFrequency(Property) == SF_Pixel);
-			// Special case MP_SubsurfaceColor as the actual property is a combination of the color and the profile but we don't want to expose the profile
-			const FString PropertyName = Property == MP_SubsurfaceColor ? "Subsurface" : FMaterialAttributeDefinitionMap::GetAttributeName(Property);
-			check(PropertyName.Len() > 0);				
-			const EMaterialValueType Type = Property == MP_SubsurfaceColor ? MCT_Float4 : FMaterialAttributeDefinitionMap::GetValueType(Property);
+		const EMaterialProperty Property = (EMaterialProperty)PropertyIndex;
+		check(FMaterialAttributeDefinitionMap::GetShaderFrequency(Property) == SF_Pixel);
+		// Special case MP_SubsurfaceColor as the actual property is a combination of the color and the profile but we don't want to expose the profile
+		const FString PropertyName = Property == MP_SubsurfaceColor ? "Subsurface" : FMaterialAttributeDefinitionMap::GetAttributeName(Property);
+		check(PropertyName.Len() > 0);				
+		const EMaterialValueType Type = Property == MP_SubsurfaceColor ? MCT_Float4 : FMaterialAttributeDefinitionMap::GetValueType(Property);
 
-			if (!bEnableExecutionFlow)
-			{
-			// Normal requires its own separate initializer
-			if (Property == MP_Normal)
-			{
-					NormalInitializerValue = FString::Printf(TEXT("\tPixelMaterialInputs.%s = %s;\n"), *PropertyName, *DerivativeVariations[DerivativeVariation].TranslatedCodeChunks[Property]);
-			}
-			else
-			{
-					if (DerivativeVariations[DerivativeVariation].TranslatedCodeChunkDefinitions[Property].Len() > 0)
-					{
-						LastProperty = Property;
-					}
+		if (!bEnableExecutionFlow)
+		{
+		// Normal requires its own separate initializer
+		if (Property == MP_Normal)
+		{
+				NormalInitializerValue = FString::Printf(TEXT("\tPixelMaterialInputs.%s = %s;\n"), *PropertyName, *DerivativeVariations[DerivativeVariation].TranslatedCodeChunks[Property]);
+		}
+		else
+		{
+				if (DerivativeVariations[DerivativeVariation].TranslatedCodeChunkDefinitions[Property].Len() > 0)
+				{
+					LastProperty = Property;
 				}
-
-				PixelInputInitializerValues += FString::Printf(TEXT("\tPixelMaterialInputs.%s = %s;\n"), *PropertyName, *DerivativeVariations[DerivativeVariation].TranslatedCodeChunks[Property]);
 			}
 
-			PixelMembersDeclaration += FString::Printf(TEXT("\t%s %s;\n"), HLSLTypeString(Type), *PropertyName);
+			PixelInputInitializerValues += FString::Printf(TEXT("\tPixelMaterialInputs.%s = %s;\n"), *PropertyName, *DerivativeVariations[DerivativeVariation].TranslatedCodeChunks[Property]);
 		}
 
-		NormalAssignment = NormalInitializerValue;
-		if (LastProperty != -1)
+		PixelMembersDeclaration += FString::Printf(TEXT("\t%s %s;\n"), HLSLTypeString(Type), *PropertyName);
+	}
+
+	NormalAssignment = NormalInitializerValue;
+	if (LastProperty != -1)
+	{
+		PixelMembersInitializationEpilog += DerivativeVariations[DerivativeVariation].TranslatedCodeChunkDefinitions[LastProperty] + TEXT("\n");
+	}
+
+	if (Substrate::IsSubstrateEnabled())
+	{
+		PixelMembersDeclaration += FString::Printf(TEXT("\t#if SUBSTRATE_USE_FULLYSIMPLIFIED_MATERIAL == 1\n"));
+		PixelMembersDeclaration += FString::Printf(TEXT("\t%s FullySimplifiedFrontMaterial;\n"), HLSLTypeString(FMaterialAttributeDefinitionMap::GetValueType(MP_FrontMaterial)));
+		PixelMembersDeclaration += FString::Printf(TEXT("\t#endif\n"));
+		if (!FullySimplifiedSubstrateFrontMaterialTranslatedCodeChunks.IsEmpty())
 		{
-			PixelMembersInitializationEpilog += DerivativeVariations[DerivativeVariation].TranslatedCodeChunkDefinitions[LastProperty] + TEXT("\n");
+			PixelInputInitializerValues += FString::Printf(TEXT("\t#if SUBSTRATE_USE_FULLYSIMPLIFIED_MATERIAL == 1\n"));
+			PixelInputInitializerValues += FString::Printf(TEXT("\tPixelMaterialInputs.FullySimplifiedFrontMaterial = %s;\n"), *FullySimplifiedSubstrateFrontMaterialTranslatedCodeChunks);
+			PixelInputInitializerValues += FString::Printf(TEXT("\t#endif\n"));
 		}
 
-		if (Substrate::IsSubstrateEnabled())
+		for (uint32 SubstrateCompilationContextIndex = 0; SubstrateCompilationContextIndex < ESubstrateCompilationContext::SCC_MAX; ++SubstrateCompilationContextIndex)
 		{
-			PixelMembersDeclaration += FString::Printf(TEXT("\t#if SUBSTRATE_USE_FULLYSIMPLIFIED_MATERIAL == 1\n"));
-			PixelMembersDeclaration += FString::Printf(TEXT("\t%s FullySimplifiedFrontMaterial;\n"), HLSLTypeString(FMaterialAttributeDefinitionMap::GetValueType(MP_FrontMaterial)));
-			PixelMembersDeclaration += FString::Printf(TEXT("\t#endif\n"));
-			if (!FullySimplifiedSubstrateFrontMaterialTranslatedCodeChunks.IsEmpty())
-			{
-				PixelInputInitializerValues += FString::Printf(TEXT("\t#if SUBSTRATE_USE_FULLYSIMPLIFIED_MATERIAL == 1\n"));
-				PixelInputInitializerValues += FString::Printf(TEXT("\tPixelMaterialInputs.FullySimplifiedFrontMaterial = %s;\n"), *FullySimplifiedSubstrateFrontMaterialTranslatedCodeChunks);
-				PixelInputInitializerValues += FString::Printf(TEXT("\t#endif\n"));
-			}
-
-			for (uint32 SubstrateCompilationContextIndex = 0; SubstrateCompilationContextIndex < ESubstrateCompilationContext::SCC_MAX; ++SubstrateCompilationContextIndex)
+			FSubstrateCompilationContext& SubstrateCtx = SubstrateCompilationContext[SubstrateCompilationContextIndex];
+			if (SubstrateCtx.CodeChunkToSubstrateSharedLocalBasis.Num() > 0)
 			{
-				FSubstrateCompilationContext& SubstrateCtx = SubstrateCompilationContext[SubstrateCompilationContextIndex];
-				if (SubstrateCtx.CodeChunkToSubstrateSharedLocalBasis.Num() > 0)
-				{
-					PixelInputInitializerValues += SubstrateCtx.SubstratePixelNormalInitializerValues;
-				}
+				PixelInputInitializerValues += SubstrateCtx.SubstratePixelNormalInitializerValues;
 			}
 		}
-
-		PixelMembersInitializationEpilog += PixelInputInitializerValues;
 	}
+
+	PixelMembersInitializationEpilog += PixelInputInitializerValues;
 }
 
 FString FHLSLMaterialTranslator::GetMaterialShaderCode()
-{	
-	// use "/Engine/Private/MaterialTemplate.ush" to create the functions to get data (e.g. material attributes) and code (e.g. material expressions to create specular color) from C++
-	TMap<FString, FString> MaterialSourceTemplateParams;
+{
+	if (!GUseMaterialTranslationResultsGrouping)
+	{
+		return GetMaterialShaderCodeOld();
+	}
+
+	// use "/Engine/Private/MaterialTemplate.ush" to create the functions to get data (e.g. material attributes) and code (e.g. material expressions to create specular color) from C++
+	int32 MaterialTemplateLineNumberNew;
+	FStringTemplateResolver Resolver = FMaterialSourceTemplate::Get().BeginResolve(GetShaderPlatform(), &MaterialTemplateLineNumberNew);
+
+	Resolver.SetParameterMap(&MaterialSourceTemplateParams);
+	MaterialSourceTemplateParams.Add({ TEXT("line_number"), FString::Printf(TEXT("%u"), MaterialTemplateLineNumberNew) });
+
+	return Resolver.Finalize();
+}
+
+FString FHLSLMaterialTranslator::GetMaterialShaderCodeOld()
+{	
+	// use "/Engine/Private/MaterialTemplate.ush" to create the functions to get data (e.g. material attributes) and code (e.g. material expressions to create specular color) from C++
+	TMap<FString, FString> Params;
 	int32 MaterialTemplateLineNumber;
 
 	FStringTemplateResolver Resolver = FMaterialSourceTemplate::Get().BeginResolve(GetShaderPlatform(), &MaterialTemplateLineNumber);
-	Resolver.SetParameterMap(&MaterialSourceTemplateParams);
+	Resolver.SetParameterMap(&Params);
 
-	MaterialSourceTemplateParams.Reserve(Resolver.GetTemplate().GetNumNamedParameters());
-	MaterialSourceTemplateParams.Add({ TEXT("line_number"), FString::Printf(TEXT("%u"), MaterialTemplateLineNumber) });
+	Params.Reserve(Resolver.GetTemplate().GetNumNamedParameters());
+	Params.Add({ TEXT("line_number"), FString::Printf(TEXT("%u"), MaterialTemplateLineNumber) });
 
 	// Assign slots to vertex interpolators
 	FString VertexInterpolatorsOffsetsDefinition;
@@ -2663,11 +3262,11 @@ FString FHLSLMaterialTranslator::GetMaterialShaderCode()
 	const uint32 NumCustomVectors = FMath::DivideAndRoundUp((uint32)CurrentCustomVertexInterpolatorOffset, 2u);
 	const uint32 NumTexCoordVectors = FinalAllocatedCoords.FindLast(true) + 1;
 
-	MaterialSourceTemplateParams.Add({ TEXT("num_material_texcoords_vertex"), FString::Printf(TEXT("%u"), NumUserVertexTexCoords) });
-	MaterialSourceTemplateParams.Add({ TEXT("num_material_texcoords"), FString::Printf(TEXT("%u"), NumUserTexCoords) });
-	MaterialSourceTemplateParams.Add({ TEXT("num_custom_vertex_interpolators"), FString::Printf(TEXT("%u"), NumCustomVectors) });
-	MaterialSourceTemplateParams.Add({ TEXT("num_tex_coord_interpolators"), FString::Printf(TEXT("%u"), NumTexCoordVectors) });
-	MaterialSourceTemplateParams.Add({ TEXT("vertex_interpolators_offsets_definition"), VertexInterpolatorsOffsetsDefinition });
+	Params.Add({ TEXT("num_material_texcoords_vertex"), FString::Printf(TEXT("%u"), NumUserVertexTexCoords) });
+	Params.Add({ TEXT("num_material_texcoords"), FString::Printf(TEXT("%u"), NumUserTexCoords) });
+	Params.Add({ TEXT("num_custom_vertex_interpolators"), FString::Printf(TEXT("%u"), NumCustomVectors) });
+	Params.Add({ TEXT("num_tex_coord_interpolators"), FString::Printf(TEXT("%u"), NumTexCoordVectors) });
+	Params.Add({ TEXT("vertex_interpolators_offsets_definition"), VertexInterpolatorsOffsetsDefinition });
 
 	FString MaterialAttributesDeclaration;
 	FString MaterialAttributesUtilities;
@@ -2714,8 +3313,8 @@ FString FHLSLMaterialTranslator::GetMaterialShaderCode()
 
 	MaterialAttributesDeclaration.Append(TEXT("};\n"));
 
-	MaterialSourceTemplateParams.Add({ TEXT("material_declarations"), MaterialAttributesDeclaration });
-	MaterialSourceTemplateParams.Add({ TEXT("material_attributes_utilities"), MaterialAttributesUtilities });
+	Params.Add({ TEXT("material_declarations"), MaterialAttributesDeclaration });
+	Params.Add({ TEXT("material_attributes_utilities"), MaterialAttributesUtilities });
 
 	// Stores the shared shader results member declarations
 	FString PixelMembersDeclaration[CompiledPDV_MAX];
@@ -2733,12 +3332,12 @@ FString FHLSLMaterialTranslator::GetMaterialShaderCode()
 	// PixelMembersDeclaration should be the same for all variations, but might change in the future. There are cases where work is shared
 	// between the pixel and vertex shader, but with Nanite all work has to be moved into the pixel shader, which means we will want
 	// different inputs. But for now, we are keeping them the same.
-	MaterialSourceTemplateParams.Add({ TEXT("pixel_material_inputs"), PixelMembersDeclaration[CompiledPDV_FiniteDifferences] });
+	Params.Add({ TEXT("pixel_material_inputs"), PixelMembersDeclaration[CompiledPDV_FiniteDifferences] });
 
 	{
 		FString DerivativeHelpers = DerivativeAutogen.GenerateUsedFunctions(*this);
 		FString DerivativeHelpersAndResources = DerivativeHelpers + ResourcesString;
-		MaterialSourceTemplateParams.Add({ TEXT("uniform_material_expressions"), DerivativeHelpersAndResources });
+		Params.Add({ TEXT("uniform_material_expressions"), DerivativeHelpersAndResources });
 	}
 
 	// Anything used bye the GenerationFunctionCode() like WorldPositionOffset shouldn't be using texures, right?
@@ -2748,38 +3347,38 @@ FString FHLSLMaterialTranslator::GetMaterialShaderCode()
 
 	if (bCompileForComputeShader)
 	{
-		MaterialSourceTemplateParams.Add({ TEXT("get_material_emissive_for_cs"), GenerateFunctionCode(CompiledMP_EmissiveColorCS, BaseDerivativeVariation) });
+		Params.Add({ TEXT("get_material_emissive_for_cs"), GenerateFunctionCode(CompiledMP_EmissiveColorCS, BaseDerivativeVariation) });
 	}
 	else
 	{
-		MaterialSourceTemplateParams.Add({ TEXT("get_material_emissive_for_cs"), TEXT("return 0") });
+		Params.Add({ TEXT("get_material_emissive_for_cs"), TEXT("return 0") });
 	}
 
-	MaterialSourceTemplateParams.Add({ TEXT("get_material_translucency_directional_lighting_intensity"), FString::Printf(TEXT("return %.5f"), Material->GetTranslucencyDirectionalLightingIntensity()) });
-	MaterialSourceTemplateParams.Add({ TEXT("get_material_translucent_shadow_density_scale"), FString::Printf(TEXT("return %.5f"), Material->GetTranslucentShadowDensityScale()) });
-	MaterialSourceTemplateParams.Add({ TEXT("get_material_translucent_self_shadow_density_scale"), FString::Printf(TEXT("return %.5f"), Material->GetTranslucentSelfShadowDensityScale()) });
-	MaterialSourceTemplateParams.Add({ TEXT("get_material_translucent_self_shadow_second_density_scale"), FString::Printf(TEXT("return %.5f"), Material->GetTranslucentSelfShadowSecondDensityScale()) });
-	MaterialSourceTemplateParams.Add({ TEXT("get_material_translucent_self_shadow_second_opacity"), FString::Printf(TEXT("return %.5f"), Material->GetTranslucentSelfShadowSecondOpacity()) });
-	MaterialSourceTemplateParams.Add({ TEXT("get_material_translucent_backscattering_exponent"), FString::Printf(TEXT("return %.5f"), Material->GetTranslucentBackscatteringExponent()) });
+	Params.Add({ TEXT("get_material_translucency_directional_lighting_intensity"), FString::Printf(TEXT("return %.5f"), Material->GetTranslucencyDirectionalLightingIntensity()) });
+	Params.Add({ TEXT("get_material_translucent_shadow_density_scale"), FString::Printf(TEXT("return %.5f"), Material->GetTranslucentShadowDensityScale()) });
+	Params.Add({ TEXT("get_material_translucent_self_shadow_density_scale"), FString::Printf(TEXT("return %.5f"), Material->GetTranslucentSelfShadowDensityScale()) });
+	Params.Add({ TEXT("get_material_translucent_self_shadow_second_density_scale"), FString::Printf(TEXT("return %.5f"), Material->GetTranslucentSelfShadowSecondDensityScale()) });
+	Params.Add({ TEXT("get_material_translucent_self_shadow_second_opacity"), FString::Printf(TEXT("return %.5f"), Material->GetTranslucentSelfShadowSecondOpacity()) });
+	Params.Add({ TEXT("get_material_translucent_backscattering_exponent"), FString::Printf(TEXT("return %.5f"), Material->GetTranslucentBackscatteringExponent()) });
 
 	{
 		FLinearColor Extinction = Material->GetTranslucentMultipleScatteringExtinction();
 
-		MaterialSourceTemplateParams.Add({ TEXT("get_material_translucent_multiple_scattering_extinction"), FString::Printf(TEXT("return MaterialFloat3(%.5f, %.5f, %.5f)"), Extinction.R, Extinction.G, Extinction.B) });
+		Params.Add({ TEXT("get_material_translucent_multiple_scattering_extinction"), FString::Printf(TEXT("return MaterialFloat3(%.5f, %.5f, %.5f)"), Extinction.R, Extinction.G, Extinction.B) });
 	}
 
-	MaterialSourceTemplateParams.Add({ TEXT("get_material_opacity_mask_clip_value"), FString::Printf(TEXT("return %.5f"), Material->GetOpacityMaskClipValue()) });
+	Params.Add({ TEXT("get_material_opacity_mask_clip_value"), FString::Printf(TEXT("return %.5f"), Material->GetOpacityMaskClipValue()) });
 
 	{
 		const FDisplacementScaling DisplacementScaling = Material->GetDisplacementScaling();
-		MaterialSourceTemplateParams.Add({ TEXT("get_material_displacement_magnitude"), FString::Printf(TEXT("return %.5f"), FMath::Max(0.0f, DisplacementScaling.Magnitude)) });
-		MaterialSourceTemplateParams.Add({ TEXT("get_material_displacement_center"), FString::Printf(TEXT("return %.5f"), FMath::Clamp(DisplacementScaling.Center, 0.0f, 1.0f)) });
+		Params.Add({ TEXT("get_material_displacement_magnitude"), FString::Printf(TEXT("return %.5f"), FMath::Max(0.0f, DisplacementScaling.Magnitude)) });
+		Params.Add({ TEXT("get_material_displacement_center"), FString::Printf(TEXT("return %.5f"), FMath::Clamp(DisplacementScaling.Center, 0.0f, 1.0f)) });
 	}
 
-	MaterialSourceTemplateParams.Add({ TEXT("get_material_world_position_offset_raw"), !bEnableExecutionFlow ? *GenerateFunctionCode(MP_WorldPositionOffset, BaseDerivativeVariation) : TEXT("return Parameters.MaterialAttributes.WorldPositionOffset") });
-	MaterialSourceTemplateParams.Add({ TEXT("get_material_previous_world_position_offset_raw"), !bEnableExecutionFlow ? *GenerateFunctionCode(CompiledMP_PrevWorldPositionOffset, BaseDerivativeVariation) : TEXT("return 0.0f") });
-	MaterialSourceTemplateParams.Add({ TEXT("get_material_custom_data0"), !bEnableExecutionFlow ? *GenerateFunctionCode(MP_CustomData0, BaseDerivativeVariation) : TEXT("return 0.0f") });
-	MaterialSourceTemplateParams.Add({ TEXT("get_material_custom_data1"), !bEnableExecutionFlow ? *GenerateFunctionCode(MP_CustomData1, BaseDerivativeVariation) : TEXT("return 0.0f") });
+	Params.Add({ TEXT("get_material_world_position_offset_raw"), !bEnableExecutionFlow ? *GenerateFunctionCode(MP_WorldPositionOffset, BaseDerivativeVariation) : TEXT("return Parameters.MaterialAttributes.WorldPositionOffset") });
+	Params.Add({ TEXT("get_material_previous_world_position_offset_raw"), !bEnableExecutionFlow ? *GenerateFunctionCode(CompiledMP_PrevWorldPositionOffset, BaseDerivativeVariation) : TEXT("return 0.0f") });
+	Params.Add({ TEXT("get_material_custom_data0"), !bEnableExecutionFlow ? *GenerateFunctionCode(MP_CustomData0, BaseDerivativeVariation) : TEXT("return 0.0f") });
+	Params.Add({ TEXT("get_material_custom_data1"), !bEnableExecutionFlow ? *GenerateFunctionCode(MP_CustomData1, BaseDerivativeVariation) : TEXT("return 0.0f") });
 
 	// Print custom texture coordinate assignments, should be fine with regular derivatives
 	FString CustomUVAssignments;
@@ -2808,7 +3407,7 @@ FString FHLSLMaterialTranslator::GetMaterialShaderCode()
 		}
 	}
 
-	MaterialSourceTemplateParams.Add({ TEXT("get_material_customized_u_vs"), CustomUVAssignments });
+	Params.Add({ TEXT("get_material_customized_u_vs"), CustomUVAssignments });
 
 	// Print custom vertex shader interpolator assignments
 	FString CustomInterpolatorAssignments;
@@ -2845,7 +3444,7 @@ FString FHLSLMaterialTranslator::GetMaterialShaderCode()
 		}
 	}
 
-	MaterialSourceTemplateParams.Add({ TEXT("get_custom_interpolators"), CustomInterpolatorAssignments });
+	Params.Add({ TEXT("get_custom_interpolators"), CustomInterpolatorAssignments });
 
 	if (bEnableExecutionFlow)
 	{
@@ -2861,7 +3460,7 @@ FString FHLSLMaterialTranslator::GetMaterialShaderCode()
 
 		EvaluateVertexCode.Append(TranslatedAttributesCodeChunks[SF_Vertex]);
 
-		MaterialSourceTemplateParams.Add({ TEXT("evaluate_material_attributes"), TranslatedAttributesCodeChunks[SF_Pixel] });
+		Params.Add({ TEXT("evaluate_material_attributes"), TranslatedAttributesCodeChunks[SF_Pixel] });
 
 		FString EvaluateMaterialAttributesCode = TEXT("    FMaterialAttributes MaterialAttributes = EvaluatePixelMaterialAttributes(Parameters);" HLSL_LINE_TERMINATOR);
 
@@ -2889,32 +3488,32 @@ FString FHLSLMaterialTranslator::GetMaterialShaderCode()
 			}
 		}
 
-		MaterialSourceTemplateParams.Add({ TEXT("calc_pixel_material_inputs_initial_calculations"), EvaluateMaterialAttributesCode });
-		MaterialSourceTemplateParams.Add({ TEXT("calc_pixel_material_inputs_normal"), TEXT("") });
-		MaterialSourceTemplateParams.Add({ TEXT("calc_pixel_material_inputs_other_inputs"), TEXT("") });
+		Params.Add({ TEXT("calc_pixel_material_inputs_initial_calculations"), EvaluateMaterialAttributesCode });
+		Params.Add({ TEXT("calc_pixel_material_inputs_normal"), TEXT("") });
+		Params.Add({ TEXT("calc_pixel_material_inputs_other_inputs"), TEXT("") });
 
-		MaterialSourceTemplateParams.Add({ TEXT("calc_pixel_material_inputs_analytic_derivatives_initial"), EvaluateMaterialAttributesCode });
-		MaterialSourceTemplateParams.Add({ TEXT("calc_pixel_material_inputs_analytic_derivatives_normal"), TEXT("") });
-		MaterialSourceTemplateParams.Add({ TEXT("calc_pixel_material_inputs_analytic_derivatives_other_inputs"), TEXT("") });
+		Params.Add({ TEXT("calc_pixel_material_inputs_analytic_derivatives_initial"), EvaluateMaterialAttributesCode });
+		Params.Add({ TEXT("calc_pixel_material_inputs_analytic_derivatives_normal"), TEXT("") });
+		Params.Add({ TEXT("calc_pixel_material_inputs_analytic_derivatives_other_inputs"), TEXT("") });
 	}
 	else
 	{
 		// skip material attributes code
-		MaterialSourceTemplateParams.Add({ TEXT("evaluate_material_attributes"), TEXT("") });
+		Params.Add({ TEXT("evaluate_material_attributes"), TEXT("") });
 
 		{
 			// Initializers required for Normal
-			MaterialSourceTemplateParams.Add({ TEXT("calc_pixel_material_inputs_initial_calculations"), DerivativeVariations[CompiledPDV_FiniteDifferences].TranslatedCodeChunkDefinitions[MP_Normal] });
-			MaterialSourceTemplateParams.Add({ TEXT("calc_pixel_material_inputs_normal"), NormalAssignment[CompiledPDV_FiniteDifferences] });
+			Params.Add({ TEXT("calc_pixel_material_inputs_initial_calculations"), DerivativeVariations[CompiledPDV_FiniteDifferences].TranslatedCodeChunkDefinitions[MP_Normal] });
+			Params.Add({ TEXT("calc_pixel_material_inputs_normal"), NormalAssignment[CompiledPDV_FiniteDifferences] });
 			// Finally the rest of common code followed by assignment into each input
-			MaterialSourceTemplateParams.Add({ TEXT("calc_pixel_material_inputs_other_inputs"), PixelMembersSetupAndAssignments[CompiledPDV_FiniteDifferences] });
+			Params.Add({ TEXT("calc_pixel_material_inputs_other_inputs"), PixelMembersSetupAndAssignments[CompiledPDV_FiniteDifferences] });
 		}
 		{
 			// Initializers required for Normal
-			MaterialSourceTemplateParams.Add({ TEXT("calc_pixel_material_inputs_analytic_derivatives_initial"), DerivativeVariations[CompiledPDV_Analytic].TranslatedCodeChunkDefinitions[MP_Normal] });
-			MaterialSourceTemplateParams.Add({ TEXT("calc_pixel_material_inputs_analytic_derivatives_normal"), NormalAssignment[CompiledPDV_Analytic] });
+			Params.Add({ TEXT("calc_pixel_material_inputs_analytic_derivatives_initial"), DerivativeVariations[CompiledPDV_Analytic].TranslatedCodeChunkDefinitions[MP_Normal] });
+			Params.Add({ TEXT("calc_pixel_material_inputs_analytic_derivatives_normal"), NormalAssignment[CompiledPDV_Analytic] });
 			// Finally the rest of common code followed by assignment into each input
-			MaterialSourceTemplateParams.Add({ TEXT("calc_pixel_material_inputs_analytic_derivatives_other_inputs"), PixelMembersSetupAndAssignments[CompiledPDV_Analytic] });
+			Params.Add({ TEXT("calc_pixel_material_inputs_analytic_derivatives_other_inputs"), PixelMembersSetupAndAssignments[CompiledPDV_Analytic] });
 		}
 	}
 
@@ -11719,6 +12318,132 @@ void FHLSLMaterialTranslator::FSubstrateCompilationContext::SubstrateEvaluateSha
 	}
 }
 
+/**
+ * This function is temporary. It's only been duplicated for validation purposes and it will be removed in a few days when the latest changes settle without problems.
+ */
+void FHLSLMaterialTranslator::FSubstrateCompilationContext::SubstrateEvaluateSharedLocalBasesNew(
+	FHLSLMaterialTranslator* Compiler,
+	uint8& OutRequestedSharedLocalBasesCount,
+	FEnvironmentDefines* OutEnvironment)
+{
+	/*
+	* The final output code/workflow for shared tangent basis should look like
+	*
+	* #define SHAREDLOCALBASIS_INDEX_0 0		// default, unused
+	* #define SHAREDLOCALBASIS_INDEX_1 0		// default, unused
+	* #define SHAREDLOCALBASIS_INDEX_2 0
+	*
+	* FSubstrateData BSDF0 = GetSubstrateSlabBSDF(... SHAREDLOCALBASIS_0, NormalCode0 ...)
+	* FSubstrateData BSDF1 = GetSubstrateSlabBSDF(... SHAREDLOCALBASIS_1, NormalCode1 ...)
+	*
+	* float3 NormalCode2 = lerp(NormalCode0, NormalCode1, mix)
+	* FSubstrateData BSDF2 = SubstrateHorizontalMixingParameterBlending(BSDF0, BSDF1, mix, NormalCode2, SHAREDLOCALBASIS_INDEX_2, SharedLocalBases.Types) // will internally create NormalCode2
+	*
+	* tParameters.SharedLocalBases.Normals[SHAREDLOCALBASIS_INDEX_2] = NormalCode2;
+	* #if MATERIAL_TANGENTSPACENORMAL
+	* Parameters.SharedLocalBases.Normals[SHAREDLOCALBASIS_INDEX_2] *= Parameters.TwoSidedSign;
+	* #endif
+	*/
+
+	FSubstrateRegisteredSharedLocalBasis UsedSharedLocalBasesInfo[SUBSTRATE_MAX_SHAREDLOCALBASES_REGISTERS];
+	FinalUsedSharedLocalBasesCount = 0;
+	OutRequestedSharedLocalBasesCount = 0;
+	const FString ParameterSharedLocalBasesName = GetParametersSharedLocalBasesName(CompilationContextIndex);
+
+	for (int32 OpIt = 0; OpIt < SubstrateMaterialExpressionRegisteredOperators.Num(); ++OpIt)
+	{
+		const FSubstrateOperator& BSDFOperator = SubstrateMaterialExpressionRegisteredOperators[OpIt];
+		if (BSDFOperator.BSDFIndex == INDEX_NONE || BSDFOperator.IsDiscarded())
+		{
+			continue;	// not a BSDF or if discarded (i.e. not the root of a parameter blending subtree), then there is no local basis to register
+		}
+
+		if (BSDFOperator.BSDFRegisteredSharedLocalBasis.NormalCodeChunk == INDEX_NONE && BSDFOperator.BSDFRegisteredSharedLocalBasis.TangentCodeChunk == INDEX_NONE)
+		{
+			continue;	// We skip null normal on certain BSDF, for instance unlit.
+		}
+		const FSubstrateSharedLocalBasesInfo& SubstrateSharedLocalBasesInfo = SubstrateCompilationInfoGetMatchingSharedLocalBasisInfo(BSDFOperator.BSDFRegisteredSharedLocalBasis);
+
+		// First, we check that the normal/tangent has not already written out (avoid 2 BSDFs sharing the same normal to note generate the same code twice)
+		bool bAlreadyProcessed = false;
+		for (uint8 i = 0; i < FinalUsedSharedLocalBasesCount; ++i)
+		{
+			if (UsedSharedLocalBasesInfo[i].NormalCodeChunkHash == SubstrateSharedLocalBasesInfo.SharedData.NormalCodeChunkHash &&
+				(UsedSharedLocalBasesInfo[i].TangentCodeChunkHash == SubstrateSharedLocalBasesInfo.SharedData.TangentCodeChunkHash || BSDFOperator.BSDFRegisteredSharedLocalBasis.TangentCodeChunk == INDEX_NONE))
+			{
+				bAlreadyProcessed = true;
+				break;
+			}
+		}
+		if (bAlreadyProcessed)
+		{
+			continue;
+		}
+
+		++OutRequestedSharedLocalBasesCount;
+		if (FinalUsedSharedLocalBasesCount >= SUBSTRATE_MAX_SHAREDLOCALBASES_REGISTERS)
+		{
+			continue;
+		}
+
+		const uint8 FinalSharedLocalBasisIndex = FinalUsedSharedLocalBasesCount++;
+		UsedSharedLocalBasesInfo[FinalSharedLocalBasisIndex] = SubstrateSharedLocalBasesInfo.SharedData;
+
+		if (CompilationContextIndex == SCC_FullySimplified)
+		{
+			SubstratePixelNormalInitializerValues += FString::Printf(TEXT("\t#if SUBSTRATE_USE_FULLYSIMPLIFIED_MATERIAL == 1\n"));
+		}
+
+		// Write out normals
+		SubstratePixelNormalInitializerValues += FString::Printf(TEXT("\tParameters.%s.Normals[%u] = %s;\n"), *ParameterSharedLocalBasesName, FinalSharedLocalBasisIndex, *SubstrateSharedLocalBasesInfo.NormalCode);
+		SubstratePixelNormalInitializerValues += FString::Printf(TEXT("\t#if MATERIAL_TANGENTSPACENORMAL\n"));
+		SubstratePixelNormalInitializerValues += FString::Printf(TEXT("\tParameters.%s.Normals[%u] *= Parameters.TwoSidedSign;\n"), *ParameterSharedLocalBasesName, FinalSharedLocalBasisIndex);
+		SubstratePixelNormalInitializerValues += FString::Printf(TEXT("\t#endif\n"));
+
+		// Write out tangents
+		if (SubstrateSharedLocalBasesInfo.SharedData.TangentCodeChunk != INDEX_NONE)
+		{
+			SubstratePixelNormalInitializerValues += FString::Printf(TEXT("\tParameters.%s.Tangents[%u] = %s;\n"), *ParameterSharedLocalBasesName, FinalSharedLocalBasisIndex, *SubstrateSharedLocalBasesInfo.TangentCode);
+		}
+		else
+		{
+			SubstratePixelNormalInitializerValues += FString::Printf(TEXT("\tParameters.%s.Tangents[%u] = Parameters.TangentToWorld[0];\n"), *ParameterSharedLocalBasesName, FinalSharedLocalBasisIndex);
+		}
+		SubstratePixelNormalInitializerValues += FString::Printf(TEXT("\t#if MATERIAL_TANGENTSPACENORMAL\n"));
+		SubstratePixelNormalInitializerValues += FString::Printf(TEXT("\tParameters.%s.Tangents[%u] *= Parameters.TwoSidedSign;\n"), *ParameterSharedLocalBasesName, FinalSharedLocalBasisIndex);
+		SubstratePixelNormalInitializerValues += FString::Printf(TEXT("\t#endif\n"));
+	}
+
+	SubstratePixelNormalInitializerValues += FString::Printf(TEXT("\tParameters.%s.Count = %u;\n"), *ParameterSharedLocalBasesName, FinalUsedSharedLocalBasesCount);
+
+	if (CompilationContextIndex == SCC_FullySimplified)
+	{
+		SubstratePixelNormalInitializerValues += FString::Printf(TEXT("\t#endif // SUBSTRATE_USE_FULLYSIMPLIFIED_MATERIAL\n"));
+	}
+
+	if (OutEnvironment)
+	{
+		// Now write out all the macros, them mapping from the BSDF to the effective position/index in the shared local basis array they should write to.
+		OutEnvironment->SubstrateDefines.Reserve(CodeChunkToSubstrateSharedLocalBasis.Num());
+		for (TMultiMap<uint64, FSubstrateSharedLocalBasesInfo>::TConstIterator It(CodeChunkToSubstrateSharedLocalBasis); It; ++It)
+		{
+			// The default linear output index will be 0 by default, and different if in fact the shared local basis points to one that is effectively in used in the array of shared local bases.
+			uint8 LinearIndex = 0;
+			for (uint8 i = 0; i < FinalUsedSharedLocalBasesCount; ++i)
+			{
+				if (UsedSharedLocalBasesInfo[i].NormalCodeChunkHash == It->Value.SharedData.NormalCodeChunkHash &&
+					(UsedSharedLocalBasesInfo[i].TangentCodeChunkHash == It->Value.SharedData.TangentCodeChunkHash || It->Value.SharedData.TangentCodeChunk == INDEX_NONE))
+				{
+					LinearIndex = i;
+					break;
+				}
+			}
+
+			OutEnvironment->SubstrateDefines.Emplace(*Compiler->GetSubstrateSharedLocalBasisIndexMacroInner(It->Value.SharedData, CompilationContextIndex), LinearIndex);
+		}
+	}
+}
+
 bool FHLSLMaterialTranslator::FSubstrateCompilationContext::SubstrateGenerateDerivedMaterialOperatorData(FHLSLMaterialTranslator* Compiler)
 {
 	FMaterial* CompilerMaterial = Compiler->Material;
@@ -14183,6 +14908,712 @@ bool FHLSLMaterialTranslator::IsDevelopmentFeatureEnabled(const FName& FeatureNa
 	return true;
 }
 
+void FHLSLMaterialTranslator::PrepareEnvironmentDefines()
+{
+	EShaderPlatform InPlatform = GetShaderPlatform();
+	*EnvironmentDefines = {};
+
+	const bool bSubstrateEnabled = Substrate::IsSubstrateEnabled();
+	bool bMaterialRequestsDualSourceBlending = false;
+
+	EnvironmentDefines->bNeedsParticlePosition = bNeedsParticlePosition || Material->ShouldGenerateSphericalParticleNormals() || bUsesSphericalParticleOpacity;
+	EnvironmentDefines->bNeedsParticleVelocity = bNeedsParticleVelocity || Material->IsUsedWithNiagaraMeshParticles();
+	EnvironmentDefines->bUseDynamicParameters = bool(DynamicParticleParameterMask);
+	EnvironmentDefines->DynamicParametersMask = DynamicParticleParameterMask;
+	EnvironmentDefines->bNeedsParticleTime = bNeedsParticleTime;
+	EnvironmentDefines->bUsesParticleMotionBlur = bUsesParticleMotionBlur;
+	EnvironmentDefines->bNeedsParticleRandom = bNeedsParticleRandom;
+	EnvironmentDefines->bSphericalParticleOpacity = bUsesSphericalParticleOpacity;
+	EnvironmentDefines->bUseParticleSubUVs = bUsesParticleSubUVs;
+	EnvironmentDefines->bLightmapUVAccess = bUsesLightmapUVs;
+	EnvironmentDefines->bUsesAOMaterialMask = bUsesAOMaterialMask;
+	EnvironmentDefines->bUsesSpeedtree = bUsesSpeedTree;
+	EnvironmentDefines->bNeedsWorldPositionExcludingShaderOffsets = bNeedsWorldPositionExcludingShaderOffsets;
+	EnvironmentDefines->bNeedsParticleSize = bNeedsParticleSize;
+	EnvironmentDefines->bNeedsParticleSpriteRotation = bNeedsParticleSpriteRotation;
+	EnvironmentDefines->bNeedsSceneTextures = MaterialCompilationOutput.bNeedsSceneTextures;
+	EnvironmentDefines->bUsesEyeAdaptation = MaterialCompilationOutput.bUsesEyeAdaptation;
+	EnvironmentDefines->bVirtualTextureOutput = MaterialCompilationOutput.bHasRuntimeVirtualTextureOutputNode;
+	EnvironmentDefines->bUsesPerInstanceCustomData = MaterialCompilationOutput.bUsesPerInstanceCustomData && Material->IsUsedWithInstancedStaticMeshes();
+	EnvironmentDefines->bUsesPerInstanceRandom = MaterialCompilationOutput.bUsesPerInstanceRandom && Material->IsUsedWithInstancedStaticMeshes();
+	EnvironmentDefines->bUsesPerInstanceFadeAmount = bUsesPerInstanceFadeAmount && Material->IsUsedWithInstancedStaticMeshes();
+	EnvironmentDefines->bUsesVertexInterpolator = MaterialCompilationOutput.bUsesVertexInterpolator;
+	EnvironmentDefines->bUsesSkyAtmosphere = bUsesSkyAtmosphere;
+	EnvironmentDefines->bUsesVertexColor = bUsesVertexColor;
+	EnvironmentDefines->bUsesParticleColor = bUsesParticleColor;
+	EnvironmentDefines->bUsesParticleLocalToWorld = bUsesParticleLocalToWorld;
+	EnvironmentDefines->bUsesParticleWorldToLocal = bUsesParticleWorldToLocal;
+	EnvironmentDefines->bUsesInstanceLocalToWorldPS = bUsesInstanceLocalToWorldPS;
+	EnvironmentDefines->bUsesInstanceWorldToLocalPS = bUsesInstanceWorldToLocalPS;
+	EnvironmentDefines->bUsesTransformVector = bUsesTransformVector;
+	EnvironmentDefines->bUsesPixelDepthOffset = bUsesPixelDepthOffset;
+	// we want USES_WORLD_POSITION_OFFSET to be readable as a bool compile argument, hence the != 0 comparison 
+	// (bUsesWorldPositionOffset is actually a 1-bit uint32 bitfield member)
+	EnvironmentDefines->bUsesWorldPositionOffset = (bUsesWorldPositionOffset != 0);
+	EnvironmentDefines->bUsesDisplacement = (bUsesDisplacement != 0);
+	EnvironmentDefines->bUsesEmissiveColor = bUsesEmissiveColor;
+	// Distortion uses tangent space transform 
+	EnvironmentDefines->bUsesDistortion = Material->IsDistorted();
+	EnvironmentDefines->bDistortionAccountForCoverage = bSubstrateEnabled && Material->GetRefractionCoverageMode() == RCM_CoverageAccountedFor ? 1 : 0;
+	EnvironmentDefines->bMaterialEnableTranslucencyFogging = Material->ShouldApplyFogging();
+	EnvironmentDefines->bMaterialEnableTranslucencyCloudFogging = Material->ShouldApplyCloudFogging();
+	EnvironmentDefines->bMaterialIsSky = Material->IsSky();
+	EnvironmentDefines->bMaterialComputeFogPerPixel = Material->ComputeFogPerPixel();
+	EnvironmentDefines->bMaterialFullyRough = bIsFullyRough || Material->IsFullyRough();
+	EnvironmentDefines->bMaterialUsesAnisotropy = bUsesAnisotropy && FDataDrivenShaderPlatformInfo::GetSupportsAnisotropicMaterials(InPlatform);
+	EnvironmentDefines->MaterialDecalReadMask = MaterialCompilationOutput.UsedDBufferTextures;
+	EnvironmentDefines->bMaterialUsesDecalLookup = MaterialCompilationOutput.bUsesDBufferTextureLookup;
+	EnvironmentDefines->MaterialPathTracingBufferRead = MaterialCompilationOutput.UsedPathTracingBufferTextures;
+	EnvironmentDefines->MaterialNeuralPostProcess = (MaterialCompilationOutput.bUsedWithNeuralNetworks || Material->IsUsedWithNeuralNetworks()) && Material->IsPostProcessMaterial();
+
+	// Count the number of VTStacks (each stack will allocate a feedback slot)
+	EnvironmentDefines->NumVirtualTextureSamples = VTStacks.Num();
+
+	// Check if any feedback slots are in use. We can simplify shader and remove EARLYZ optimizations if none are.
+	bool bGenerateFeedback = false;
+	for (int i = 0; i < VTStacks.Num() && !bGenerateFeedback; ++i)
+	{
+		bGenerateFeedback |= VTStacks[i].bGenerateFeedback;
+	}
+	EnvironmentDefines->bMaterialVirtualTextureFeedback = bGenerateFeedback;
+
+	// Setup defines to map each VT stack to either 1 or 2 page table textures, depending on how many layers it uses
+	EnvironmentDefines->VirtualPageTypes.SetNumUninitialized(VTStacks.Num());
+	for (int i = 0; i < VTStacks.Num(); ++i)
+	{
+		const FMaterialVirtualTextureStack& Stack = MaterialCompilationOutput.UniformExpressionSet.VTStacks[i];
+		EnvironmentDefines->VirtualPageTypes[i] = FEnvironmentDefines::TABLE0;
+		if (Stack.GetNumLayers() > 4u)
+		{
+			EnvironmentDefines->VirtualPageTypes[i] |= FEnvironmentDefines::TABLE1;
+		}
+		if (VTStacks[i].bAdaptive)
+		{
+			EnvironmentDefines->VirtualPageTypes[i] |= FEnvironmentDefines::TABLE_INDIRECTION;
+		}
+	}
+
+	// Set all the shading models for this material here 
+	FMaterialShadingModelField ShadingModels = Material->GetShadingModels();
+
+	// If the material gets its shading model from the material expressions, then we use the result from the compilation (assuming it's valid).
+	// This result will potentially be tighter than what GetShadingModels() returns, because it only picks up the shading models from the expressions that get compiled for a specific feature level and quality level
+	// For example, the material might have shading models behind static switches. GetShadingModels() will return both the true and the false paths from that switch, whereas the shading model field from the compilation will only contain the actual shading model selected 
+	if (Material->IsShadingModelFromMaterialExpression() && ShadingModelsFromCompilation.IsValid())
+	{
+		// Shading models fetched from the compilation of the expression graph
+		ShadingModels = ShadingModelsFromCompilation;
+	}
+
+	ensure(ShadingModels.IsValid());
+
+	// This is to have platforms use the simple single layer water shading similar to mobile: no dynamic lights, only sun and sky, no distortion, no colored transmittance on background, no custom depth read.
+	const bool bSingleLayerWaterUsesSimpleShading = FDataDrivenShaderPlatformInfo::GetWaterUsesSimpleForwardShading(InPlatform) && IsForwardShadingEnabled(InPlatform);
+	// Value must match SINGLE_LAYER_WATER_SHADING_QUALITY_MOBILE_WITH_DEPTH_TEXTURE in SingleLayerWaterCommon.ush!
+	EnvironmentDefines->bSingleLayerWaterShadingQuality = (ShadingModels.HasShadingModel(MSM_SingleLayerWater) && bSingleLayerWaterUsesSimpleShading);
+
+	if (ShadingModels.IsLit())
+	{
+		EnvironmentDefines->bShadingModelsIsLit = true;
+
+		int NumSetMaterials = 0;
+
+		for (int i = 0; i < MSM_NUM; ++i)
+		{
+			if (ShadingModels.HasShadingModel((EMaterialShadingModel)i))
+			{
+				EnvironmentDefines->MaterialShadingModelEnabled |= 1 << i;
+				++NumSetMaterials;
+			}
+		}
+
+		EnvironmentDefines->bMaterialSubsurfaceProfileUseCurvature = ShadingModels.HasShadingModel(MSM_SubsurfaceProfile) && bUsesCurvature;
+		EnvironmentDefines->bMaterialShadingModelEyeUseCurvature = ShadingModels.HasShadingModel(MSM_Eye) && bUsesCurvature;
+		EnvironmentDefines->bDisableForwardLocalLights = ShadingModels.HasShadingModel(MSM_SingleLayerWater) && FDataDrivenShaderPlatformInfo::GetRequiresDisableForwardLocalLights(Platform);
+
+		const bool bIsWaterDistanceFieldShadowEnabled = IsWaterDistanceFieldShadowEnabled(InPlatform);
+		const bool bIsWaterVSMFilteringEnabled = IsWaterVirtualShadowMapFilteringEnabled(InPlatform);
+		EnvironmentDefines->bSingleLayerWaterSeparatedMainLight = (ShadingModels.HasShadingModel(MSM_SingleLayerWater) && (bIsWaterDistanceFieldShadowEnabled || bIsWaterVSMFilteringEnabled));
+		EnvironmentDefines->bMaterialSingleShadingModel = (NumSetMaterials == 1);
+
+		ensure(NumSetMaterials != 0);
+		if (NumSetMaterials == 0)
+		{
+			// Should not really end up here
+			UE_LOG(LogMaterial, Warning, TEXT("Unknown material shading model(s). Setting to MSM_DefaultLit"));
+			EnvironmentDefines->MaterialShadingModelEnabled |= 1 << MSM_DefaultLit;
+		}
+	}
+	else
+	{
+		// Unlit shading model can only exist by itself
+		EnvironmentDefines->bMaterialSingleShadingModel = true;
+		EnvironmentDefines->MaterialShadingModelEnabled |= 1 << MSM_Unlit;
+	}
+
+	if (Material->GetMaterialDomain() == MD_Volume)
+	{
+		TArray<const UMaterialExpressionVolumetricAdvancedMaterialOutput*> VolumetricAdvancedExpressions;
+		Material->GetMaterialInterface()->GetMaterial()->GetAllExpressionsOfType(VolumetricAdvancedExpressions);
+		if (VolumetricAdvancedExpressions.Num() > 0)
+		{
+			if (VolumetricAdvancedExpressions.Num() > 1)
+			{
+				UE_LOG(LogMaterial, Fatal, TEXT("Only a single UMaterialExpressionVolumetricAdvancedMaterialOutput node is supported."));
+			}
+
+			EnvironmentDefines->bMaterialVolumetricAdvanced = true;
+
+			const UMaterialExpressionVolumetricAdvancedMaterialOutput* VolumetricAdvancedNode = VolumetricAdvancedExpressions[0];
+			EnvironmentDefines->bMaterialVolumetricAdvancedPhasePerSample = VolumetricAdvancedNode->GetEvaluatePhaseOncePerSample();
+
+			EnvironmentDefines->bMaterialVolumetricAdvancedGreyscaleMaterial = VolumetricAdvancedNode->bGrayScaleMaterial;
+			EnvironmentDefines->bMaterialVolumetricAdvancedRaymarchVolumeShadow = VolumetricAdvancedNode->bRayMarchVolumeShadow;
+			EnvironmentDefines->bMaterialVolumetricAdvancedClampMultiscatteringContribution = VolumetricAdvancedNode->bClampMultiScatteringContribution;
+			EnvironmentDefines->MaterialVolumetricAdvancedMultiscatteringOctaveCount = VolumetricAdvancedNode->GetMultiScatteringApproximationOctaveCount();
+			EnvironmentDefines->bMaterialVolumetricAdvancedConservativeDensity = VolumetricAdvancedNode->ConservativeDensity.IsConnected();
+			EnvironmentDefines->bMaterialVolumetricAdvancedOverrideAmbientOcclusion = Material->HasAmbientOcclusionConnected() || bSubstrateWritesAmbientOcclusion;
+			EnvironmentDefines->bMaterialVolumetricAdvancedAdvancedGroundContribution = VolumetricAdvancedNode->bGroundContribution;
+		}
+
+		TArray<const UMaterialExpressionVolumetricCloudEmptySpaceSkippingOutput*> EmptySpaceSkippingOutputExpressions;
+		Material->GetMaterialInterface()->GetMaterial()->GetAllExpressionsOfType(EmptySpaceSkippingOutputExpressions);
+		if (EmptySpaceSkippingOutputExpressions.Num() > 0)
+		{
+			if (VolumetricAdvancedExpressions.Num() > 1)
+			{
+				UE_LOG(LogMaterial, Fatal, TEXT("Only a single UMaterialExpressionEmptySpaceSkippingOutput node is supported."));
+			}
+			EnvironmentDefines->bMaterialVolumetricCloudEmptySpaceSkippingOutput = true;
+		}
+	}
+
+	EnvironmentDefines->bMaterialIsSubstrate = bMaterialIsSubstrate;
+
+	// Substrate requests dual source blending only for BLEND_TranslucentColoredTransmittance
+	bMaterialRequestsDualSourceBlending |= bMaterialIsSubstrate && Material->GetBlendMode() == EBlendMode::BLEND_TranslucentColoredTransmittance;
+
+	// if duals source blending (colored transmittance) is not supported on a platform, it will fall back to standard alpha blending (grey scale transmittance)
+	EnvironmentDefines->bDualSourceColorBlendingEnabled = Material->IsDualBlendingEnabled(Platform);
+	EnvironmentDefines->bSubstratePremultipliedAlphaOpacityOverridden = bMaterialIsSubstrate && bOpacityPropertyIsUsed;
+
+	if (bMaterialIsSubstrate)
+	{
+		EnvironmentDefines->bSubstrateUsesConversionFromLegacy = bSubstrateUsesConversionFromLegacy;
+		EnvironmentDefines->bSubstrateMaterialOutputOpaqueRoughRefractions = bSubstrateOutputsOpaqueRoughRefractions;
+
+		EnvironmentDefines->SubstrateMaterialExportType = (int32)GetSubstrateMaterialExportType();
+		EnvironmentDefines->SubstrateMaterialExportContext = (int32)GetSubstrateMaterialExportContext();
+		EnvironmentDefines->SubstrateMaterialExportLegacyBlendMode = (int32)GetSubstrateMaterialExportLegacyBlendMode();
+
+		// Unlit cannot be combined with other BSDF so we can simply pick the default strata context
+		EnvironmentDefines->bSubstrateOptimizedUnlit = SubstrateCompilationContext[ESubstrateCompilationContext::SCC_Default].bSubstrateMaterialIsUnlitNode;
+
+		{
+			// For now, the fully simplified mode is used for Lumen or anything else supported inlined evaluation. The export is only valid for the default case.
+			// STRATA_TODO: generate an export for the different context (need to generate two export functions: the default one and the FullSimplification one)
+			FSubstrateCompilationContext& SubstrateCtx = SubstrateCompilationContext[ESubstrateCompilationContext::SCC_Default];
+			EnvironmentDefines->bSubstrateSinglePath = SubstrateCtx.SubstrateMaterialComplexity.IsSingle();
+			EnvironmentDefines->bSubstrateFastPath = SubstrateCtx.SubstrateMaterialComplexity.IsSimple();
+			EnvironmentDefines->SubstrateClampedClosureCount = SubstrateCtx.SubstrateMaterialEffectiveClosureCount;
+			EnvironmentDefines->bSubstrateComplexSpecialPath = SubstrateCtx.SubstrateMaterialComplexity.IsComplexSpecial();
+		}
+
+
+		FString SubstrateMaterialDescription;
+		for (uint32 SubstrateCompilationContextIndex = 0; SubstrateCompilationContextIndex < ESubstrateCompilationContext::SCC_MAX; ++SubstrateCompilationContextIndex)
+		{
+			FSubstrateCompilationContext& SubstrateCtx = SubstrateCompilationContext[SubstrateCompilationContextIndex];
+			const FSubstrateSimplificationStatus& SubstrateSimplificationStatus = SubstrateCtx.SubstrateSimplificationStatus;
+
+			check(SubstrateCtx.SubstrateMaterialRootOperator);
+			uint32 RootMaximumDistanceToLeaves = SubstrateCtx.SubstrateMaterialRootOperator->MaxDistanceFromLeaves;
+
+			// Compute the shared local basis count and generate the hlsl shader code for it.
+			SubstrateCtx.SubstratePixelNormalInitializerValues = FString::Printf(TEXT("\n\n\n\t// Substrate normal and tangent\n"));
+			SubstrateCtx.FinalUsedSharedLocalBasesCount = 0;
+			uint8 RequestedSharedLocalBasesCount = 0;
+			SubstrateCtx.SubstrateEvaluateSharedLocalBasesNew(this, RequestedSharedLocalBasesCount, EnvironmentDefines.Get());
+
+#if WITH_EDITOR
+			// Now write some feedback to the user, but only produce debug string if in editor
+			{
+				// Output some debug info as comment in code and in the material stat window
+				const uint32 SubstrateBytePerPixel_Platform = Substrate::GetBytePerPixel(InPlatform);
+				const uint32 SubstrateClosurePerPixel_Platform = Substrate::GetClosurePerPixel(InPlatform);
+				FString SubstrateMaterialContextDescription;
+
+				auto GetSubstrateCompilationContextName = [](uint32 Index)
+				{
+					switch (Index)
+					{
+					case ESubstrateCompilationContext::SCC_Default:
+						return TEXT("Default");
+					case ESubstrateCompilationContext::SCC_FullySimplified:
+						return TEXT("FullySimplified");
+					default:
+						check(false);
+					}
+					return TEXT("ERROR");
+				};
+				FString SubstrateCompilationContextName = GetSubstrateCompilationContextName(SubstrateCompilationContextIndex);
+				SubstrateMaterialContextDescription += FString::Printf(TEXT("----- SUBSTRATE - %s -----\n"), *SubstrateCompilationContextName);
+				SubstrateMaterialContextDescription += FString::Printf(TEXT("SubstrateCompilationInfo -\n"));
+				SubstrateMaterialContextDescription += FString::Printf(TEXT(" - Byte Per Pixel Budget                           %u\n"), SubstrateBytePerPixel_Platform);
+				SubstrateMaterialContextDescription += FString::Printf(TEXT(" - Closure Per Pixel Budget                        %u\n"), SubstrateClosurePerPixel_Platform);
+				SubstrateMaterialContextDescription += FString::Printf(TEXT(" - Requested Byte Size before simplification       %u (%d UINT32)\n"), SubstrateSimplificationStatus.OriginalRequestedByteSize, SubstrateSimplificationStatus.OriginalRequestedByteSize / 4);
+				SubstrateMaterialContextDescription += FString::Printf(TEXT(" - Requested Byte Size after simplification        %u (%d UINT32)\n"), SubstrateCtx.SubstrateMaterialRequestedSizeByte, SubstrateCtx.SubstrateMaterialRequestedSizeByte / 4);
+				SubstrateMaterialContextDescription += FString::Printf(TEXT(" - Requested Closure Count before simplification   %u\n"), SubstrateSimplificationStatus.OriginalRequestedClosureCount);
+				SubstrateMaterialContextDescription += FString::Printf(TEXT(" - Requested Closure Count after simplification    %u\n"), SubstrateCtx.SubstrateMaterialClosureCount);
+				SubstrateMaterialContextDescription += FString::Printf(TEXT(" - Material complexity                             %s\n"), *FSubstrateMaterialComplexity::ToString(SubstrateCtx.SubstrateMaterialComplexity.SubstrateMaterialType(), true /* Upper case */));
+				SubstrateMaterialContextDescription += FString::Printf(TEXT(" - BSDF Count                                      %i\n"), SubstrateCtx.SubstrateMaterialEffectiveClosureCount); // REMOVE?
+				if (RequestedSharedLocalBasesCount > SUBSTRATE_MAX_SHAREDLOCALBASES_REGISTERS)
+				{
+					SubstrateMaterialDescription += FString::Printf(TEXT(" - SharedLocalBasesCount                      %i (Requested:%i)\n"), SubstrateCtx.FinalUsedSharedLocalBasesCount, RequestedSharedLocalBasesCount);
+				}
+				else
+				{
+					SubstrateMaterialDescription += FString::Printf(TEXT(" - SharedLocalBasesCount                      %i\n"), SubstrateCtx.FinalUsedSharedLocalBasesCount);
+				}
+
+
+				for (int32 OpIt = 0; OpIt < SubstrateCtx.SubstrateMaterialExpressionRegisteredOperators.Num(); ++OpIt)
+				{
+					FSubstrateOperator& BSDFOperator = SubstrateCtx.SubstrateMaterialExpressionRegisteredOperators[OpIt];
+					if (BSDFOperator.BSDFIndex == INDEX_NONE || BSDFOperator.IsDiscarded())
+					{
+						continue;	// not a BSDF or if discarded (i.e. not the root of a parameter blending subtree), then there is no local basis to register
+					}
+					if (BSDFOperator.BSDFRegisteredSharedLocalBasis.NormalCodeChunk == INDEX_NONE && BSDFOperator.BSDFRegisteredSharedLocalBasis.TangentCodeChunk == INDEX_NONE)
+					{
+						continue;	// We skip null normal on certain BSDF, for instance unlit.
+					}
+					const FSubstrateSharedLocalBasesInfo& SubstrateSharedLocalBasesInfo = SubstrateCtx.SubstrateCompilationInfoGetMatchingSharedLocalBasisInfo(BSDFOperator.BSDFRegisteredSharedLocalBasis);
+
+					SubstrateMaterialContextDescription += FString::Printf(TEXT("     - %s - SharedLocalBasisIndexMacro = %s \n"), *GetSubstrateBSDFName(BSDFOperator.BSDFType), *GetSubstrateSharedLocalBasisIndexMacroInner(BSDFOperator.BSDFRegisteredSharedLocalBasis, SubstrateCompilationContextIndex));
+				}
+
+				SubstrateMaterialContextDescription += FString::Printf(TEXT("----------- SUBSTRATE TREE - %s -----------\n"), *SubstrateCompilationContextName);
+				SubstrateMaterialContextDescription += FString::Printf(TEXT("Graph maximum distance to leaves %u\n"), RootMaximumDistanceToLeaves);
+				// Debug print operators according to depth from root.
+				{
+
+					for (int32 DistanceToLeaves = RootMaximumDistanceToLeaves; DistanceToLeaves >= 0; --DistanceToLeaves)
+					{
+						SubstrateMaterialContextDescription += FString::Printf(TEXT("----- DistanceFromLeaves = %d -----\n"), DistanceToLeaves);
+						for (auto& It : SubstrateCtx.SubstrateMaterialExpressionRegisteredOperators)
+						{
+							if (!It.IsDiscarded() && It.MaxDistanceFromLeaves == DistanceToLeaves)
+							{
+								SubstrateMaterialContextDescription += FString::Printf(TEXT("\tIdx=%d Op=%s ParentIdx=%d LeftIndex=%d RightIndex=%d BSDFIdx=%d LayerDepth=%d IsTop=%d IsBot=%d BSDFType=%s SSS=%d MFP=%d F90=%d Rough2=%d Fuzz=%d Aniso=%d Glint=%d SpecularProfile=%d\n"),
+									It.Index, GetSubstrateOperatorStr(It.OperatorType), It.ParentIndex, It.LeftIndex, It.RightIndex, It.BSDFIndex, It.LayerDepth, It.bIsTop, It.bIsBottom,
+									*GetSubstrateBSDFName(It.BSDFType), It.bBSDFHasSSS, It.bBSDFHasMFPPluggedIn, It.bBSDFHasEdgeColor, It.bBSDFHasSecondRoughnessOrSimpleClearCoat, It.bBSDFHasFuzz, It.bBSDFHasAnisotropy, It.bBSDFHasGlint, It.bBSDFHasSpecularProfile);
+							}
+						}
+					}
+				}
+
+				ResourcesString += TEXT("/*");
+				ResourcesString += SubstrateMaterialContextDescription;
+				ResourcesString += TEXT("*/");
+
+				SubstrateMaterialDescription += SubstrateMaterialContextDescription;
+
+				if (SubstrateCompilationContextIndex == ESubstrateCompilationContext::SCC_Default)
+				{
+					MaterialCompilationOutput.SubstrateMaterialCompilationOutput.SharedLocalBasesCount = SubstrateCtx.FinalUsedSharedLocalBasesCount;
+				}
+			}
+#endif // WITH_EDITOR
+		}
+		MaterialCompilationOutput.SubstrateMaterialCompilationOutput.SubstrateMaterialDescription = SubstrateMaterialDescription;
+	}
+
+	EnvironmentDefines->bTextureSampleDebug = IsDebugTextureSampleEnabled();
+}
+
+bool FHLSLMaterialTranslator::QueryDDCCachedTranslationResults()
+{
+	if (!GJobMaterialTranslateDDCEnable)
+	{
+		return false;
+	}
+
+	// Try fetching the translation results from the DDC using the generated key hash
+	FSharedBuffer MaterialCompilationOutputBuffer;
+	FSharedBuffer TranslationResultsBuffer;
+	FSharedBuffer EnvironmentDefinesBuffer;
+
+	UE::DerivedData::FRequestOwner RequestOwner(UE::DerivedData::EPriority::Blocking);
+	UE::DerivedData::FCacheKey CacheKey{ MaterialTranslationDDCBucket, DDCKeyHash };
+
+	{
+		UE::DerivedData::FCacheGetRequest Request;
+		Request.Name = Material->GetMaterialInterface()->GetName();
+		Request.Key = CacheKey;
+		Request.Policy = UE::DerivedData::ECachePolicy::Default;
+
+		UE::DerivedData::GetCache().Get(
+			{ Request },
+			RequestOwner,
+			[&](UE::DerivedData::FCacheGetResponse&& Response)
+			{
+				if (Response.Status == UE::DerivedData::EStatus::Ok)
+				{
+					MaterialCompilationOutputBuffer = Response.Record.GetValue(MaterialCompilationOutputId).GetData().Decompress();
+					TranslationResultsBuffer = Response.Record.GetValue(MaterialResultsOutputId).GetData().Decompress();
+					EnvironmentDefinesBuffer = Response.Record.GetValue(EnvironmentDefinesId).GetData().Decompress();
+				}
+			}
+		);
+	}
+
+	RequestOwner.Wait();
+
+	// Load the results if we hit the cache.
+	if (!MaterialCompilationOutputBuffer.IsNull() && !TranslationResultsBuffer.IsNull() && !EnvironmentDefinesBuffer.IsNull())
+	{
+		// Read the environment defines
+		FMemoryReaderView EnvironmentDefinesBufferReader{ TArrayView<uint8>{ (uint8*)EnvironmentDefinesBuffer.GetData(), (int)EnvironmentDefinesBuffer.GetSize() } };
+		EnvironmentDefines->Serialize(EnvironmentDefinesBufferReader);
+
+		// Read the material compilation output
+		FShaderMapPointerTable PointerTable;
+		FPlatformTypeLayoutParameters LayoutParams;
+		LayoutParams.InitializeForPlatform(GetTargetPlatform());
+		FMemoryReaderView MaterialCompilationOutputReader{ TArrayView<uint8>{ (uint8*)MaterialCompilationOutputBuffer.GetData(), (int)MaterialCompilationOutputBuffer.GetSize() } };
+		FMemoryImageObject LoadedContent = FMemoryImageResult::LoadFromArchive(MaterialCompilationOutputReader, StaticGetTypeLayoutDesc<FMaterialCompilationOutput>(), &PointerTable, LayoutParams);
+		MaterialCompilationOutput = *(FMaterialCompilationOutput*)LoadedContent.Object;
+
+		// Read the array of material string parameters
+		FMemoryReaderView ResultsMemoryReader{ TArrayView<uint8>{ (uint8*)TranslationResultsBuffer.GetData(), (int)TranslationResultsBuffer.GetSize() } };
+		ResultsMemoryReader << MaterialSourceTemplateParams;
+
+		return true;
+	}
+
+	return false;
+}
+
+void FHLSLMaterialTranslator::PushResultsToDDCCache()
+{
+	// We currently don't support caching parameter collections on the DDC, as it is invalid at this point
+	// to serialize an array of object pointers. Only allow putting the translation results on the DDC
+	// if the this material does not use parameter collections.
+	if (!GJobMaterialTranslateDDCEnable || !ParameterCollections.IsEmpty())
+	{
+		GShaderCompilerStats->IncrementMaterialTranslationSkippedDDC();
+		return;
+	}
+
+	UE::DerivedData::FCacheKey CacheKey{ MaterialTranslationDDCBucket, DDCKeyHash };
+	UE::DerivedData::FCacheRecordBuilder RecordBuilder{ CacheKey };
+
+	// Push the material compilation output
+	FBufferArchive MaterialCompilationOutputBuffer;
+	FShaderMapPointerTable PointerTable;
+
+	FMemoryImage MemoryImage;
+	MemoryImage.PrevPointerTable = &PointerTable;
+	MemoryImage.PointerTable = &PointerTable;
+	MemoryImage.TargetLayoutParameters.InitializeForArchive(MaterialCompilationOutputBuffer);
+
+	FMemoryImageWriter Writer(MemoryImage);
+	Writer.WriteRootObject(&MaterialCompilationOutput, StaticGetTypeLayoutDesc<FMaterialCompilationOutput>());
+
+	FMemoryImageResult MemoryImageResult;
+	MemoryImage.Flatten(MemoryImageResult, true);
+	MemoryImageResult.SaveToArchive(MaterialCompilationOutputBuffer);
+	RecordBuilder.AddValue(MaterialCompilationOutputId, FSharedBuffer::MakeView(MaterialCompilationOutputBuffer.GetData(), MaterialCompilationOutputBuffer.Num()));
+
+	// Push the material source string parameters
+	FBufferArchive ResultsBuffer;
+	ResultsBuffer << MaterialSourceTemplateParams;
+	RecordBuilder.AddValue(MaterialResultsOutputId, FSharedBuffer::MakeView(ResultsBuffer.GetData(), ResultsBuffer.Num()));
+
+	// Push the material shader defines
+	FBufferArchive EnvironmentDefinesBuffer;
+	EnvironmentDefines->Serialize(EnvironmentDefinesBuffer);
+	RecordBuilder.AddValue(EnvironmentDefinesId, FSharedBuffer::MakeView(EnvironmentDefinesBuffer.GetData(), EnvironmentDefinesBuffer.Num()));
+
+	// Push the all the data to the DDC
+	UE::DerivedData::FRequestOwner RequestOwner{ UE::DerivedData::EPriority::Normal };
+	UE::DerivedData::ECachePolicy Policy = UE::DerivedData::ECachePolicy::Default;
+	UE::DerivedData::FCachePutRequest Request{ { Material->GetMaterialInterface()->GetName() }, RecordBuilder.Build(), Policy };
+	UE::DerivedData::GetCache().Put({ Request }, RequestOwner);
+	RequestOwner.KeepAlive();
+}
+
+void FHLSLMaterialTranslator::PrepareMaterialSourceStringParameters()
+{
+	MaterialSourceTemplateParams.Empty();
+	MaterialSourceTemplateParams.Reserve(FMaterialSourceTemplate::Get().GetTemplate(GetShaderPlatform()).GetNumNamedParameters());
+
+	// Assign slots to vertex interpolators
+	FString VertexInterpolatorsOffsetsDefinition;
+	TBitArray<> FinalAllocatedCoords = GetVertexInterpolatorsOffsets(VertexInterpolatorsOffsetsDefinition);
+
+	const uint32 NumUserVertexTexCoords = GetNumUserVertexTexCoords();
+	const uint32 NumUserTexCoords = GetNumUserTexCoords();
+	const uint32 NumCustomVectors = FMath::DivideAndRoundUp((uint32)CurrentCustomVertexInterpolatorOffset, 2u);
+	const uint32 NumTexCoordVectors = FinalAllocatedCoords.FindLast(true) + 1;
+
+	MaterialSourceTemplateParams.Add({ TEXT("num_material_texcoords_vertex"), FString::Printf(TEXT("%u"), NumUserVertexTexCoords) });
+	MaterialSourceTemplateParams.Add({ TEXT("num_material_texcoords"), FString::Printf(TEXT("%u"), NumUserTexCoords) });
+	MaterialSourceTemplateParams.Add({ TEXT("num_custom_vertex_interpolators"), FString::Printf(TEXT("%u"), NumCustomVectors) });
+	MaterialSourceTemplateParams.Add({ TEXT("num_tex_coord_interpolators"), FString::Printf(TEXT("%u"), NumTexCoordVectors) });
+	MaterialSourceTemplateParams.Add({ TEXT("vertex_interpolators_offsets_definition"), VertexInterpolatorsOffsetsDefinition });
+
+	FString MaterialAttributesDeclaration;
+	FString MaterialAttributesUtilities;
+
+	const TArray<FGuid>& OrderedVisibleAttributes = FMaterialAttributeDefinitionMap::GetOrderedVisibleAttributeList();
+
+	// Reserve enough space vased on an estimate of the length of a single attribute string
+	MaterialAttributesDeclaration.Reserve(30 + MaterialAttributesDeclaration.Len() * 128);
+	MaterialAttributesUtilities.Reserve(MaterialAttributesDeclaration.Len() * 256);
+
+	MaterialAttributesDeclaration.Append(TEXT("struct FMaterialAttributes\n{\n"));
+
+	const EMaterialShadingModel DefaultShadingModel = Material->GetShadingModels().GetFirstShadingModel();
+	for (const FGuid& AttributeID : OrderedVisibleAttributes)
+	{
+		const FString PropertyName = FMaterialAttributeDefinitionMap::GetAttributeName(AttributeID);
+		const EMaterialValueType PropertyType = FMaterialAttributeDefinitionMap::GetValueType(AttributeID);
+		const TCHAR* HLSLType = nullptr;
+
+		switch (PropertyType)
+		{
+			case MCT_Float1: case MCT_Float: HLSLType = TEXT("float"); break;
+			case MCT_Float2: HLSLType = TEXT("float2"); break;
+			case MCT_Float3: HLSLType = TEXT("float3"); break;
+			case MCT_Float4: HLSLType = TEXT("float4"); break;
+			case MCT_UInt: case MCT_UInt1: case MCT_ShadingModel: HLSLType = TEXT("uint"); break;
+			case MCT_UInt2: HLSLType = TEXT("uint2"); break;
+			case MCT_UInt3: HLSLType = TEXT("uint3"); break;
+			case MCT_UInt4: HLSLType = TEXT("uint4"); break;
+			case MCT_Substrate: HLSLType = TEXT("FSubstrateData"); break;
+			default: break;
+		}
+
+		if (HLSLType)
+		{
+			const FVector4f DefaultValue = FMaterialAttributeDefinitionMap::GetDefaultValue(AttributeID);
+
+			MaterialAttributesDeclaration.Appendf(TEXT("\t%s %s;") HLSL_LINE_TERMINATOR, HLSLType, *PropertyName);
+
+			// Chainable method to set the attribute
+			MaterialAttributesUtilities.Appendf(TEXT("FMaterialAttributes FMaterialAttributes_Set%s(FMaterialAttributes InAttributes, %s InValue) { InAttributes.%s = InValue; return InAttributes; }") HLSL_LINE_TERMINATOR,
+				*PropertyName, HLSLType, *PropertyName);
+		}
+	}
+
+	MaterialAttributesDeclaration.Append(TEXT("};\n"));
+
+	MaterialSourceTemplateParams.Add({ TEXT("material_declarations"), MaterialAttributesDeclaration });
+	MaterialSourceTemplateParams.Add({ TEXT("material_attributes_utilities"), MaterialAttributesUtilities });
+
+	// Stores the shared shader results member declarations
+	FString PixelMembersDeclaration[CompiledPDV_MAX];
+
+	FString NormalAssignment[CompiledPDV_MAX];
+
+	// Stores the code to initialize all inputs after MP_Normal
+	FString PixelMembersSetupAndAssignments[CompiledPDV_MAX];
+
+	for (int32 Index = 0; Index < CompiledPDV_MAX; Index++)
+	{
+		GetSharedInputsMaterialCode(PixelMembersDeclaration[Index], NormalAssignment[Index], PixelMembersSetupAndAssignments[Index], (ECompiledPartialDerivativeVariation)Index);
+	}
+
+	// PixelMembersDeclaration should be the same for all variations, but might change in the future. There are cases where work is shared
+	// between the pixel and vertex shader, but with Nanite all work has to be moved into the pixel shader, which means we will want
+	// different inputs. But for now, we are keeping them the same.
+	MaterialSourceTemplateParams.Add({ TEXT("pixel_material_inputs"), PixelMembersDeclaration[CompiledPDV_FiniteDifferences] });
+
+	{
+		FString DerivativeHelpers = DerivativeAutogen.GenerateUsedFunctions(*this);
+		FString DerivativeHelpersAndResources = DerivativeHelpers + ResourcesString;
+		MaterialSourceTemplateParams.Add({ TEXT("uniform_material_expressions"), DerivativeHelpersAndResources });
+	}
+
+	// Anything used bye the GenerationFunctionCode() like WorldPositionOffset shouldn't be using texures, right?
+	// Let those use the standard finite differences textures, since they should be the same. If we actually want
+	// those to handle texture reads properly, we'll have to make extra versions.
+	ECompiledPartialDerivativeVariation BaseDerivativeVariation = CompiledPDV_FiniteDifferences;
+
+	if (bCompileForComputeShader)
+	{
+		MaterialSourceTemplateParams.Add({ TEXT("get_material_emissive_for_cs"), GenerateFunctionCode(CompiledMP_EmissiveColorCS, BaseDerivativeVariation) });
+	}
+	else
+	{
+		MaterialSourceTemplateParams.Add({ TEXT("get_material_emissive_for_cs"), TEXT("return 0") });
+	}
+
+	MaterialSourceTemplateParams.Add({ TEXT("get_material_translucency_directional_lighting_intensity"), FString::Printf(TEXT("return %.5f"), Material->GetTranslucencyDirectionalLightingIntensity()) });
+	MaterialSourceTemplateParams.Add({ TEXT("get_material_translucent_shadow_density_scale"), FString::Printf(TEXT("return %.5f"), Material->GetTranslucentShadowDensityScale()) });
+	MaterialSourceTemplateParams.Add({ TEXT("get_material_translucent_self_shadow_density_scale"), FString::Printf(TEXT("return %.5f"), Material->GetTranslucentSelfShadowDensityScale()) });
+	MaterialSourceTemplateParams.Add({ TEXT("get_material_translucent_self_shadow_second_density_scale"), FString::Printf(TEXT("return %.5f"), Material->GetTranslucentSelfShadowSecondDensityScale()) });
+	MaterialSourceTemplateParams.Add({ TEXT("get_material_translucent_self_shadow_second_opacity"), FString::Printf(TEXT("return %.5f"), Material->GetTranslucentSelfShadowSecondOpacity()) });
+	MaterialSourceTemplateParams.Add({ TEXT("get_material_translucent_backscattering_exponent"), FString::Printf(TEXT("return %.5f"), Material->GetTranslucentBackscatteringExponent()) });
+
+	{
+		FLinearColor Extinction = Material->GetTranslucentMultipleScatteringExtinction();
+
+		MaterialSourceTemplateParams.Add({ TEXT("get_material_translucent_multiple_scattering_extinction"), FString::Printf(TEXT("return MaterialFloat3(%.5f, %.5f, %.5f)"), Extinction.R, Extinction.G, Extinction.B) });
+	}
+
+	MaterialSourceTemplateParams.Add({ TEXT("get_material_opacity_mask_clip_value"), FString::Printf(TEXT("return %.5f"), Material->GetOpacityMaskClipValue()) });
+
+	{
+		const FDisplacementScaling DisplacementScaling = Material->GetDisplacementScaling();
+		MaterialSourceTemplateParams.Add({ TEXT("get_material_displacement_magnitude"), FString::Printf(TEXT("return %.5f"), FMath::Max(0.0f, DisplacementScaling.Magnitude)) });
+		MaterialSourceTemplateParams.Add({ TEXT("get_material_displacement_center"), FString::Printf(TEXT("return %.5f"), FMath::Clamp(DisplacementScaling.Center, 0.0f, 1.0f)) });
+	}
+
+	MaterialSourceTemplateParams.Add({ TEXT("get_material_world_position_offset_raw"), !bEnableExecutionFlow ? *GenerateFunctionCode(MP_WorldPositionOffset, BaseDerivativeVariation) : TEXT("return Parameters.MaterialAttributes.WorldPositionOffset") });
+	MaterialSourceTemplateParams.Add({ TEXT("get_material_previous_world_position_offset_raw"), !bEnableExecutionFlow ? *GenerateFunctionCode(CompiledMP_PrevWorldPositionOffset, BaseDerivativeVariation) : TEXT("return 0.0f") });
+	MaterialSourceTemplateParams.Add({ TEXT("get_material_custom_data0"), !bEnableExecutionFlow ? *GenerateFunctionCode(MP_CustomData0, BaseDerivativeVariation) : TEXT("return 0.0f") });
+	MaterialSourceTemplateParams.Add({ TEXT("get_material_custom_data1"), !bEnableExecutionFlow ? *GenerateFunctionCode(MP_CustomData1, BaseDerivativeVariation) : TEXT("return 0.0f") });
+
+	// Print custom texture coordinate assignments, should be fine with regular derivatives
+	FString CustomUVAssignments;
+
+	int32 LastProperty = -1;
+	for (uint32 CustomUVIndex = 0; CustomUVIndex < NumUserTexCoords; CustomUVIndex++)
+	{
+		if (bEnableExecutionFlow)
+		{
+			const FString AttributeName = FMaterialAttributeDefinitionMap::GetAttributeName((EMaterialProperty)(MP_CustomizedUVs0 + CustomUVIndex));
+			CustomUVAssignments.Appendf(TEXT("\tOutTexCoords[%u] = Parameters.MaterialAttributes.%s;") HLSL_LINE_TERMINATOR, CustomUVIndex, *AttributeName);
+		}
+		else
+		{
+			if (CustomUVIndex == 0)
+			{
+				CustomUVAssignments.Append(DerivativeVariations[BaseDerivativeVariation].TranslatedCodeChunkDefinitions[MP_CustomizedUVs0 + CustomUVIndex]);
+			}
+
+			if (DerivativeVariations[BaseDerivativeVariation].TranslatedCodeChunkDefinitions[MP_CustomizedUVs0 + CustomUVIndex].Len() > 0)
+			{
+				LastProperty = MP_CustomizedUVs0 + CustomUVIndex;
+			}
+			CustomUVAssignments.Appendf(TEXT("\tOutTexCoords[%u] = %s;") HLSL_LINE_TERMINATOR, CustomUVIndex, *DerivativeVariations[BaseDerivativeVariation].TranslatedCodeChunks[MP_CustomizedUVs0 + CustomUVIndex]);
+		}
+	}
+
+	MaterialSourceTemplateParams.Add({ TEXT("get_material_customized_u_vs"), CustomUVAssignments });
+
+	// Print custom vertex shader interpolator assignments
+	FString CustomInterpolatorAssignments;
+
+	for (UMaterialExpressionVertexInterpolator* Interpolator : CustomVertexInterpolators)
+	{
+		if (Interpolator->InterpolatorOffset != INDEX_NONE)
+		{
+			check(Interpolator->InterpolatorIndex != INDEX_NONE);
+			check(Interpolator->InterpolatedType & MCT_Float);
+
+			const EMaterialValueType Type = Interpolator->InterpolatedType == MCT_Float ? MCT_Float1 : Interpolator->InterpolatedType;
+			const TCHAR* Swizzle[2] = { TEXT("x"), TEXT("y") };
+			const int32 Offset = Interpolator->InterpolatorOffset;
+			const int32 Index = Interpolator->InterpolatorIndex;
+
+			// Note: We reference the UV define directly to avoid having to pre-accumulate UV counts before property translation
+			CustomInterpolatorAssignments += FString::Printf(TEXT("\tOutTexCoords[VERTEX_INTERPOLATOR_%i_TEXCOORDS_X].%s = VertexInterpolator%i(Parameters).x;") HLSL_LINE_TERMINATOR, Index, Swizzle[Offset % 2], Index);
+
+			if (Type >= MCT_Float2)
+			{
+				CustomInterpolatorAssignments += FString::Printf(TEXT("\tOutTexCoords[VERTEX_INTERPOLATOR_%i_TEXCOORDS_Y].%s = VertexInterpolator%i(Parameters).y;") HLSL_LINE_TERMINATOR, Index, Swizzle[(Offset + 1) % 2], Index);
+
+				if (Type >= MCT_Float3)
+				{
+					CustomInterpolatorAssignments += FString::Printf(TEXT("\tOutTexCoords[VERTEX_INTERPOLATOR_%i_TEXCOORDS_Z].%s = VertexInterpolator%i(Parameters).z;") HLSL_LINE_TERMINATOR, Index, Swizzle[(Offset + 2) % 2], Index);
+
+					if (Type == MCT_Float4)
+					{
+						CustomInterpolatorAssignments += FString::Printf(TEXT("\tOutTexCoords[VERTEX_INTERPOLATOR_%i_TEXCOORDS_W].%s = VertexInterpolator%i(Parameters).w;") HLSL_LINE_TERMINATOR, Index, Swizzle[(Offset + 3) % 2], Index);
+					}
+				}
+			}
+		}
+	}
+
+	MaterialSourceTemplateParams.Add({ TEXT("get_custom_interpolators"), CustomInterpolatorAssignments });
+
+	if (bEnableExecutionFlow)
+	{
+		FString EvaluateVertexCode;
+
+		// Set default texcoords in the VS
+		for (uint32 TexCoordIndex = 0; TexCoordIndex < NumUserVertexTexCoords; ++TexCoordIndex)
+		{
+			const FString AttributeName = FMaterialAttributeDefinitionMap::GetAttributeName((EMaterialProperty)(MP_CustomizedUVs0 + TexCoordIndex));
+
+			EvaluateVertexCode += FString::Printf(TEXT("\tDefaultMaterialAttributes.%s = Parameters.TexCoords[%d];") HLSL_LINE_TERMINATOR, *AttributeName, TexCoordIndex);
+		}
+
+		EvaluateVertexCode += TranslatedAttributesCodeChunks[SF_Vertex];
+
+		MaterialSourceTemplateParams.Add({ TEXT("evaluate_material_attributes"), TranslatedAttributesCodeChunks[SF_Pixel] });
+
+		FString EvaluateMaterialAttributesCode = TEXT("    FMaterialAttributes MaterialAttributes = EvaluatePixelMaterialAttributes(Parameters);" HLSL_LINE_TERMINATOR);
+
+		for (int32 PropertyIndex = 0; PropertyIndex < MP_MAX; ++PropertyIndex)
+		{
+			// Skip non-shared properties
+			if (!SharedPixelProperties[PropertyIndex])
+			{
+				continue;
+			}
+
+			const EMaterialProperty Property = (EMaterialProperty)PropertyIndex;
+			check(FMaterialAttributeDefinitionMap::GetShaderFrequency(Property) == SF_Pixel);
+			// Special case MP_SubsurfaceColor as the actual property is a combination of the color and the profile but we don't want to expose the profile
+			const FString PropertyName = FMaterialAttributeDefinitionMap::GetAttributeName(Property);
+
+			if (PropertyIndex == MP_SubsurfaceColor)
+			{
+				// TODO - properly handle subsurface profile
+				EvaluateMaterialAttributesCode += FString::Printf("    PixelMaterialInputs.Subsurface = float4(MaterialAttributes.%s, 0.0f);" HLSL_LINE_TERMINATOR, *PropertyName);
+			}
+			else
+			{
+				EvaluateMaterialAttributesCode += FString::Printf("    PixelMaterialInputs.%s = MaterialAttributes.%s;" HLSL_LINE_TERMINATOR, *PropertyName, *PropertyName);
+			}
+		}
+
+		MaterialSourceTemplateParams.Add({ TEXT("calc_pixel_material_inputs_initial_calculations"), EvaluateMaterialAttributesCode });
+		MaterialSourceTemplateParams.Add({ TEXT("calc_pixel_material_inputs_normal"), TEXT("") });
+		MaterialSourceTemplateParams.Add({ TEXT("calc_pixel_material_inputs_other_inputs"), TEXT("") });
+
+		MaterialSourceTemplateParams.Add({ TEXT("calc_pixel_material_inputs_analytic_derivatives_initial"), EvaluateMaterialAttributesCode });
+		MaterialSourceTemplateParams.Add({ TEXT("calc_pixel_material_inputs_analytic_derivatives_normal"), TEXT("") });
+		MaterialSourceTemplateParams.Add({ TEXT("calc_pixel_material_inputs_analytic_derivatives_other_inputs"), TEXT("") });
+	}
+	else
+	{
+		// skip material attributes code
+		MaterialSourceTemplateParams.Add({ TEXT("evaluate_material_attributes"), TEXT("") });
+
+		{
+			// Initializers required for Normal
+			MaterialSourceTemplateParams.Add({ TEXT("calc_pixel_material_inputs_initial_calculations"), DerivativeVariations[CompiledPDV_FiniteDifferences].TranslatedCodeChunkDefinitions[MP_Normal] });
+			MaterialSourceTemplateParams.Add({ TEXT("calc_pixel_material_inputs_normal"), NormalAssignment[CompiledPDV_FiniteDifferences] });
+			// Finally the rest of common code followed by assignment into each input
+			MaterialSourceTemplateParams.Add({ TEXT("calc_pixel_material_inputs_other_inputs"), PixelMembersSetupAndAssignments[CompiledPDV_FiniteDifferences] });
+		}
+		{
+			// Initializers required for Normal
+			MaterialSourceTemplateParams.Add({ TEXT("calc_pixel_material_inputs_analytic_derivatives_initial"), DerivativeVariations[CompiledPDV_Analytic].TranslatedCodeChunkDefinitions[MP_Normal] });
+			MaterialSourceTemplateParams.Add({ TEXT("calc_pixel_material_inputs_analytic_derivatives_normal"), NormalAssignment[CompiledPDV_Analytic] });
+			// Finally the rest of common code followed by assignment into each input
+			MaterialSourceTemplateParams.Add({ TEXT("calc_pixel_material_inputs_analytic_derivatives_other_inputs"), PixelMembersSetupAndAssignments[CompiledPDV_Analytic] });
+		}
+	}
+}
+
 int32 FHLSLMaterialTranslator::SparseVolumeTexture(USparseVolumeTexture* Texture, int32& TextureReferenceIndex, EMaterialSamplerType SamplerType)
 {
 	TextureReferenceIndex = Material->GetReferencedTextures().Find(Texture);
diff --git a/Engine/Source/Runtime/Engine/Private/Materials/HLSLMaterialTranslator.h b/Engine/Source/Runtime/Engine/Private/Materials/HLSLMaterialTranslator.h
--- a/Engine/Source/Runtime/Engine/Private/Materials/HLSLMaterialTranslator.h
+++ b/Engine/Source/Runtime/Engine/Private/Materials/HLSLMaterialTranslator.h
@@ -426,6 +426,8 @@ protected:
 	/** Will contain all the shading models picked up from the material expression graph */
 	FMaterialShadingModelField ShadingModelsFromCompilation;
 
+	struct FEnvironmentDefines;
+
 	// Describe the simplification status. Once the material has been compiled, it can be used to understand if and how it has been simplified.
 	struct FSubstrateSimplificationStatus
 	{
@@ -526,6 +528,7 @@ protected:
 		bool SubstrateGenerateDerivedMaterialOperatorData(FHLSLMaterialTranslator* Compiler);
 
 		void SubstrateEvaluateSharedLocalBases(FHLSLMaterialTranslator* Compiler, uint8& RequestedSharedLocalBasesCount, FShaderCompilerEnvironment* OutEnvironment);
+		void SubstrateEvaluateSharedLocalBasesNew(FHLSLMaterialTranslator* Compiler, uint8& RequestedSharedLocalBasesCount, FEnvironmentDefines* OutEnvironment);
 
 		FSubstrateSharedLocalBasesInfo SubstrateCompilationInfoGetMatchingSharedLocalBasisInfo(const FSubstrateRegisteredSharedLocalBasis& SearchedSharedLocalBasis);
 
@@ -552,7 +555,15 @@ protected:
 
 	/** Substrate material compilation and simplification configuration. */
 	FSubstrateCompilationConfig SubstrateCompilationConfig;
+	
+	/** The DDC key for the material translation results */
+	FIoHash DDCKeyHash;
+
 public: 
+	/**
+	 * Returns a string representation that identify the translator version. Used to keep shader DDC keys valid when material translation internals change.
+	 */
+	static void AppendVersionString(FString& Output, EShaderPlatform Platform);
 
 	FHLSLMaterialTranslator(FMaterial* InMaterial,
 		FMaterialCompilationOutput& InMaterialCompilationOutput,
@@ -561,7 +572,8 @@ public:
 		EMaterialQualityLevel::Type InQualityLevel,
 		ERHIFeatureLevel::Type InFeatureLevel,
 		const ITargetPlatform* InTargetPlatform = nullptr,
-		const FSubstrateCompilationConfig* InSubstrateCompilationConfig = nullptr);
+		const FSubstrateCompilationConfig* InSubstrateCompilationConfig = nullptr,
+		FString MaterialTranslationDDCKeyString = {});
 
 	~FHLSLMaterialTranslator();
 
@@ -621,6 +633,11 @@ public:
 	virtual FString GetParameterCodeDeriv(int32 Index, ECompiledPartialDerivativeVariation Variation);
 
 protected:
+	// These -Old functions are temporarily preserved in order to switch back to the old way of
+	// generating material translation results. We intend to remove these as soon as we are sure
+	// the new implementations of these functions work well.
+	void GetMaterialEnvironmentOld(EShaderPlatform InPlatform, FShaderCompilerEnvironment& OutEnvironment);
+	FString GetMaterialShaderCodeOld();
 
 	uint64 GetParameterHash(int32 Index);
 
@@ -1346,6 +1363,38 @@ protected:
 	virtual bool IsCurrentlyCompilingForPreviousFrame() const;
 
 	virtual bool IsDevelopmentFeatureEnabled(const FName& FeatureName) const override;
+	
+	/**
+	 * EFfectively performs the translation without querying the DDC first.
+	 */
+	void DoTranslate();
+
+	/**
+	 * Queries the DDC cache for a cached translation.
+	 * @return Whether the speecified key is in the DDC.
+	 */
+	bool QueryDDCCachedTranslationResults();
+
+	/**
+	 * Pushes the final results to the DDC cache.
+	 */
+	void PushResultsToDDCCache();
+
+	/**
+	 * Prepares the material source generation parameters.
+	 */
+	void PrepareMaterialSourceStringParameters();
+
+	/**
+	 * Prepares the Environment Defines array based on compilation results.
+	 */
+	void PrepareEnvironmentDefines();
+
+	/** The material shader source template parameters */
+	TMap<FString, FString> MaterialSourceTemplateParams;
+
+	/** The output material shader defines */
+	TUniquePtr<FEnvironmentDefines> EnvironmentDefines;
 };
 
 #endif // WITH_EDITORONLY_DATA
diff --git a/Engine/Source/Runtime/Engine/Private/Materials/MaterialShader.cpp b/Engine/Source/Runtime/Engine/Private/Materials/MaterialShader.cpp
--- a/Engine/Source/Runtime/Engine/Private/Materials/MaterialShader.cpp
+++ b/Engine/Source/Runtime/Engine/Private/Materials/MaterialShader.cpp
@@ -27,6 +27,7 @@
 #include "UObject/UE5ReleaseStreamObjectVersion.h"
 #include "Materials/Material.h"
 #include "MaterialHLSLGenerator.h"
+#include "HLSLMaterialTranslator.h"
 
 int32 GMaterialExcludeNonPipelinedShaders = 1;
 static FAutoConsoleVariableRef CVarMaterialExcludeNonPipelinedShaders(
@@ -178,7 +179,7 @@ FString GetBlendModeString(EBlendMode BlendMode)
 
 #if WITH_EDITOR
 /** Creates a string key for the derived data cache given a shader map id. */
-static FString GetMaterialShaderMapKeyString(const FMaterialShaderMapId& ShaderMapId, EShaderPlatform Platform)
+FString GetMaterialShaderMapKeyString(const FMaterialShaderMapId& ShaderMapId, EShaderPlatform Platform, bool bIncludeKeyStringShaderDependencies)
 {
 	FName Format = LegacyShaderPlatformToShaderFormat(Platform);
 	FString ShaderMapKeyString;
@@ -194,10 +195,12 @@ static FString GetMaterialShaderMapKeyString(const FMaterialShaderMapId& ShaderM
 	ShaderMapKeyString.AppendChar('_');
 
 	ShaderMapAppendKeyString(Platform, ShaderMapKeyString);
-	ShaderMapId.AppendKeyString(ShaderMapKeyString);
+	ShaderMapId.AppendKeyString(ShaderMapKeyString, true, bIncludeKeyStringShaderDependencies);
 	FMaterialAttributeDefinitionMap::AppendDDCKeyString(ShaderMapKeyString);
 	FShaderCompileUtilities::AppendGBufferDDCKeyString(Platform, ShaderMapKeyString);
 
+	FHLSLMaterialTranslator::AppendVersionString(ShaderMapKeyString, Platform);
+
 	return ShaderMapKeyString;
 }
 
@@ -1077,7 +1080,7 @@ void FMaterialShaderMapId::AppendStaticParametersString(FString& ParamsString) c
 	}
 }
 
-void FMaterialShaderMapId::AppendKeyString(FString& KeyString, bool bIncludeSourceAndMaterialState) const
+void FMaterialShaderMapId::AppendKeyString(FString& KeyString, bool bIncludeSourceAndMaterialState, bool bIncludeKeyStringShaderDependencies) const
 {
 	check(IsContentValid());
 	if (bIncludeSourceAndMaterialState)
@@ -1140,13 +1143,16 @@ void FMaterialShaderMapId::AppendKeyString(FString& KeyString, bool bIncludeSour
 	}
 
 	// Add the inputs for any shaders that are stored inline in the shader map
-	AppendKeyStringShaderDependencies(
-		MakeArrayView(ShaderTypeDependencies),
-		MakeArrayView(ShaderPipelineTypeDependencies),
-		MakeArrayView(VertexFactoryTypeDependencies),
-		LayoutParams, 
-		KeyString,
-		bIncludeSourceAndMaterialState);
+	if (bIncludeKeyStringShaderDependencies)
+	{
+		AppendKeyStringShaderDependencies(
+			MakeArrayView(ShaderTypeDependencies),
+			MakeArrayView(ShaderPipelineTypeDependencies),
+			MakeArrayView(VertexFactoryTypeDependencies),
+			LayoutParams,
+			KeyString,
+			bIncludeSourceAndMaterialState);
+	}
 
 	BytesToHex(&TextureReferencesHash.Hash[0], sizeof(TextureReferencesHash.Hash), KeyString);
 
@@ -1440,7 +1446,8 @@ TSharedRef<FMaterialShaderMap::FAsyncLoadContext> FMaterialShaderMap::BeginLoadF
 				ShaderMap->Serialize(Ar);
 				//InOutShaderMap->RegisterSerializedShaders(false);
 
-				const FString InDataKey = GetMaterialShaderMapKeyString(ShaderMap->GetShaderMapId(), Platform);
+				const FString InDataKey = GetMaterialShaderMapKeyString(ShaderMap->GetShaderMapId(), Platform, true);
+
 				if (InDataKey != DataKey)
 				{
 					UE_LOG(LogMaterial, Warning, TEXT("Shader map key recomputed from DDC data: %s"), *InDataKey);
diff --git a/Engine/Source/Runtime/Engine/Private/Materials/MaterialShared.cpp b/Engine/Source/Runtime/Engine/Private/Materials/MaterialShared.cpp
--- a/Engine/Source/Runtime/Engine/Private/Materials/MaterialShared.cpp
+++ b/Engine/Source/Runtime/Engine/Private/Materials/MaterialShared.cpp
@@ -3310,7 +3310,8 @@ bool FMaterial::Translate_Legacy(const FMaterialShaderMapId& ShaderMapId,
 	FMaterialCompilationOutput& OutCompilationOutput,
 	TRefCountPtr<FSharedShaderCompilerEnvironment>& OutMaterialEnvironment)
 {
-	FHLSLMaterialTranslator MaterialTranslator(this, OutCompilationOutput, InStaticParameters, InPlatform, GetQualityLevel(), ShaderMapId.FeatureLevel, InTargetPlatform, &ShaderMapId.SubstrateCompilationConfig);
+	FString MaterialTranslationDDCKeyString = GetMaterialShaderMapKeyString(ShaderMapId, InPlatform, false);
+	FHLSLMaterialTranslator MaterialTranslator(this, OutCompilationOutput, InStaticParameters, InPlatform, GetQualityLevel(), ShaderMapId.FeatureLevel, InTargetPlatform, &ShaderMapId.SubstrateCompilationConfig, MoveTemp(MaterialTranslationDDCKeyString));
 	const bool bSuccess = MaterialTranslator.Translate();
 	if (bSuccess)
 	{
diff --git a/Engine/Source/Runtime/Engine/Private/Materials/MaterialSourceTemplate.cpp b/Engine/Source/Runtime/Engine/Private/Materials/MaterialSourceTemplate.cpp
--- a/Engine/Source/Runtime/Engine/Private/Materials/MaterialSourceTemplate.cpp
+++ b/Engine/Source/Runtime/Engine/Private/Materials/MaterialSourceTemplate.cpp
@@ -25,10 +25,16 @@ FStringTemplateResolver FMaterialSourceTemplate::BeginResolve(EShaderPlatform Sh
 	return { Templates[ShaderPlatform], 50 * 1024 };
 }
 
-const FSHA1& FMaterialSourceTemplate::GetTemplateHash(EShaderPlatform ShaderPlatform)
+const FStringTemplate& FMaterialSourceTemplate::GetTemplate(EShaderPlatform ShaderPlatform)
 {
 	Preload(ShaderPlatform);
-	return TemplateHash[ShaderPlatform];
+	return Templates[ShaderPlatform];
+}
+
+const FString& FMaterialSourceTemplate::GetTemplateHashString(EShaderPlatform ShaderPlatform)
+{
+	Preload(ShaderPlatform);
+	return TemplateHashString[ShaderPlatform];
 }
 
 bool FMaterialSourceTemplate::Preload(EShaderPlatform ShaderPlatform)
@@ -88,13 +94,14 @@ bool FMaterialSourceTemplate::Preload(EShaderPlatform ShaderPlatform)
 	}
 
 	// Extract the material template string TemplateVersion parameter
+	FSHA1 TemplateHash = {};
 	const TStringView TemplateVersionKeyword = TEXT("$TemplateVersion{");
 	int Begin = MaterialTemplateString.Find(TemplateVersionKeyword.GetData());
 	int End = MaterialTemplateString.Find(TEXT("}"), ESearchCase::CaseSensitive, ESearchDir::FromStart, Begin + TemplateVersionKeyword.Len());
 	if (Begin > 0 && End > 0)
 	{
 		Begin += TemplateVersionKeyword.Len();
-		TemplateHash[ShaderPlatform].UpdateWithString(*MaterialTemplateString + Begin, End - Begin);
+		TemplateHash.UpdateWithString(*MaterialTemplateString + Begin, End - Begin);
 	}
 
 	TArray<FStringView> Parameters;
@@ -102,9 +109,11 @@ bool FMaterialSourceTemplate::Preload(EShaderPlatform ShaderPlatform)
 	Parameters.Sort();
 	for (const FStringView& Param : Parameters)
 	{
-		TemplateHash[ShaderPlatform].UpdateWithString(Param.GetData(), Param.Len());
+		TemplateHash.UpdateWithString(Param.GetData(), Param.Len());
 	}
 
+	TemplateHashString[ShaderPlatform] = LexToString(TemplateHash.Finalize());
+
 	return true;
 }
 
diff --git a/Engine/Source/Runtime/Engine/Private/Materials/MaterialSourceTemplate.h b/Engine/Source/Runtime/Engine/Private/Materials/MaterialSourceTemplate.h
--- a/Engine/Source/Runtime/Engine/Private/Materials/MaterialSourceTemplate.h
+++ b/Engine/Source/Runtime/Engine/Private/Materials/MaterialSourceTemplate.h
@@ -18,11 +18,14 @@ struct FMaterialSourceTemplate
 	 * The appropriate template is loaded for the first time if it wasn't loaded before.*/
 	FStringTemplateResolver BeginResolve(EShaderPlatform ShaderPlatform, int32* MaterialTemplateLineNumber = nullptr);
 
+	/* Returns the string template for specified shader platform.  */
+	const FStringTemplate& GetTemplate(EShaderPlatform ShaderPlatform);
+
 	/** Return the hash of the material template for specified shader platform 
 	  * The information hashed is the string contained within "$TemplateVersion{...}" and the name of the parameters
 	  * contained in the template string (insensitive to ordering).
 	  */
-	const FSHA1& GetTemplateHash(EShaderPlatform ShaderPlatform);
+	const FString& GetTemplateHashString(EShaderPlatform ShaderPlatform);
 
 	/** Loads a material source template if still unloaded. */
 	bool Preload(EShaderPlatform ShaderPlatform);
@@ -31,7 +34,7 @@ struct FMaterialSourceTemplate
 	FStringTemplate Templates[SP_NumPlatforms];
 
 	/** Hash of the template parameters and TemplateVersion parameter extracted from file */
-	FSHA1 TemplateHash[SP_NumPlatforms];
+	FString TemplateHashString[SP_NumPlatforms];
 
 	/** Cached material template line numbers per shader platform */
 	int32 MaterialTemplateLineNumbers[SP_NumPlatforms];
diff --git a/Engine/Source/Runtime/Engine/Private/ShaderCompiler/ShaderCompiler.cpp b/Engine/Source/Runtime/Engine/Private/ShaderCompiler/ShaderCompiler.cpp
--- a/Engine/Source/Runtime/Engine/Private/ShaderCompiler/ShaderCompiler.cpp
+++ b/Engine/Source/Runtime/Engine/Private/ShaderCompiler/ShaderCompiler.cpp
@@ -5173,6 +5173,12 @@ void FShaderCompilerStats::FMaterialCounters::WriteStatSummary(const TCHAR* Aggr
 	UE_LOG(LogShaderCompilers, Display, TEXT("Materials Cooked:        %d"), NumMaterialsCooked);
 	UE_LOG(LogShaderCompilers, Display, TEXT("Materials Translated:    %d"), MaterialTranslateCalls);
 	UE_LOG(LogShaderCompilers, Display, TEXT("Material Translate Time: %.2f s"), MaterialTranslateTimeSec);
+	UE_LOG(LogShaderCompilers, Display, TEXT("Material Cache Time: %.2f s"), MaterialCacheTimeSec);
+
+	int HitsPercentage = MaterialTranslateCalls ? MaterialCacheHits * 100 / MaterialTranslateCalls : 0;
+	UE_LOG(LogShaderCompilers, Display, TEXT("Material Cache Hits: %d (%d%%)"), MaterialCacheHits, HitsPercentage);
+
+	UE_LOG(LogShaderCompilers, Display, TEXT("Material translations not on DDC: %d"), MaterialTranslationSkippedDDCCount);
 }
 
 void FShaderCompilerStats::FMaterialCounters::GatherAnalytics(TArray<FAnalyticsEventAttribute>& Attributes)
@@ -5180,6 +5186,9 @@ void FShaderCompilerStats::FMaterialCounters::GatherAnalytics(TArray<FAnalyticsE
 	Attributes.Emplace(TEXT("Material_NumMaterialsCooked"), NumMaterialsCooked);
 	Attributes.Emplace(TEXT("Material_MaterialTranslateCalls"), MaterialTranslateCalls);
 	Attributes.Emplace(TEXT("Material_MaterialTranslateTimeSec"), MaterialTranslateTimeSec);
+	Attributes.Emplace(TEXT("Material_MaterialCacheTimeSec"), MaterialCacheTimeSec);
+	Attributes.Emplace(TEXT("Material_MaterialCacheHits"), MaterialCacheHits);
+	Attributes.Emplace(TEXT("Material_MaterialTranslationSkippedDDCCount"), MaterialTranslationSkippedDDCCount);
 }
 
 void FShaderCompilerStats::IncrementMaterialCook()
@@ -5188,16 +5197,23 @@ void FShaderCompilerStats::IncrementMaterialCook()
 	MaterialCounters.NumMaterialsCooked++;
 }
 
-void FShaderCompilerStats::IncrementMaterialsTranslated()
+void FShaderCompilerStats::IncrementMaterialTranslated(double InTime)
 {
 	FScopeLock Lock(&CompileStatsLock);
 	MaterialCounters.MaterialTranslateCalls++;
+	MaterialCounters.MaterialTranslateTimeSec += InTime;
 }
 
-void FShaderCompilerStats::IncrementMaterialTranslateTime(double InTime)
+void FShaderCompilerStats::IncrementMaterialCacheHit(double InCacheTime)
 {
 	FScopeLock Lock(&CompileStatsLock);
-	MaterialCounters.MaterialTranslateTimeSec += InTime;
+	MaterialCounters.MaterialCacheHits++;
+	MaterialCounters.MaterialCacheTimeSec += InCacheTime;
+}
+
+void FShaderCompilerStats::IncrementMaterialTranslationSkippedDDC()
+{
+	MaterialCounters.MaterialTranslationSkippedDDCCount++;
 }
 
 void FShaderCompilerStats::RegisterCookedShaders(uint32 NumCooked, float CompileTime, EShaderPlatform Platform, const FString MaterialPath, FString PermutationString)
diff --git a/Engine/Source/Runtime/Engine/Public/MaterialShaderType.h b/Engine/Source/Runtime/Engine/Public/MaterialShaderType.h
--- a/Engine/Source/Runtime/Engine/Public/MaterialShaderType.h
+++ b/Engine/Source/Runtime/Engine/Public/MaterialShaderType.h
@@ -51,6 +51,9 @@ extern ENGINE_API FString GetShadingModelFieldString(FMaterialShadingModelField
 /** Converts an EBlendMode to a string description. */
 extern ENGINE_API FString GetBlendModeString(EBlendMode BlendMode);
 
+/** Creates a string key for the derived data cache given a shader map id. */
+extern ENGINE_API FString GetMaterialShaderMapKeyString(const FMaterialShaderMapId& ShaderMapId, EShaderPlatform Platform, bool bIncludeKeyStringShaderDependencies = true);
+
 /** Called for every material shader to update the appropriate stats. */
 extern void UpdateMaterialShaderCompilingStats(const FMaterial* Material);
 
diff --git a/Engine/Source/Runtime/Engine/Public/MaterialShared.h b/Engine/Source/Runtime/Engine/Public/MaterialShared.h
--- a/Engine/Source/Runtime/Engine/Public/MaterialShared.h
+++ b/Engine/Source/Runtime/Engine/Public/MaterialShared.h
@@ -1083,7 +1083,7 @@ public:
 	void UpdateFromParameterSet(const FStaticParameterSet& StaticParameters);
 
 	/** Appends string representations of this Id to a key string. */
-	void AppendKeyString(FString& KeyString, bool bIncludeSourceAndMaterialState = true) const;
+	void AppendKeyString(FString& KeyString, bool bIncludeSourceAndMaterialState = true, bool bIncludeKeyStringShaderDependencies = true) const;
 	void AppendStaticParametersString(FString& ParamsString) const;
 
 	const TArray<FStaticSwitchParameter> &GetStaticSwitchParameters() const 					{ return StaticSwitchParameters; }
diff --git a/Engine/Source/Runtime/Engine/Public/ShaderCompiler.h b/Engine/Source/Runtime/Engine/Public/ShaderCompiler.h
--- a/Engine/Source/Runtime/Engine/Public/ShaderCompiler.h
+++ b/Engine/Source/Runtime/Engine/Public/ShaderCompiler.h
@@ -367,8 +367,9 @@ public:
 	};
 
 	void IncrementMaterialCook();
-	void IncrementMaterialsTranslated();
-	void IncrementMaterialTranslateTime(double InTime);
+	void IncrementMaterialTranslated(double InTime);
+	void IncrementMaterialCacheHit(double InTime);
+	void IncrementMaterialTranslationSkippedDDC();
 
 	ENGINE_API void RegisterCookedShaders(uint32 NumCooked, float CompileTime, EShaderPlatform Platform, const FString MaterialPath, FString PermutationString = FString(""));
 	ENGINE_API void RegisterCompiledShaders(uint32 NumPermutations, EShaderPlatform Platform, const FString MaterialPath, FString PermutationString = FString(""));
@@ -566,13 +567,24 @@ private:
 		int32 MaterialTranslateCalls = 0;
 
 		/** The total time in seconds to translate all materials.  */
-		double MaterialTranslateTimeSec = 0.0f;
+		double MaterialTranslateTimeSec = 0.0;
+
+		/** The total number times a material translation was skipped because the the results were in the DDC. */
+		int32 MaterialCacheHits = 0;
+		
+		/** The total time spent querying the DDC cache. */
+		double MaterialCacheTimeSec = 0.0;
+
+		int32 MaterialTranslationSkippedDDCCount = 0;
 
 		FMaterialCounters& operator+=(const FMaterialCounters& Other)
 		{
 			NumMaterialsCooked += Other.NumMaterialsCooked;
 			MaterialTranslateCalls += Other.MaterialTranslateCalls;
 			MaterialTranslateTimeSec += Other.MaterialTranslateTimeSec;
+			MaterialCacheHits += Other.MaterialCacheHits;
+			MaterialCacheTimeSec += Other.MaterialCacheTimeSec;
+			MaterialTranslationSkippedDDCCount += Other.MaterialTranslationSkippedDDCCount;
 
 			return *this;
 		}
