World Partition Runtime HashSet: fix content bundles external streaming objects not getting the proper object name, resulting in no content bundles loaded in cooked builds (worked fine in PIE).

diff --git a/Engine/Source/Runtime/Engine/Private/WorldPartition/RuntimeHashSet/WorldPartitionRuntimeHashSet.cpp b/Engine/Source/Runtime/Engine/Private/WorldPartition/RuntimeHashSet/WorldPartitionRuntimeHashSet.cpp
--- a/Engine/Source/Runtime/Engine/Private/WorldPartition/RuntimeHashSet/WorldPartitionRuntimeHashSet.cpp
+++ b/Engine/Source/Runtime/Engine/Private/WorldPartition/RuntimeHashSet/WorldPartitionRuntimeHashSet.cpp
@@ -240,7 +240,7 @@ URuntimeHashExternalStreamingObjectBase* UWorldPartitionRuntimeHashSet::StoreToE
 {
 	check(!RuntimeStreamingData.IsEmpty());
 
-	URuntimeHashSetExternalStreamingObject* NewStreamingObject = CreateExternalStreamingObject<URuntimeHashSetExternalStreamingObject>(this, MakeUniqueObjectName(this, URuntimeHashExternalStreamingObjectBase::StaticClass()));
+	URuntimeHashSetExternalStreamingObject* NewStreamingObject = CreateExternalStreamingObject<URuntimeHashSetExternalStreamingObject>(StreamingObjectOuter, StreamingObjectName);
 	NewStreamingObject->RuntimeStreamingData = MoveTemp(RuntimeStreamingData);
 
 	for (FRuntimePartitionStreamingData& StreamingData : NewStreamingObject->RuntimeStreamingData)
