Turn off incorrect check in VT Builder ; not true wth padding options on texture

#rnx

diff --git a/Engine/Source/Runtime/Engine/Private/Texture.cpp b/Engine/Source/Runtime/Engine/Private/Texture.cpp
--- a/Engine/Source/Runtime/Engine/Private/Texture.cpp
+++ b/Engine/Source/Runtime/Engine/Private/Texture.cpp
@@ -546,7 +546,8 @@ void UTexture::ValidateSettingsAfterImportOrEdit(bool * pRequiresNotifyMaterials
 		ResizeDuringBuildX = FMath::Max(0, FMath::Min((int32)GetMaximumDimension(), ResizeDuringBuildX));
 		ResizeDuringBuildY = FMath::Max(0, FMath::Min((int32)GetMaximumDimension(), ResizeDuringBuildY));
 
-		// IsPowerOfTwo only checks XY :
+		// IsPowerOfTwo only checks XY : and only checks BlockIndex 0 
+		//	<- this is wrong for VT, need to check if all blocks are Pow2
 		bool bIsPowerOfTwo = Source.IsPowerOfTwo();
 		if ( ! FMath::IsPowerOfTwo(Source.GetVolumeSizeZ()) )
 		{
@@ -2576,6 +2577,8 @@ int64 FTextureSource::GetBytesPerPixel(int32 LayerIndex) const
 	return GetBytesPerPixel(GetFormat(LayerIndex));
 }
 
+// IsPowerOfTwo() with no BlockIndex just checks the first block (and only XY)
+//	not the overall dimensions of UDIM
 bool FTextureSource::IsPowerOfTwo(int32 BlockIndex) const
 {
 	FTextureSourceBlock Block;
@@ -3027,6 +3030,8 @@ int64 FTextureSource::CalcMipOffset(int32 BlockIndex, int32 LayerIndex, int32 Of
 	return MipOffset.Get(0);
 }
 
+// UseHashAsGuid is done automatically in UnlockMip, should not usually be called directly
+//	calling it multiple times does not re-hash the data; it's harmless
 void FTextureSource::UseHashAsGuid()
 {
 	if (HasPayloadData())
diff --git a/Engine/Source/Runtime/Engine/Private/VT/VirtualTextureDataBuilder.cpp b/Engine/Source/Runtime/Engine/Private/VT/VirtualTextureDataBuilder.cpp
--- a/Engine/Source/Runtime/Engine/Private/VT/VirtualTextureDataBuilder.cpp
+++ b/Engine/Source/Runtime/Engine/Private/VT/VirtualTextureDataBuilder.cpp
@@ -1049,8 +1049,9 @@ void FVirtualTextureDataBuilder::BuildLayerBlocks(FSlowTask& BuildTask, uint32 L
 
 		// We skip the first compressed mip output, since that will just be a copy of the input
 		check(CompressedMips.Num() >= BlockData.NumMips + 1);
-		check(BlockData.SizeX == CompressedMips[1].SizeX);
-		check(BlockData.SizeY == CompressedMips[1].SizeY);
+		// not true with padding options :
+		//check(BlockData.SizeX == CompressedMips[1].SizeX);
+		//check(BlockData.SizeY == CompressedMips[1].SizeY);
 
 		BlockData.Mips.Reserve(CompressedMips.Num() - 1);
 		for (int32 MipIndex = 1; MipIndex < BlockData.NumMips + 1; ++MipIndex)
