Allow warnings as errors cvar to be set per platform

diff --git a/Engine/Source/Runtime/Engine/Private/ShaderCompiler/ShaderCompiler.cpp b/Engine/Source/Runtime/Engine/Private/ShaderCompiler/ShaderCompiler.cpp
--- a/Engine/Source/Runtime/Engine/Private/ShaderCompiler/ShaderCompiler.cpp
+++ b/Engine/Source/Runtime/Engine/Private/ShaderCompiler/ShaderCompiler.cpp
@@ -7848,8 +7848,9 @@ void GlobalBeginCompileShader(
 		SET_SHADER_DEFINE(Input.Environment, DO_GUARD_SLOW, GSShaderCheckLevel > 1 ? 1 : 0);
 	}
 
-	if (int WarnLevel = CVarShaderWarningsAsErrors.GetValueOnAnyThread(); WarnLevel != 0)
 	{
+		static FShaderPlatformCachedIniValue<int32> CVarWarningsAsErrorsPerPlatform(TEXT("r.Shaders.WarningsAsErrors"));
+		const int WarnLevel = CVarWarningsAsErrorsPerPlatform.Get(ShaderPlatform);
 		if ((WarnLevel == 1 && ShaderType->GetTypeForDynamicCast() == FShaderType::EShaderTypeForDynamicCast::Global) || WarnLevel > 1)
 		{
 			Input.Environment.CompilerFlags.Add(CFLAG_WarningsAsErrors);
diff --git a/Engine/Source/Runtime/RenderCore/Private/Shader.cpp b/Engine/Source/Runtime/RenderCore/Private/Shader.cpp
--- a/Engine/Source/Runtime/RenderCore/Private/Shader.cpp
+++ b/Engine/Source/Runtime/RenderCore/Private/Shader.cpp
@@ -1614,8 +1614,11 @@ void ShaderMapAppendKeyString(EShaderPlatform Platform, FString& KeyString)
 	}
 
 	{
-		static const auto CVar = IConsoleManager::Get().FindConsoleVariable(TEXT("r.Shaders.WarningsAsErrors"));
-		KeyString += (CVar && CVar->GetInt() == 1) ? TEXT("_WX") : TEXT("");
+		static FShaderPlatformCachedIniValue<int32> CVarWarningsAsErrorsPerPlatform(TEXT("r.Shaders.WarningsAsErrors"));
+		if (const int32 Level = CVarWarningsAsErrorsPerPlatform.Get(Platform); Level != 0)
+		{
+			KeyString.Appendf(TEXT("_WX%d"), Level);
+		}
 	}
 	
 	{
