Fix VT Builder crash when first block is pow2 but later blocks are not
Texture IsPowerOfTwo check was confusingly only checking first block
most call sites want to check AreAllBlocksPowerOfTwo
deprecate the old problematic function name

#rnx

diff --git a/Engine/Plugins/Enterprise/MDLImporter/Source/MDLImporter/Private/generator/MaterialTextureFactory.cpp b/Engine/Plugins/Enterprise/MDLImporter/Source/MDLImporter/Private/generator/MaterialTextureFactory.cpp
--- a/Engine/Plugins/Enterprise/MDLImporter/Source/MDLImporter/Private/generator/MaterialTextureFactory.cpp
+++ b/Engine/Plugins/Enterprise/MDLImporter/Source/MDLImporter/Private/generator/MaterialTextureFactory.cpp
@@ -133,7 +133,7 @@ namespace Generator
 	{
 		TextureMipGenSettings MipGenSettings = Property.MipGenSettings;
 
-		if (!Texture->Source.IsPowerOfTwo())
+		if (!Texture->Source.AreAllBlocksPowerOfTwo())
 		{
 			MipGenSettings = TMGS_NoMipmaps;
 
diff --git a/Engine/Plugins/Interchange/Runtime/Source/Import/Private/Texture/InterchangeTextureFactory.cpp b/Engine/Plugins/Interchange/Runtime/Source/Import/Private/Texture/InterchangeTextureFactory.cpp
--- a/Engine/Plugins/Interchange/Runtime/Source/Import/Private/Texture/InterchangeTextureFactory.cpp
+++ b/Engine/Plugins/Interchange/Runtime/Source/Import/Private/Texture/InterchangeTextureFactory.cpp
@@ -1507,7 +1507,7 @@ namespace UE::Interchange::Private::InterchangeTextureFactory
 
 #if WITH_EDITORONLY_DATA
 		// At some point, these decisions should be left to the pipeline
-		if (!Texture->Source.IsPowerOfTwo())
+		if ( Texture->Source.GetNumBlocks() == 1 && !Texture->Source.IsBlockPowerOfTwo(0))
 		{
 			// try to set some better default options for non-pow2 textures
 
diff --git a/Engine/Source/Editor/UnrealEd/Private/Factories/EditorFactories.cpp b/Engine/Source/Editor/UnrealEd/Private/Factories/EditorFactories.cpp
--- a/Engine/Source/Editor/UnrealEd/Private/Factories/EditorFactories.cpp
+++ b/Engine/Source/Editor/UnrealEd/Private/Factories/EditorFactories.cpp
@@ -4155,7 +4155,7 @@ UObject* UTextureFactory::FactoryCreateBinary
 	
 	Texture->bPreserveBorder		= bPreserveBorder;
 
-	if ( ! Texture->Source.IsPowerOfTwo() )
+	if ( Texture->Source.GetNumBlocks() == 1 && ! Texture->Source.IsBlockPowerOfTwo(0) )
 	{
 		// try to set some better default options for non-pow2 textures
 
diff --git a/Engine/Source/Editor/VirtualTexturingEditor/Private/SConvertToVirtualTexture.cpp b/Engine/Source/Editor/VirtualTexturingEditor/Private/SConvertToVirtualTexture.cpp
--- a/Engine/Source/Editor/VirtualTexturingEditor/Private/SConvertToVirtualTexture.cpp
+++ b/Engine/Source/Editor/VirtualTexturingEditor/Private/SConvertToVirtualTexture.cpp
@@ -461,7 +461,7 @@ void SConvertToVirtualTexture::UpdateList()
 		FConversionStatus* Status = new(AssetStatus) FConversionStatus();
 		Status->UserSelected = Worker.UserTextures.Contains(Texture);
 		Status->UnderSized = CheckUnderSized(Texture,ThresholdValue,bBackwards);
-		Status->NonPowerOf2 = !Texture->Source.IsPowerOfTwo() && Texture->PowerOfTwoMode == ETexturePowerOfTwoSetting::None;
+		Status->NonPowerOf2 = !Texture->Source.AreAllBlocksPowerOfTwo() && Texture->PowerOfTwoMode == ETexturePowerOfTwoSetting::None;
 	}
 
 	for (UTexture2D* Texture : Worker.MaterialRejectedTextures)
@@ -470,7 +470,7 @@ void SConvertToVirtualTexture::UpdateList()
 		FConversionStatus* Status = new(AssetStatus) FConversionStatus();
 		Status->UserSelected = Worker.UserTextures.Contains(Texture);
 		Status->UnderSized = CheckUnderSized(Texture,ThresholdValue,bBackwards);
-		Status->NonPowerOf2 = !Texture->Source.IsPowerOfTwo() && Texture->PowerOfTwoMode == ETexturePowerOfTwoSetting::None;
+		Status->NonPowerOf2 = !Texture->Source.AreAllBlocksPowerOfTwo() && Texture->PowerOfTwoMode == ETexturePowerOfTwoSetting::None;
 		Status->InvalidMaterialUsage = true;
 	}
 
diff --git a/Engine/Source/Editor/VirtualTexturingEditor/Private/VirtualTextureConversionWorker.cpp b/Engine/Source/Editor/VirtualTexturingEditor/Private/VirtualTextureConversionWorker.cpp
--- a/Engine/Source/Editor/VirtualTexturingEditor/Private/VirtualTextureConversionWorker.cpp
+++ b/Engine/Source/Editor/VirtualTexturingEditor/Private/VirtualTextureConversionWorker.cpp
@@ -157,7 +157,7 @@ void FVirtualTextureConversionWorker::FindAllTexturesAndMaterials_Iteration(TSet
 			// since we're not changing the original texture, we don't need to bring in any additional dependencies
 			// Non-power-2 textures won't convert to VT, don't bring any dependencies for them
 			if (!Tex2D->GetPathName().StartsWith("/Engine/") &&
-				(Tex2D->Source.IsPowerOfTwo() || Tex2D->PowerOfTwoMode != ETexturePowerOfTwoSetting::None))
+				(Tex2D->Source.AreAllBlocksPowerOfTwo() || Tex2D->PowerOfTwoMode != ETexturePowerOfTwoSetting::None))
 			{
 				// Also get any preview materials that reference the given texture
 				// We need to convert these to ensure any active material editors remain valid
@@ -483,7 +483,7 @@ void FVirtualTextureConversionWorker::FilterList(int32 SizeThreshold)
 	{
 		bool DoInclude;
 
-		if ( !Texture->Source.IsPowerOfTwo() && Texture->PowerOfTwoMode == ETexturePowerOfTwoSetting::None )
+		if ( !Texture->Source.AreAllBlocksPowerOfTwo() && Texture->PowerOfTwoMode == ETexturePowerOfTwoSetting::None )
 		{
 			// not pow2, reject
 			DoInclude = false;
diff --git a/Engine/Source/Runtime/Engine/Classes/Engine/Texture.h b/Engine/Source/Runtime/Engine/Classes/Engine/Texture.h
--- a/Engine/Source/Runtime/Engine/Classes/Engine/Texture.h
+++ b/Engine/Source/Runtime/Engine/Classes/Engine/Texture.h
@@ -289,7 +289,12 @@ struct FTextureSource
 	ENGINE_API int64 GetBytesPerPixel(int32 LayerIndex = 0) const;
 
 	/** Return true if the source XY size is power-of-2.  Does not check Z size for volumes.  */
-	ENGINE_API bool IsPowerOfTwo(int32 BlockIndex = 0) const;
+	UE_DEPRECATED(5.5,"Prefer AreAllBlocksPowerOfTwo, or IsBlockPowerOfTwo if you really only want one block")
+	ENGINE_API bool IsPowerOfTwo(int32 BlockIndex = 0) const { return IsBlockPowerOfTwo(BlockIndex); }
+	
+	/** Return true if the source XY size is power-of-2.  Does not check Z size for volumes.  */
+	ENGINE_API bool IsBlockPowerOfTwo(int32 BlockIndex) const;
+	ENGINE_API bool AreAllBlocksPowerOfTwo() const;
 
 	/** Returns true if source art is available. */
 	ENGINE_API bool IsValid() const;
diff --git a/Engine/Source/Runtime/Engine/Private/Texture.cpp b/Engine/Source/Runtime/Engine/Private/Texture.cpp
--- a/Engine/Source/Runtime/Engine/Private/Texture.cpp
+++ b/Engine/Source/Runtime/Engine/Private/Texture.cpp
@@ -546,9 +546,8 @@ void UTexture::ValidateSettingsAfterImportOrEdit(bool * pRequiresNotifyMaterials
 		ResizeDuringBuildX = FMath::Max(0, FMath::Min((int32)GetMaximumDimension(), ResizeDuringBuildX));
 		ResizeDuringBuildY = FMath::Max(0, FMath::Min((int32)GetMaximumDimension(), ResizeDuringBuildY));
 
-		// IsPowerOfTwo only checks XY : and only checks BlockIndex 0 
-		//	<- this is wrong for VT, need to check if all blocks are Pow2
-		bool bIsPowerOfTwo = Source.IsPowerOfTwo();
+		// IsPowerOfTwo only checks XY
+		bool bIsPowerOfTwo = Source.AreAllBlocksPowerOfTwo();
 		if ( ! FMath::IsPowerOfTwo(Source.GetVolumeSizeZ()) )
 		{
 			bIsPowerOfTwo = false;
@@ -1609,7 +1608,7 @@ bool UTexture::IsPossibleToStream() const
 		// duplicate the checks done for NeverStream :
 
 		// IsPowerOfTwo only checks XY :
-		bool bIsPowerOfTwo = Source.IsPowerOfTwo();
+		bool bIsPowerOfTwo = Source.AreAllBlocksPowerOfTwo();
 		if ( ! FMath::IsPowerOfTwo(Source.GetVolumeSizeZ()) )
 		{
 			bIsPowerOfTwo = false;
@@ -2577,15 +2576,26 @@ int64 FTextureSource::GetBytesPerPixel(int32 LayerIndex) const
 	return GetBytesPerPixel(GetFormat(LayerIndex));
 }
 
-// IsPowerOfTwo() with no BlockIndex just checks the first block (and only XY)
-//	not the overall dimensions of UDIM
-bool FTextureSource::IsPowerOfTwo(int32 BlockIndex) const
+bool FTextureSource::IsBlockPowerOfTwo(int32 BlockIndex) const
 {
 	FTextureSourceBlock Block;
 	GetBlock(BlockIndex, Block);
 	return FMath::IsPowerOfTwo(Block.SizeX) && FMath::IsPowerOfTwo(Block.SizeY);
 }
 
+bool FTextureSource::AreAllBlocksPowerOfTwo() const
+{
+	int32 NumBlocks = GetNumBlocks();
+	for(int32 BlockIndex=0;BlockIndex<NumBlocks;BlockIndex++)
+	{
+		if ( ! IsBlockPowerOfTwo(BlockIndex) )
+		{
+			return false;
+		}
+	}
+	return true;
+}
+
 bool FTextureSource::IsValid() const
 {
 	return SizeX > 0 && SizeY > 0 && NumSlices > 0 && NumLayers > 0 && NumMips > 0 &&
