Skip the update of the replicated rigid state of a cluster union if our value is "Uninitialized" which means we haven't actually received an update to apply

diff --git a/Engine/Source/Runtime/Engine/Private/PhysicsEngine/ClusterUnionComponent.cpp b/Engine/Source/Runtime/Engine/Private/PhysicsEngine/ClusterUnionComponent.cpp
--- a/Engine/Source/Runtime/Engine/Private/PhysicsEngine/ClusterUnionComponent.cpp
+++ b/Engine/Source/Runtime/Engine/Private/PhysicsEngine/ClusterUnionComponent.cpp
@@ -43,6 +43,8 @@ namespace
 	bool bFlushNetDormancyOnSyncProxy = true;
 	FAutoConsoleVariableRef CVarFlushNetDormancyOnSyncProxy(TEXT("ClusterUnion.FlushNetDormancyOnSyncProxy"), bFlushNetDormancyOnSyncProxy, TEXT("When there is a new rigid state on the authority, flush net dormancy so that even if this object is net dorman the rigid state will come through to the client."));
 
+	bool GSkipZeroStateInOnRep = true;
+	FAutoConsoleVariableRef CVarSkipZeroState(TEXT("ClusterUnion.SkipZeroStateInOnRep"), GSkipZeroStateInOnRep, TEXT("Whether we skip 0 (uninitialized) states when running the onrep for the replicated rigid state of the cluster union"));
 
 	template<typename PayloadType>
 	struct TClusterUnionAABBTreeStorageTraits
@@ -1296,7 +1298,11 @@ void UClusterUnionComponent::BroadcastComponentRemovedEvents(UPrimitiveComponent
 
 void UClusterUnionComponent::OnRep_RigidState()
 {
-	if (!PhysicsProxy)
+	// Prior to having a proxy we can't set its state.
+	// Also, if our state is 0 we have not yet recieved a valid object state from
+	// the server (we call this onrep in OnCreatePhysicsState too incase we got a
+	// new state but didn't have a proxy at that time)
+	if (!PhysicsProxy || (ReplicatedRigidState.ObjectState == 0 && GSkipZeroStateInOnRep))
 	{
 		return;
 	}
