Predictively update the newly cloned prefab to the old prefab's child to parent transform
(on behalf of todd.eckert)

diff --git a/Engine/Source/Runtime/Engine/Classes/PhysicsEngine/ClusterUnionComponent.h b/Engine/Source/Runtime/Engine/Classes/PhysicsEngine/ClusterUnionComponent.h
--- a/Engine/Source/Runtime/Engine/Classes/PhysicsEngine/ClusterUnionComponent.h
+++ b/Engine/Source/Runtime/Engine/Classes/PhysicsEngine/ClusterUnionComponent.h
@@ -347,12 +347,14 @@ public:
 
 	friend class UClusterUnionReplicatedProxyComponent;
 	friend class UModularVehicleBaseComponent;
-protected:
 
 	// This should only be called on the client when replication happens.
 	UFUNCTION()
 	ENGINE_API void ForceSetChildToParent(UPrimitiveComponent* InComponent, const TArray<int32>& BoneIds, const TArray<FTransform>& ChildToParent);
 
+	ENGINE_API virtual void OnChildToParentUpdated(UPrimitiveComponent* ChangedComponent, const FLocalBonesToTransformMap& PerBoneChildToParent) {}
+protected:
+
 	ENGINE_API void BroadcastComponentAddedEvents(UPrimitiveComponent* ChangedComponent, const TArray<FClusterUnionBoneData>& BoneIds, bool bIsNew, const TArray<FClusterUnionBoneData>& RemovedBoneIDs);
 	ENGINE_API void BroadcastComponentRemovedEvents(UPrimitiveComponent* ChangedComponent, const TArray<FClusterUnionBoneData>& InRemovedBonesData);
 
@@ -360,7 +362,6 @@ protected:
 	Chaos::FClusterUnionPhysicsProxy* GetPhysicsProxy() { return PhysicsProxy; }
 
 	ENGINE_API virtual void SetRigidState(Chaos::EObjectStateType ObjectState);
-	ENGINE_API virtual void OnChildToParentUpdated(UPrimitiveComponent* ChangedComponent, const FLocalBonesToTransformMap& PerBoneChildToParent) {}
 
 	// We need to keep track of the mapping of primitive components to physics objects.
 	// This way we know the right physics objects to pass when removing the component (because
diff --git a/Engine/Source/Runtime/Engine/Classes/PhysicsEngine/ClusterUnionReplicatedProxyComponent.h b/Engine/Source/Runtime/Engine/Classes/PhysicsEngine/ClusterUnionReplicatedProxyComponent.h
--- a/Engine/Source/Runtime/Engine/Classes/PhysicsEngine/ClusterUnionReplicatedProxyComponent.h
+++ b/Engine/Source/Runtime/Engine/Classes/PhysicsEngine/ClusterUnionReplicatedProxyComponent.h
@@ -58,6 +58,8 @@ public:
 	/* reinitialize the internal transient state of the replicated proxy */
 	ENGINE_API void ResetTransientState();
 
+	ENGINE_API const TArray<FTransform>& GetParticleChildToParents() { return ParticleChildToParents; }
+
 protected:
 
 	UFUNCTION()
