Fixes to ensure the Simulation State is set to kinematic on Cluster Unions the client when the player is outside of the configured simulation range

- Now we block any update to the rigid state if we are out of relevancy
- When we go outside the simulation range, we make the CU anchored

#rnx
#tests Tested in a cooked client and server, verified CUs become kinematic far away from the player and dynamic when they are close. Tested this with newly placed platforms, persistence, and POIs

diff --git a/Engine/Source/Runtime/Engine/Classes/PhysicsEngine/ClusterUnionComponent.h b/Engine/Source/Runtime/Engine/Classes/PhysicsEngine/ClusterUnionComponent.h
--- a/Engine/Source/Runtime/Engine/Classes/PhysicsEngine/ClusterUnionComponent.h
+++ b/Engine/Source/Runtime/Engine/Classes/PhysicsEngine/ClusterUnionComponent.h
@@ -374,6 +374,8 @@ protected:
 	UFUNCTION()
 	ENGINE_API bool IsAuthority() const;
 
+	const FClusterUnionReplicatedData& GetReplicatedRigidState() const { return ReplicatedRigidState; }
+
 private:
 	// These are the statically clustered components. These should
 	// be specified in the editor and never change.
@@ -411,10 +413,13 @@ private:
 	mutable bool bHasCachedLocalBounds;
 	mutable FBoxSphereBounds CachedLocalBounds;
 
+protected:
 	// Handles changes to ReplicatedRigidState. Note that this function does not handle replication of X/R since we make use
 	// of the scene component's default replication for that.
 	UFUNCTION()
-	ENGINE_API void OnRep_RigidState();
+	ENGINE_API virtual void OnRep_RigidState();
+
+private:
 
 	ENGINE_API FPhysScene_Chaos* GetChaosScene() const;
 
