Fix ClusterUnion child transforms for components added while asleep

We were relying on the cluster union to call MoveComponent on all components in the cluster, but this does not happen is the cluster is asleep or has not moved. So if a new componet was added (on the client), but the cluster does not wake or move (or OnRep_ChildToParent is called after it moved) the new component will have whatever transform it was spawned with, which is usually the world transfomr from the server. This meant that the visials and game thread collisions were not in the correct location which can caus eproblem for system that depend on it.

diff --git a/Engine/Source/Runtime/Engine/Classes/PhysicsEngine/ClusterUnionComponent.h b/Engine/Source/Runtime/Engine/Classes/PhysicsEngine/ClusterUnionComponent.h
--- a/Engine/Source/Runtime/Engine/Classes/PhysicsEngine/ClusterUnionComponent.h
+++ b/Engine/Source/Runtime/Engine/Classes/PhysicsEngine/ClusterUnionComponent.h
@@ -298,7 +298,7 @@ public:
 	using FLocalBonesToTransformMap = TSortedMap<int32, FMappedBoneData, TInlineAllocator<1>>;
 
 	// SyncClusterUnionFromProxy will examine the make up of the cluster union (particles, child to parent, etc.) and do whatever is needed on the GT in terms of bookkeeping.
-	ENGINE_API void SyncClusterUnionFromProxy();
+	ENGINE_API void SyncClusterUnionFromProxy(const FTransform& NewTransform, TArray<TTuple<UPrimitiveComponent*, FTransform>>* OutNewComponents);
 
 	UFUNCTION()
 	bool IsComponentAdded(UPrimitiveComponent* Component) { return PerComponentData.Contains(Component) || PendingComponentSync.Contains(Component); }
@@ -352,7 +352,8 @@ public:
 	UFUNCTION()
 	ENGINE_API void ForceSetChildToParent(UPrimitiveComponent* InComponent, const TArray<int32>& BoneIds, const TArray<FTransform>& ChildToParent);
 
-	ENGINE_API virtual void OnChildToParentUpdated(UPrimitiveComponent* ChangedComponent, const FLocalBonesToTransformMap& PerBoneChildToParent) {}
+	ENGINE_API virtual void OnChildToParentUpdated(UPrimitiveComponent* ChangedComponent, const FLocalBonesToTransformMap& PerBoneChildToParent, const FTransform& NewTransform, TArray<TTuple<UPrimitiveComponent*, FTransform>>* OutNewComponents) {}
+
 protected:
 
 	ENGINE_API void BroadcastComponentAddedEvents(UPrimitiveComponent* ChangedComponent, const TArray<FClusterUnionBoneData>& BoneIds, bool bIsNew, const TArray<FClusterUnionBoneData>& RemovedBoneIDs);
@@ -455,7 +456,7 @@ private:
 
 	// These functions only get called when the physics thread syncs to the game thread thereby enforcing a physics thread authoritative view of
 	// what particles are currently contained within the cluster union.
-	ENGINE_API void HandleAddOrModifiedClusteredComponent(const FMappedComponentKey& ChangedComponentData, const FLocalBonesToTransformMap& PerBoneChildToParent);
+	ENGINE_API void HandleAddOrModifiedClusteredComponent(const FMappedComponentKey& ChangedComponentData, const FLocalBonesToTransformMap& PerBoneChildToParent, const FTransform& NewTransform, TArray<TTuple<UPrimitiveComponent*, FTransform>>* OutNewComponents);
 	ENGINE_API void HandleRemovedClusteredComponent(TObjectKey<UPrimitiveComponent> RemovedComponent, const FClusteredComponentData& ComponentData);
 
 	ENGINE_API TArray<UPrimitiveComponent*> GetAllCurrentChildComponents() const;
diff --git a/Engine/Source/Runtime/Engine/Private/PhysicsEngine/ClusterUnionComponent.cpp b/Engine/Source/Runtime/Engine/Private/PhysicsEngine/ClusterUnionComponent.cpp
--- a/Engine/Source/Runtime/Engine/Private/PhysicsEngine/ClusterUnionComponent.cpp
+++ b/Engine/Source/Runtime/Engine/Private/PhysicsEngine/ClusterUnionComponent.cpp
@@ -887,7 +887,7 @@ void UClusterUnionComponent::HandleComponentPhysicsStateChangePostAddIntoCluster
 }
 
 DECLARE_CYCLE_STAT(TEXT("UClusterUnionComponent::SyncClusterUnionFromProxy"), STAT_ClusterUnionComponent_SyncClusterUnionFromProxy, STATGROUP_Chaos);
-void UClusterUnionComponent::SyncClusterUnionFromProxy()
+void UClusterUnionComponent::SyncClusterUnionFromProxy(const FTransform& NewTransform, TArray<TTuple<UPrimitiveComponent*, FTransform>>* OutNewComponents)
 {
 	SCOPE_CYCLE_COUNTER(STAT_ClusterUnionComponent_SyncClusterUnionFromProxy);
 
@@ -1013,7 +1013,7 @@ void UClusterUnionComponent::SyncClusterUnionFromProxy()
 	// If a component lives in both, then it's a modified component.
 	for (const TPair<FMappedComponentKey, FLocalBonesToTransformMap>& Kvp : MappedData)
 	{
-		HandleAddOrModifiedClusteredComponent(Kvp.Key, Kvp.Value);
+		HandleAddOrModifiedClusteredComponent(Kvp.Key, Kvp.Value, NewTransform, OutNewComponents);
 	}
 
 	// If a component lives in PerComponentData but not in MappedData, deleted component!
@@ -1047,7 +1047,7 @@ void UClusterUnionComponent::SyncClusterUnionFromProxy()
 	}
 }
 
-void UClusterUnionComponent::HandleAddOrModifiedClusteredComponent(const FMappedComponentKey& ChangedComponentData, const FLocalBonesToTransformMap& PerBoneChildToParent)
+void UClusterUnionComponent::HandleAddOrModifiedClusteredComponent(const FMappedComponentKey& ChangedComponentData, const FLocalBonesToTransformMap& PerBoneChildToParent, const FTransform& NewTransform, TArray<TTuple<UPrimitiveComponent*, FTransform>>* OutNewComponents)
 {
 	if (!ChangedComponentData.ComponentPtr || !ChangedComponentData.ComponentPtr->HasValidPhysicsState() || !ChangedComponentData.ComponentPtr->GetWorld())
 	{
@@ -1158,7 +1158,7 @@ void UClusterUnionComponent::HandleAddOrModifiedClusteredComponent(const FMapped
 		}
 	}
 
-	OnChildToParentUpdated(ChangedComponentData.ComponentPtr, PerBoneChildToParent);
+	OnChildToParentUpdated(ChangedComponentData.ComponentPtr, PerBoneChildToParent, NewTransform, OutNewComponents);
 
 	if (PerBoneChildToParent.Num())
 	{
diff --git a/Engine/Source/Runtime/Engine/Private/PhysicsEngine/Experimental/PhysScene_Chaos.cpp b/Engine/Source/Runtime/Engine/Private/PhysicsEngine/Experimental/PhysScene_Chaos.cpp
--- a/Engine/Source/Runtime/Engine/Private/PhysicsEngine/Experimental/PhysScene_Chaos.cpp
+++ b/Engine/Source/Runtime/Engine/Private/PhysicsEngine/Experimental/PhysScene_Chaos.cpp
@@ -62,6 +62,9 @@ FAutoConsoleVariableRef CVar_KinematicDeferralLogInvalidBodies(TEXT("p.Kinematic
 float GReplicationCacheLingerForNSeconds = 3.f;
 FAutoConsoleVariableRef CVar_ReplicationCacheLingerForNSeconds(TEXT("np2.ReplicationCache.LingerForNSeconds"), GReplicationCacheLingerForNSeconds, TEXT("How long to keep data in the replication cache without the actor accessing it, after this we stop caching the actors state until it tries to access it again."));
 
+bool bGClusterUnionSyncBodiesMoveNewComponents = false;
+FAutoConsoleVariableRef CVar_GClusterUnionSyncBodiesCheckDirtyFlag(TEXT("p.ClusterUnion.SyncBodiesMoveNewComponents"), bGClusterUnionSyncBodiesMoveNewComponents, TEXT("Enable a fix to ensure new components in a cluster union are moved once on add (even if the cluster is not moving)."));
+
 DECLARE_CYCLE_STAT(TEXT("Update Kinematics On Deferred SkelMeshes"), STAT_UpdateKinematicsOnDeferredSkelMeshesChaos, STATGROUP_Physics);
 
 struct FPendingAsyncPhysicsCommand
@@ -2355,7 +2358,27 @@ void FPhysScene_Chaos::OnSyncBodies(Chaos::FPhysicsSolverBase* Solver)
 					PendingTransforms.Add(FPhysScenePendingComponentTransform_Chaos(ParentComponent, DirtyParticle->GetWakeEvent()));
 				}
 
-				ParentComponent->SyncClusterUnionFromProxy();
+				if (bHasMoved || !bGClusterUnionSyncBodiesMoveNewComponents)
+				{
+					ParentComponent->SyncClusterUnionFromProxy(NewTransform, nullptr);
+				}
+				else
+				{
+					// We must to call MoveComponent on any newly added component. The Cluster Union will take care of
+					// that but only if it moved. If it did not move we need to handle it manually
+					TArray<TTuple<UPrimitiveComponent*, FTransform>> NewComponents;
+					ParentComponent->SyncClusterUnionFromProxy(NewTransform, &NewComponents);
+
+					for (TTuple<UPrimitiveComponent*, FTransform>& NewComponentTuple : NewComponents)
+					{
+						UPrimitiveComponent* NewComponent = NewComponentTuple.Get<0>();
+						const FTransform& NewComponentTransform = NewComponentTuple.Get<1>();
+						const FVector NewComponentCurrentLocation = NewComponent->GetComponentLocation();
+						const FVector NewComponentDelta = NewComponentTransform.GetLocation() - NewComponentCurrentLocation;
+
+						PendingTransforms.Add(FPhysScenePendingComponentTransform_Chaos(NewComponent, NewComponentDelta, NewComponentTransform.GetRotation(), DirtyParticle->GetWakeEvent()));
+					}
+				}
 
 				// make sure we have at least a child to be added to the acceleration structure 
 				// this avoid the invalid bounds to cause the particle to be added to the global acceleration structure array
diff --git a/Engine/Source/Runtime/Experimental/Chaos/Private/Chaos/ClusterUnionManager.cpp b/Engine/Source/Runtime/Experimental/Chaos/Private/Chaos/ClusterUnionManager.cpp
--- a/Engine/Source/Runtime/Experimental/Chaos/Private/Chaos/ClusterUnionManager.cpp
+++ b/Engine/Source/Runtime/Experimental/Chaos/Private/Chaos/ClusterUnionManager.cpp
@@ -1361,6 +1361,8 @@ namespace Chaos
 
 		if (bMadeChanges)
 		{
+			MEvolution.GetParticles().MarkTransientDirtyParticle(ClusterUnion->InternalCluster, false);
+
 			constexpr EUpdateClusterUnionPropertiesFlags Flags = EUpdateClusterUnionPropertiesFlags::IncrementalGenerateConnectionGraph
 				| EUpdateClusterUnionPropertiesFlags::IncrementalGenerateGeometry;
 			RequestDeferredClusterPropertiesUpdate(ClusterIndex, Flags);
