[NetPhysics] Don't allow duplicate or out of order physics replication targets
Also raise the target alignment span via a CVar multiplier.

diff --git a/Engine/Source/Runtime/Engine/Private/PhysicsEngine/PhysicsReplication.cpp b/Engine/Source/Runtime/Engine/Private/PhysicsEngine/PhysicsReplication.cpp
--- a/Engine/Source/Runtime/Engine/Private/PhysicsEngine/PhysicsReplication.cpp
+++ b/Engine/Source/Runtime/Engine/Private/PhysicsEngine/PhysicsReplication.cpp
@@ -190,6 +190,12 @@ namespace PhysicsReplicationCVars
 		
 		static float SleepSecondsClearTarget = 15.0f;
 		static FAutoConsoleVariableRef CVarSleepSecondsClearTarget(TEXT("np2.PredictiveInterpolation.SleepSecondsClearTarget"), SleepSecondsClearTarget, TEXT("Wait for the object to sleep for this many seconds before clearing the replication target, to ensure nothing wakes up the object just after it goes to sleep on the client."));
+		
+		static int32 TargetTickAlignmentClampMultiplier = 2;
+		static FAutoConsoleVariableRef CVarTargetTickAlignmentClampMultiplier(TEXT("np2.PredictiveInterpolation.TargetTickAlignmentClampMultiplier"), TargetTickAlignmentClampMultiplier, TEXT("Multiplier to adjust clamping of target alignment via TickCount. Multiplier is performed on AverageReceiveInterval."));
+		
+		static bool LegacyTargetUpdateCheck = false;
+		static FAutoConsoleVariableRef CVarTargetLegacyTargetUpdateCheck(TEXT("np2.PredictiveInterpolation.LegacyTargetUpdateCheck"), LegacyTargetUpdateCheck, TEXT("Use Input.ServerFrame >= Target->PrevServerFrame when checking if a received input state should be cached as a target. This is the legacy version of the check."));
 	}
 
 }
@@ -856,12 +862,15 @@ void FPhysicsReplicationAsync::UpdateAsyncTarget(const FPhysicsRepAsyncInputData
 		Target->PrevLinVel = Input.TargetState.LinVel;
 	}
 
-	if (Input.ServerFrame >= Target->PrevServerFrame)
+	// Note: If target is waiting, check PrevServerFrame since ServerFrame has been modified and we still want to cache the input target
+	if (PhysicsReplicationCVars::PredictiveInterpolationCVars::LegacyTargetUpdateCheck
+		? (Input.ServerFrame >= Target->PrevServerFrame) // Legacy check
+		: (Input.ServerFrame == 0 || (Target->bWaiting ? Input.ServerFrame > Target->PrevServerFrame : Input.ServerFrame > Target->ServerFrame)))
 	{
 		const int32 PrevTickCount = Target->TickCount;
 		const int32 PrevReceiveInterval = Target->ReceiveInterval;
 		const int32 SendInterval = Input.ServerFrame - Target->ServerFrame;
-		const int32 AverageReceiveInterval = FMath::CeilToInt(Target->AverageReceiveInterval);
+		const int32 AdjustedAverageReceiveInterval = FMath::CeilToInt(Target->AverageReceiveInterval) * PhysicsReplicationCVars::PredictiveInterpolationCVars::TargetTickAlignmentClampMultiplier;
 
 		Target->PrevServerFrame = Target->bWaiting ? Input.ServerFrame : Target->ServerFrame;
 		Target->ServerFrame = Target->bWaiting ? Target->ServerFrame : Input.ServerFrame;
@@ -871,7 +880,7 @@ void FPhysicsReplicationAsync::UpdateAsyncTarget(const FPhysicsRepAsyncInputData
 		Target->TargetState = Input.TargetState;
 		Target->RepMode = Input.RepMode;
 		Target->FrameOffset = Input.FrameOffset;
-		Target->TickCount = bFirstTarget ? 0 : FMath::Clamp((PrevTickCount - (SendInterval > 0 ? SendInterval : PrevReceiveInterval)), -AverageReceiveInterval, AverageReceiveInterval);
+		Target->TickCount = bFirstTarget ? 0 : FMath::Clamp((PrevTickCount - (SendInterval > 0 ? SendInterval : PrevReceiveInterval)), -AdjustedAverageReceiveInterval, AdjustedAverageReceiveInterval);
 		Target->AccumulatedSleepSeconds = 0.0f;
 
 		if (Input.RepMode == EPhysicsReplicationMode::PredictiveInterpolation)
