[Insights] Fixed asset metadata association with mem allocations to match LLM.

diff --git a/Engine/Source/Runtime/Core/Public/ProfilingDebugging/AssetMetadataTrace.h b/Engine/Source/Runtime/Core/Public/ProfilingDebugging/AssetMetadataTrace.h
--- a/Engine/Source/Runtime/Core/Public/ProfilingDebugging/AssetMetadataTrace.h
+++ b/Engine/Source/Runtime/Core/Public/ProfilingDebugging/AssetMetadataTrace.h
@@ -8,21 +8,19 @@
 
 namespace UE { namespace Trace { class FChannel; } }
 
-////////////////////////////////////////////////////////////////////////////////////////////////////
-#if !defined(UE_TRACE_ASSET_METADATA_ENABLED) 
+#if !defined(UE_TRACE_ASSET_METADATA_ENABLED)
 	#define UE_TRACE_ASSET_METADATA_ENABLED UE_TRACE_METADATA_ENABLED
 #endif
 
-////////////////////////////////////////////////////////////////////////////////////////////////////
 #if UE_TRACE_ASSET_METADATA_ENABLED
 
 /**
- * Channel that asset meta data is output on
+ * Channel that asset metadata is output on
  */
 CORE_API UE_TRACE_CHANNEL_EXTERN(AssetMetadataChannel);
 
 /**
- * Metadata scope to instrument operations belonging to a certain asset.
+ * Metadata scope to instrument operations belonging to a certain asset
  */
 CORE_API UE_TRACE_METADATA_EVENT_BEGIN_EXTERN(Asset)
 	UE_TRACE_METADATA_EVENT_REFERENCE_FIELD(Strings, FName, Name)
@@ -31,15 +29,14 @@ CORE_API UE_TRACE_METADATA_EVENT_BEGIN_EXTERN(Asset)
 UE_TRACE_METADATA_EVENT_END()
 
 /**
- * Utility macro to create asset scope from and object and object class
- * @note When using this macro in an module outside of Core, make sure a dependency to the "TraceLog" module
- *		 is added.
+ * Utility macro to create asset scope from an object and object class
+ * @note When using this macro in a module outside of Core, make sure a dependency to the "TraceLog" module is added.
  */
 #define UE_TRACE_METADATA_SCOPE_ASSET(Object, ObjClass) \
-	UE_TRACE_METADATA_SCOPE_ASSET_FNAME(Object->GetFName(), ObjClass->GetFName(), Object->GetOutermost()->GetFName())
+	UE_TRACE_METADATA_SCOPE_ASSET_FNAME(Object->GetFName(), ObjClass->GetFName(), Object->GetPackage()->GetFName())
 
 /**
- * Utility macro to create an asset scope by specifying object name, class name and package name explicitly.
+ * Utility macro to create an asset scope by specifying object name, class name and package name explicitly
  */
 #define UE_TRACE_METADATA_SCOPE_ASSET_FNAME(ObjectName, ObjClassName, PackageName) \
 	auto MetaNameRef = UE_TRACE_CHANNELEXPR_IS_ENABLED(MetadataChannel) ? FStringTrace::GetNameRef(ObjectName) : UE::Trace::FEventRef32(0,0);\
@@ -50,7 +47,7 @@ UE_TRACE_METADATA_EVENT_END()
 		<< Asset.Class(ClassNameRef)\
 		<< Asset.Package(PackageNameRef);
 
-#else
+#else // UE_TRACE_ASSET_METADATA_ENABLED
 
 #define UE_TRACE_METADATA_SCOPE_ASSET(...)
 #define UE_TRACE_METADATA_SCOPE_ASSET_FNAME(...)
diff --git a/Engine/Source/Runtime/Core/Public/ProfilingDebugging/MetadataTrace.h b/Engine/Source/Runtime/Core/Public/ProfilingDebugging/MetadataTrace.h
--- a/Engine/Source/Runtime/Core/Public/ProfilingDebugging/MetadataTrace.h
+++ b/Engine/Source/Runtime/Core/Public/ProfilingDebugging/MetadataTrace.h
@@ -53,10 +53,11 @@ public:
 /**
  * A scope to push/pop data on the metadata stack.
  */
-#define UE_TRACE_METADATA_SCOPE(ScopeName, ChannelExpr) UE_TRACE_LOG_SCOPED(Metadata, ScopeName, ChannelExpr)
+#define UE_TRACE_METADATA_SCOPE(ScopeName, ChannelExpr) \
+	UE_TRACE_LOG_SCOPED(Metadata, ScopeName, ChannelExpr)
 
 /**
- * Temporarily clears the current meta data stack and starts a new one. When this scope is exited the previous
+ * Temporarily clears the current metadata stack and starts a new one. When this scope is exited the previous
  * stack is restored. This is useful for example in situations where an asset causes allocation in an unrelated
  * system, like shared buffers, but that allocation should not be ascribed to a specific asset.
  */
@@ -64,19 +65,19 @@ public:
 	UE_TRACE_LOG_SCOPED(MetadataStack, ClearScope, MetadataChannel)
 
 /**
- * Saves the currents stack, and returns an identifier. The stack can be restored on another thread using
+ * Saves the current metadata stack, and returns an identifier. The stack can be restored on another thread using
  * the identifier. This can be used when the context needs to be applied on a task.
  */
 #define UE_TRACE_METADATA_SAVE_STACK() \
 	FMetadataTrace::SaveStack();
 
 /**
- * A scope that restores a meta stack previously stored with UE_TRACE_METADATA_SAVE_STACK during it's lifetime.
+ * A scope that restores a metadata stack previously stored with UE_TRACE_METADATA_SAVE_STACK during it's lifetime.
  */
 class FMetadataRestoreScope
 {
 public:
-	CORE_API FMetadataRestoreScope(uint32 SavedMetadataIdentifier);
+	CORE_API FMetadataRestoreScope(uint32 InSavedMetadataIdentifier);
 	CORE_API ~FMetadataRestoreScope();
 private:
 	void ActivateScope(uint32 InStackId);
diff --git a/Engine/Source/Runtime/CoreUObject/Private/Serialization/AsyncLoading.cpp b/Engine/Source/Runtime/CoreUObject/Private/Serialization/AsyncLoading.cpp
--- a/Engine/Source/Runtime/CoreUObject/Private/Serialization/AsyncLoading.cpp
+++ b/Engine/Source/Runtime/CoreUObject/Private/Serialization/AsyncLoading.cpp
@@ -3163,22 +3163,23 @@ UObject* FAsyncPackage::EventDrivenIndexToObject(FPackageIndex Index, bool bChec
 
 void FAsyncPackage::EventDrivenCreateExport(int32 LocalExportIndex)
 {
+	LLM_SCOPE(ELLMTag::AsyncLoading);
 	SCOPED_LOADTIMER(Package_CreateExports);
+
 	FObjectExport& Export = Linker->ExportMap[LocalExportIndex];
 
 	TRACE_LOADTIME_CREATE_EXPORT_SCOPE(this, &Export.Object);
 
-	const UClass* Class = CastEventDrivenIndexToObject<UClass>(Export.ClassIndex, false);
-	LLM_SCOPE(ELLMTag::AsyncLoading);
-	LLM_SCOPE_DYNAMIC_STAT_OBJECTPATH(GetLinkerRoot(), ELLMTagSet::Assets);
-	LLM_SCOPE_DYNAMIC_STAT_OBJECTPATH(Class, ELLMTagSet::AssetClasses);
-   	UE_TRACE_METADATA_SCOPE_ASSET_FNAME(Export.ObjectName, Class->GetFName(), GetLinkerRoot()->GetFName());
-
 	// Check whether we already loaded the object and if not whether the context flags allow loading it.
 	//check(!Export.Object || Export.Object->HasAnyFlags(RF_ClassDefaultObject)); // we should not have this yet, unless it is a CDO
 	check(!Export.Object); // we should not have this yet
 	if (!Export.Object && !Export.bExportLoadFailed)
 	{
+		const UClass* Class = CastEventDrivenIndexToObject<UClass>(Export.ClassIndex, false);
+		LLM_SCOPE_DYNAMIC_STAT_OBJECTPATH(GetLinkerRoot(), ELLMTagSet::Assets);
+		LLM_SCOPE_DYNAMIC_STAT_OBJECTPATH(Class, ELLMTagSet::AssetClasses);
+		UE_TRACE_METADATA_SCOPE_ASSET_FNAME(Export.ObjectName, Class->GetFName(), GetLinkerRoot()->GetFName());
+
 		FUObjectSerializeContext* LoadContext = GetSerializeContext();
 		FScopedAddObjectreference OnExitAddReference(*this, Export.Object);
 
@@ -3456,21 +3457,20 @@ void FAsyncPackage::EventDrivenSerializeExport(int32 LocalExportIndex)
 
 	FObjectExport& Export = Linker->ExportMap[LocalExportIndex];
 
-	const UClass* Class = CastEventDrivenIndexToObject<UClass>(Export.ClassIndex, false);
-	LLM_SCOPE_DYNAMIC_STAT_OBJECTPATH(GetLinkerRoot(), ELLMTagSet::Assets);
-	LLM_SCOPE_DYNAMIC_STAT_OBJECTPATH(Class, ELLMTagSet::AssetClasses);
-
 	UObject* Object = Export.Object;
 	if (Object && Object->HasAnyFlags(RF_NeedLoad))
 	{
-  		UE_TRACE_METADATA_SCOPE_ASSET(Object, Class);
+		const UClass* Class = CastEventDrivenIndexToObject<UClass>(Export.ClassIndex, false);
+		LLM_SCOPE_DYNAMIC_STAT_OBJECTPATH(GetLinkerRoot(), ELLMTagSet::Assets);
+		LLM_SCOPE_DYNAMIC_STAT_OBJECTPATH(Class, ELLMTagSet::AssetClasses);
+		UE_TRACE_METADATA_SCOPE_ASSET_FNAME(Export.ObjectName, Class->GetFName(), GetLinkerRoot()->GetFName());
+
 		Linker->GetAsyncLoader()->LogItem(TEXT("EventDrivenSerializeExport"), Export.SerialOffset, Export.SerialSize);
 
 		LastTypeOfWorkPerformed = TEXT("EventDrivenSerializeExport");
 		LastObjectWorkWasPerformedOn = Object;
 		check(Object->GetLinker() == Linker);
 		check(Object->GetLinkerIndex() == LocalExportIndex);
-		UClass* Cls = nullptr;
 
 		// If this is a struct, make sure that its parent struct is completely loaded
 		if (UStruct* Struct = dynamic_cast<UStruct*>(Object))
diff --git a/Engine/Source/Runtime/CoreUObject/Private/Serialization/AsyncLoading2.cpp b/Engine/Source/Runtime/CoreUObject/Private/Serialization/AsyncLoading2.cpp
--- a/Engine/Source/Runtime/CoreUObject/Private/Serialization/AsyncLoading2.cpp
+++ b/Engine/Source/Runtime/CoreUObject/Private/Serialization/AsyncLoading2.cpp
@@ -6507,7 +6507,7 @@ void FAsyncPackage2::EventDrivenCreateExport(const FAsyncPackageHeaderData& Head
 
 	LLM_SCOPE_DYNAMIC_STAT_OBJECTPATH(GetLinkerRoot(), ELLMTagSet::Assets);
 	LLM_SCOPE_DYNAMIC_STAT_OBJECTPATH(LoadClass, ELLMTagSet::AssetClasses);
-    UE_TRACE_METADATA_SCOPE_ASSET_FNAME(ObjectName, LoadClass->GetFName(), GetLinkerRoot()->GetFName());
+	UE_TRACE_METADATA_SCOPE_ASSET_FNAME(ObjectName, LoadClass->GetFName(), GetLinkerRoot()->GetFName());
 
 	// Try to find existing object first as we cannot in-place replace objects, could have been created by other export in this package
 	{
@@ -6706,9 +6706,9 @@ bool FAsyncPackage2::EventDrivenSerializeExport(const FAsyncPackageHeaderData& H
 	}
 
 	const UClass* LoadClass = Export.ClassIndex.IsNull() ? UClass::StaticClass() : CastEventDrivenIndexToObject<UClass>(Header, Export.ClassIndex, true);
-	UE_TRACE_METADATA_SCOPE_ASSET(Object, LoadClass);
-	LLM_SCOPE_DYNAMIC_STAT_OBJECTPATH(Object->GetOutermost(), ELLMTagSet::Assets);
+	LLM_SCOPE_DYNAMIC_STAT_OBJECTPATH(Object->GetPackage(), ELLMTagSet::Assets);
 	LLM_SCOPE_DYNAMIC_STAT_OBJECTPATH(LoadClass, ELLMTagSet::AssetClasses);
+	UE_TRACE_METADATA_SCOPE_ASSET(Object, LoadClass);
 
 	// cache archetype
 	// prevents GetArchetype from hitting the expensive GetArchetypeFromRequiredInfoImpl
diff --git a/Engine/Source/Runtime/CoreUObject/Private/UObject/Obj.cpp b/Engine/Source/Runtime/CoreUObject/Private/UObject/Obj.cpp
--- a/Engine/Source/Runtime/CoreUObject/Private/UObject/Obj.cpp
+++ b/Engine/Source/Runtime/CoreUObject/Private/UObject/Obj.cpp
@@ -48,6 +48,7 @@
 #include "Misc/RedirectCollector.h"
 #include "Misc/PlayInEditorLoadingScope.h"
 #include "UObject/GCScopeLock.h"
+#include "ProfilingDebugging/AssetMetadataTrace.h"
 #include "ProfilingDebugging/CookStats.h"
 #include "ProfilingDebugging/LoadTimeTracker.h"
 #include "UObject/GCObject.h"
@@ -1278,9 +1279,9 @@ void UObject::ConditionalPostLoad()
 				UE_SCOPED_COOK_STAT(Package->GetFName(), EPackageEventStatType::LoadPackage);
 				LLM_SCOPE_DYNAMIC_STAT_OBJECTPATH(Package, ELLMTagSet::Assets);
 				LLM_SCOPE_DYNAMIC_STAT_OBJECTPATH(GetClass(), ELLMTagSet::AssetClasses);
-				UE_TRACE_METADATA_SCOPE_ASSET(this, GetClass());
+				UE_TRACE_METADATA_SCOPE_ASSET_FNAME(GetFName(), GetClass()->GetFName(), Package->GetFName())
 				TRACE_LOADTIME_POSTLOAD_OBJECT_SCOPE(this);
-				
+
 				PostLoad();
 
 				LLM_PUSH_STATS_FOR_ASSET_TAGS();
diff --git a/Engine/Source/Runtime/CoreUObject/Private/UObject/UObjectGlobals.cpp b/Engine/Source/Runtime/CoreUObject/Private/UObject/UObjectGlobals.cpp
--- a/Engine/Source/Runtime/CoreUObject/Private/UObject/UObjectGlobals.cpp
+++ b/Engine/Source/Runtime/CoreUObject/Private/UObject/UObjectGlobals.cpp
@@ -67,7 +67,7 @@
 #include "ProfilingDebugging/CsvProfiler.h"
 #include "IO/IoDispatcher.h"
 #include "ProfilingDebugging/LoadTimeTracker.h"
-#include "Misc/PackageAccessTracking.h"
+#include "ProfilingDebugging/AssetMetadataTrace.h"
 #include "Misc/PackageAccessTracking.h"
 #include "UObject/PropertyWithSetterAndGetter.h"
 #include "UObject/AnyPackagePrivate.h"
@@ -2057,6 +2057,7 @@ UPackage* LoadPackage(UPackage* InOuter, const FPackagePath& PackagePath, uint32
 #endif // STATS
 
 	LLM_SCOPE_DYNAMIC_STAT_OBJECTPATH_FNAME(PackagePath.GetPackageFName(), ELLMTagSet::Assets);
+	UE_TRACE_METADATA_SCOPE_ASSET_FNAME(NAME_None, NAME_None, PackagePath.GetPackageFName());
 	TRACE_LOADTIME_REQUEST_GROUP_SCOPE(TEXT("SyncLoad - %s"), *PackagePath.GetDebugName());
 	
 	// if this is a supported asset, it should be loaded fully rather than just for diffing
diff --git a/Engine/Source/Runtime/Engine/Private/Components/ActorComponent.cpp b/Engine/Source/Runtime/Engine/Private/Components/ActorComponent.cpp
--- a/Engine/Source/Runtime/Engine/Private/Components/ActorComponent.cpp
+++ b/Engine/Source/Runtime/Engine/Private/Components/ActorComponent.cpp
@@ -2,45 +2,47 @@
 // ActorComponent.cpp: Actor component implementation.
 
 #include "Components/ActorComponent.h"
-#include "Engine/LevelStreaming.h"
-#include "EngineStats.h"
-#include "Engine/MemberReference.h"
-#include "Engine/Level.h"
-#include "Components/PrimitiveComponent.h"
+
 #include "AI/NavigationSystemBase.h"
-#include "Engine/BlueprintGeneratedClass.h"
-#include "Elements/Framework/EngineElementsLibrary.h"
+#include "Async/ParallelFor.h"
+#include "ComponentRecreateRenderStateContext.h"
 #include "ComponentReregisterContext.h"
+#include "Components/PrimitiveComponent.h"
+#include "ComponentUtils.h"
+#include "Elements/Framework/EngineElementsLibrary.h"
+#include "Engine/ActorChannel.h"
 #include "Engine/AssetUserData.h"
+#include "Engine/BlueprintGeneratedClass.h"
+#include "Engine/Engine.h"
+#include "Engine/InputDelegateBinding.h"
+#include "Engine/Level.h"
+#include "Engine/LevelStreaming.h"
 #include "Engine/LevelStreamingPersistent.h"
+#include "Engine/MemberReference.h"
 #include "Engine/NetDriver.h"
-#include "Engine/ActorChannel.h"
+#include "Engine/SimpleConstructionScript.h"
+#include "EngineStats.h"
+#include "GameFramework/InputSettings.h"
 #include "HAL/LowLevelMemStats.h"
-#include "Misc/ScopeRWLock.h"
-#include "Net/UnrealNetwork.h"
 #include "Logging/MessageLog.h"
-#include "Misc/UObjectToken.h"
 #include "Misc/MapErrors.h"
-#include "ComponentRecreateRenderStateContext.h"
-#include "Engine/SimpleConstructionScript.h"
-#include "ComponentUtils.h"
-#include "Engine/Engine.h"
+#include "Misc/ScopeRWLock.h"
+#include "Misc/UObjectToken.h"
 #include "Net/Core/PushModel/PushModel.h"
+#include "Net/UnrealNetwork.h"
+#include "ObjectTrace.h"
 #include "Physics/Experimental/PhysScene_Chaos.h"
 #include "PhysicsEngine/BodySetup.h"
+#include "ProfilingDebugging/AssetMetadataTrace.h"
 #include "SceneInterface.h"
-#include "UObject/FrameworkObjectVersion.h"
 #include "UObject/FortniteReleaseBranchCustomObjectVersion.h"
-#include "Async/ParallelFor.h"
-#include "Engine/InputDelegateBinding.h"
-#include "GameFramework/InputSettings.h"
+#include "UObject/FrameworkObjectVersion.h"
 
 #if WITH_EDITOR
 #include "Kismet2/ComponentEditorUtils.h"
 #include "ObjectCacheEventSink.h"
 #include "StaticMeshCompiler.h"
 #endif
-#include "ObjectTrace.h"
 
 #if UE_WITH_IRIS
 #include "Iris/Core/IrisLog.h"
@@ -1644,7 +1646,8 @@
 void UActorComponent::CreatePhysicsState(bool bAllowDeferral)
 {
 	LLM_SCOPE(ELLMTag::Chaos);
-	LLM_SCOPE_DYNAMIC_STAT_OBJECTPATH(GetOutermost(), ELLMTagSet::Assets);
+	LLM_SCOPE_DYNAMIC_STAT_OBJECTPATH(GetPackage(), ELLMTagSet::Assets);
+	UE_TRACE_METADATA_SCOPE_ASSET_FNAME(NAME_None, NAME_None, GetPackage()->GetFName());
 
 	SCOPE_CYCLE_COUNTER(STAT_ComponentCreatePhysicsState);
 
@@ -1730,7 +1733,8 @@
 	{
 		SCOPE_CYCLE_COUNTER(STAT_ComponentCreateRenderState);
 		LLM_SCOPE(ELLMTag::SceneRender);
-		LLM_SCOPE_DYNAMIC_STAT_OBJECTPATH(GetOutermost(), ELLMTagSet::Assets);
+		LLM_SCOPE_DYNAMIC_STAT_OBJECTPATH(GetPackage(), ELLMTagSet::Assets);
+		UE_TRACE_METADATA_SCOPE_ASSET_FNAME(NAME_None, NAME_None, GetPackage()->GetFName());
 		CreateRenderState_Concurrent(Context);
 		checkf(bRenderStateCreated, TEXT("Failed to route CreateRenderState_Concurrent (%s)"), *GetFullName());
 	}
@@ -1840,7 +1844,8 @@
 void UActorComponent::DoDeferredRenderUpdates_Concurrent()
 {
 	LLM_SCOPE(ELLMTag::SceneRender);
-	LLM_SCOPE_DYNAMIC_STAT_OBJECTPATH(GetOutermost(), ELLMTagSet::Assets);
+	LLM_SCOPE_DYNAMIC_STAT_OBJECTPATH(GetPackage(), ELLMTagSet::Assets);
+	UE_TRACE_METADATA_SCOPE_ASSET_FNAME(NAME_None, NAME_None, GetPackage()->GetFName());
 
 	checkf(!IsUnreachable(), TEXT("%s"), *GetFullName());
 	checkf(!IsTemplate(), TEXT("%s"), *GetFullName());
diff --git a/Engine/Source/Runtime/Engine/Private/Components/SkinnedMeshComponent.cpp b/Engine/Source/Runtime/Engine/Private/Components/SkinnedMeshComponent.cpp
--- a/Engine/Source/Runtime/Engine/Private/Components/SkinnedMeshComponent.cpp
+++ b/Engine/Source/Runtime/Engine/Private/Components/SkinnedMeshComponent.cpp
@@ -40,6 +40,7 @@
 #include "HAL/LowLevelMemStats.h"
 #include "UObject/Package.h"
 #include "Rendering/RenderCommandPipes.h"
+#include "ProfilingDebugging/AssetMetadataTrace.h"
 
 DEFINE_LOG_CATEGORY_STATIC(LogSkinnedMeshComp, Log, All);
 
@@ -918,7 +919,8 @@ void USkinnedMeshComponent::CreateRenderState_Concurrent(FRegisterComponentConte
 
 	if( GetSkinnedAsset() )
 	{
-		LLM_SCOPE_DYNAMIC_STAT_OBJECTPATH(GetSkinnedAsset()->GetOutermost(), ELLMTagSet::Assets);
+		LLM_SCOPE_DYNAMIC_STAT_OBJECTPATH(GetSkinnedAsset()->GetPackage(), ELLMTagSet::Assets);
+		UE_TRACE_METADATA_SCOPE_ASSET_FNAME(NAME_None, NAME_None, GetSkinnedAsset()->GetPackage()->GetFName());
 
 		// Attempting to track down UE-45505, where it looks as if somehow a skeletal mesh component's mesh has only been partially loaded, causing a mismatch in the LOD arrays
 		checkf(!GetSkinnedAsset()->HasAnyFlags(RF_NeedLoad | RF_NeedPostLoad | RF_NeedPostLoadSubobjects | RF_WillBeLoaded), TEXT("Attempting to create render state for a skeletal mesh that is is not fully loaded. Mesh: %s"), *GetSkinnedAsset()->GetName());
@@ -2276,7 +2278,8 @@ bool USkinnedMeshComponent::AllocateTransformData()
 	// Allocate transforms if not present.
 	if (GetSkinnedAsset() != nullptr && LeaderPoseComponent == nullptr )
 	{
-		LLM_SCOPE_DYNAMIC_STAT_OBJECTPATH(GetSkinnedAsset()->GetOutermost(), ELLMTagSet::Assets);
+		LLM_SCOPE_DYNAMIC_STAT_OBJECTPATH(GetSkinnedAsset()->GetPackage(), ELLMTagSet::Assets);
+		UE_TRACE_METADATA_SCOPE_ASSET_FNAME(NAME_None, NAME_None, GetSkinnedAsset()->GetPackage()->GetFName());
 
 		const int32 NumBones = GetSkinnedAsset()->GetRefSkeleton().GetNum();
 
diff --git a/Engine/Source/Runtime/Engine/Private/Texture2D.cpp b/Engine/Source/Runtime/Engine/Private/Texture2D.cpp
--- a/Engine/Source/Runtime/Engine/Private/Texture2D.cpp
+++ b/Engine/Source/Runtime/Engine/Private/Texture2D.cpp
@@ -44,6 +44,7 @@
 #include "Misc/ScopedSlowTask.h"
 #include "UObject/ObjectSaveContext.h"
 #include "UObject/StrongObjectPtr.h"
+#include "ProfilingDebugging/AssetMetadataTrace.h"
 
 #include UE_INLINE_GENERATED_CPP_BY_NAME(Texture2D)
 
@@ -1493,6 +1494,7 @@ void FVirtualTexture2DResource::InitRHI(FRHICommandListBase& RHICmdList)
 {
 	TRACE_CPUPROFILER_EVENT_SCOPE(FVirtualTexture2DResource::InitRHI);
 	LLM_SCOPE_DYNAMIC_STAT_OBJECTPATH_FNAME(PackageName, ELLMTagSet::Assets);
+	UE_TRACE_METADATA_SCOPE_ASSET_FNAME(NAME_None, NAME_None, PackageName);
 
 	uint32 MaxAnisotropy = 0;
 	if (VirtualTextureScalability::IsAnisotropicFilteringEnabled())
diff --git a/Engine/Source/Runtime/Engine/Private/TextureRenderTarget2D.cpp b/Engine/Source/Runtime/Engine/Private/TextureRenderTarget2D.cpp
--- a/Engine/Source/Runtime/Engine/Private/TextureRenderTarget2D.cpp
+++ b/Engine/Source/Runtime/Engine/Private/TextureRenderTarget2D.cpp
@@ -17,6 +17,8 @@
 #include "RenderGraphUtils.h"
 #include "UObject/Package.h"
 #include "UObject/UnrealType.h"
+#include "ProfilingDebugging/AssetMetadataTrace.h"
+
 #if WITH_EDITOR
 #include "Components/SceneCaptureComponent2D.h"
 #include "TextureCompiler.h"
@@ -470,7 +472,8 @@ ETextureCreateFlags FTextureRenderTarget2DResource::GetCreateFlags()
  */
 void FTextureRenderTarget2DResource::InitRHI(FRHICommandListBase& RHICmdList)
 {
-	LLM_SCOPE_DYNAMIC_STAT_OBJECTPATH(Owner->GetOutermost(), ELLMTagSet::Assets);
+	LLM_SCOPE_DYNAMIC_STAT_OBJECTPATH(Owner->GetPackage(), ELLMTagSet::Assets);
+	UE_TRACE_METADATA_SCOPE_ASSET_FNAME(NAME_None, NAME_None, Owner->GetPackage()->GetFName());
 
 	if( TargetSizeX > 0 && TargetSizeY > 0 )
 	{
@@ -554,7 +557,8 @@ void FTextureRenderTarget2DResource::ReleaseRHI()
  */
 void FTextureRenderTarget2DResource::UpdateDeferredResource( FRHICommandListImmediate& RHICmdList, bool bClearRenderTarget/*=true*/ )
 {
-	LLM_SCOPE_DYNAMIC_STAT_OBJECTPATH(Owner->GetOutermost(), ELLMTagSet::Assets);
+	LLM_SCOPE_DYNAMIC_STAT_OBJECTPATH(Owner->GetPackage(), ELLMTagSet::Assets);
+	UE_TRACE_METADATA_SCOPE_ASSET_FNAME(NAME_None, NAME_None, Owner->GetPackage()->GetFName());
 
 	SCOPED_DRAW_EVENT(RHICmdList, GPUResourceUpdate)
 	RemoveFromDeferredUpdateList();
diff --git a/Engine/Source/Runtime/Engine/Private/TextureRenderTarget2DArray.cpp b/Engine/Source/Runtime/Engine/Private/TextureRenderTarget2DArray.cpp
--- a/Engine/Source/Runtime/Engine/Private/TextureRenderTarget2DArray.cpp
+++ b/Engine/Source/Runtime/Engine/Private/TextureRenderTarget2DArray.cpp
@@ -14,6 +14,7 @@
 #include "Engine/Texture2DArray.h"
 #include "RHIUtilities.h"
 #include "UObject/Package.h"
+#include "ProfilingDebugging/AssetMetadataTrace.h"
 
 #include UE_INLINE_GENERATED_CPP_BY_NAME(TextureRenderTarget2DArray)
 
@@ -195,7 +196,8 @@ UTexture2DArray* UTextureRenderTarget2DArray::ConstructTexture2DArray(UObject* I
  */
 void FTextureRenderTarget2DArrayResource::InitRHI(FRHICommandListBase& RHICmdList)
 {
-	LLM_SCOPE_DYNAMIC_STAT_OBJECTPATH(Owner->GetOutermost(), ELLMTagSet::Assets);
+	LLM_SCOPE_DYNAMIC_STAT_OBJECTPATH(Owner->GetPackage(), ELLMTagSet::Assets);
+	UE_TRACE_METADATA_SCOPE_ASSET_FNAME(NAME_None, NAME_None, Owner->GetPackage()->GetFName());
 
 	if((Owner->SizeX > 0) && (Owner->SizeY > 0) && (Owner->Slices > 0))
 	{
@@ -277,7 +279,8 @@ void FTextureRenderTarget2DArrayResource::ReleaseRHI()
  */
 void FTextureRenderTarget2DArrayResource::UpdateDeferredResource(FRHICommandListImmediate& RHICmdList, bool bClearRenderTarget/*=true*/)
 {
-	LLM_SCOPE_DYNAMIC_STAT_OBJECTPATH(Owner->GetOutermost(), ELLMTagSet::Assets);
+	LLM_SCOPE_DYNAMIC_STAT_OBJECTPATH(Owner->GetPackage(), ELLMTagSet::Assets);
+	UE_TRACE_METADATA_SCOPE_ASSET_FNAME(NAME_None, NAME_None, Owner->GetPackage()->GetFName());
 
 	RemoveFromDeferredUpdateList();
 
diff --git a/Engine/Source/Runtime/Engine/Private/TextureRenderTargetCube.cpp b/Engine/Source/Runtime/Engine/Private/TextureRenderTargetCube.cpp
--- a/Engine/Source/Runtime/Engine/Private/TextureRenderTargetCube.cpp
+++ b/Engine/Source/Runtime/Engine/Private/TextureRenderTargetCube.cpp
@@ -15,6 +15,7 @@
 #include "DeviceProfiles/DeviceProfileManager.h"
 #include "Engine/TextureCube.h"
 #include "RHIUtilities.h"
+#include "ProfilingDebugging/AssetMetadataTrace.h"
 
 #if WITH_EDITOR
 #include "Components/SceneCaptureComponentCube.h"
@@ -215,7 +216,8 @@ UTextureCube* UTextureRenderTargetCube::ConstructTextureCube(UObject* InOuter, c
  */
 void FTextureRenderTargetCubeResource::InitRHI(FRHICommandListBase& RHICmdList)
 {
-	LLM_SCOPE_DYNAMIC_STAT_OBJECTPATH(Owner->GetOutermost(), ELLMTagSet::Assets);
+	LLM_SCOPE_DYNAMIC_STAT_OBJECTPATH(Owner->GetPackage(), ELLMTagSet::Assets);
+	UE_TRACE_METADATA_SCOPE_ASSET_FNAME(NAME_None, NAME_None, Owner->GetPackage()->GetFName());
 
 	if(Owner->SizeX > 0)
 	{
@@ -307,7 +309,8 @@ void FTextureRenderTargetCubeResource::ReleaseRHI()
  */
 void FTextureRenderTargetCubeResource::UpdateDeferredResource(FRHICommandListImmediate& RHICmdList, bool bClearRenderTarget/*=true*/)
 {
-	LLM_SCOPE_DYNAMIC_STAT_OBJECTPATH(Owner->GetOutermost(), ELLMTagSet::Assets);
+	LLM_SCOPE_DYNAMIC_STAT_OBJECTPATH(Owner->GetPackage(), ELLMTagSet::Assets);
+	UE_TRACE_METADATA_SCOPE_ASSET_FNAME(NAME_None, NAME_None, Owner->GetPackage()->GetFName());
 
 	RemoveFromDeferredUpdateList();
 
diff --git a/Engine/Source/Runtime/Engine/Private/TextureRenderTargetVolume.cpp b/Engine/Source/Runtime/Engine/Private/TextureRenderTargetVolume.cpp
--- a/Engine/Source/Runtime/Engine/Private/TextureRenderTargetVolume.cpp
+++ b/Engine/Source/Runtime/Engine/Private/TextureRenderTargetVolume.cpp
@@ -14,6 +14,7 @@
 #include "HAL/LowLevelMemStats.h"
 #include "RHIUtilities.h"
 #include "UObject/Package.h"
+#include "ProfilingDebugging/AssetMetadataTrace.h"
 
 #include UE_INLINE_GENERATED_CPP_BY_NAME(TextureRenderTargetVolume)
 
@@ -201,7 +202,8 @@ UVolumeTexture* UTextureRenderTargetVolume::ConstructTextureVolume(UObject* InOu
  */
 void FTextureRenderTargetVolumeResource::InitRHI(FRHICommandListBase& RHICmdList)
 {
-	LLM_SCOPE_DYNAMIC_STAT_OBJECTPATH(Owner->GetOutermost(), ELLMTagSet::Assets);
+	LLM_SCOPE_DYNAMIC_STAT_OBJECTPATH(Owner->GetPackage(), ELLMTagSet::Assets);
+	UE_TRACE_METADATA_SCOPE_ASSET_FNAME(NAME_None, NAME_None, Owner->GetPackage()->GetFName());
 
 	if((Owner->SizeX > 0) && (Owner->SizeY > 0) && (Owner->SizeZ > 0))
 	{
@@ -279,7 +281,8 @@ void FTextureRenderTargetVolumeResource::ReleaseRHI()
  */
 void FTextureRenderTargetVolumeResource::UpdateDeferredResource(FRHICommandListImmediate& RHICmdList, bool bClearRenderTarget/*=true*/)
 {
-	LLM_SCOPE_DYNAMIC_STAT_OBJECTPATH(Owner->GetOutermost(), ELLMTagSet::Assets);
+	LLM_SCOPE_DYNAMIC_STAT_OBJECTPATH(Owner->GetPackage(), ELLMTagSet::Assets);
+	UE_TRACE_METADATA_SCOPE_ASSET_FNAME(NAME_None, NAME_None, Owner->GetPackage()->GetFName());
 
 	RemoveFromDeferredUpdateList();
 
diff --git a/Engine/Source/Runtime/Renderer/Private/RendererScene.cpp b/Engine/Source/Runtime/Renderer/Private/RendererScene.cpp
--- a/Engine/Source/Runtime/Renderer/Private/RendererScene.cpp
+++ b/Engine/Source/Runtime/Renderer/Private/RendererScene.cpp
@@ -75,6 +75,7 @@
 #include "RenderCore.h"
 #include "IESTextureManager.h"
 #include "Materials/MaterialRenderProxy.h"
+#include "ProfilingDebugging/AssetMetadataTrace.h"
 #include "ProfilingDebugging/CountersTrace.h"
 #include "SceneCulling/SceneCulling.h"
 #include "InstanceCulling/InstanceCullingOcclusionQuery.h"
@@ -1853,6 +1854,7 @@ void FScene::BatchAddPrimitivesInternal(TArrayView<T*> InPrimitives)
 	}
 #endif
 	LLM_SCOPE_DYNAMIC_STAT_OBJECTPATH(InPrimitives[0]->GetOutermost(), ELLMTagSet::Assets);
+	UE_TRACE_METADATA_SCOPE_ASSET_FNAME(NAME_None, NAME_None, InPrimitives[0]->GetOutermost()->GetFName());
 
 	SCOPE_CYCLE_COUNTER(STAT_AddScenePrimitiveGT);
 	SCOPED_NAMED_EVENT(FScene_AddPrimitive, FColor::Green);
diff --git a/Engine/Source/Runtime/Renderer/Renderer.Build.cs b/Engine/Source/Runtime/Renderer/Renderer.Build.cs
--- a/Engine/Source/Runtime/Renderer/Renderer.Build.cs
+++ b/Engine/Source/Runtime/Renderer/Renderer.Build.cs
@@ -30,14 +30,15 @@ public class Renderer : ModuleRules
 
 		PrivateDependencyModuleNames.AddRange(
 			new string[] {
-				"CoreUObject", 
+				"CoreUObject",
 				"ApplicationCore",
-				"RenderCore", 
+				"RenderCore",
 				"ImageWriteQueue",
 				"RHI",
-				"MaterialShaderQualitySettings"
+				"MaterialShaderQualitySettings",
+				"TraceLog"
 			}
-            );
+		);
 
         PrivateIncludePathModuleNames.AddRange(new string[] { "HeadMountedDisplay" });
         DynamicallyLoadedModuleNames.AddRange(new string[] { "HeadMountedDisplay" });
