Material Editor - Turns on material translation results collection into "intermediate data", that is in turn used to generate material shader and the MaterialEnvironment.

diff --git a/Engine/Source/Runtime/Engine/Private/Materials/HLSLMaterialTranslator.cpp b/Engine/Source/Runtime/Engine/Private/Materials/HLSLMaterialTranslator.cpp
--- a/Engine/Source/Runtime/Engine/Private/Materials/HLSLMaterialTranslator.cpp
+++ b/Engine/Source/Runtime/Engine/Private/Materials/HLSLMaterialTranslator.cpp
@@ -190,7 +190,7 @@ static FAutoConsoleVariableRef CVarPedanticErrorChecksEnabled(
 	TEXT("Enables material compilation pedantic error checking"));
 
 /* Controls whether to use the new GetMaterialShaderCode() and GetMaterialEnvironment() implementations. */
-static bool GUseMaterialTranslationResultsGrouping = false;
+static bool GUseMaterialTranslationResultsGrouping = true;
 
 /* Controls whether DDC caching of material translation results is enabled. */
 static bool GJobMaterialTranslateDDCEnable = false;
