[ZenLoader]
  - Fix crash when objects are moved between packages during a rename operation and then imported from another package being loaded
  - Added some rename tests in memory and on disk along with some reload to exercise the code we're fixing
  - Improve async loading test framework to be more generic

diff --git a/Engine/Source/Runtime/CoreUObject/Private/Serialization/AsyncLoading2.cpp b/Engine/Source/Runtime/CoreUObject/Private/Serialization/AsyncLoading2.cpp
--- a/Engine/Source/Runtime/CoreUObject/Private/Serialization/AsyncLoading2.cpp
+++ b/Engine/Source/Runtime/CoreUObject/Private/Serialization/AsyncLoading2.cpp
@@ -735,7 +735,7 @@ private:
 			Values[Index] = GUObjectArray.ObjectToIndex(Object);
 		}
 
-		void Remove(uint64 ExportHash)
+		bool Remove(uint64 ExportHash)
 		{
 			TArrayView<uint64> Keys = GetKeys();
 			int32 Index = Algo::LowerBound(Keys, ExportHash);
@@ -743,7 +743,10 @@ private:
 			{
 				TArrayView<int32> Values = GetValues();
 				Values[Index] = -1;
+				return true;
 			}
+
+			return false;
 		}
 
 		UObject* Find(uint64 ExportHash)
@@ -955,8 +958,10 @@ public:
 	{
 		check(!bIsMissing);
 		check(HasPackage());
-		bAreAllPublicExportsLoaded = false;
-		PublicExportMap.Remove(ExportHash);
+		if (PublicExportMap.Remove(ExportHash))
+		{
+			bAreAllPublicExportsLoaded = false;
+		}
 	}
 
 	UObject* GetPublicExport(uint64 ExportHash)
@@ -1418,14 +1423,24 @@ public:
 		}
 
 		FPublicExportKey* ExistingKey = ObjectIndexToPublicExport.Find(ObjectIndex);
-		if (ExistingKey && (*ExistingKey != Key))
+		if (ExistingKey && *ExistingKey != Key)
 		{
-			UE_LOG(LogStreaming, Fatal,
-				TEXT("FGlobalImportStore::StoreGlobalObject: The constructed public export object '%s' with index %d and id %s:0x%llX already exists in GlobalImportStore but with a different key %s:0x%llX."),
+			UE_LOG(LogStreaming, Log,
+				TEXT("FGlobalImportStore::StoreGlobalObject: The constructed public export object '%s' with index %d and id %s:0x%llX already exists in GlobalImportStore but with a different key %s:0x%llX.")
+				TEXT("The existing object will be replaced since it or its package was most likely renamed after it was loaded the first time."),
 				Object ? *Object->GetFullName() : TEXT("null"),
 				ObjectIndex,
 				*FormatPackageId(Key.GetPackageId()), Key.GetExportHash(),
 				*FormatPackageId(ExistingKey->GetPackageId()), ExistingKey->GetExportHash());
+
+			// Break the link with the old package now because otherwise, we wouldn't be able to remove the export
+			// during GC since the ObjectIndex can only be linked to a single package ref.
+			if (FLoadedPackageRef* ExistingPackageRef = FindPackageRef(ExistingKey->GetPackageId()))
+			{
+				ExistingPackageRef->RemovePublicExport(ExistingKey->GetExportHash());
+			}
+
+			ObjectIndexToPublicExport.Remove(ObjectIndex);
 		}
 
 		FLoadedPackageRef& PackageRef = FindPackageRefChecked(Key.GetPackageId());
@@ -5594,7 +5609,7 @@ bool FAsyncPackage2::ResolveLinkerLoadImports(FAsyncLoadingThreadState2& ThreadS
 		}
 	}
 
-	const int32 ImportCount = LinkerLoadState->Linker->ImportMap.Num();
+	const int32 ImportCount = HeaderData.ImportMap.Num();
 	while (LinkerLoadState->CreateImportIndex < ImportCount)
 	{
 		const int32 ImportIndex = LinkerLoadState->CreateImportIndex++;
diff --git a/Engine/Source/Runtime/Engine/Private/Tests/Loading/AsyncLoadingTests_Renames.cpp b/Engine/Source/Runtime/Engine/Private/Tests/Loading/AsyncLoadingTests_Renames.cpp
new file mode 100644
--- /dev/null
+++ b/Engine/Source/Runtime/Engine/Private/Tests/Loading/AsyncLoadingTests_Renames.cpp
@@ -0,0 +1,135 @@
+// Copyright Epic Games, Inc. All Rights Reserved.
+
+#include "AsyncLoadingTests_Shared.h"
+#include "Misc/AutomationTest.h"
+#include "Misc/PackageName.h"
+#include "UObject/CoreRedirects.h"
+#include "UObject/SavePackage.h"
+
+#if WITH_DEV_AUTOMATION_TESTS
+
+// All Renames tests should run on zenloader only as the other loaders are not compliant.
+class FLoadingTests_Renames_Base : public FLoadingTests_ZenLoaderOnly_Base
+{
+public:
+	using FLoadingTests_ZenLoaderOnly_Base::FLoadingTests_ZenLoaderOnly_Base;
+
+protected:
+	void DoRenameTest(bool bWithSave, bool bWithGC)
+	{
+		// Package 1 has a hard ref to package 2.
+		auto MutateObjects =
+			[](FLoadingTestsScope& Scope)
+			{
+				Scope.Object1->HardReference = Scope.Object2;
+			};
+
+		FLoadingTestsScope LoadingTestScope(this, MutateObjects);
+
+		UAsyncLoadingTests_Shared* Object2 = LoadObject<UAsyncLoadingTests_Shared>(nullptr, FLoadingTestsScope::ObjectPath2);
+
+		// Move the Object in package 2 to a new package 3, leaving a redirector
+		// This leaves the zenloader global import store in a bad state since it still thinks object in Package 2 is still an export.
+		UPackage* NewPackage = LoadingTestScope.CreatePackage();
+		Object2->Rename(nullptr, NewPackage);
+
+		FString MovedObjectName = Object2->GetFullName();
+
+		if (bWithSave)
+		{
+			LoadingTestScope.SavePackages();
+		}
+
+		if (bWithGC)
+		{
+			LoadingTestScope.GarbageCollect();
+		}
+
+		// Load object 1, which should trigger a fixup during import because object2 will be resolved in the wrong package.
+		UAsyncLoadingTests_Shared* Object1 = LoadObject<UAsyncLoadingTests_Shared>(nullptr, FLoadingTestsScope::ObjectPath1);
+
+		TestNotEqual<UObject*>("Object1 should have been loaded", Object1, nullptr);
+		if (Object1)
+		{
+			TestNotEqual<UObject*>("Object1's hard reference should have been loaded", Object1->HardReference, nullptr);
+
+			if (Object1->HardReference)
+			{
+				TestEqual("Object1's hard reference object name should be the moved object", Object1->HardReference.GetFullName(), MovedObjectName);
+			}
+		}
+
+		// Try to reload package2 for good measure as the previous operation might have caused the loader to become confused.
+		UPackage* Package = LoadPackage(nullptr, FLoadingTestsScope::PackagePath2, LOAD_None);
+
+		TestNotEqual<UObject*>("Package should have been loaded", Package, nullptr);
+		if (Package)
+		{
+			TestEqual("Package name is wrong", Package->GetPathName(), FLoadingTestsScope::PackagePath2);
+		}
+		
+		Object2 = LoadObject<UAsyncLoadingTests_Shared>(nullptr, FLoadingTestsScope::ObjectPath2);
+
+		TestNotEqual<UObject*>("Object should have been loaded", Object2, nullptr);
+		if (Object2)
+		{
+			TestEqual("Object name is wrong", Object2->GetFullName(), MovedObjectName);
+		}
+	}
+};
+
+/**
+ * This test validates the loader behavior when an object is moved to a different package in memory only.
+ */
+IMPLEMENT_CUSTOM_SIMPLE_AUTOMATION_TEST(
+	FLoadingTests_Renames_ObjectMovedToAnotherPackage_InMemory,
+	FLoadingTests_Renames_Base,
+	TEXT("System.Engine.Loading.Renames.ObjectMovedToAnotherPackage.InMemory"),
+	EAutomationTestFlags::EditorContext | EAutomationTestFlags::EngineFilter
+)
+bool FLoadingTests_Renames_ObjectMovedToAnotherPackage_InMemory::RunTest(const FString& Parameters)
+{
+	static constexpr bool bWithSave = false;
+	static constexpr bool bWithGC   = false;
+	DoRenameTest(bWithSave, bWithGC);
+
+	return true;
+}
+
+/**
+ * This test validates the loader behavior when an object is moved to a different package and dirtied packages are written to disk.
+ */
+IMPLEMENT_CUSTOM_SIMPLE_AUTOMATION_TEST(
+	FLoadingTests_Renames_ObjectMovedToAnotherPackage_OnDisk,
+	FLoadingTests_Renames_Base,
+	TEXT("System.Engine.Loading.Renames.ObjectMovedToAnotherPackage.OnDisk"),
+	EAutomationTestFlags::EditorContext | EAutomationTestFlags::EngineFilter
+)
+bool FLoadingTests_Renames_ObjectMovedToAnotherPackage_OnDisk::RunTest(const FString& Parameters)
+{
+	static constexpr bool bWithSave = true;
+	static constexpr bool bWithGC = false;
+	DoRenameTest(bWithSave, bWithGC);
+
+	return true;
+}
+
+/**
+ * This test validates the loader behavior when an object is moved to a different package and dirtied packages are written to disk, then reloaded after a full GC.
+ */
+IMPLEMENT_CUSTOM_SIMPLE_AUTOMATION_TEST(
+	FLoadingTests_Renames_ObjectMovedToAnotherPackage_OnDiskReload,
+	FLoadingTests_Renames_Base,
+	TEXT("System.Engine.Loading.Renames.ObjectMovedToAnotherPackage.OnDiskReload"),
+	EAutomationTestFlags::EditorContext | EAutomationTestFlags::EngineFilter
+)
+bool FLoadingTests_Renames_ObjectMovedToAnotherPackage_OnDiskReload::RunTest(const FString& Parameters)
+{
+	static constexpr bool bWithSave = true;
+	static constexpr bool bWithGC = true;
+	DoRenameTest(bWithSave, bWithGC);
+
+	return true;
+}
+
+#endif // WITH_DEV_AUTOMATION_TESTS
diff --git a/Engine/Source/Runtime/Engine/Private/Tests/Loading/AsyncLoadingTests_Shared.cpp b/Engine/Source/Runtime/Engine/Private/Tests/Loading/AsyncLoadingTests_Shared.cpp
--- a/Engine/Source/Runtime/Engine/Private/Tests/Loading/AsyncLoadingTests_Shared.cpp
+++ b/Engine/Source/Runtime/Engine/Private/Tests/Loading/AsyncLoadingTests_Shared.cpp
@@ -15,49 +15,77 @@ UAsyncLoadingTests_Shared::FOnIsPostLoadThreadSafeDelegate UAsyncLoadingTests_Sh
 
 void FLoadingTestsScope::CreateObjects()
 {
-	Package1 = CreatePackage(PackagePath1);
+	Package1 = CreatePackage();
 	Object1 = NewObject<UAsyncLoadingTests_Shared>(Package1, ObjectName, RF_Public | RF_Standalone);
 
-	Package2 = CreatePackage(PackagePath2);
+	Package2 = CreatePackage();
 	Object2 = NewObject<UAsyncLoadingTests_Shared>(Package2, ObjectName, RF_Public | RF_Standalone);
 }
 
-void FLoadingTestsScope::MutateObjects()
+void FLoadingTestsScope::DefaultMutateObjects()
 {
 	// This is the soft reference that we want to test loading for
 	Object1->SoftReference = Object2;
 }
 
-void FLoadingTestsScope::SaveObjects()
+void FLoadingTestsScope::SavePackages()
 {
 	// To avoid an error on save, we need to mark the package as fully loaded.
-	Package1->MarkAsFullyLoaded();
-	Package2->MarkAsFullyLoaded();
+	for (const FString& PackageName : PackageNames)
+	{
+		if (UPackage* Package = FindObject<UPackage>(nullptr, *PackageName))
+		{
+			Package->MarkAsFullyLoaded();
+		}
+	}
 
 	// Save packages to disk.
-	check(UPackage::SavePackage(Package1, nullptr, *FPackageName::LongPackageNameToFilename(PackagePath1, FPackageName::GetAssetPackageExtension()), FSavePackageArgs()));
-	check(UPackage::SavePackage(Package2, nullptr, *FPackageName::LongPackageNameToFilename(PackagePath2, FPackageName::GetAssetPackageExtension()), FSavePackageArgs()));
+	for (const FString& PackageName : PackageNames)
+	{
+		if (UPackage* Package = FindObject<UPackage>(nullptr, *PackageName))
+		{
+			check(UPackage::SavePackage(Package, nullptr, *FPackageName::LongPackageNameToFilename(*PackageName, FPackageName::GetAssetPackageExtension()), FSavePackageArgs()));
+		}
+	}
+}
+
+void FLoadingTestsScope::GarbageCollect()
+{
+	TArray<FString> ObjectPaths;
 
-	// Remove RF_Standalone from top level object.
-	Object1->ClearFlags(RF_Standalone);
-	Object2->ClearFlags(RF_Standalone);
+	// Remove RF_Standalone from package and objects inside it.
+	for (const FString& PackageName : PackageNames)
+	{
+		if (UPackage* Package = FindObject<UPackage>(nullptr, *PackageName))
+		{
+			Package->GetMetaData()->ClearFlags(RF_Standalone);
+			ObjectPaths.Add(PackageName);
+
+			ForEachObjectWithPackage(Package,
+				[&ObjectPaths](UObject* Object)
+				{
+					ObjectPaths.Add(Object->GetPathName());
+					Object->ClearFlags(RF_Standalone);
+					return true;
+				}
+			);
+		}
+	}
 
-	// Remove RF_Standalone from the UMetaData otherwise the package will not GC.
-	Package1->GetMetaData()->ClearFlags(RF_Standalone);
-	Package2->GetMetaData()->ClearFlags(RF_Standalone);
+	// Make sure everything we gathered can be properly found
+	for (const FString& ObjectPath : ObjectPaths)
+	{
+		checkf(FindObject<UObject>(nullptr, *ObjectPath) != nullptr, TEXT("%s should be present in memory"), *ObjectPath);
+	}
 
 	// GC and make sure everything gets cleaned up before loading
-	check(FindObject<UAsyncLoadingTests_Shared>(nullptr, ObjectPath1) != nullptr);
-	check(FindObject<UPackage>(nullptr, PackagePath1) != nullptr);
-	check(FindObject<UAsyncLoadingTests_Shared>(nullptr, ObjectPath2) != nullptr);
-	check(FindObject<UPackage>(nullptr, PackagePath2) != nullptr);
-
 	CollectGarbage(GARBAGE_COLLECTION_KEEPFLAGS);
 
-	check(FindObject<UAsyncLoadingTests_Shared>(nullptr, ObjectPath1) == nullptr);
-	check(FindObject<UPackage>(nullptr, PackagePath1) == nullptr);
-	check(FindObject<UAsyncLoadingTests_Shared>(nullptr, ObjectPath2) == nullptr);
-	check(FindObject<UPackage>(nullptr, PackagePath2) == nullptr);
+	// Now make sure everything is gone
+	for (const FString& ObjectPath : ObjectPaths)
+	{
+		checkf(FindObject<UObject>(nullptr, *ObjectPath) == nullptr, TEXT("%s should have been garbage collected"), *ObjectPath);
+	}
 }
 
 void FLoadingTestsScope::LoadObjects()
@@ -74,38 +102,7 @@ void FLoadingTestsScope::LoadObjects()
 
 void FLoadingTestsScope::CleanupObjects()
 {
-	// GC and make sure everything gets cleaned up before exiting
-	Object1 = FindObject<UAsyncLoadingTests_Shared>(nullptr, ObjectPath1);
-	Object2 = FindObject<UAsyncLoadingTests_Shared>(nullptr, ObjectPath2);
-	
-	if (Object1)
-	{
-		Object1->ClearFlags(RF_Standalone);
-	}
-	if (Object2)
-	{
-		Object2->ClearFlags(RF_Standalone);
-	}
-
-	Package1 = FindObject<UPackage>(nullptr, PackagePath1);
-	Package2 = FindObject<UPackage>(nullptr, PackagePath2);
-
-	if (Package1)
-	{
-		Package1->GetMetaData()->ClearFlags(RF_Standalone);
-	}
-	
-	if (Package2)
-	{
-		Package2->GetMetaData()->ClearFlags(RF_Standalone);
-	}
-
-	CollectGarbage(GARBAGE_COLLECTION_KEEPFLAGS);
-
-	check(FindObject<UAsyncLoadingTests_Shared>(nullptr, ObjectPath1) == nullptr);
-	check(FindObject<UPackage>(nullptr, PackagePath1) == nullptr);
-	check(FindObject<UAsyncLoadingTests_Shared>(nullptr, ObjectPath2) == nullptr);
-	check(FindObject<UPackage>(nullptr, PackagePath2) == nullptr);
+	GarbageCollect();
 
 	UAsyncLoadingTests_Shared::OnPostLoad.Unbind();
 	UAsyncLoadingTests_Shared::OnSerialize.Unbind();
diff --git a/Engine/Source/Runtime/Engine/Private/Tests/Loading/AsyncLoadingTests_Shared.h b/Engine/Source/Runtime/Engine/Private/Tests/Loading/AsyncLoadingTests_Shared.h
--- a/Engine/Source/Runtime/Engine/Private/Tests/Loading/AsyncLoadingTests_Shared.h
+++ b/Engine/Source/Runtime/Engine/Private/Tests/Loading/AsyncLoadingTests_Shared.h
@@ -18,6 +18,9 @@ public:
 	UPROPERTY()
 	TSoftObjectPtr<UObject> SoftReference;
 
+	UPROPERTY()
+	TObjectPtr<UObject> HardReference;
+
 	// Those delegates allow to easily change behavior between tests
 	using FOnPostLoadDelegate = TDelegate<void(UAsyncLoadingTests_Shared*), FDefaultTSDelegateUserPolicy>;
 	static FOnPostLoadDelegate OnPostLoad;
@@ -72,12 +75,15 @@ class FLoadingTestsScope
 {
 private:
 	FAutomationTestBase& AutomationTest;
+
+	TArray<FString>     PackageNames;
+	std::atomic<uint32> PackageIndex;
 public:
 	static constexpr const TCHAR* ObjectName = TEXT("TestObject");
-	static constexpr const TCHAR* PackagePath1 = TEXT("/Engine/TestRecursiveLoads_Package1");
-	static constexpr const TCHAR* ObjectPath1 = TEXT("/Engine/TestRecursiveLoads_Package1.TestObject");
-	static constexpr const TCHAR* PackagePath2 = TEXT("/Engine/TestRecursiveLoads_Package2");
-	static constexpr const TCHAR* ObjectPath2 = TEXT("/Engine/TestRecursiveLoads_Package2.TestObject");
+	static constexpr const TCHAR* PackagePath1 = TEXT("/Engine/LoadingTestsScope_Package1");
+	static constexpr const TCHAR* ObjectPath1 = TEXT("/Engine/LoadingTestsScope_Package1.TestObject");
+	static constexpr const TCHAR* PackagePath2 = TEXT("/Engine/LoadingTestsScope_Package2");
+	static constexpr const TCHAR* ObjectPath2 = TEXT("/Engine/LoadingTestsScope_Package2.TestObject");
 
 	UPackage* Package1 = nullptr;
 	UPackage* Package2 = nullptr;
@@ -85,13 +91,22 @@ public:
 	UAsyncLoadingTests_Shared* Object1 = nullptr;
 	UAsyncLoadingTests_Shared* Object2 = nullptr;
 
+	UPackage* CreatePackage()
+	{
+		FString PackageName = FString::Printf(TEXT("/Engine/LoadingTestsScope_Package%u"), ++PackageIndex);
+		UPackage* Package = ::CreatePackage(*PackageName);
+		PackageNames.Add(PackageName);
+		return Package;
+	}
+
 	void CreateObjects();
-	void MutateObjects();
-	void SaveObjects();
+	void DefaultMutateObjects();
+	void SavePackages();
 	void LoadObjects();
 	void CleanupObjects();
+	void GarbageCollect();
 
-	FLoadingTestsScope(FAutomationTestBase* InAutomationTest)
+	FLoadingTestsScope(FAutomationTestBase* InAutomationTest, TFunction<void (FLoadingTestsScope&)> InMutateObjects = nullptr)
 		: AutomationTest(*InAutomationTest)
 	{
 		// Just make sure the async loading queue is empty before beginning.
@@ -99,9 +114,18 @@ public:
 
 		CreateObjects();
 
-		MutateObjects();
+		if (InMutateObjects)
+		{
+			InMutateObjects(*this);
+		}
+		else
+		{
+			DefaultMutateObjects();
+		}
+
+		SavePackages();
 
-		SaveObjects();
+		GarbageCollect();
 	}
 
 	virtual ~FLoadingTestsScope()
