[Navmesh] Fix navmesh tiles being removed and not generating for some modifiers when navmesh is in DynamicModifierOnly mode.
(Before the change, if we were DynamicModifierOnly mode and a resolution was set, it would set the Geometry flag as well as the DynamicModifier flag, resulting in the touched tile being flushed but not fully rebuilt.)
Also fixes an issue with bFillCollisionUnderneathForNavmesh.

diff --git a/Engine/Source/Runtime/Engine/Classes/AI/Navigation/NavRelevantInterface.h b/Engine/Source/Runtime/Engine/Classes/AI/Navigation/NavRelevantInterface.h
--- a/Engine/Source/Runtime/Engine/Classes/AI/Navigation/NavRelevantInterface.h
+++ b/Engine/Source/Runtime/Engine/Classes/AI/Navigation/NavRelevantInterface.h
@@ -119,6 +119,7 @@ struct FNavigationRelevantData : public TSharedFromThis<FNavigationRelevantData,
 
 	FORCEINLINE bool HasGeometry() const { return VoxelData.Num() || CollisionData.Num(); }
 	FORCEINLINE bool HasModifiers() const { return !Modifiers.IsEmpty(); }
+	FORCEINLINE bool HasDynamicModifiers() const { return Modifiers.IsDynamic(); }
 	FORCEINLINE bool IsPendingLazyGeometryGathering() const { return bPendingLazyGeometryGathering; }
 	FORCEINLINE bool IsPendingLazyModifiersGathering() const { return bPendingLazyModifiersGathering; }
 	FORCEINLINE bool IsPendingChildLazyModifiersGathering() const { return bPendingChildLazyModifiersGathering; }
@@ -136,11 +137,12 @@ struct FNavigationRelevantData : public TSharedFromThis<FNavigationRelevantData,
 	FORCEINLINE SIZE_T GetGeometryAllocatedSize() const { return CollisionData.GetAllocatedSize() + VoxelData.GetAllocatedSize(); }
 	FORCEINLINE int32 GetDirtyFlag() const
 	{
-		const bool bSetGeometryFlag = HasGeometry() || IsPendingLazyGeometryGathering() || Modifiers.GetMaskFillCollisionUnderneathForNavmesh() ||
+		const bool bSetGeometryFlag = HasGeometry() || IsPendingLazyGeometryGathering() ||
+			Modifiers.GetFillCollisionUnderneathForNavmesh() || Modifiers.GetMaskFillCollisionUnderneathForNavmesh() ||
 			(Modifiers.GetNavMeshResolution() != ENavigationDataResolution::Invalid);
 		
 		return (bSetGeometryFlag ? ENavigationDirtyFlag::Geometry : 0) |
-			((HasModifiers() || NeedAnyPendingLazyModifiersGathering()) ? ENavigationDirtyFlag::DynamicModifier : 0) |
+			((HasDynamicModifiers() || NeedAnyPendingLazyModifiersGathering()) ? ENavigationDirtyFlag::DynamicModifier : 0) |
 			(Modifiers.HasAgentHeightAdjust() ? ENavigationDirtyFlag::UseAgentHeight : 0);
 	}
 
diff --git a/Engine/Source/Runtime/Engine/Public/AI/NavigationModifier.h b/Engine/Source/Runtime/Engine/Public/AI/NavigationModifier.h
--- a/Engine/Source/Runtime/Engine/Public/AI/NavigationModifier.h
+++ b/Engine/Source/Runtime/Engine/Public/AI/NavigationModifier.h
@@ -290,6 +290,12 @@ struct FCompositeNavModifier : public FNavigationModifier
 			NavMeshResolution == ENavigationDataResolution::Invalid;
 	}
 
+	FORCEINLINE bool IsDynamic() const 
+	{
+		// Excluding bFillCollisionUnderneathForNavmesh, bMaskFillCollisionUnderneathForNavmesh and NavMeshResolution since they require full tile rebuild.
+		return !Areas.IsEmpty() || !SimpleLinks.IsEmpty() || !CustomLinks.IsEmpty();
+	}
+	
 	void Add(const FAreaNavModifier& Area)
 	{
 		Areas.Add(Area);
