Added a way to convert FPrimitiveInstanceId to instance index via InstancedStaticMeshComponent's API

diff --git a/Engine/Source/Runtime/Engine/Classes/Components/InstancedStaticMeshComponent.h b/Engine/Source/Runtime/Engine/Classes/Components/InstancedStaticMeshComponent.h
--- a/Engine/Source/Runtime/Engine/Classes/Components/InstancedStaticMeshComponent.h
+++ b/Engine/Source/Runtime/Engine/Classes/Components/InstancedStaticMeshComponent.h
@@ -213,6 +213,9 @@ class UInstancedStaticMeshComponent : public UStaticMeshComponent, public ISMIns
 	/**
 	 */
 	ENGINE_API bool IsValidId(FPrimitiveInstanceId InstanceId);
+	
+	/** Fetches current instance index for a given InstanceId */
+	FORCEINLINE int32 GetInstanceIndexForId(FPrimitiveInstanceId InstanceId) const { return PrimitiveInstanceDataManager.IdToIndex(InstanceId); }
 
 	ENGINE_API void SetHasPerInstancePrevTransforms(bool bInHasPreviousTransforms);
 
