Fixed when the HISM gets into a state where the hierarchy is built from initial template data, but then not rebuilt after the actual intstance data has been populated, during BP initialization.

diff --git a/Engine/Source/Runtime/Engine/Private/HierarchicalInstancedStaticMesh.cpp b/Engine/Source/Runtime/Engine/Private/HierarchicalInstancedStaticMesh.cpp
--- a/Engine/Source/Runtime/Engine/Private/HierarchicalInstancedStaticMesh.cpp
+++ b/Engine/Source/Runtime/Engine/Private/HierarchicalInstancedStaticMesh.cpp
@@ -2759,7 +2759,7 @@ bool UHierarchicalInstancedStaticMeshComponent::BuildTreeIfOutdated(bool Async,
 
 	if (ForceUpdate 
 		|| bIsOutOfDate
-		|| PrimitiveInstanceDataManager.HasAnyInstanceChanges()
+		|| PrimitiveInstanceDataManager.HasAnyChanges()
 		|| InstanceReorderTable.Num() != PerInstanceSMData.Num()
 		|| NumBuiltInstances != PerInstanceSMData.Num() 
 		|| (GetStaticMesh() != nullptr && CacheMeshExtendedBounds != GetStaticMesh()->GetBounds())
diff --git a/Engine/Source/Runtime/Engine/Public/InstancedStaticMesh/ISMInstanceDataManager.h b/Engine/Source/Runtime/Engine/Public/InstancedStaticMesh/ISMInstanceDataManager.h
--- a/Engine/Source/Runtime/Engine/Public/InstancedStaticMesh/ISMInstanceDataManager.h
+++ b/Engine/Source/Runtime/Engine/Public/InstancedStaticMesh/ISMInstanceDataManager.h
@@ -123,6 +123,12 @@ public:
 
 	ENGINE_API bool HasAnyInstanceChanges() const;
 
+	/** 
+	 * Returns true if there are explicitly tracked instance changes, or the state is not tracked (because no proxy has been created yet),
+	 * and the tracking state is not Disabled.
+	 */
+	inline bool HasAnyChanges() const { return GetState() != ETrackingState::Disabled && (GetState() != ETrackingState::Tracked || HasAnyInstanceChanges());}
+
 	void SerializeRenderData(FArchive& Ar, bool bCooked);
 
 	bool FlushChanges(FInstanceUpdateComponentDesc &&ComponentData, bool bNewPrimitiveProxy);
