Level Instance: don't allow partial loading if the level instance has an active child edit.

diff --git a/Engine/Source/Runtime/Engine/Private/LevelInstance/LevelInstanceInterface.cpp b/Engine/Source/Runtime/Engine/Private/LevelInstance/LevelInstanceInterface.cpp
--- a/Engine/Source/Runtime/Engine/Private/LevelInstance/LevelInstanceInterface.cpp
+++ b/Engine/Source/Runtime/Engine/Private/LevelInstance/LevelInstanceInterface.cpp
@@ -25,7 +25,7 @@ bool ILevelInstanceInterface::SupportsPartialEditorLoading() const
 	{
 		if (ULevelInstanceSubsystem* LevelInstanceSubsystem = GetLevelInstanceSubsystem())
 		{
-			if (LevelInstanceSubsystem->IsEditingLevelInstance(this))
+			if (LevelInstanceSubsystem->IsEditingLevelInstance(this) || HasChildEdit())
 			{
 				return false;
 			}
