CoreRedirects: Functions: Change MemberReference to support a function call on one class being redirected to a function call on another class. Used when we have static utility functions that need to change.
#rnx

diff --git a/Engine/Source/Runtime/Engine/Classes/Engine/MemberReference.h b/Engine/Source/Runtime/Engine/Classes/Engine/MemberReference.h
--- a/Engine/Source/Runtime/Engine/Classes/Engine/MemberReference.h
+++ b/Engine/Source/Runtime/Engine/Classes/Engine/MemberReference.h
@@ -349,17 +349,22 @@ public:
 	 * @param	FieldClass		UClass of field type we are looking for
 	 * @param	InitialScope	The scope the field was initially defined in.  The function will search up into parent scopes to attempt to find remappings
 	 * @param	InitialName		The name of the field to attempt to find a redirector for
-	 * @param	bInitialScopeMustBeOwnerOfField		if true the InitialScope must be Child of the field's owner
+	 * @param	bInitialScopeMustBeOwnerOfFieldForParentScopeRedirect If true, and a redirect is found in a parent scope, and the targetscope of the redirect
+	 *          is a child scope of the parent scope that is not the initial scope or a parent of the initial scope, then the redirect will be ignored.
 	 * @return	The remapped field, if one exists
 	 */
-	ENGINE_API static UField* FindRemappedField(UClass *FieldClass, UClass* InitialScope, FName InitialName, bool bInitialScopeMustBeOwnerOfField = false);
-	ENGINE_API static FField* FindRemappedField(FFieldClass* FieldClass, UClass* InitialScope, FName InitialName, bool bInitialScopeMustBeOwnerOfField = false);
+	ENGINE_API static UField* FindRemappedField(UClass *FieldClass, UClass* InitialScope, FName InitialName,
+		bool bInitialScopeMustBeOwnerOfFieldForParentScopeRedirect = false);
+	ENGINE_API static FField* FindRemappedField(FFieldClass* FieldClass, UClass* InitialScope, FName InitialName,
+		bool bInitialScopeMustBeOwnerOfFieldForParentScopeRedirect = false);
 
 	/** Templated version of above, extracts FieldClass and Casts result */
 	template<class TFieldType>
-	static TFieldType* FindRemappedField(UClass* InitialScope, FName InitialName, bool bInitialScopeMustBeOwnerOfField = false)
+	static TFieldType* FindRemappedField(UClass* InitialScope, FName InitialName,
+		bool bInitialScopeMustBeOwnerOfFieldForParentScopeRedirect = false)
 	{
-		return FFieldVariant(FindRemappedField(TFieldType::StaticClass(), InitialScope, InitialName, bInitialScopeMustBeOwnerOfField)).Get<TFieldType>();
+		return FFieldVariant(FindRemappedField(TFieldType::StaticClass(), InitialScope, InitialName,
+			bInitialScopeMustBeOwnerOfFieldForParentScopeRedirect)).Get<TFieldType>();
 	}
 
 #if WITH_EDITOR
diff --git a/Engine/Source/Runtime/Engine/Private/MemberReference.cpp b/Engine/Source/Runtime/Engine/Private/MemberReference.cpp
--- a/Engine/Source/Runtime/Engine/Private/MemberReference.cpp
+++ b/Engine/Source/Runtime/Engine/Private/MemberReference.cpp
@@ -371,7 +371,14 @@ TFieldType* FMemberReference::ResolveMemberImpl(UClass* SelfScope, TFieldTypeCla
 		UClass* TargetScope = GetScope(SelfScope);
 		if (bCanFollowRedirects && TargetScope)
 		{
-			ReturnField = static_cast<TFieldType*>(FindRemappedField(FieldClass, TargetScope, MemberName, true));
+			// bInitialScopeMustBeOwnerOfFieldForParentScopeRedirect is required to avoid invalid redirects. Usecase
+			// showing why: Consider a BPGC based on Pawn, and it overrides the OnLand event. But FindRemappedField
+			// returns Character::OnLanded because of a redirector on Pawn:
+			//    "+K2FieldRedirects=(OldFieldName="Pawn.OnLanded",NewFieldName="Character.OnLanded")"
+			// Character.OnLanded is not a valid field for the BPGC that inherits from Pawn.
+			constexpr bool bInitialScopeMustBeOwnerOfFieldForParentScopeRedirect = true;
+			ReturnField = static_cast<TFieldType*>(FindRemappedField(FieldClass, TargetScope, MemberName,
+				bInitialScopeMustBeOwnerOfFieldForParentScopeRedirect));
 		}
 
 		if (ReturnField != nullptr)
@@ -512,7 +519,8 @@ UFunction* FMemberReference::ResolveMemberFunction(UClass* SelfScope, const bool
 }
 
 template <typename TFieldType>
-TFieldType* FindRemappedFieldImpl(FName FieldClassOutermostName, FName FieldClassName, UClass* InitialScope, FName InitialName, bool bInitialScopeMustBeOwnerOfField)
+TFieldType* FindRemappedFieldImpl(FName FieldClassOutermostName, FName FieldClassName, UClass* InitialScope,
+	FName InitialName, bool bInitialScopeMustBeOwnerOfFieldForParentScopeRedirect)
 {
 	DECLARE_SCOPE_CYCLE_COUNTER(TEXT("FMemberReference::FindRemappedField"), STAT_LinkerLoad_FindRemappedField, STATGROUP_LoadTimeVerbose);
 
@@ -571,7 +579,8 @@ TFieldType* FindRemappedFieldImpl(FName FieldClassOutermostName, FName FieldClas
 			TFieldType* NewField = FindUFieldOrFProperty<TFieldType>(SearchClass, NewFieldName);
 			if (NewField != nullptr)
 			{
-				if (bInitialScopeMustBeOwnerOfField && !InitialScope->IsChildOf(SearchClass))
+				if (bInitialScopeMustBeOwnerOfFieldForParentScopeRedirect &&
+					TestRemapClass != InitialScope && !InitialScope->IsChildOf(SearchClass))
 				{
 					UE_LOG(LogBlueprint, Log, TEXT("UK2Node:  Unable to update field. Remapped field '%s' in not owned by given scope. Scope: '%s', Owner: '%s'."), *InitialName.ToString(), *InitialScope->GetName(), *NewFieldName.ToString());
 				}
@@ -595,13 +604,17 @@ TFieldType* FindRemappedFieldImpl(FName FieldClassOutermostName, FName FieldClas
 	return nullptr;
 }
 
-UField* FMemberReference::FindRemappedField(UClass* FieldClass, UClass* InitialScope, FName InitialName, bool bInitialScopeMustBeOwnerOfField)
+UField* FMemberReference::FindRemappedField(UClass* FieldClass, UClass* InitialScope, FName InitialName,
+	bool bInitialScopeMustBeOwnerOfFieldForParentScopeRedirect)
 {	
-	return FindRemappedFieldImpl<UField>(FieldClass->GetOutermost()->GetFName(), FieldClass->GetFName(), InitialScope, InitialName, bInitialScopeMustBeOwnerOfField);
+	return FindRemappedFieldImpl<UField>(FieldClass->GetOutermost()->GetFName(), FieldClass->GetFName(), InitialScope,
+		InitialName, bInitialScopeMustBeOwnerOfFieldForParentScopeRedirect);
 }
 
-FField* FMemberReference::FindRemappedField(FFieldClass* FieldClass, UClass* InitialScope, FName InitialName, bool bInitialScopeMustBeOwnerOfField)
+FField* FMemberReference::FindRemappedField(FFieldClass* FieldClass, UClass* InitialScope, FName InitialName,
+	bool bInitialScopeMustBeOwnerOfFieldForParentScopeRedirect)
 {
-	return FindRemappedFieldImpl<FField>(GLongCoreUObjectPackageName, FieldClass->GetFName(), InitialScope, InitialName, bInitialScopeMustBeOwnerOfField);
+	return FindRemappedFieldImpl<FField>(GLongCoreUObjectPackageName, FieldClass->GetFName(), InitialScope, InitialName,
+		bInitialScopeMustBeOwnerOfFieldForParentScopeRedirect);
 }
 
