Material Editor:  Added statistics for preshaders, and dumping of preshader expressions and parameter fetches in human readable form to the HLSL source.  Similar statistics had been implemented in debug code to optimize a worst case shader, but now they are user facing to allow anyone to do similar content optimizations.

diff --git a/Engine/Source/Editor/MaterialEditor/Private/MaterialEditor.cpp b/Engine/Source/Editor/MaterialEditor/Private/MaterialEditor.cpp
--- a/Engine/Source/Editor/MaterialEditor/Private/MaterialEditor.cpp
+++ b/Engine/Source/Editor/MaterialEditor/Private/MaterialEditor.cpp
@@ -3159,6 +3159,14 @@ void FMaterialEditor::UpdateMaterialinfoList_Old()
 					TSharedRef<FTokenizedMessage> ShaderCountLine = FTokenizedMessage::Create(EMessageSeverity::Info);
 					ShaderCountLine->AddToken(FTextToken::Create(FText::FromString(ShaderCountString)));
 					Messages.Add(ShaderCountLine);
+
+					// Add number of preshaders and stats
+					uint32 TotalParams, TotalOps;
+					MaterialResource->GetPreshaderStats(TotalParams, TotalOps);
+					FString PreshaderCountString = FString::Printf(TEXT("Preshaders: %u  (%u param fetches, %u ops)"), ShaderMap->GetNumPreshaders(), TotalParams, TotalOps);
+					TSharedRef<FTokenizedMessage> PreshaderCountLine = FTokenizedMessage::Create(EMessageSeverity::Info);
+					PreshaderCountLine->AddToken(FTextToken::Create(FText::FromString(PreshaderCountString)));
+					Messages.Add(PreshaderCountLine);
 				}
 			}
 
diff --git a/Engine/Source/Runtime/Engine/Private/Materials/MaterialShared.cpp b/Engine/Source/Runtime/Engine/Private/Materials/MaterialShared.cpp
--- a/Engine/Source/Runtime/Engine/Private/Materials/MaterialShared.cpp
+++ b/Engine/Source/Runtime/Engine/Private/Materials/MaterialShared.cpp
@@ -46,6 +46,7 @@
 #include "MaterialHLSLGenerator.h"
 #include "MaterialHLSLEmitter.h"
 #include "HLSLTree/HLSLTreeCommon.h"
+#include "Shader/PreshaderEvaluate.h"
 #endif
 #if WITH_ODSC
 #include "ODSC/ODSCManager.h"
@@ -75,6 +76,14 @@ namespace MaterialSharedCookStats
 }
 #endif
 
+#if WITH_EDITOR
+static TAutoConsoleVariable<bool> CVarMaterialEdPreshaderDumpToHLSL(
+	TEXT("r.MaterialEditor.PreshaderDumpToHLSL"),
+	true,
+	TEXT("Controls whether to append preshader expressions and parameter reference counts to the HLSL source window (as comments at the end of the code)."),
+	ECVF_RenderThreadSafe);
+#endif
+
 IMPLEMENT_TYPE_LAYOUT(FHashedMaterialParameterInfo);
 IMPLEMENT_TYPE_LAYOUT(FUniformExpressionSet);
 IMPLEMENT_TYPE_LAYOUT(FMaterialCompilationOutput);
@@ -4332,6 +4341,72 @@ bool FMaterial::GetMaterialExpressionSource( FString& OutSource )
 		if (Source)
 		{
 			OutSource = MoveTemp(*Source);
+
+			if (CVarMaterialEdPreshaderDumpToHLSL.GetValueOnGameThread())
+			{
+				OutSource.AppendChar('\n');
+
+				TMap<FString, uint32> ParameterReferences;
+				FMaterialRenderContext MaterialContext(nullptr, *this, nullptr);
+				UE::Shader::FPreshaderDataContext PreshaderContextBase(NewCompilationOutput.UniformExpressionSet.UniformPreshaderData);
+				for (int32 PreshaderIndex = 0; PreshaderIndex < NewCompilationOutput.UniformExpressionSet.UniformPreshaders.Num(); PreshaderIndex++)
+				{
+					const FMaterialUniformPreshaderHeader& PreshaderHeader = NewCompilationOutput.UniformExpressionSet.UniformPreshaders[PreshaderIndex];
+					const FMaterialUniformPreshaderField& PreshaderField = NewCompilationOutput.UniformExpressionSet.UniformPreshaderFields[PreshaderIndex];
+
+					UE::Shader::FPreshaderDataContext PreshaderContext(PreshaderContextBase, PreshaderHeader.OpcodeOffset, PreshaderHeader.OpcodeSize);
+					FString PreshaderDebug = PreshaderGenerateDebugString(&NewCompilationOutput.UniformExpressionSet, MaterialContext, PreshaderContext, &ParameterReferences);
+
+					// If this is a numeric field, add a swizzle for it
+					const TCHAR* SwizzleSuffix = TEXT("");
+					if (((uint8)PreshaderField.Type >= (uint8)UE::Shader::EValueType::Float1 && (uint8)PreshaderField.Type <= (uint8)UE::Shader::EValueType::Float4) ||
+						((uint8)PreshaderField.Type >= (uint8)UE::Shader::EValueType::Int1 && (uint8)PreshaderField.Type <= (uint8)UE::Shader::EValueType::Int4) ||
+						((uint8)PreshaderField.Type >= (uint8)UE::Shader::EValueType::Bool1 && (uint8)PreshaderField.Type <= (uint8)UE::Shader::EValueType::Bool4))
+					{
+						UE::Shader::FType ShaderType(PreshaderField.Type);
+
+						// First axis is offset, second axis is number of components (minus one)
+						static const TCHAR* SwizzleTable[4][4] =
+						{
+							{ TEXT(".x"), TEXT(".xy"), TEXT(".xyz"), TEXT(".xyzw") },
+							{ TEXT(".y"), TEXT(".yz"), TEXT(".yzw"), TEXT(".yzw?") },
+							{ TEXT(".z"), TEXT(".zw"), TEXT(".zw?"), TEXT(".zw??") },
+							{ TEXT(".w"), TEXT(".w?"), TEXT(".w??"), TEXT(".w???") },
+						};
+						SwizzleSuffix = SwizzleTable[PreshaderField.BufferOffset % 4][ShaderType.GetNumComponents() - 1];
+					}
+
+					OutSource.Appendf(TEXT("// PreshaderBuffer[%d]%s = %s\n"), PreshaderField.BufferOffset / 4, SwizzleSuffix, *PreshaderDebug);
+				}
+
+				// Sort parameter references by frequency
+				TArray<FSetElementId> ParameterReferencesSort;
+				ParameterReferencesSort.Reserve(ParameterReferences.Num());
+				for (auto ParameterReferenceIt = ParameterReferences.CreateConstIterator(); ParameterReferenceIt; ++ParameterReferenceIt)
+				{
+					ParameterReferencesSort.Add(ParameterReferenceIt.GetId());
+				}
+				Algo::Sort(ParameterReferencesSort, [&ParameterReferences](const FSetElementId& A, const FSetElementId& B)
+					{
+						const auto& ParameterReferenceA = ParameterReferences.Get(A);
+						const auto& ParameterReferenceB = ParameterReferences.Get(B);
+						if (ParameterReferenceA.Value != ParameterReferenceB.Value)
+						{
+							// Descending count
+							return ParameterReferenceA.Value > ParameterReferenceB.Value;
+						}
+						return ParameterReferenceA.Key < ParameterReferenceB.Key;
+					});
+
+				// Print parameter references
+				OutSource.Append("\n// Preshader parameter reference counts:\n");
+
+				for (int32 ParameterReferenceIndex = 0; ParameterReferenceIndex < ParameterReferencesSort.Num(); ++ParameterReferenceIndex)
+				{
+					const auto& ParameterReference = ParameterReferences.Get(ParameterReferencesSort[ParameterReferenceIndex]);
+					OutSource.Appendf(TEXT("// Param[\"%s\"] = %d\n"), *ParameterReference.Key, ParameterReference.Value);
+				}
+			}
 			return true;
 		}
 	}
@@ -4342,6 +4417,29 @@ bool FMaterial::GetMaterialExpressionSource( FString& OutSource )
 #endif
 }
 
+void FMaterial::GetPreshaderStats(uint32& TotalParameters, uint32& TotalOps) const
+{
+	TotalParameters = 0;
+	TotalOps = 0;
+
+#if WITH_EDITORONLY_DATA
+	const FMaterialShaderMap* ShaderMap = GetGameThreadShaderMap();
+	if (ShaderMap)
+	{
+		const FUniformExpressionSet& UniformExpressionSet = ShaderMap->GetUniformExpressionSet();
+		FMaterialRenderContext MaterialContext(nullptr, *this, nullptr);
+		UE::Shader::FPreshaderDataContext PreshaderContextBase(UniformExpressionSet.UniformPreshaderData);
+		for (const FMaterialUniformPreshaderHeader& PreshaderHeader : UniformExpressionSet.UniformPreshaders)
+		{
+			UE::Shader::FPreshaderDataContext PreshaderContext(PreshaderContextBase, PreshaderHeader.OpcodeOffset, PreshaderHeader.OpcodeSize);
+			PreshaderComputeDebugStats(&UniformExpressionSet, MaterialContext, PreshaderContext, TotalParameters, TotalOps);
+		}
+	}
+#else
+	UE_LOG(LogMaterial, Fatal,TEXT("GetPreshaderStats is only supported for WITH_EDITOR builds."));
+#endif
+}
+
 bool FMaterial::WritesEveryPixel(bool bShadowPass) const
 {
 	const bool bVFTypeSupportsNullPixelShader = !IsUsedWithInstancedStaticMeshes();
diff --git a/Engine/Source/Runtime/Engine/Private/Shader/Preshader.cpp b/Engine/Source/Runtime/Engine/Private/Shader/Preshader.cpp
--- a/Engine/Source/Runtime/Engine/Private/Shader/Preshader.cpp
+++ b/Engine/Source/Runtime/Engine/Private/Shader/Preshader.cpp
@@ -20,6 +20,77 @@ namespace UE
 namespace Shader
 {
 
+// For prettier debug printing of preshaders, we'd like add / subtract / multiply / divide to be operators instead of function calls.  To
+// implement this, we need to know when a given operator requires parentheses for correct order of operations.  We track the type of
+// expression for each stack element, and can decide whether an operator is valid without parentheses based on the state of the elements:
+//
+//		Singular,		Singular		+-*/ valid
+//		Singular,		SumTerms		+    valid
+//		Singular,		ProductTerms	+-*  valid
+//		SumTerms,		Singular		+-   valid
+//		SumTerms,		SumTerms		+    valid
+//		SumTerms,		ProductTerms	+-   valid
+//		ProductTerms,	Singular		+-*/ valid
+//		ProductTerms,	SumTerms		+    valid
+//		ProductTerms,	ProductTerms	+-*  valid
+//
+// Besides the above, swizzles require parentheses for non-singular elements, so we can use the state for that purpose as well.  Other
+// operators (like comparison), are quite rare in Preshaders, so it's not worth the effort to try to convert those.
+//
+enum class EPreshaderDebugStackType : uint8
+{
+	Singular,				// Stack element is a singular term (constant, param, function, parentheses wrapped, etc)
+	SumTerms,				// Stack element is a series of added / subtracted terms
+	ProductTerms,			// Stack element is a series of multiplied / divided terms
+};
+
+enum class EPreshaderOperatorType : uint8
+{
+	Add,
+	Sub,
+	Mul,
+	Div,
+	None
+};
+
+// Each entry in the table is the highest operator enum value that's valid without parentheses for that state combination.
+// As you can see from the valid combinations listed in the comment above, valid combinations are all monotonic sequences,
+// which is why we only need to store the maximum valid operator type, as opposed to flags.
+static const EPreshaderOperatorType PreshaderDebugStackLookup[3][3] =
+{
+	// Second element:
+	// Singular						SumTerms						ProductTerms
+	{ EPreshaderOperatorType::Div,	EPreshaderOperatorType::Add,	EPreshaderOperatorType::Mul },		// First element:  Singular
+	{ EPreshaderOperatorType::Sub,	EPreshaderOperatorType::Add,	EPreshaderOperatorType::Sub },		// First element:  SumTerms
+	{ EPreshaderOperatorType::Div,	EPreshaderOperatorType::Add,	EPreshaderOperatorType::Mul },		// First element:  ProductTerms
+};
+
+struct FPreshaderDebugStackElement
+{
+	FString Text;
+	EPreshaderDebugStackType DebugStackType;
+};
+
+struct FPreshaderDebugStack
+{
+	TArray<FPreshaderDebugStackElement> Elements;
+	TMap<FString, uint32>* ParameterReferences = nullptr;
+	uint32 OpCount = 0;
+	uint32 ParameterCount = 0;
+	bool bGenerateString = true;
+
+	void Push(FString Text, EPreshaderDebugStackType DebugStackType)
+	{
+		Elements.Push({ MoveTemp(Text), DebugStackType });
+	}
+
+	void Reset()
+	{
+		Elements.Reset();
+		// By design, doesn't reset "ParameterReferences", so they can be accumulated for the caller
+	}
+};
+
 FPreshaderType::FPreshaderType(const FType& InType) : ValueType(InType.ValueType)
 {
 	if (InType.IsStruct())
@@ -326,6 +397,22 @@ static void EvaluateConstantZero(FPreshaderStack& Stack, FPreshaderDataContext&
 	Stack.PushEmptyValue(Type); // Leave the empty value zero-initialized
 }
 
+static void EvaluateConstantZero(FPreshaderDebugStack& Stack, FPreshaderDataContext& RESTRICT Data)
+{
+	Stack.OpCount++;
+
+	FPreshaderStack TempStack;
+	EvaluateConstantZero(TempStack, Data);
+
+	if (Stack.bGenerateString)
+	{
+		// Display scalar constants without type information, for better readability
+		bool bIsScalar = TempStack.PeekValue().Type.ValueType == EValueType::Float1;
+		TStringBuilder<256> ValueBuffer;
+		Stack.Push(TempStack.PeekValue().AsShaderValue().ToString(bIsScalar ? EValueStringFormat::Description : EValueStringFormat::HLSL, ValueBuffer), EPreshaderDebugStackType::Singular);
+	}
+}
+
 static void EvaluateConstant(FPreshaderStack& Stack, FPreshaderDataContext& RESTRICT Data)
 {
 	const FPreshaderType Type = ReadPreshaderValue<FPreshaderType>(Data);
@@ -351,6 +438,22 @@ static void EvaluateConstant(FPreshaderStack& Stack, FPreshaderDataContext& REST
 	}
 }
 
+static void EvaluateConstant(FPreshaderDebugStack& Stack, FPreshaderDataContext& RESTRICT Data)
+{
+	Stack.OpCount++;
+
+	FPreshaderStack TempStack;
+	EvaluateConstant(TempStack, Data);
+
+	if (Stack.bGenerateString)
+	{
+		// Display scalar constants without type information, for better readability
+		bool bIsScalar = TempStack.PeekValue().Type.ValueType == EValueType::Float1;
+		TStringBuilder<256> ValueBuffer;
+		Stack.Push(TempStack.PeekValue().AsShaderValue().ToString(bIsScalar ? EValueStringFormat::Description : EValueStringFormat::HLSL, ValueBuffer), EPreshaderDebugStackType::Singular);
+	}
+}
+
 static void EvaluateSetField(FPreshaderStack& Stack, FPreshaderDataContext& RESTRICT Data)
 {
 	const FPreshaderValue Value = Stack.PopValue();
@@ -383,6 +486,22 @@ static void EvaluateSetField(FPreshaderStack& Stack, FPreshaderDataContext& REST
 	}
 }
 
+static void EvaluateSetField(FPreshaderDebugStack& Stack, FPreshaderDataContext& RESTRICT Data)
+{
+	Stack.OpCount++;
+
+	const int32 ComponentIndex = ReadPreshaderValue<int32>(Data);
+	const int32 ComponentNum = ReadPreshaderValue<int32>(Data);
+
+	if (Stack.bGenerateString)
+	{
+		FString Value = Stack.Elements.Pop().Text;
+
+		Stack.Elements.Last().Text = FString::Printf(TEXT("SetField(%s[%d,%d], %s)"), *Stack.Elements.Last().Text, ComponentIndex, ComponentNum, *Value);
+		Stack.Elements.Last().DebugStackType = EPreshaderDebugStackType::Singular;
+	}
+}
+
 static void EvaluateGetField(FPreshaderStack& Stack, FPreshaderDataContext& RESTRICT Data)
 {
 	const FPreshaderValue StructValue = Stack.PopValue();
@@ -400,6 +519,28 @@ static void EvaluateGetField(FPreshaderStack& Stack, FPreshaderDataContext& REST
 	Stack.PushValue(FieldType, FieldComponents);
 }
 
+static void EvaluateGetField(FPreshaderDebugStack& Stack, FPreshaderDataContext& RESTRICT Data)
+{
+	Stack.OpCount++;
+
+	const FPreshaderType FieldType = ReadPreshaderValue<FPreshaderType>(Data);
+	const int32 ComponentIndex = ReadPreshaderValue<int32>(Data);
+	const int32 ComponentNum = FieldType.GetNumComponents();
+
+	if (Stack.bGenerateString)
+	{
+		if (Stack.Elements.Last().DebugStackType == EPreshaderDebugStackType::Singular)
+		{
+			Stack.Elements.Last().Text.Append(FString::Printf(TEXT("[%d]"), ComponentIndex));
+		}
+		else
+		{
+			Stack.Elements.Last().Text = FString::Printf(TEXT("(%s)[%d]"), *Stack.Elements.Last().Text, ComponentIndex);
+			Stack.Elements.Last().DebugStackType = EPreshaderDebugStackType::Singular;
+		}
+	}
+}
+
 static void EvaluatePushValue(FPreshaderStack& Stack, FPreshaderDataContext& RESTRICT Data)
 {
 	const int32 StackOffset = ReadPreshaderValue<uint16>(Data);
@@ -409,6 +550,17 @@ static void EvaluatePushValue(FPreshaderStack& Stack, FPreshaderDataContext& RES
 	Stack.PushValue(Value.Type, LocalComponent);
 }
 
+static void EvaluatePushValue(FPreshaderDebugStack& Stack, FPreshaderDataContext& RESTRICT Data)
+{
+	Stack.OpCount++;
+
+	const int32 StackOffset = ReadPreshaderValue<uint16>(Data);
+	if (Stack.bGenerateString)
+	{
+		Stack.Elements.Push(Stack.Elements[StackOffset]);
+	}
+}
+
 static void EvaluateAssign(FPreshaderStack& Stack)
 {
 	const FPreshaderValue Value = Stack.PopValue();
@@ -421,6 +573,18 @@ static void EvaluateAssign(FPreshaderStack& Stack)
 	Stack.PushValue(Value.Type, LocalComponent);
 }
 
+static void EvaluateAssign(FPreshaderDebugStack& Stack)
+{
+	Stack.OpCount++;
+
+	if (Stack.bGenerateString)
+	{
+		FPreshaderDebugStackElement Value = Stack.Elements.Pop();
+		Stack.Elements.Pop();
+		Stack.Elements.Push(Value);
+	}
+}
+
 static void EvaluateParameter(FPreshaderStack& Stack, const FUniformExpressionSet* UniformExpressionSet, uint32 ParameterIndex, const FMaterialRenderContext& Context)
 {
 	if (!UniformExpressionSet)
@@ -499,6 +663,47 @@ static void EvaluateParameter(FPreshaderStack& Stack, const FUniformExpressionSe
 	}
 }
 
+static void EvaluateParameter(FPreshaderDebugStack& Stack, const FUniformExpressionSet* UniformExpressionSet, uint32 ParameterIndex, const FMaterialRenderContext& Context)
+{
+	Stack.OpCount++;
+	Stack.ParameterCount++;
+
+	if (!Stack.bGenerateString && !Stack.ParameterReferences)
+	{
+		return;
+	}
+
+	const FMaterialNumericParameterInfo& Parameter = UniformExpressionSet->GetNumericParameter(ParameterIndex);
+	FString ParameterName = Parameter.ParameterInfo.Name.ToString();
+
+	// Count the number of references to each parameter
+	if (Stack.ParameterReferences)
+	{
+		uint32 ParameterNameHash = GetTypeHash(ParameterName);
+		uint32* ParameterCount = Stack.ParameterReferences->FindByHash(ParameterNameHash, ParameterName);
+		if (ParameterCount)
+		{
+			(*ParameterCount)++;
+		}
+		else
+		{
+			Stack.ParameterReferences->AddByHash(ParameterNameHash, ParameterName, 1);
+		}
+	}
+
+	if (Stack.bGenerateString)
+	{
+		if (!UniformExpressionSet)
+		{
+			// return 0 for parameters if we don't have UniformExpressionSet
+			Stack.Push(TEXT("0.0"), EPreshaderDebugStackType::Singular);
+			return;
+		}
+
+		Stack.Push(FString::Printf(TEXT("Param[\"%s\"]"), *ParameterName), EPreshaderDebugStackType::Singular);
+	}
+}
+
 // Utility function to normalize a vector in place, for use with EvaluateUnaryOpInPlace
 EValueType NormalizeInPlace(EValueType Type, TArrayView<FValueComponent> Component)
 {
@@ -597,6 +802,18 @@ static inline void EvaluateUnaryOp(FPreshaderStack& Stack, const Operation& Op)
 	Stack.PushValue(Op(Value));
 }
 
+static inline void EvaluateUnaryOpDebug(FPreshaderDebugStack& Stack, const TCHAR* Op)
+{
+	Stack.OpCount++;
+
+	if (Stack.bGenerateString)
+	{
+		// Length op is in a lambda, rather than a standalone function, so replace that with "Length" if we detect that's the Op passed in
+		Stack.Elements.Last().Text = FString::Printf(TEXT("%s( %s )"), FCString::Strstr(Op, TEXT("Sqrt(Dot(")) ? TEXT("Length") : Op, *Stack.Elements.Last().Text);
+		Stack.Elements.Last().DebugStackType = EPreshaderDebugStackType::Singular;
+	}
+}
+
 template<typename Operation>
 static inline void EvaluateUnaryOpInPlace(FPreshaderStack& Stack, const Operation& Op)
 {
@@ -612,6 +829,19 @@ static inline void EvaluateUnaryOpInPlace(FPreshaderStack& Stack, const Operatio
 	}
 }
 
+static inline void EvaluateUnaryOpInPlaceDebug(FPreshaderDebugStack& Stack, const TCHAR* Op)
+{
+	Stack.OpCount++;
+
+	if (Stack.bGenerateString)
+	{
+		FString LocalOp = Op;
+		LocalOp.RemoveFromEnd("InPlace");
+		Stack.Elements.Last().Text = FString::Printf(TEXT("%s( %s )"), *LocalOp, *Stack.Elements.Last().Text);
+		Stack.Elements.Last().DebugStackType = EPreshaderDebugStackType::Singular;
+	}
+}
+
 template<typename Operation>
 static inline void EvaluateBinaryOp(FPreshaderStack& Stack, const Operation& Op)
 {
@@ -620,6 +850,70 @@ static inline void EvaluateBinaryOp(FPreshaderStack& Stack, const Operation& Op)
 	Stack.PushValue(Op(Value0, Value1));
 }
 
+static inline void EvaluateBinaryOpDebug(FPreshaderDebugStack& Stack, const TCHAR* Op)
+{
+	Stack.OpCount++;
+
+	if (Stack.bGenerateString)
+	{
+		FPreshaderDebugStackElement Value1 = Stack.Elements.Pop();
+		FPreshaderDebugStackElement Value0 = Stack.Elements.Pop();
+
+		EPreshaderOperatorType PreshaderOperator = EPreshaderOperatorType::None;
+		if (!FCString::Strcmp(Op, TEXT("Add")))
+		{
+			PreshaderOperator = EPreshaderOperatorType::Add;
+		}
+		else if (!FCString::Strcmp(Op, TEXT("Sub")))
+		{
+			PreshaderOperator = EPreshaderOperatorType::Sub;
+		}
+		else if (!FCString::Strcmp(Op, TEXT("Mul")))
+		{
+			PreshaderOperator = EPreshaderOperatorType::Mul;
+		}
+		else if (!FCString::Strcmp(Op, TEXT("Div")))
+		{
+			PreshaderOperator = EPreshaderOperatorType::Div;
+		}
+
+		if (PreshaderOperator != EPreshaderOperatorType::None)
+		{
+			static const TCHAR PreshaderOperatorChars[4] = { '+', '-', '*', '/' };
+
+			TCHAR OperatorChar = PreshaderOperatorChars[(uint8)PreshaderOperator];
+			EPreshaderDebugStackType NewDebugStackType = PreshaderOperator >= EPreshaderOperatorType::Mul ? EPreshaderDebugStackType::ProductTerms : EPreshaderDebugStackType::SumTerms;
+			EPreshaderOperatorType MaxOperatorValid = PreshaderDebugStackLookup[(uint8)Value0.DebugStackType][(uint8)Value1.DebugStackType];
+
+			if (PreshaderOperator <= MaxOperatorValid)
+			{
+				// Operator is valid without parentheses
+				Stack.Push(FString::Printf(TEXT("%s %c %s"), *Value0.Text, OperatorChar, *Value1.Text), NewDebugStackType);
+			}
+			else
+			{
+				// Add parentheses around non-singular parameters.  At least one is non-singular, otherwise we would have gone through the above path
+				if (Value0.DebugStackType == EPreshaderDebugStackType::Singular)
+				{
+					Stack.Push(FString::Printf(TEXT("%s %c (%s)"), *Value0.Text, OperatorChar, *Value1.Text), NewDebugStackType);
+				}
+				else if (Value1.DebugStackType == EPreshaderDebugStackType::Singular)
+				{
+					Stack.Push(FString::Printf(TEXT("(%s) %c %s"), *Value0.Text, OperatorChar, *Value1.Text), NewDebugStackType);
+				}
+				else
+				{
+					Stack.Push(FString::Printf(TEXT("(%s) %c (%s)"), *Value0.Text, OperatorChar, *Value1.Text), NewDebugStackType);
+				}
+			}
+		}
+		else
+		{
+			Stack.Push(FString::Printf(TEXT("%s(%s, %s)"), Op, *Value0.Text, *Value1.Text), EPreshaderDebugStackType::Singular);
+		}
+	}
+}
+
 template<typename Operation, typename OperationFallback>
 static inline void EvaluateBinaryOpInPlace(FPreshaderStack& Stack, const Operation& Op, const OperationFallback& FallbackOp)
 {
@@ -649,6 +943,19 @@ static inline void EvaluateTernaryOp(FPreshaderStack& Stack, const Operation& Op
 	Stack.PushValue(Op(Value0, Value1, Value2));
 }
 
+static inline void EvaluateTernaryOpDebug(FPreshaderDebugStack& Stack, const TCHAR* Op)
+{
+	Stack.OpCount++;
+
+	if (Stack.bGenerateString)
+	{
+		FString Value2 = Stack.Elements.Pop().Text;
+		FString Value1 = Stack.Elements.Pop().Text;
+		FString Value0 = Stack.Elements.Pop().Text;
+		Stack.Push(FString::Printf(TEXT("%s(%s, %s, %s)"), Op, *Value0, *Value1, *Value2), EPreshaderDebugStackType::Singular);
+	}
+}
+
 // Runs the swizzle in place on the top item on the stack
 static void EvaluateComponentSwizzle(FPreshaderStack& Stack, FPreshaderDataContext& RESTRICT Data)
 {
@@ -727,6 +1034,44 @@ static void EvaluateComponentSwizzle(FPreshaderStack& Stack, FPreshaderDataConte
 	}
 }
 
+static void EvaluateComponentSwizzle(FPreshaderDebugStack& Stack, FPreshaderDataContext& RESTRICT Data)
+{
+	Stack.OpCount++;
+
+	const uint8 NumElements = ReadPreshaderValue<uint8>(Data);
+	const uint8 IndexR = FMath::Min(ReadPreshaderValue<uint8>(Data), (uint8)3);
+	const uint8 IndexG = FMath::Min(ReadPreshaderValue<uint8>(Data), (uint8)3);
+	const uint8 IndexB = FMath::Min(ReadPreshaderValue<uint8>(Data), (uint8)3);
+	const uint8 IndexA = FMath::Min(ReadPreshaderValue<uint8>(Data), (uint8)3);
+
+	if (Stack.bGenerateString)
+	{
+		if (Stack.Elements.Last().DebugStackType != EPreshaderDebugStackType::Singular)
+		{
+			Stack.Elements.Last().Text = FString::Printf(TEXT("(%s)"), *Stack.Elements.Last().Text);
+			Stack.Elements.Last().DebugStackType = EPreshaderDebugStackType::Singular;
+		}
+
+		static const TCHAR ComponentChars[4] = { 'x', 'y', 'z', 'w' };
+
+		switch (NumElements)
+		{
+		case 4:
+			Stack.Elements.Last().Text = FString::Printf(TEXT("%s.%c%c%c%c"), *Stack.Elements.Last().Text, ComponentChars[IndexR], ComponentChars[IndexG], ComponentChars[IndexB], ComponentChars[IndexA]);
+			break;
+		case 3:
+			Stack.Elements.Last().Text = FString::Printf(TEXT("%s.%c%c%c"), *Stack.Elements.Last().Text, ComponentChars[IndexR], ComponentChars[IndexG], ComponentChars[IndexB]);
+			break;
+		case 2:
+			Stack.Elements.Last().Text = FString::Printf(TEXT("%s.%c%c"), *Stack.Elements.Last().Text, ComponentChars[IndexR], ComponentChars[IndexG]);
+			break;
+		case 1:
+			Stack.Elements.Last().Text = FString::Printf(TEXT("%s.%c"), *Stack.Elements.Last().Text, ComponentChars[IndexR]);
+			break;
+		}
+	}
+}
+
 static void GetTextureParameter(const FMaterialRenderContext& Context, FPreshaderDataContext& RESTRICT Data, const UTexture*& OutTexture, const USparseVolumeTexture*& OutSparseVolumeTexture)
 {
 	const FHashedMaterialParameterInfo ParameterInfo = ReadPreshaderValue<FHashedMaterialParameterInfo>(Data);
@@ -767,6 +1112,18 @@ static void EvaluateTextureSize(const FMaterialRenderContext& Context, FPreshade
 	}
 }
 
+static void EvaluateTextureSize(const FMaterialRenderContext& Context, FPreshaderDebugStack& Stack, FPreshaderDataContext& RESTRICT Data)
+{
+	Stack.OpCount++;
+
+	const FHashedMaterialParameterInfo ParameterInfo = ReadPreshaderValue<FHashedMaterialParameterInfo>(Data);
+	const int32 TextureIndex = ReadPreshaderValue<int32>(Data);
+	if (Stack.bGenerateString)
+	{
+		Stack.Push(FString::Printf(TEXT("TEX[\"%s\"][%d].TextureSize"), *ParameterInfo.GetName().ToString(), TextureIndex), EPreshaderDebugStackType::Singular);
+	}
+}
+
 static void EvaluateTexelSize(const FMaterialRenderContext& Context, FPreshaderStack& Stack, FPreshaderDataContext& RESTRICT Data)
 {
 	const UTexture* Texture = nullptr;
@@ -792,6 +1149,18 @@ static void EvaluateTexelSize(const FMaterialRenderContext& Context, FPreshaderS
 	}
 }
 
+static void EvaluateTexelSize(const FMaterialRenderContext& Context, FPreshaderDebugStack& Stack, FPreshaderDataContext& RESTRICT Data)
+{
+	Stack.OpCount++;
+
+	const FHashedMaterialParameterInfo ParameterInfo = ReadPreshaderValue<FHashedMaterialParameterInfo>(Data);
+	const int32 TextureIndex = ReadPreshaderValue<int32>(Data);
+	if (Stack.bGenerateString)
+	{
+		Stack.Push(FString::Printf(TEXT("TEX[\"%s\"][%d].TexelSize"), *ParameterInfo.GetName().ToString(), TextureIndex), EPreshaderDebugStackType::Singular);
+	}
+}
+
 static FGuid GetExternalTextureGuid(const FMaterialRenderContext& Context, FPreshaderDataContext& RESTRICT Data)
 {
 	const FScriptName ParameterName = ReadPreshaderValue<FScriptName>(Data);
@@ -822,6 +1191,34 @@ static void EvaluateExternalTextureCoordinateOffset(const FMaterialRenderContext
 	Stack.PushValue(Result);
 }
 
+static void EvaluateExternalTextureCoordinateScaleRotation(const FMaterialRenderContext& Context, FPreshaderDebugStack& Stack, FPreshaderDataContext& RESTRICT Data)
+{
+	Stack.OpCount++;
+
+	const FScriptName ParameterName = ReadPreshaderValue<FScriptName>(Data);
+	const FGuid ExternalTextureGuid = ReadPreshaderValue<FGuid>(Data);
+	const int32 TextureIndex = ReadPreshaderValue<int32>(Data);
+
+	if (Stack.bGenerateString)
+	{
+		Stack.Push(FString::Printf(TEXT("TEX[\"%s\"].ExternalTextureScaleRotation"), *ParameterName.ToString()), EPreshaderDebugStackType::Singular);
+	}
+}
+
+static void EvaluateExternalTextureCoordinateOffset(const FMaterialRenderContext& Context, FPreshaderDebugStack& Stack, FPreshaderDataContext& RESTRICT Data)
+{
+	Stack.OpCount++;
+
+	const FScriptName ParameterName = ReadPreshaderValue<FScriptName>(Data);
+	const FGuid ExternalTextureGuid = ReadPreshaderValue<FGuid>(Data);
+	const int32 TextureIndex = ReadPreshaderValue<int32>(Data);
+
+	if (Stack.bGenerateString)
+	{
+		Stack.Push(FString::Printf(TEXT("TEX[\"%s\"].ExternalTextureOffset"), *ParameterName.ToString()), EPreshaderDebugStackType::Singular);
+	}
+}
+
 static void EvaluateRuntimeVirtualTextureUniform(const FMaterialRenderContext& Context, FPreshaderStack& Stack, FPreshaderDataContext& RESTRICT Data)
 {
 	const FHashedMaterialParameterInfo ParameterInfo = ReadPreshaderValue<FHashedMaterialParameterInfo>(Data);
@@ -843,6 +1240,20 @@ static void EvaluateRuntimeVirtualTextureUniform(const FMaterialRenderContext& C
 	}
 }
 
+static void EvaluateRuntimeVirtualTextureUniform(const FMaterialRenderContext& Context, FPreshaderDebugStack& Stack, FPreshaderDataContext& RESTRICT Data)
+{
+	Stack.OpCount++;
+
+	const FHashedMaterialParameterInfo ParameterInfo = ReadPreshaderValue<FHashedMaterialParameterInfo>(Data);
+	const int32 TextureIndex = ReadPreshaderValue<int32>(Data);
+	const int32 VectorIndex = ReadPreshaderValue<int32>(Data);
+
+	if (Stack.bGenerateString)
+	{
+		Stack.Push(FString::Printf(TEXT("TEX[\"%s\"].VirtualTextureUniform[%d]"), *ParameterInfo.Name.ToString(), VectorIndex), EPreshaderDebugStackType::Singular);
+	}
+}
+
 static void EvaluateSparseTextureUniform(const FMaterialRenderContext& Context, FPreshaderStack& Stack, FPreshaderDataContext& RESTRICT Data)
 {
 	const FHashedMaterialParameterInfo ParameterInfo = ReadPreshaderValue<FHashedMaterialParameterInfo>(Data);
@@ -864,6 +1275,20 @@ static void EvaluateSparseTextureUniform(const FMaterialRenderContext& Context,
 	}
 }
 
+static void EvaluateSparseTextureUniform(const FMaterialRenderContext& Context, FPreshaderDebugStack& Stack, FPreshaderDataContext& RESTRICT Data)
+{
+	Stack.OpCount++;
+
+	const FHashedMaterialParameterInfo ParameterInfo = ReadPreshaderValue<FHashedMaterialParameterInfo>(Data);
+	const int32 TextureIndex = ReadPreshaderValue<int32>(Data);
+	const int32 VectorIndex = ReadPreshaderValue<int32>(Data);
+
+	if (Stack.bGenerateString)
+	{
+		Stack.Push(FString::Printf(TEXT("TEX[\"%s\"].SparseTextureUniform[%d]"), *ParameterInfo.Name.ToString(), VectorIndex), EPreshaderDebugStackType::Singular);
+	}
+}
+
 static void EvaluateJump(FPreshaderDataContext& RESTRICT Data)
 {
 	const int32 JumpOffset = ReadPreshaderValue<int32>(Data);
@@ -883,6 +1308,22 @@ static void EvaluateJumpIfFalse(FPreshaderStack& Stack, FPreshaderDataContext& R
 	}
 }
 
+static void EvaluateJumpIfFalse(FPreshaderDebugStack& Stack, FPreshaderDataContext& RESTRICT Data)
+{
+	Stack.OpCount++;
+
+	const int32 JumpOffset = ReadPreshaderValue<int32>(Data);
+	check(Data.Ptr + JumpOffset <= Data.EndPtr);
+
+	// It's complicated to include conditional jumps in a stack evaluator.  Ideally we'd want to evaluate and display both
+	// branches, but that requires evaluating both branches forward to where the branches reach the same instruction
+	// again, then converting stack entries that diverge between the two code paths to include a conditional expression.
+	// I don't have any examples of bytecode with conditional branches to be able to test an implementation of that (I
+	// think only the new translator generates these), so we'll just ignore the condition completely for now, and it won't
+	// show up in the debug expression.
+	FString ConditionValue = Stack.Elements.Pop().Text;
+}
+
 FPreshaderValue EvaluatePreshader(const FUniformExpressionSet* UniformExpressionSet, const FMaterialRenderContext& Context, FPreshaderStack& Stack, FPreshaderDataContext& RESTRICT Data)
 {
 	uint8 const* const DataEnd = Data.EndPtr;
@@ -964,6 +1405,117 @@ FPreshaderValue EvaluatePreshader(const FUniformExpressionSet* UniformExpression
 	return Result;
 }
 
+#define EvaluateUnaryOp(STACK, OP) EvaluateUnaryOpDebug(STACK, TEXT(#OP))
+#define EvaluateUnaryOpInPlace(STACK, IN_PLACE_OP) EvaluateUnaryOpInPlaceDebug(STACK, TEXT(#IN_PLACE_OP))
+#define EvaluateBinaryOp(STACK, OP) EvaluateBinaryOpDebug(STACK, TEXT(#OP))
+#define EvaluateBinaryOpInPlace(STACK, IN_PLACE_OP, OP) EvaluateBinaryOpDebug(STACK, TEXT(#OP))
+#define EvaluateTernaryOp(STACK, OP) EvaluateTernaryOpDebug(STACK, TEXT(#OP))
+
+// The body of this function is designed to be completely identical to EvaluatePreshader above, other than the lack of popping the result
+// at the end, so new opcodes can be added and the body trivially cut-and-pasted, for ease of maintenance.  Temporary redefining some Evaluate
+// functions as macros is done above to automate string-izing of function names.
+static void EvaluatePreshaderDebug(const FUniformExpressionSet* UniformExpressionSet, const FMaterialRenderContext& Context, FPreshaderDebugStack& Stack, FPreshaderDataContext& RESTRICT Data)
+{
+	uint8 const* const DataEnd = Data.EndPtr;
+
+	Stack.Reset();
+	while (Data.Ptr < DataEnd)
+	{
+		const EPreshaderOpcode Opcode = (EPreshaderOpcode)ReadPreshaderValue<uint8>(Data);
+		switch (Opcode)
+		{
+		case EPreshaderOpcode::ConstantZero: EvaluateConstantZero(Stack, Data); break;
+		case EPreshaderOpcode::Constant: EvaluateConstant(Stack, Data); break;
+		case EPreshaderOpcode::GetField: EvaluateGetField(Stack, Data); break;
+		case EPreshaderOpcode::SetField: EvaluateSetField(Stack, Data); break;
+		case EPreshaderOpcode::Parameter:
+			EvaluateParameter(Stack, UniformExpressionSet, ReadPreshaderValue<uint16>(Data), Context);
+			break;
+		case EPreshaderOpcode::PushValue: EvaluatePushValue(Stack, Data); break;
+		case EPreshaderOpcode::Assign: EvaluateAssign(Stack); break;
+		case EPreshaderOpcode::Add: EvaluateBinaryOpInPlace(Stack, AddInPlace, Add); break;
+		case EPreshaderOpcode::Sub: EvaluateBinaryOpInPlace(Stack, SubInPlace, Sub); break;
+		case EPreshaderOpcode::Mul: EvaluateBinaryOpInPlace(Stack, MulInPlace, Mul); break;
+		case EPreshaderOpcode::Div: EvaluateBinaryOpInPlace(Stack, DivInPlace, Div); break;
+		case EPreshaderOpcode::Less: EvaluateBinaryOp(Stack, Less); break;
+		case EPreshaderOpcode::Greater: EvaluateBinaryOp(Stack, Greater); break;
+		case EPreshaderOpcode::LessEqual: EvaluateBinaryOp(Stack, LessEqual); break;
+		case EPreshaderOpcode::GreaterEqual: EvaluateBinaryOp(Stack, GreaterEqual); break;
+		case EPreshaderOpcode::Fmod: EvaluateBinaryOpInPlace(Stack, FmodInPlace, Fmod); break;
+		case EPreshaderOpcode::Min: EvaluateBinaryOpInPlace(Stack, MinInPlace, Min); break;
+		case EPreshaderOpcode::Max: EvaluateBinaryOpInPlace(Stack, MaxInPlace, Max); break;
+		case EPreshaderOpcode::Clamp: EvaluateTernaryOp(Stack, Clamp); break;
+		case EPreshaderOpcode::Dot: EvaluateBinaryOp(Stack, Dot); break;
+		case EPreshaderOpcode::Cross: EvaluateBinaryOp(Stack, Cross); break;
+		case EPreshaderOpcode::Neg: EvaluateUnaryOpInPlace(Stack, NegInPlace); break;
+		case EPreshaderOpcode::Sqrt: EvaluateUnaryOpInPlace(Stack, SqrtInPlace); break;
+		case EPreshaderOpcode::Rcp: EvaluateUnaryOpInPlace(Stack, RcpInPlace); break;
+		case EPreshaderOpcode::Length: EvaluateUnaryOp(Stack, [](const FValue& Value) { return Sqrt(Dot(Value, Value)); }); break;
+		case EPreshaderOpcode::Normalize: EvaluateUnaryOpInPlace(Stack, NormalizeInPlace); break;
+		case EPreshaderOpcode::Sin: EvaluateUnaryOpInPlace(Stack, SinInPlace); break;
+		case EPreshaderOpcode::Cos: EvaluateUnaryOpInPlace(Stack, CosInPlace); break;
+		case EPreshaderOpcode::Tan: EvaluateUnaryOpInPlace(Stack, TanInPlace); break;
+		case EPreshaderOpcode::Asin: EvaluateUnaryOpInPlace(Stack, AsinInPlace); break;;
+		case EPreshaderOpcode::Acos: EvaluateUnaryOpInPlace(Stack, AcosInPlace); break;
+		case EPreshaderOpcode::Atan: EvaluateUnaryOpInPlace(Stack, AtanInPlace); break;
+		case EPreshaderOpcode::Atan2: EvaluateBinaryOpInPlace(Stack, Atan2InPlace, Atan2); break;
+		case EPreshaderOpcode::Abs: EvaluateUnaryOpInPlace(Stack, AbsInPlace); break;
+		case EPreshaderOpcode::Saturate: EvaluateUnaryOpInPlace(Stack, SaturateInPlace); break;
+		case EPreshaderOpcode::Floor: EvaluateUnaryOpInPlace(Stack, FloorInPlace); break;
+		case EPreshaderOpcode::Ceil: EvaluateUnaryOpInPlace(Stack, CeilInPlace); break;
+		case EPreshaderOpcode::Round: EvaluateUnaryOpInPlace(Stack, RoundInPlace); break;
+		case EPreshaderOpcode::Trunc: EvaluateUnaryOpInPlace(Stack, TruncInPlace); break;
+		case EPreshaderOpcode::Sign: EvaluateUnaryOpInPlace(Stack, SignInPlace); break;
+		case EPreshaderOpcode::Frac: EvaluateUnaryOpInPlace(Stack, FracInPlace); break;
+		case EPreshaderOpcode::Fractional: EvaluateUnaryOpInPlace(Stack, FractionalInPlace); break;
+		case EPreshaderOpcode::Exp: EvaluateUnaryOpInPlace(Stack, ExpInPlace); break;
+		case EPreshaderOpcode::Exp2: EvaluateUnaryOpInPlace(Stack, Exp2InPlace); break;
+		case EPreshaderOpcode::Log: EvaluateUnaryOpInPlace(Stack, LogInPlace); break;
+		case EPreshaderOpcode::Log2: EvaluateUnaryOpInPlace(Stack, Log2InPlace); break;
+		case EPreshaderOpcode::Log10: EvaluateUnaryOpInPlace(Stack, Log10InPlace); break;
+		case EPreshaderOpcode::ComponentSwizzle: EvaluateComponentSwizzle(Stack, Data); break;
+		case EPreshaderOpcode::AppendVector: EvaluateBinaryOpInPlace(Stack, AppendInPlace, Append); break;
+		case EPreshaderOpcode::TextureSize: EvaluateTextureSize(Context, Stack, Data); break;
+		case EPreshaderOpcode::TexelSize: EvaluateTexelSize(Context, Stack, Data); break;
+		case EPreshaderOpcode::ExternalTextureCoordinateScaleRotation: EvaluateExternalTextureCoordinateScaleRotation(Context, Stack, Data); break;
+		case EPreshaderOpcode::ExternalTextureCoordinateOffset: EvaluateExternalTextureCoordinateOffset(Context, Stack, Data); break;
+		case EPreshaderOpcode::RuntimeVirtualTextureUniform: EvaluateRuntimeVirtualTextureUniform(Context, Stack, Data); break;
+		case EPreshaderOpcode::SparseVolumeTextureUniform: EvaluateSparseTextureUniform(Context, Stack, Data); break;
+		case EPreshaderOpcode::Jump: EvaluateJump(Data); break;
+		case EPreshaderOpcode::JumpIfFalse: EvaluateJumpIfFalse(Stack, Data); break;
+		default:
+			UE_LOG(LogMaterial, Fatal, TEXT("Unknown preshader opcode %d"), (uint8)Opcode);
+			break;
+		}
+	}
+	check(Data.Ptr == DataEnd);
+}
+
+#undef EvaluateUnaryOp
+#undef EvaluateUnaryOpInPlace
+#undef EvaluateBinaryOp
+#undef EvaluateBinaryOpInPlace
+#undef EvaluateTernaryOp
+
+FString PreshaderGenerateDebugString(const FUniformExpressionSet* UniformExpressionSet, const FMaterialRenderContext& Context, FPreshaderDataContext& RESTRICT Data, TMap<FString, uint32>* ParameterReferences)
+{
+	FPreshaderDebugStack Stack;
+	Stack.ParameterReferences = ParameterReferences;
+	EvaluatePreshaderDebug(UniformExpressionSet, Context, Stack, Data);
+
+	return Stack.Elements.Pop().Text;
+}
+
+void PreshaderComputeDebugStats(const FUniformExpressionSet* UniformExpressionSet, const FMaterialRenderContext& Context, FPreshaderDataContext& RESTRICT Data, uint32& TotalParameters, uint32& TotalOps)
+{
+	FPreshaderDebugStack Stack;
+	Stack.bGenerateString = false;
+	EvaluatePreshaderDebug(UniformExpressionSet, Context, Stack, Data);
+
+	TotalParameters += Stack.ParameterCount;
+	TotalOps += Stack.OpCount;
+}
+
 FPreshaderValue FPreshaderData::Evaluate(FUniformExpressionSet* UniformExpressionSet, const struct FMaterialRenderContext& Context, FPreshaderStack& Stack) const
 {
 	FPreshaderDataContext PreshaderContext(*this);
diff --git a/Engine/Source/Runtime/Engine/Private/Shader/PreshaderEvaluate.h b/Engine/Source/Runtime/Engine/Private/Shader/PreshaderEvaluate.h
--- a/Engine/Source/Runtime/Engine/Private/Shader/PreshaderEvaluate.h
+++ b/Engine/Source/Runtime/Engine/Private/Shader/PreshaderEvaluate.h
@@ -28,5 +28,8 @@ struct FPreshaderDataContext
 
 ENGINE_API FPreshaderValue EvaluatePreshader(const FUniformExpressionSet* UniformExpressionSet, const FMaterialRenderContext& Context, FPreshaderStack& Stack, FPreshaderDataContext& RESTRICT Data);
 
+ENGINE_API FString PreshaderGenerateDebugString(const FUniformExpressionSet* UniformExpressionSet, const FMaterialRenderContext& Context, FPreshaderDataContext& RESTRICT Data, TMap<FString, uint32>* ParameterReferences = nullptr);
+ENGINE_API void PreshaderComputeDebugStats(const FUniformExpressionSet* UniformExpressionSet, const FMaterialRenderContext& Context, FPreshaderDataContext& RESTRICT Data, uint32& TotalParameters, uint32& TotalOps);
+
 } // namespace Shader
 } // namespace UE
diff --git a/Engine/Source/Runtime/Engine/Public/MaterialShared.h b/Engine/Source/Runtime/Engine/Public/MaterialShared.h
--- a/Engine/Source/Runtime/Engine/Public/MaterialShared.h
+++ b/Engine/Source/Runtime/Engine/Public/MaterialShared.h
@@ -1431,6 +1431,7 @@ public:
 		CopyAssignItems(Result.GetData(), GetContent()->MaterialCompilationOutput.EstimatedLWCFuncUsages, (int)ELWCFunctionKind::Max);
 		return Result;
 	}
+	uint32 GetNumPreshaders() const { return GetContent()->MaterialCompilationOutput.UniformExpressionSet.UniformPreshaders.Num(); }
 #endif
 	uint32 GetNumVirtualTextureStacks() const { return GetContent()->MaterialCompilationOutput.UniformExpressionSet.VTStacks.Num(); }
 	uint8 GetRuntimeVirtualTextureOutputAttributeMask() const { return GetContent()->MaterialCompilationOutput.RuntimeVirtualTextureOutputAttributeMask; }
@@ -2274,6 +2275,9 @@ public:
 	*/
 	ENGINE_API bool GetMaterialExpressionSource(FString& OutSource);
 
+	/** Returns summary statistics for preshaders */
+	ENGINE_API void GetPreshaderStats(uint32& TotalParameters, uint32& TotalOps) const;
+
 	/* Helper function to look at both IsMasked and IsDitheredLODTransition to determine if it writes every pixel */
 	ENGINE_API bool WritesEveryPixel(bool bShadowPass = false) const;
 	ENGINE_API bool WritesEveryPixel(bool bShadowPass, bool bVFTypeSupportsNullPixelShader) const;
