Material Editor - This change enables material translation DDC caching by default.

diff --git a/Engine/Source/Runtime/Engine/Private/Materials/HLSLMaterialTranslator.cpp b/Engine/Source/Runtime/Engine/Private/Materials/HLSLMaterialTranslator.cpp
--- a/Engine/Source/Runtime/Engine/Private/Materials/HLSLMaterialTranslator.cpp
+++ b/Engine/Source/Runtime/Engine/Private/Materials/HLSLMaterialTranslator.cpp
@@ -193,11 +193,11 @@ static FAutoConsoleVariableRef CVarPedanticErrorChecksEnabled(
 static bool GUseMaterialTranslationResultsGrouping = true;
 
 /* Controls whether DDC caching of material translation results is enabled. */
-static bool GJobMaterialTranslateDDCEnable = false;
-static FAutoConsoleVariableRef CVarJobMaterialTranslateDDCEnable(
-	TEXT("r.Material.TranslateDDCEnable"),
-	GJobMaterialTranslateDDCEnable,
-	TEXT("Whether to enable material translation DDC caching.\n"));
+static bool GJobDisableMaterialTranslateDDC = false;
+static FAutoConsoleVariableRef CVarJobDisableMaterialTranslateDDC(
+	TEXT("r.Material.DisableTranslateDDC"),
+	GJobDisableMaterialTranslateDDC,
+	TEXT("Whether to disable material translation DDC caching.\n"));
 
 UE::DerivedData::FCacheBucket MaterialTranslationDDCBucket = UE::DerivedData::FCacheBucket(TEXT("MaterialTranslation"));
 UE::DerivedData::FValueId MaterialCompilationOutputId = UE::DerivedData::FValueId::FromName("FHLSLMaterialTranslator_MaterialCompilationOutput");
@@ -15200,7 +15200,7 @@ void FHLSLMaterialTranslator::PrepareEnvironmentDefines()
 
 bool FHLSLMaterialTranslator::QueryDDCCachedTranslationResults()
 {
-	if (!GJobMaterialTranslateDDCEnable)
+	if (GJobDisableMaterialTranslateDDC)
 	{
 		return false;
 	}
@@ -15266,7 +15266,7 @@ void FHLSLMaterialTranslator::PushResultsToDDCCache()
 	// We currently don't support caching parameter collections on the DDC, as it is invalid at this point
 	// to serialize an array of object pointers. Only allow putting the translation results on the DDC
 	// if the this material does not use parameter collections.
-	if (!GJobMaterialTranslateDDCEnable || !ParameterCollections.IsEmpty())
+	if (GJobDisableMaterialTranslateDDC || !ParameterCollections.IsEmpty())
 	{
 		GShaderCompilerStats->IncrementMaterialTranslationSkippedDDC();
 		return;
