Fixed curve filtering on DeadBlending. Fixed issue when BlendProfiles and BlendModes were used in combination on DeadBlending. Added bone filter to Inertialization node. Memory improvements to Inertialization node and some cleanup/deprecation of old code.

diff --git a/Engine/Source/Runtime/Engine/Classes/Animation/AnimNode_DeadBlending.h b/Engine/Source/Runtime/Engine/Classes/Animation/AnimNode_DeadBlending.h
--- a/Engine/Source/Runtime/Engine/Classes/Animation/AnimNode_DeadBlending.h
+++ b/Engine/Source/Runtime/Engine/Classes/Animation/AnimNode_DeadBlending.h
@@ -3,7 +3,6 @@
 #pragma once
 
 #include "AnimNode_Inertialization.h"
-#include "Interfaces/Interface_BoneReferenceSkeletonProvider.h"
 #include "AnimNode_DeadBlending.generated.h"
 
 /**
@@ -75,6 +74,13 @@ private:
 	UPROPERTY(EditAnywhere, Category = Filter)
 	TArray<FName> FilteredCurves;
 
+	/**
+	 * List of curves that will not be included in the extrapolation. Curves in this list will effectively act like they have had their value reset 
+	 * to zero at the point of the transition, and will be blended in with the curve values in the animation that is being transitioned to.
+	 */ 
+	UPROPERTY(EditAnywhere, Category = Filter)
+	TArray<FName> ExtrapolationFilteredCurves;
+
 	// List of bones that should not use inertial blending. These bones will change instantly when the animation switches.
 	UPROPERTY(EditAnywhere, Category = Filter)
 	TArray<FBoneReference> FilteredBones;
@@ -131,14 +137,6 @@ private:
 	UPROPERTY(EditAnywhere, Category = Extrapolation, meta = (Min = "0.0", UIMin = "0.0"))
 	float MaximumCurveVelocity = 100.0f;
 
-	/**
-	 * Enable this to pre-allocate memory for the node rather than to allocate and deallocate memory when blending 
-	 * becomes active and inactive. This improves performance, but causes larger memory usage, in particular when you 
-	 * have multiple Dead Blending nodes in an animation graph that are not all used at once.
-	 */
-	UPROPERTY(EditAnywhere, Category = Memory)
-	bool bPreallocateMemory = false;
-
 	/**
 	* Clear any active blends if we just became relevant, to avoid carrying over undesired blends.
 	*/	
@@ -175,7 +173,7 @@ public: // FAnimNode_Base
 private:
 	
 	/**
-	 * Deactivates the inertialization and frees any temporary memory (unless bPreallocateMemory is set).
+	 * Deactivates the inertialization and frees any temporary memory.
 	 */
 	void Deactivate();
 
@@ -209,6 +207,9 @@ private:
 	// Cached curve filter built from FilteredCurves
 	UE::Anim::FCurveFilter CurveFilter;
 
+	// Cached curve filter built from ExtrapolationFilteredCurves
+	UE::Anim::FCurveFilter ExtrapolatedCurveFilter;
+
 	// Cache compact pose bone index for FilteredBones
 	TArray<FCompactPoseBoneIndex, TInlineAllocator<8>> BoneFilter;
 
@@ -227,7 +228,7 @@ private:
 
 	// Recorded pose state at point of transition.
 
-	TArray<uint16> BoneIndices;
+	TArray<int32> BoneIndices;
 	
 	TArray<FVector> BoneTranslations;
 	TArray<FQuat> BoneRotations;
@@ -252,8 +253,12 @@ private:
 		FDeadBlendingCurveElement() = default;
 	};
 
+	// Recorded curve state at the point of transition
 	TBaseBlendedCurve<TInlineAllocator<8>, FDeadBlendingCurveElement> CurveData;
 
+	// Temporary storage for curve data of the Destination Pose
+	TBaseBlendedCurve<TInlineAllocator<8>, UE::Anim::FCurveElement> PoseCurveData;
+
 private:
 
 	// Variable used to store the elapsed delta time between calls to evaluate.
diff --git a/Engine/Source/Runtime/Engine/Classes/Animation/AnimNode_Inertialization.h b/Engine/Source/Runtime/Engine/Classes/Animation/AnimNode_Inertialization.h
--- a/Engine/Source/Runtime/Engine/Classes/Animation/AnimNode_Inertialization.h
+++ b/Engine/Source/Runtime/Engine/Classes/Animation/AnimNode_Inertialization.h
@@ -8,6 +8,7 @@
 #include "Animation/AnimCurveTypes.h"
 #include "Animation/AnimNodeMessages.h"
 #include "AlphaBlend.h" // Required for EAlphaBlendOption
+#include "Interfaces/Interface_BoneReferenceSkeletonProvider.h"
 
 #include "AnimNode_Inertialization.generated.h"
 
@@ -20,7 +21,8 @@
 // https://www.gdcvault.com/play/1025331/Inertialization
 
 
-namespace UE { namespace Anim {
+namespace UE::Anim
+{
 
 // Event that can be subscribed to request inertialization-based blends
 class IInertializationRequester : public IGraphMessage
@@ -42,7 +44,7 @@ public:
 	virtual void AddDebugRecord(const FAnimInstanceProxy& InSourceProxy, int32 InSourceNodeId) = 0;
 };
 
-}}	// namespace UE::Anim
+}	// namespace UE::Anim
 
 UENUM()
 enum class EInertializationState : uint8
@@ -52,9 +54,8 @@ enum class EInertializationState : uint8
 	Active			// Inertialization active... apply the previously captured pose difference
 };
 
-
 UENUM()
-enum class EInertializationBoneState : uint8
+enum class UE_DEPRECATED(5.4, "Internal private pose storage is now used by inertialization.") EInertializationBoneState : uint8
 {
 	Invalid,		// Invalid bone (ie: bone was present in the skeleton but was not present in the pose when it was captured)
 	Valid,			// Valid bone
@@ -186,29 +187,17 @@ struct FInertializationRequest
 // endif ANIM_TRACE_ENABLED
 };
 
-
 USTRUCT()
-struct FInertializationPose
+struct UE_DEPRECATED(5.4, "Internal private pose storage is now used by inertialization.") FInertializationPose
 {
 	GENERATED_BODY()
 
 	FTransform ComponentTransform;
-
-	// Bone transforms indexed by skeleton bone index.  Transforms are in local space except for direct descendants of
-	// the root which are in component space (ie: they have been multiplied by the root).  Invalid bones (ie: bones
-	// that are present in the skeleton but were not present in the pose when it was captured) are all zero
-	//
 	TArray<FTransform> BoneTransforms;
-
-	// Bone states indexed by skeleton bone index
-	//
+	PRAGMA_DISABLE_DEPRECATION_WARNINGS
 	TArray<EInertializationBoneState> BoneStates;
-
-	// Snapshot of active curves
-	// 
+	PRAGMA_ENABLE_DEPRECATION_WARNINGS
 	FInertializationCurve Curves;
-
-
 	FName AttachParentName;
 	float DeltaTime;
 
@@ -228,13 +217,15 @@ struct FInertializationPose
 	void InitFrom(const FCompactPose& Pose, const FBlendedCurve& InCurves, const FTransform& InComponentTransform, const FName& InAttachParentName, float InDeltaTime);
 };
 
+PRAGMA_DISABLE_DEPRECATION_WARNINGS
 template <>
 struct TUseBitwiseSwap<FInertializationPose>
 {
 	enum { Value = false };
 };
+PRAGMA_ENABLE_DEPRECATION_WARNINGS
 
-// Structure for storing a pose snapshot sparsely (i.e. when we may not have the full set of transform for every bone)
+// Internal private structure used for storing a pose snapshots sparsely (i.e. when we may not have the full set of transform for every bone)
 struct FInertializationSparsePose
 {
 	friend struct FAnimNode_Inertialization;
@@ -242,11 +233,12 @@ struct FInertializationSparsePose
 
 private:
 
+	// Transform of the component at the point of the snapshot
 	FTransform ComponentTransform;
 	
 	// For each SkeletonPoseBoneIndex this array stores the index into the BoneTranslations, BoneRotations, and 
 	// BoneScales arrays which contains that bone's data. Or INDEX_NONE if this bone's data is not in the snapshot.
-	TArray<uint16> BoneIndices;
+	TArray<int32> BoneIndices;
 	
 	// Bone translation Data
 	TArray<FVector> BoneTranslations;
@@ -260,16 +252,19 @@ private:
     // Curve Data
 	FInertializationCurve Curves;
 
-	FName AttachParentName;
-	float DeltaTime;
+	// Attached Parent object Name
+	FName AttachParentName = NAME_None;
+	
+	// Delta Time since last snapshot
+	float DeltaTime = 0.0f;
 
-	void InitFrom(const FCompactPose& Pose, const FBlendedCurve& InCurves, const FTransform& InComponentTransform, const FName& InAttachParentName, float InDeltaTime);
+	void InitFrom(const FCompactPose& Pose, const FBlendedCurve& InCurves, const FTransform& InComponentTransform, const FName InAttachParentName, const float InDeltaTime);
 	bool IsEmpty() const;
 	void Empty();
 };
 
 USTRUCT()
-struct FInertializationBoneDiff
+struct UE_DEPRECATED(5.4, "Internal private pose difference storage is now used by inertialization.") FInertializationBoneDiff
 {
 	GENERATED_BODY()
 
@@ -328,9 +323,8 @@ struct FInertializationCurveDiffElement : public UE::Anim::FCurveElement
 	}
 };
 
-
 USTRUCT()
-struct FInertializationPoseDiff
+struct UE_DEPRECATED(5.4, "Internal private pose difference storage is now used by inertialization.") FInertializationPoseDiff
 {
 	GENERATED_BODY()
 
@@ -346,29 +340,11 @@ struct FInertializationPoseDiff
 		InertializationSpace = EInertializationSpace::Default;
 	}
 
-	// Initialize the pose difference from the current pose and the two previous snapshots
-	//
-	// Pose								the current frame's pose
-	// ComponentTransform				the current frame's component to world transform
-	// AttachParentName					the current frame's attach parent name (for checking if the attachment has changed)
-	// Prev1							the previous frame's pose
-	// Prev2							the pose from two frames before
-	// CurveFilter						filter for curves we don't want to inertialize
-	//
-	void InitFrom(const FCompactPose& Pose, const FBlendedCurve& Curves, const FTransform& ComponentTransform, const FName& AttachParentName, const FInertializationPose& Prev1, const FInertializationPose& Prev2, const UE::Anim::FCurveFilter& CurveFilter);
+	PRAGMA_DISABLE_DEPRECATION_WARNINGS
 
-	// Apply this difference to a pose, decaying over time as InertializationElapsedTime approaches InertializationDuration
-	//
-	// Pose								[in/out] the current frame's pose
-	// Curves							[in/out] the current frame's animation curves
-	// InertializationElapsedTime		time elapsed since the start of the inertialization
-	// InertializationDuration			total inertialization duration (used for curves)
-	// InertializationDurationPerBone	inertialization duration per bone (indexed by skeleton bone index) (used for pose)
-	//
+	void InitFrom(const FCompactPose& Pose, const FBlendedCurve& Curves, const FTransform& ComponentTransform, const FName& AttachParentName, const FInertializationPose& Prev1, const FInertializationPose& Prev2, const UE::Anim::FCurveFilter& CurveFilter);
 	void ApplyTo(FCompactPose& Pose, FBlendedCurve& Curves, float InertializationElapsedTime, float InertializationDuration, TArrayView<const float> InertializationDurationPerBone) const;
 
-	// Get the inertialization space for this pose diff (for debug display)
-	//
 	EInertializationSpace GetInertializationSpace() const
 	{
 		return InertializationSpace;
@@ -376,21 +352,15 @@ struct FInertializationPoseDiff
 
 private:
 
-	static float CalcInertialFloat(float x0, float v0, float t, float t1);
-
-	// Bone differences indexed by skeleton bone index
 	TArray<FInertializationBoneDiff> BoneDiffs;
-
-	// Curve differences
 	TBaseBlendedCurve<FDefaultAllocator, FInertializationCurveDiffElement> CurveDiffs;
-
-	// Inertialization space (local vs world for situations where we wish to correct a world-space discontinuity such as an abrupt orientation change)
 	EInertializationSpace InertializationSpace;
-};
 
+	PRAGMA_ENABLE_DEPRECATION_WARNINGS
+};
 
 USTRUCT(BlueprintInternalUseOnly)
-struct FAnimNode_Inertialization : public FAnimNode_Base
+struct FAnimNode_Inertialization : public FAnimNode_Base, public IBoneReferenceSkeletonProvider
 {
 	GENERATED_BODY()
 
@@ -407,13 +377,15 @@ private:
 	UPROPERTY(EditAnywhere, Category = Filter)
 	TArray<FName> FilteredCurves;
 
-	/**
-	 * Enable this to pre-allocate memory for the node rather than to allocate and deallocate memory when blending
-	 * becomes active and inactive. This improves performance, but causes larger memory usage, in particular when you
-	 * have multiple Inertialization nodes in an animation graph that are not all used at once.
-	 */
-	UPROPERTY(EditAnywhere, Category = Memory)
-	bool bPreallocateMemory = false;
+	// List of bones that should not use inertial blending. These bones will change instantly when the animation switches.
+	UPROPERTY(EditAnywhere, Category = Filter)
+	TArray<FBoneReference> FilteredBones;
+
+#if WITH_EDITORONLY_DATA
+	UE_DEPRECATED(5.4, "Preallocate Memory has been deprecated.")
+	UPROPERTY(meta = (DeprecatedProperty, DeprecationMessage = "Preallocate Memory has been deprecated."))
+	bool bPreallocateMemory_DEPRECATED = false;
+#endif
 
 	/**
 	* Clear any active blends if we just became relevant, to avoid carrying over undesired blends.
@@ -432,7 +404,15 @@ private:
 
 public: // FAnimNode_Inertialization
 
-	ENGINE_API FAnimNode_Inertialization();
+	// Note: We need to explicitly disable warnings on these constructors/operators for clang to be happy with deprecated variables
+	PRAGMA_DISABLE_DEPRECATION_WARNINGS
+	ENGINE_API FAnimNode_Inertialization() = default;
+	ENGINE_API ~FAnimNode_Inertialization() = default;
+	ENGINE_API FAnimNode_Inertialization(const FAnimNode_Inertialization&) = default;
+	ENGINE_API FAnimNode_Inertialization(FAnimNode_Inertialization&&) = default;
+	ENGINE_API FAnimNode_Inertialization& operator=(const FAnimNode_Inertialization&) = default;
+	ENGINE_API FAnimNode_Inertialization& operator=(FAnimNode_Inertialization&&) = default;
+	PRAGMA_ENABLE_DEPRECATION_WARNINGS
 	
 	// Request to activate inertialization for a duration.
 	// If multiple requests are made on the same inertialization node, the minimum requested time will be used.
@@ -462,29 +442,55 @@ public: // FAnimNode_Base
 
 protected:
 
-	// Start Inertialization
-	//
-	// Computes the inertialization pose difference from the current pose and the two previous poses (to capture velocity).  This function
-	// is virtual so that a derived class could optionally regularize the pose snapshots to align better with the current frame's pose
-	// before computing the inertial difference (for example to correct for instantaneous changes in the root relative to its children).
-	//
+	PRAGMA_DISABLE_DEPRECATION_WARNINGS
+
+	UE_DEPRECATED(5.4, "This function is longer called by the node internally as inertialization method is now private.")
 	ENGINE_API virtual void StartInertialization(FPoseContext& Context, FInertializationPose& PreviousPose1, FInertializationPose& PreviousPose2, float Duration, TArrayView<const float> DurationPerBone, /*OUT*/ FInertializationPoseDiff& OutPoseDiff);
 
-	// Apply Inertialization
-	//
-	// Applies the inertialization pose difference to the current pose (feathering down to zero as ElapsedTime approaches Duration).  This
-	// function is virtual so that a derived class could optionally adjust the pose based on any regularization done in StartInertialization.
-	//
+	UE_DEPRECATED(5.4, "This function is longer called by the node internally as inertialization method is now private.")
 	ENGINE_API virtual void ApplyInertialization(FPoseContext& Context, const FInertializationPoseDiff& PoseDiff, float ElapsedTime, float Duration, TArrayView<const float> DurationPerBone);
 
+	PRAGMA_ENABLE_DEPRECATION_WARNINGS
 
 private:
 
-	// Snapshots of the actor pose from past frames
-	TArray<FInertializationPose> PoseSnapshots;
+	/**
+	 * Deactivates the inertialization and frees any temporary memory.
+	 */
+	void Deactivate();
+
+	/**
+	 * Computes the inertialization pose difference between the current pose and the previous pose and computes the velocity of this difference.
+	 *
+	 * @param InPose				The current pose for the animation being transitioned to.
+	 * @param InCurves				The current curves for the animation being transitioned to.
+	 * @param ComponentTransform	The component transform of the current pose
+	 * @param AttachParentName		The name of the attached parent object
+	 * @param PreviousPose1			The pose recorded as output of the inertializer on the previous frame.
+	 * @param PreviousPose2			The pose recorded as output of the inertializer two frames ago.
+	 */
+	void InitFrom(
+		const FCompactPose& InPose, 
+		const FBlendedCurve& InCurves, 
+		const FTransform& ComponentTransform, 
+		const FName AttachParentName, 
+		const FInertializationSparsePose& PreviousPose1, 
+		const FInertializationSparsePose& PreviousPose2);
+
+	/**
+	 * Applies the inertialization difference to the given pose (decaying to zero as ElapsedTime approaches Duration)
+	 *
+	 * @param InOutPose		The current pose to blend with the extrapolated pose.
+	 * @param InOutCurves	The current curves to blend with the extrapolated curves.
+	 */
+	void ApplyTo(FCompactPose& InOutPose, FBlendedCurve& InOutCurves);
+
+	// Snapshots of the actor pose generated as output.
+	FInertializationSparsePose PrevPoseSnapshot;
+	FInertializationSparsePose CurrPoseSnapshot;
 
 	// Elapsed delta time between calls to evaluate
-	float DeltaTime;
+	float DeltaTime = 0.0f;
 
 	// Pending inertialization requests
 	UPROPERTY(Transient)
@@ -493,36 +499,52 @@ private:
 	// Update Counter for detecting being relevant
 	FGraphTraversalCounter UpdateCounter;
 
-	// Teleport type
-	ETeleportType TeleportType;
-
 	// Inertialization state
-	EInertializationState InertializationState;
+	EInertializationState InertializationState = EInertializationState::Inactive;
 
 	// Amount of time elapsed during the Inertialization
-	float InertializationElapsedTime;
+	float InertializationElapsedTime = 0.0f;
 
 	// Inertialization duration for the main inertialization request (used for curve blending and deficit tracking)
-	float InertializationDuration;
+	float InertializationDuration = 0.0f;
 
 	// Inertialization durations indexed by skeleton bone index (used for per-bone blending)
 	TCustomBoneIndexArray<float, FSkeletonPoseBoneIndex> InertializationDurationPerBone;
 
 	// Maximum of InertializationDuration and all entries in InertializationDurationPerBone (used for knowing when to shutdown the inertialization)
-	float InertializationMaxDuration;
+	float InertializationMaxDuration = 0.0f;
 
 	// Inertialization deficit (for tracking and reducing 'pose melting' when thrashing inertialization requests)
-	float InertializationDeficit;
+	float InertializationDeficit = 0.0f;
 
 	// Inertialization pose differences
-	FInertializationPoseDiff InertializationPoseDiff;
+	TArray<int32> BoneIndices;
+	TArray<FVector3f> BoneTranslationDiffDirection;
+	TArray<float> BoneTranslationDiffMagnitude;
+	TArray<float> BoneTranslationDiffSpeed;
+	TArray<FVector3f> BoneRotationDiffAxis;
+	TArray<float> BoneRotationDiffAngle;
+	TArray<float> BoneRotationDiffSpeed;
+	TArray<FVector3f> BoneScaleDiffAxis;
+	TArray<float> BoneScaleDiffMagnitude;
+	TArray<float> BoneScaleDiffSpeed;
+
+	// Curve differences
+	TBaseBlendedCurve<FDefaultAllocator, FInertializationCurveDiffElement> CurveDiffs;
+
+	// Temporary storage for curve data of the Destination Pose
+	TBaseBlendedCurve<TInlineAllocator<8>, UE::Anim::FCurveElement> PoseCurveData;
+
+public: // IBoneReferenceSkeletonProvider
+	ENGINE_API class USkeleton* GetSkeleton(bool& bInvalidSkeletonIsError, const IPropertyHandle* PropertyHandle) override;
+
+private:
 
 	// Cached curve filter built from FilteredCurves
 	UE::Anim::FCurveFilter CurveFilter;
-	
-	// Reset inertialization timing and state
-	void Deactivate();
 
+	// Cache compact pose bone index for FilteredBones
+	TArray<FCompactPoseBoneIndex, TInlineAllocator<8>> BoneFilter;
 
 // if ANIM_TRACE_ENABLED - these properties are only used for debugging when ANIM_TRACE_ENABLED == 1
 
diff --git a/Engine/Source/Runtime/Engine/Private/Animation/AnimNode_DeadBlending.cpp b/Engine/Source/Runtime/Engine/Private/Animation/AnimNode_DeadBlending.cpp
--- a/Engine/Source/Runtime/Engine/Private/Animation/AnimNode_DeadBlending.cpp
+++ b/Engine/Source/Runtime/Engine/Private/Animation/AnimNode_DeadBlending.cpp
@@ -18,8 +18,8 @@ LLM_DEFINE_TAG(Animation_DeadBlending);
 
 TAutoConsoleVariable<int32> CVarAnimDeadBlendingEnable(TEXT("a.AnimNode.DeadBlending.Enable"), 1, TEXT("Enable / Disable DeadBlending"));
 
-namespace UE::Anim {
-
+namespace UE::Anim
+{
 	// Inertialization request event bound to a node
 	class FDeadBlendingRequester : public IInertializationRequester
 	{
@@ -261,25 +261,22 @@ void FAnimNode_DeadBlending::Deactivate()
 {
 	InertializationState = EInertializationState::Inactive;
 
-	if (!bPreallocateMemory)
-	{
-		BoneIndices.Empty();
+	BoneIndices.Empty();
 
-		BoneTranslations.Empty();
-		BoneRotations.Empty();
-		BoneRotationDirections.Empty();
-		BoneScales.Empty();
+	BoneTranslations.Empty();
+	BoneRotations.Empty();
+	BoneRotationDirections.Empty();
+	BoneScales.Empty();
 
-		BoneTranslationVelocities.Empty();
-		BoneRotationVelocities.Empty();
-		BoneScaleVelocities.Empty();
+	BoneTranslationVelocities.Empty();
+	BoneRotationVelocities.Empty();
+	BoneScaleVelocities.Empty();
 
-		BoneTranslationDecayHalfLives.Empty();
-		BoneRotationDecayHalfLives.Empty();
-		BoneScaleDecayHalfLives.Empty();
+	BoneTranslationDecayHalfLives.Empty();
+	BoneRotationDecayHalfLives.Empty();
+	BoneScaleDecayHalfLives.Empty();
 
-		InertializationDurationPerBone.Empty();
-	}
+	InertializationDurationPerBone.Empty();
 }
 
 void FAnimNode_DeadBlending::InitFrom(const FCompactPose& InPose, const FBlendedCurve& InCurves, const FInertializationSparsePose& SrcPosePrev, const FInertializationSparsePose& SrcPoseCurr)
@@ -331,8 +328,8 @@ void FAnimNode_DeadBlending::InitFrom(const FCompactPose& InPose, const FBlended
 		const int32 SkeletonPoseBoneIndex = BoneContainer.GetSkeletonIndex(BoneIndex);
 
 		if (SkeletonPoseBoneIndex == INDEX_NONE || 
-			SrcPoseCurr.BoneIndices[SkeletonPoseBoneIndex] == static_cast<uint16>(INDEX_NONE) ||
-			SrcPosePrev.BoneIndices[SkeletonPoseBoneIndex] == static_cast<uint16>(INDEX_NONE))
+			SrcPoseCurr.BoneIndices[SkeletonPoseBoneIndex] == INDEX_NONE ||
+			SrcPosePrev.BoneIndices[SkeletonPoseBoneIndex] == INDEX_NONE)
 		{
 			continue;
 		}
@@ -462,12 +459,20 @@ void FAnimNode_DeadBlending::InitFrom(const FCompactPose& InPose, const FBlended
 				ExtrapolationHalfLifeMax);
 		});
 
-	// Apply filtering to remove anything we don't want to inertialize
-
+	// Apply filtering to remove filtered curves from extrapolation. This does not actually
+	// prevent these curves from being blended, but does stop them appearing as empty
+	// in the output curves created by the Union in ApplyTo unless they are already in the
+	// destination animation.
 	if (CurveFilter.Num() > 0)
 	{
 		UE::Anim::FCurveUtils::Filter(CurveData, CurveFilter);
 	}
+
+	// Apply filtering to remove curves that are not meant to be extrapolated
+	if (ExtrapolatedCurveFilter.Num() > 0)
+	{
+		UE::Anim::FCurveUtils::Filter(CurveData, ExtrapolatedCurveFilter);
+	}
 }
 
 void FAnimNode_DeadBlending::ApplyTo(FCompactPose& InOutPose, FBlendedCurve& InOutCurves)
@@ -480,7 +485,7 @@ void FAnimNode_DeadBlending::ApplyTo(FCompactPose& InOutPose, FBlendedCurve& InO
 	{
 		const int32 SkeletonPoseBoneIndex = BoneContainer.GetSkeletonIndex(BoneIndex);
 
-		if (SkeletonPoseBoneIndex == INDEX_NONE || BoneIndices[SkeletonPoseBoneIndex] == static_cast<uint16>(INDEX_NONE) || BoneFilter.Contains(BoneIndex))
+		if (SkeletonPoseBoneIndex == INDEX_NONE || BoneIndices[SkeletonPoseBoneIndex] == INDEX_NONE || BoneFilter.Contains(BoneIndex))
 		{
 			continue;
 		}
@@ -547,7 +552,7 @@ void FAnimNode_DeadBlending::ApplyTo(FCompactPose& InOutPose, FBlendedCurve& InO
 		// Compute Blend Alpha
 
 		const float Alpha = 1.0f - FAlphaBlend::AlphaToBlendOption(
-			InertializationTime / FMath::Max(InertializationDurationPerBone[SkeletonPoseBoneIndex], UE_SMALL_NUMBER),
+			FMath::Clamp(InertializationTime / FMath::Max(InertializationDurationPerBone[SkeletonPoseBoneIndex], UE_SMALL_NUMBER), 0.0f, 1.0f),
 			InertializationBlendMode, InertializationCustomBlendCurve);
 
 		// Perform Blend
@@ -576,14 +581,28 @@ void FAnimNode_DeadBlending::ApplyTo(FCompactPose& InOutPose, FBlendedCurve& InO
 	// Compute Blend Alpha
 
 	const float CurveAlpha = 1.0f - FAlphaBlend::AlphaToBlendOption(
-		InertializationTime / FMath::Max(InertializationDuration, UE_SMALL_NUMBER),
+		FMath::Clamp(InertializationTime / FMath::Max(InertializationDuration, UE_SMALL_NUMBER), 0.0f, 1.0f),
 		InertializationBlendMode, InertializationCustomBlendCurve);
 
 	// Blend Curves
 
-	UE::Anim::FNamedValueArrayUtils::Union(InOutCurves, CurveData,
-		[CurveAlpha, this](UE::Anim::FCurveElement& OutResultElement, const FDeadBlendingCurveElement& InElement1, UE::Anim::ENamedValueUnionFlags InFlags)
+	PoseCurveData.CopyFrom(InOutCurves);
+
+	UE::Anim::FNamedValueArrayUtils::Union(InOutCurves, PoseCurveData, CurveData, [CurveAlpha, this](
+			UE::Anim::FCurveElement& OutResultElement, 
+			const UE::Anim::FCurveElement& InElement0,
+			const FDeadBlendingCurveElement& InElement1, 
+			UE::Anim::ENamedValueUnionFlags InFlags)
 		{
+			// For filtered Curves take destination value
+
+			if (FilteredCurves.Contains(OutResultElement.Name))
+			{
+				OutResultElement.Value = InElement0.Value;
+				OutResultElement.Flags = InElement0.Flags;
+				return;
+			}
+
 			// Compute Extrapolated Curve Value
 
 			const float ExtrapolatedCurve = UE::Anim::DeadBlending::Private::ExtrapolateCurve(
@@ -596,13 +615,13 @@ void FAnimNode_DeadBlending::ApplyTo(FCompactPose& InOutPose, FBlendedCurve& InO
 			if (bShowExtrapolations)
 			{
 				OutResultElement.Value = ExtrapolatedCurve;
-				OutResultElement.Flags |= InElement1.Flags;
+				OutResultElement.Flags = InElement0.Flags | InElement1.Flags;
 				return;
 			}
 #endif
 
-			OutResultElement.Value = FMath::Lerp(OutResultElement.Value, ExtrapolatedCurve, CurveAlpha);
-			OutResultElement.Flags |= InElement1.Flags;
+			OutResultElement.Value = FMath::Lerp(InElement0.Value, ExtrapolatedCurve, CurveAlpha);
+			OutResultElement.Flags = InElement0.Flags | InElement1.Flags;
 		});
 }
 
@@ -640,29 +659,31 @@ void FAnimNode_DeadBlending::Initialize_AnyThread(const FAnimationInitializeCont
 	CurveFilter.SetFilterMode(UE::Anim::ECurveFilterMode::DisallowFiltered);
 	CurveFilter.AppendNames(FilteredCurves);
 
-	BoneFilter.Init(FCompactPoseBoneIndex(INDEX_NONE), FilteredBones.Num());
+	ExtrapolatedCurveFilter.Empty();
+	ExtrapolatedCurveFilter.SetFilterMode(UE::Anim::ECurveFilterMode::DisallowFiltered);
+	ExtrapolatedCurveFilter.AppendNames(ExtrapolationFilteredCurves);
 
+	BoneFilter.Init(FCompactPoseBoneIndex(INDEX_NONE), FilteredBones.Num());
+	
 	PrevPoseSnapshot.Empty();
 	CurrPoseSnapshot.Empty();
 
 	RequestQueue.Reserve(8);
 
-	const int32 NumSkeletonBones = bPreallocateMemory ? Context.AnimInstanceProxy->GetSkeleton()->GetReferenceSkeleton().GetNum() : 0;
-
-	BoneIndices.Empty(NumSkeletonBones);
+	BoneIndices.Empty();
 	
-	BoneTranslations.Empty(NumSkeletonBones);
-	BoneRotations.Empty(NumSkeletonBones);
-	BoneRotationDirections.Empty(NumSkeletonBones);
-	BoneScales.Empty(NumSkeletonBones);
+	BoneTranslations.Empty();
+	BoneRotations.Empty();
+	BoneRotationDirections.Empty();
+	BoneScales.Empty();
 
-	BoneTranslationVelocities.Empty(NumSkeletonBones);
-	BoneRotationVelocities.Empty(NumSkeletonBones);
-	BoneScaleVelocities.Empty(NumSkeletonBones);
+	BoneTranslationVelocities.Empty();
+	BoneRotationVelocities.Empty();
+	BoneScaleVelocities.Empty();
 
-	BoneTranslationDecayHalfLives.Empty(NumSkeletonBones);
-	BoneRotationDecayHalfLives.Empty(NumSkeletonBones);
-	BoneScaleDecayHalfLives.Empty(NumSkeletonBones);
+	BoneTranslationDecayHalfLives.Empty();
+	BoneRotationDecayHalfLives.Empty();
+	BoneScaleDecayHalfLives.Empty();
 
 	CurveData.Empty();
 
@@ -672,7 +693,7 @@ void FAnimNode_DeadBlending::Initialize_AnyThread(const FAnimationInitializeCont
 	InertializationTime = 0.0f;
 
 	InertializationDuration = 0.0f;
-	InertializationDurationPerBone.Empty(NumSkeletonBones);
+	InertializationDurationPerBone.Empty();
 	InertializationMaxDuration = 0.0f;
 
 	InertializationBlendMode = DefaultBlendMode;
@@ -690,6 +711,7 @@ void FAnimNode_DeadBlending::CacheBones_AnyThread(const FAnimationCacheBonesCont
 	// Compute Compact Pose Bone Index for each bone in Filter
 
 	const FBoneContainer& RequiredBones = Context.AnimInstanceProxy->GetRequiredBones();
+	BoneFilter.Init(FCompactPoseBoneIndex(INDEX_NONE), FilteredBones.Num());
 	for (int32 FilterBoneIdx = 0; FilterBoneIdx < FilteredBones.Num(); FilterBoneIdx++)
 	{
 		FilteredBones[FilterBoneIdx].Initialize(Context.AnimInstanceProxy->GetSkeleton());
@@ -766,9 +788,9 @@ void FAnimNode_DeadBlending::Evaluate_AnyThread(FPoseContext& Output)
 	{
 		const FVector RootWorldSpaceLocation = ComponentTransform.TransformPosition(Output.Pose[FCompactPoseBoneIndex(0)].GetTranslation());
 		
-		const uint16 RootBoneIndex = CurrPoseSnapshot.BoneIndices[0];
+		const int32 RootBoneIndex = CurrPoseSnapshot.BoneIndices[0];
 
-		if (RootBoneIndex != static_cast<uint16>(INDEX_NONE))
+		if (RootBoneIndex != INDEX_NONE)
 		{
 			const FVector PrevRootWorldSpaceLocation = CurrPoseSnapshot.ComponentTransform.TransformPosition(CurrPoseSnapshot.BoneTranslations[RootBoneIndex]);
 
@@ -929,19 +951,14 @@ void FAnimNode_DeadBlending::Evaluate_AnyThread(FPoseContext& Output)
 
 	// Record Pose Snapshot
 
-	if (CurrPoseSnapshot.IsEmpty())
-	{
-		// Initialize the current pose
-		CurrPoseSnapshot.InitFrom(Output.Pose, Output.Curve, ComponentTransform, AttachParentName, DeltaTime);
-	}
-	else
+	if (!CurrPoseSnapshot.IsEmpty())
 	{
 		// Directly swap the memory of the current pose with the prev pose snapshot (to avoid allocations and copies)
 		Swap(PrevPoseSnapshot, CurrPoseSnapshot);
-
-		// Initialize the (now irrelevant) current pose
-		CurrPoseSnapshot.InitFrom(Output.Pose, Output.Curve, ComponentTransform, AttachParentName, DeltaTime);
 	}
+	
+	// Initialize the current pose
+	CurrPoseSnapshot.InitFrom(Output.Pose, Output.Curve, ComponentTransform, AttachParentName, DeltaTime);
 
 	// Reset Delta Time
 
@@ -949,7 +966,7 @@ void FAnimNode_DeadBlending::Evaluate_AnyThread(FPoseContext& Output)
 
 	const float InertializationWeight = InertializationState == EInertializationState::Active ?
 		1.0f - FAlphaBlend::AlphaToBlendOption(
-			InertializationTime / FMath::Max(InertializationDuration, UE_SMALL_NUMBER),
+			FMath::Clamp(InertializationTime / FMath::Max(InertializationDuration, UE_SMALL_NUMBER), 0.0f, 1.0f),
 			InertializationBlendMode, InertializationCustomBlendCurve) : 0.0f;
 
 	TRACE_ANIM_NODE_VALUE_WITH_ID(Output, GetNodeIndex(), TEXT("State"), *UEnum::GetValueAsString(InertializationState));
diff --git a/Engine/Source/Runtime/Engine/Private/Animation/AnimNode_Inertialization.cpp b/Engine/Source/Runtime/Engine/Private/Animation/AnimNode_Inertialization.cpp
--- a/Engine/Source/Runtime/Engine/Private/Animation/AnimNode_Inertialization.cpp
+++ b/Engine/Source/Runtime/Engine/Private/Animation/AnimNode_Inertialization.cpp
@@ -27,68 +27,201 @@ TAutoConsoleVariable<int32> CVarAnimInertializationIgnoreVelocity(TEXT("a.AnimNo
 TAutoConsoleVariable<int32> CVarAnimInertializationIgnoreDeficit(TEXT("a.AnimNode.Inertialization.IgnoreDeficit"), 0, TEXT("Ignore inertialization time deficit caused by interruptions"));
 
 
-static constexpr int32 INERTIALIZATION_MAX_POSE_SNAPSHOTS = 2;
-static constexpr float INERTIALIZATION_TIME_EPSILON = 1.0e-7f;
-
-namespace UE { namespace Anim {
-
-void IInertializationRequester::RequestInertialization(const FInertializationRequest& InInertializationRequest)
-{
-	RequestInertialization(InInertializationRequest.Duration, InInertializationRequest.BlendProfile);
-}
-
-// Inertialization request event bound to a node
-class FInertializationRequester : public IInertializationRequester
+namespace UE::Anim
 {
-public:
-	FInertializationRequester(const FAnimationBaseContext& InContext, FAnimNode_Inertialization* InNode)
-		: Node(*InNode)
-		, NodeId(InContext.GetCurrentNodeId())
-		, Proxy(*InContext.AnimInstanceProxy)
-	{}
-
-private:
-	// IInertializationRequester interface
-	virtual void RequestInertialization(float InRequestedDuration, const UBlendProfile* InBlendProfile) override
-	{ 
-		Node.RequestInertialization(InRequestedDuration, InBlendProfile); 
-	}
 
-	virtual void RequestInertialization(const FInertializationRequest& InInertializationRequest) override
+	void IInertializationRequester::RequestInertialization(const FInertializationRequest& InInertializationRequest)
 	{
-		// The Blend Mode parameters will be ignored as FAnimNode_Inertialization does not support them.
-		Node.RequestInertialization(InInertializationRequest);
+		RequestInertialization(InInertializationRequest.Duration, InInertializationRequest.BlendProfile);
 	}
 
-	virtual void AddDebugRecord(const FAnimInstanceProxy& InSourceProxy, int32 InSourceNodeId)
+	// Inertialization request event bound to a node
+	class FInertializationRequester : public IInertializationRequester
 	{
+	public:
+		FInertializationRequester(const FAnimationBaseContext& InContext, FAnimNode_Inertialization* InNode)
+			: Node(*InNode)
+			, NodeId(InContext.GetCurrentNodeId())
+			, Proxy(*InContext.AnimInstanceProxy)
+		{}
+
+	private:
+		// IInertializationRequester interface
+		virtual void RequestInertialization(float InRequestedDuration, const UBlendProfile* InBlendProfile) override
+		{
+			Node.RequestInertialization(InRequestedDuration, InBlendProfile);
+		}
+
+		virtual void RequestInertialization(const FInertializationRequest& InInertializationRequest) override
+		{
+			// The Blend Mode parameters will be ignored as FAnimNode_Inertialization does not support them.
+			Node.RequestInertialization(InInertializationRequest);
+		}
+
+		virtual void AddDebugRecord(const FAnimInstanceProxy& InSourceProxy, int32 InSourceNodeId)
+		{
 #if WITH_EDITORONLY_DATA
-		Proxy.RecordNodeAttribute(InSourceProxy, NodeId, InSourceNodeId, IInertializationRequester::Attribute);
+			Proxy.RecordNodeAttribute(InSourceProxy, NodeId, InSourceNodeId, IInertializationRequester::Attribute);
 #endif
-		TRACE_ANIM_NODE_ATTRIBUTE(Proxy, InSourceProxy, NodeId, InSourceNodeId, IInertializationRequester::Attribute);
+			TRACE_ANIM_NODE_ATTRIBUTE(Proxy, InSourceProxy, NodeId, InSourceNodeId, IInertializationRequester::Attribute);
+		}
+
+		// Node to target
+		FAnimNode_Inertialization& Node;
+
+		// Node index
+		int32 NodeId;
+
+		// Proxy currently executing
+		FAnimInstanceProxy& Proxy;
+	};
+}
+
+namespace UE::Anim::Inertialization::Private
+{
+	static inline int32 GetNumSkeletonBones(const FBoneContainer& BoneContainer)
+	{
+		const USkeleton* SkeletonAsset = BoneContainer.GetSkeletonAsset();
+		check(SkeletonAsset);
+
+		const FReferenceSkeleton& RefSkeleton = SkeletonAsset->GetReferenceSkeleton();
+		return RefSkeleton.GetNum();
 	}
 
-	// Node to target
-	FAnimNode_Inertialization& Node;
+	static constexpr float INERTIALIZATION_TIME_EPSILON = 1.0e-7f;
 
-	// Node index
-	int32 NodeId;
+	// Calculate the "inertialized" value of a single float at time t
+	//
+	// @param x0	Initial value of the float (at time 0)
+	// @param v0	Initial "velocity" (first derivative) of the float (at time 0)
+	// @param t		Time at which to evaluate the float
+	// @param t1	Ending inertialization time (ie: the time at which the curve must be zero)
+	//
+	// Evaluates a quintic polynomial curve with the specified initial conditions (x0, v0) which hits zero at time t1.  As well,
+	// the curve is designed so that the first and second derivatives are also zero at time t1.
+	//
+	// The initial second derivative (a0) is chosen such that it is as close to zero as possible, but large enough to prevent any
+	// overshoot (ie: enforce x >= 0 for t between 0 and t1).  If necessary, the ending time t1 will be adjusted (shortened) to
+	// guarantee that there is no overshoot, even for very large initial velocities.
+	//
+	static float CalcInertialFloat(float x0, float v0, float t, float t1)
+	{
+		static_assert(INERTIALIZATION_TIME_EPSILON * INERTIALIZATION_TIME_EPSILON * INERTIALIZATION_TIME_EPSILON * INERTIALIZATION_TIME_EPSILON * INERTIALIZATION_TIME_EPSILON > FLT_MIN,
+			"INERTIALIZATION_TIME_EPSILON^5 must be greater than FLT_MIN to avoid denormalization (and potential division by zero) for very small values of t1");
 
-	// Proxy currently executing
-	FAnimInstanceProxy& Proxy;
-};
+		if (t < 0.0f)
+		{
+			t = 0.0f;
+		}
 
+		if (t >= t1 - INERTIALIZATION_TIME_EPSILON)
+		{
+			return 0.0f;
+		}
 
-static int32 GetNumSkeletonBones(const FBoneContainer& BoneContainer)
-{
-	const USkeleton* SkeletonAsset = BoneContainer.GetSkeletonAsset();
-	check(SkeletonAsset);
+		// Assume that x0 >= 0... if this is not the case, then simply invert everything (both input and output)
+		float sign = 1.0f;
+		if (x0 < 0.0f)
+		{
+			x0 = -x0;
+			v0 = -v0;
+			sign = -1.0f;
+		}
 
-	const FReferenceSkeleton& RefSkeleton = SkeletonAsset->GetReferenceSkeleton();
-	return RefSkeleton.GetNum();
-}
+		// If v0 > 0, then the curve will overshoot away from zero, so clamp v0 here to guarantee that there is no overshoot
+		if (v0 > 0.0f)
+		{
+			v0 = 0.0f;
+		}
+
+		// Check for invalid values - this is only expected to occur if NaNs or other invalid values are coming into the node
+		if (!ensureMsgf(x0 >= 0.0f && v0 <= 0.0f && t >= 0.0f && t1 >= 0.0f,
+			TEXT("Invalid Value(s) in Inertialization - x0: %f, v0: %f, t: %f, t1: %f"), x0, v0, t, t1))
+		{
+			x0 = 0.0f;
+			v0 = 0.0f;
+			t = 0.0f;
+			t1 = 0.0f;
+		}
+
+		// Limit t1 such that the curve does not overshoot below zero (ensuring that x >= 0 for all t between 0 and t1).
+		//
+		// We observe that if the curve does overshoot below zero, it must have an inflection point somewhere between 0 and t1
+		// (since we know that both x0 and x1 are >= 0).  Therefore, we can prevent overshoot by adjusting t1 such that any
+		// inflection point is at t >= t1.
+		//
+		// Assuming that we are using the zero jerk version of the curve (see below) whose velocity has a triple root at t1, then
+		// we can prevent overshoot by forcing the remaining root to be at time t >= t1, or equivalently, we can set t1 to be the
+		// lesser of the original t1 or the value that gives us a solution with a quadruple velocity root at t1.
+		//
+		// The following Mathematica expression solves for t1 that gives us the quadruple velocity root:
+		//
+		//		v := q * (t-t1)^4
+		//		x := Integrate[Expand[v], t] + x0
+		//		eq1 := (v/.t->0)==v0
+		//		eq2 := (x/.t->t1)==0
+		//		Solve[{eq1 && eq2}, {q,t1}]
+		//
+		if (v0 < -UE_KINDA_SMALL_NUMBER)
+		{
+			t1 = FMath::Min(t1, -5.0f * x0 / v0);
+		}
+
+		if (t >= t1 - INERTIALIZATION_TIME_EPSILON)
+		{
+			return 0.0f;
+		}
+
+		const float t1_2 = t1 * t1;
+		const float t1_3 = t1 * t1_2;
+		const float t1_4 = t1 * t1_3;
+		const float t1_5 = t1 * t1_4;
+
+		// Compute the initial acceleration value (a0) for this curve.  Ideally we want to use an initial acceleration of zero, but
+		// if there is a large negative initial velocity, then we will need to use a larger acceleration in order to ensure that
+		// the curve does not overshoot below zero (ie: to ensure that x >= 0 for all t between 0 and t1).
+		//
+		// To compute a0, we first compute the a0 that we would get if we also specified that the third derivative (the "jerk" j)
+		// is zero at t1.  If this value of a0 is positive (and therefore opposing the initial velocity), then we use that.  If it
+		// is negative, then we simply use an initial a0 of zero.
+		//
+		// The following Mathematica expression solves for a0 that gives us zero jerk at t1:
+		//
+		//		x:= A*t^5 + B*t^4 + C*t^3 + D*t^2 + v0*t + x0
+		//		v:=Dt[x, t, Constants->{A,B,C,D,v0,x0}]
+		//		a:=Dt[v, t, Constants->{A,B,C,D,v0,x0}]
+		//		j:=Dt[a, t, Constants->{A,B,C,D,v0,x0}]
+		//		eq1:=(x/.t->t1)==0
+		//		eq2:=(v/.t->t1)==0
+		//		eq3:=(a/.t->t1)==0
+		//		eq4:=(j/.t->t1)==0
+		//		a0:=a/.t->0/.Solve[{eq1 && eq2 && eq3 && eq4}, {A,B,C,D}]
+		//		ExpandNumerator[a0]
+		//
+		const float a0 = FMath::Max(0.0f, (-8.0f * t1 * v0 - 20.0f * x0) / t1_2);
+
+		// Compute the polynomial coefficients given the starting and ending conditions, solved from:
+		//
+		//		x:= A*t^5 + B*t^4 + C*t^3 + D*t^2 + v0*t + x0
+		//		v:=Dt[x, t, Constants->{A,B,C,D,v0,x0}]
+		//		a:=Dt[v, t, Constants->{A,B,C,D,v0,x0}]
+		//		eq1:=(x/.t->t1)==0
+		//		eq2:=(v/.t->t1)==0
+		//		eq3:=(a/.t->t1)==0
+		//		eq4:=(a/.t->0)==a0
+		//		Simplify[Solve[{eq1 && eq2 && eq3 && eq4}, {A,B,C,D}]]
+		//
+		const float A = -0.5f * (a0 * t1_2 + 6.0f * t1 * v0 + 12.0f * x0) / t1_5;
+		const float B = 0.5f * (3.0f * a0 * t1_2 + 16.0f * t1 * v0 + 30.0f * x0) / t1_4;
+		const float C = -0.5f * (3.0f * a0 * t1_2 + 12.0f * t1 * v0 + 20.0f * x0) / t1_3;
+		const float D = 0.5f * a0;
+
+		const float x = (((((A * t) + B) * t + C) * t + D) * t + v0) * t + x0;
+
+		return x * sign;
+	}
 
-}}	// namespace UE::Anim
+}	// namespace UE::Anim::Inertialization::Private
 
 PRAGMA_DISABLE_DEPRECATION_WARNINGS
 
@@ -115,17 +248,17 @@ void FInertializationRequest::Clear()
 #endif
 }
 
-FAnimNode_Inertialization::FAnimNode_Inertialization()
-	: DeltaTime(0.0f)
-	, TeleportType(ETeleportType::None)
-	, InertializationState(EInertializationState::Inactive)
-	, InertializationElapsedTime(0.0f)
-	, InertializationDuration(0.0f)
-	, InertializationMaxDuration(0.0f)
-	, InertializationDeficit(0.0f)
+class USkeleton* FAnimNode_Inertialization::GetSkeleton(bool& bInvalidSkeletonIsError, const IPropertyHandle* PropertyHandle)
 {
-}
+	bInvalidSkeletonIsError = false;
 
+	if (const IAnimClassInterface* AnimClassInterface = GetAnimClassInterface())
+	{
+		return AnimClassInterface->GetTargetSkeleton();
+	}
+
+	return nullptr;
+}
 
 void FAnimNode_Inertialization::RequestInertialization(float Duration, const UBlendProfile* BlendProfile)
 {
@@ -172,20 +305,43 @@ void FAnimNode_Inertialization::Initialize_AnyThread(const FAnimationInitializeC
 	FAnimNode_Base::Initialize_AnyThread(Context);
 	Source.Initialize(Context);
 
+	CurveFilter.Empty();
+	CurveFilter.SetFilterMode(UE::Anim::ECurveFilterMode::DisallowFiltered);
+	CurveFilter.AppendNames(FilteredCurves);
 
-	const int32 NumSkeletonBones = bPreallocateMemory ? Context.AnimInstanceProxy->GetSkeleton()->GetReferenceSkeleton().GetNum() : 0;
+	BoneFilter.Init(FCompactPoseBoneIndex(INDEX_NONE), FilteredBones.Num());
+
+	PrevPoseSnapshot.Empty();
+	CurrPoseSnapshot.Empty();
 
-	PoseSnapshots.Empty(INERTIALIZATION_MAX_POSE_SNAPSHOTS);
 	RequestQueue.Reserve(8);
-	InertializationDurationPerBone.Reserve(NumSkeletonBones);
+
+	BoneIndices.Empty();
+
+	BoneTranslationDiffDirection.Empty();
+	BoneTranslationDiffMagnitude.Empty();
+	BoneTranslationDiffSpeed.Empty();
+
+	BoneRotationDiffAxis.Empty();
+	BoneRotationDiffAngle.Empty();
+	BoneRotationDiffSpeed.Empty();
+
+	BoneScaleDiffAxis.Empty();
+	BoneScaleDiffMagnitude.Empty();
+	BoneScaleDiffSpeed.Empty();
+
+	CurveDiffs.Empty();
 
 	DeltaTime = 0.0f;
 
-	TeleportType = ETeleportType::None;
+	InertializationState = EInertializationState::Inactive;
+	InertializationElapsedTime = 0.0f;
 
-	Deactivate();
+	InertializationDuration = 0.0f;
+	InertializationDurationPerBone.Empty();
+	InertializationMaxDuration = 0.0f;
 
-	InertializationPoseDiff.Reset(NumSkeletonBones);
+	InertializationDeficit = 0.0f;
 }
 
 
@@ -195,6 +351,16 @@ void FAnimNode_Inertialization::CacheBones_AnyThread(const FAnimationCacheBonesC
 
 	FAnimNode_Base::CacheBones_AnyThread(Context);
 	Source.CacheBones(Context);
+
+	// Compute Compact Pose Bone Index for each bone in Filter
+
+	const FBoneContainer& RequiredBones = Context.AnimInstanceProxy->GetRequiredBones();
+	BoneFilter.Init(FCompactPoseBoneIndex(INDEX_NONE), FilteredBones.Num());
+	for (int32 FilterBoneIdx = 0; FilterBoneIdx < FilteredBones.Num(); FilterBoneIdx++)
+	{
+		FilteredBones[FilterBoneIdx].Initialize(Context.AnimInstanceProxy->GetSkeleton());
+		BoneFilter[FilterBoneIdx] = FilteredBones[FilterBoneIdx].GetCompactPoseIndex(RequiredBones);
+	}
 }
 
 
@@ -276,18 +442,18 @@ void FAnimNode_Inertialization::Evaluate_AnyThread(FPoseContext& Output)
 		Deactivate();
 
 		// Clear the pose history
-		PoseSnapshots.Reset();
+		PrevPoseSnapshot.Empty();
+		CurrPoseSnapshot.Empty();
 
-		// Reset the cached time accumulator and teleport state
+		// Reset the cached time accumulator
 		DeltaTime = 0.0f;
-		TeleportType = ETeleportType::None;
 
 		return;
 	}
 
 	// Update the inertialization state if a new inertialization request is pending
 	const int32 NumRequests = RequestQueue.Num();
-	if (NumRequests > 0 && PoseSnapshots.Num() > 0)
+	if (NumRequests > 0 && !CurrPoseSnapshot.IsEmpty())
 	{
 		float AppliedDeficit = 0.0f;
 		if (InertializationState == EInertializationState::Active)
@@ -303,7 +469,7 @@ void FAnimNode_Inertialization::Evaluate_AnyThread(FPoseContext& Output)
 		InertializationState = EInertializationState::Pending;
 		InertializationElapsedTime = 0.0f;
 		
-		const int32 NumSkeletonBones = UE::Anim::GetNumSkeletonBones(Output.AnimInstanceProxy->GetRequiredBones());
+		const int32 NumSkeletonBones = UE::Anim::Inertialization::Private::GetNumSkeletonBones(Output.AnimInstanceProxy->GetRequiredBones());
 
 		const USkeleton* TargetSkeleton = Output.AnimInstanceProxy->GetRequiredBones().GetSkeletonAsset();
 		auto FillSkeletonBoneDurationsArray = [this, NumSkeletonBones, TargetSkeleton](auto& DurationPerBone, float Duration, const UBlendProfile* BlendProfile) {
@@ -404,21 +570,25 @@ void FAnimNode_Inertialization::Evaluate_AnyThread(FPoseContext& Output)
 	// than the mesh component's location because we still want to inertialize instances where the skeletal mesh component has been moved
 	// while simultaneously counter-moving the root bone (as is the case when mounting and dismounting vehicles for example)
 
-	// bool bTeleported = (TeleportType != ETeleportType::None);
-	// Ignore TeleportType for now. See UE-78594.
 	bool bTeleported = false;
 	const float TeleportDistanceThreshold = Output.AnimInstanceProxy->GetSkelMeshComponent()->GetTeleportDistanceThreshold();
-	if (!bTeleported && PoseSnapshots.Num() > 0 && TeleportDistanceThreshold > 0.0f)
+	if (!bTeleported && !CurrPoseSnapshot.IsEmpty() && TeleportDistanceThreshold > 0.0f)
 	{
-		const FInertializationPose& Prev = PoseSnapshots[PoseSnapshots.Num() - 1];
-
 		const FVector RootWorldSpaceLocation = ComponentTransform.TransformPosition(Output.Pose[FCompactPoseBoneIndex(0)].GetTranslation());
-		const FVector PrevRootWorldSpaceLocation = Prev.ComponentTransform.TransformPosition(Prev.BoneTransforms[0].GetTranslation());
-		if (FVector::DistSquared(RootWorldSpaceLocation, PrevRootWorldSpaceLocation) > FMath::Square(TeleportDistanceThreshold))
+
+		const int32 RootBoneIndex = CurrPoseSnapshot.BoneIndices[0];
+
+		if (RootBoneIndex != INDEX_NONE)
 		{
-			bTeleported = true;
+			const FVector PrevRootWorldSpaceLocation = CurrPoseSnapshot.ComponentTransform.TransformPosition(CurrPoseSnapshot.BoneTranslations[RootBoneIndex]);
+
+			if (FVector::DistSquared(RootWorldSpaceLocation, PrevRootWorldSpaceLocation) > FMath::Square(TeleportDistanceThreshold))
+			{
+				bTeleported = true;
+			}
 		}
 	}
+
 	if (bTeleported)
 	{
 		// Cancel inertialization requests during teleports
@@ -438,26 +608,6 @@ void FAnimNode_Inertialization::Evaluate_AnyThread(FPoseContext& Output)
 		DeltaTime = 0.0f;
 	}
 
-	// Inertialize the pose
-	if (InertializationState == EInertializationState::Pending)
-	{
-		StartInertialization(Output,
-			PoseSnapshots[PoseSnapshots.Num() - 1],
-			PoseSnapshots[FMath::Max(PoseSnapshots.Num() - 2, 0)],
-			InertializationDuration,
-			InertializationDurationPerBone,
-			InertializationPoseDiff);
-		InertializationState = EInertializationState::Active;
-	}
-	if (InertializationState == EInertializationState::Active)
-	{
-		ApplyInertialization(Output,
-			InertializationPoseDiff,
-			InertializationElapsedTime,
-			InertializationDuration,
-			InertializationDurationPerBone);
-	}
-
 	// Get the parent actor attachment information (to detect and counteract discontinuities when changing parents)
 	FName AttachParentName = NAME_None;
 	if (AActor* Owner = Output.AnimInstanceProxy->GetSkelMeshComponent()->GetOwner())
@@ -468,30 +618,39 @@ void FAnimNode_Inertialization::Evaluate_AnyThread(FPoseContext& Output)
 		}
 	}
 
-	// Capture the pose snapshot for this frame
-	if (PoseSnapshots.Num() < INERTIALIZATION_MAX_POSE_SNAPSHOTS)
+	// Inertialize the pose
+	if (InertializationState == EInertializationState::Pending)
 	{
-		// Add the pose to the end of the buffer
-		FInertializationPose& Snapshot = PoseSnapshots.AddDefaulted_GetRef();
-		Snapshot.InitFrom(Output.Pose, Output.Curve, ComponentTransform, AttachParentName, DeltaTime);
+		InitFrom(
+			Output.Pose,
+			Output.Curve, 
+			ComponentTransform, 
+			AttachParentName, 
+			CurrPoseSnapshot, 
+			PrevPoseSnapshot);
+			
+		InertializationState = EInertializationState::Active;
 	}
-	else
+
+	if (InertializationState == EInertializationState::Active)
 	{
-		// Bubble the old poses forward in the buffer (using swaps to avoid allocations and copies)
-		for (int32 i = 0; i < INERTIALIZATION_MAX_POSE_SNAPSHOTS - 1; ++i)
-		{
-			Swap(PoseSnapshots[i], PoseSnapshots[i+1]);
-		}
+		ApplyTo(Output.Pose, Output.Curve);
+	}
+
+	// Record Pose Snapshot
 
-		// Overwrite the (now irrelevant) pose in the last slot with the new post snapshot
-		// (thereby avoiding the reallocation costs we would have incurred had we simply added a new pose at the end)
-		FInertializationPose& Snapshot = PoseSnapshots[INERTIALIZATION_MAX_POSE_SNAPSHOTS - 1];
-		Snapshot.InitFrom(Output.Pose, Output.Curve, ComponentTransform, AttachParentName, DeltaTime);
+	if (!CurrPoseSnapshot.IsEmpty())
+	{
+		// Directly swap the memory of the current pose with the prev pose snapshot (to avoid allocations and copies)
+		Swap(PrevPoseSnapshot, CurrPoseSnapshot);
 	}
+	
+	// Initialize the current pose
+	CurrPoseSnapshot.InitFrom(Output.Pose, Output.Curve, ComponentTransform, AttachParentName, DeltaTime);
+	
 
 	// Reset the time accumulator and teleport state
 	DeltaTime = 0.0f;
-	TeleportType = ETeleportType::None;
 
 	const float NormalizedInertializationTime = InertializationDuration > UE_KINDA_SMALL_NUMBER ? (InertializationElapsedTime / InertializationDuration) : 0.0f;
 	const float InertializationWeight = InertializationState == EInertializationState::Active ? 1.0f - NormalizedInertializationTime : 0.0f;
@@ -546,13 +705,12 @@ bool FAnimNode_Inertialization::NeedsDynamicReset() const
 
 void FAnimNode_Inertialization::ResetDynamics(ETeleportType InTeleportType)
 {
-	if (InTeleportType > TeleportType)
-	{
-		TeleportType = InTeleportType;
-	}
+	// Note: InTeleportType is unused and teleports are detected automatically (UE-78594)
 }
 
+PRAGMA_DISABLE_DEPRECATION_WARNINGS
 
+// DEPRECATED: See FAnimNode_Inertialization::InitFrom
 void FAnimNode_Inertialization::StartInertialization(FPoseContext& Context, FInertializationPose& PreviousPose1, FInertializationPose& PreviousPose2, float Duration, TArrayView<const float> DurationPerBone, /*OUT*/ FInertializationPoseDiff& OutPoseDiff)
 {
 	// Determine if this skeletal mesh's actor is attached to another actor
@@ -566,7 +724,7 @@ void FAnimNode_Inertialization::StartInertialization(FPoseContext& Context, FIne
 	}
 
 	// Initialize curve filter if required 
-	if(FilteredCurves.Num() != CurveFilter.Num())
+	if (FilteredCurves.Num() != CurveFilter.Num())
 	{
 		CurveFilter.Empty();
 		CurveFilter.SetFilterMode(UE::Anim::ECurveFilterMode::DisallowFiltered);
@@ -576,35 +734,328 @@ void FAnimNode_Inertialization::StartInertialization(FPoseContext& Context, FIne
 	OutPoseDiff.InitFrom(Context.Pose, Context.Curve, Context.AnimInstanceProxy->GetComponentTransform(), AttachParentName, PreviousPose1, PreviousPose2, CurveFilter);
 }
 
-
+// DEPRECATED: See FAnimNode_Inertialization::ApplyTo
 void FAnimNode_Inertialization::ApplyInertialization(FPoseContext& Context, const FInertializationPoseDiff& PoseDiff, float ElapsedTime, float Duration, TArrayView<const float> DurationPerBone)
 {
 	PoseDiff.ApplyTo(Context.Pose, Context.Curve, ElapsedTime, Duration, DurationPerBone);
 }
 
+PRAGMA_ENABLE_DEPRECATION_WARNINGS
 
-void FAnimNode_Inertialization::Deactivate()
+void FAnimNode_Inertialization::InitFrom(const FCompactPose& InPose, const FBlendedCurve& InCurves, const FTransform& ComponentTransform, const FName AttachParentName, const FInertializationSparsePose& Prev1, const FInertializationSparsePose& Prev2)
+{	
+	// Compute the Inertialization Space
+
+	const FQuat ComponentTransformGetRotationInverse = ComponentTransform.GetRotation().Inverse();
+
+	// Determine if we should initialize in local space (the default) or in world space (for situations where we wish to correct
+	// a world-space discontinuity such as an abrupt orientation change)
+	EInertializationSpace InertializationSpace = EInertializationSpace::Default;
+	if (AttachParentName != Prev1.AttachParentName || AttachParentName != Prev2.AttachParentName)
+	{
+		// If the parent space has changed, then inertialize in world space
+		InertializationSpace = EInertializationSpace::WorldSpace;
+	}
+	else if (AttachParentName == NAME_None)
+	{
+		// If there was a discontinuity in ComponentTransform orientation, then correct for that by inertializing the orientation in world space
+		// (but only if the mesh is not attached to another actor, because we don't want to dampen the connection between attached actors)
+		if ((FMath::Abs((Prev1.ComponentTransform.GetRotation() * ComponentTransformGetRotationInverse).W) < 0.999f) ||	// (W < 0.999f --> angle > 5 degrees)
+			(FMath::Abs((Prev2.ComponentTransform.GetRotation() * ComponentTransformGetRotationInverse).W) < 0.999f))	// (W < 0.999f --> angle > 5 degrees)
+		{
+			InertializationSpace = EInertializationSpace::WorldRotation;
+		}
+	}
+
+	// Compute the Inertialization Bone Indices which we will use to index into BoneTranslations, BoneRotations, etc
+
+	const FBoneContainer& BoneContainer = InPose.GetBoneContainer();
+
+	const int32 NumSkeletonBones = UE::Anim::Inertialization::Private::GetNumSkeletonBones(BoneContainer);
+
+	BoneIndices.Init(INDEX_NONE, NumSkeletonBones);
+
+	int32 NumInertializationBones = 0;
+
+	for (FCompactPoseBoneIndex BoneIndex : InPose.ForEachBoneIndex())
+	{
+		const int32 SkeletonPoseBoneIndex = BoneContainer.GetSkeletonIndex(BoneIndex);
+
+		if (SkeletonPoseBoneIndex == INDEX_NONE || 
+			Prev1.BoneIndices[SkeletonPoseBoneIndex] == INDEX_NONE ||
+			Prev2.BoneIndices[SkeletonPoseBoneIndex] == INDEX_NONE)
+		{
+			continue;
+		}
+
+		BoneIndices[SkeletonPoseBoneIndex] = NumInertializationBones;
+		NumInertializationBones++;
+	}
+
+	// Allocate Inertialization Bones
+
+	BoneTranslationDiffDirection.Init(FVector3f::ZeroVector, NumInertializationBones);
+	BoneTranslationDiffMagnitude.Init(0.0f, NumInertializationBones);
+	BoneTranslationDiffSpeed.Init(0.0f, NumInertializationBones);
+	BoneRotationDiffAxis.Init(FVector3f::ZeroVector, NumInertializationBones);
+	BoneRotationDiffAngle.Init(0.0f, NumInertializationBones);
+	BoneRotationDiffSpeed.Init(0.0f, NumInertializationBones);
+	BoneScaleDiffAxis.Init(FVector3f::ZeroVector, NumInertializationBones);
+	BoneScaleDiffMagnitude.Init(0.0f, NumInertializationBones);
+	BoneScaleDiffSpeed.Init(0.0f, NumInertializationBones);
+	
+	// Compute Pose Differences
+
+	for (FCompactPoseBoneIndex BoneIndex : InPose.ForEachBoneIndex())
+	{
+		const int32 SkeletonPoseBoneIndex = BoneContainer.GetSkeletonIndex(BoneIndex);
+
+		if (SkeletonPoseBoneIndex == INDEX_NONE ||
+			Prev1.BoneIndices[SkeletonPoseBoneIndex] == INDEX_NONE ||
+			Prev2.BoneIndices[SkeletonPoseBoneIndex] == INDEX_NONE)
+		{
+			continue;
+		}
+
+		// Get Bone Indices for Inertialization Bone, Prev and Curr Pose Bones
+
+		const int32 InertializationBoneIndex = BoneIndices[SkeletonPoseBoneIndex];
+		const int32 Prev1PoseBoneIndex = Prev1.BoneIndices[SkeletonPoseBoneIndex];
+		const int32 Prev2PoseBoneIndex = Prev2.BoneIndices[SkeletonPoseBoneIndex];
+
+		check(InertializationBoneIndex != INDEX_NONE);
+		check(Prev1PoseBoneIndex != INDEX_NONE);
+		check(Prev2PoseBoneIndex != INDEX_NONE);
+
+		const FTransform PoseTransform = InPose[BoneIndex];
+		FTransform Prev1Transform = FTransform(Prev1.BoneRotations[Prev1PoseBoneIndex], Prev1.BoneTranslations[Prev1PoseBoneIndex], Prev1.BoneScales[Prev1PoseBoneIndex]);
+		FTransform Prev2Transform = FTransform(Prev2.BoneRotations[Prev2PoseBoneIndex], Prev2.BoneTranslations[Prev2PoseBoneIndex], Prev2.BoneScales[Prev2PoseBoneIndex]);
+
+		if (BoneIndex.IsRootBone())
+		{
+			// If we are inertializing in world space, then adjust the historical root bones to be in a consistent reference frame
+			if (InertializationSpace == EInertializationSpace::WorldSpace)
+			{
+				Prev1Transform *= Prev1.ComponentTransform.GetRelativeTransform(ComponentTransform);
+				Prev2Transform *= Prev2.ComponentTransform.GetRelativeTransform(ComponentTransform);
+			}
+			else if (InertializationSpace == EInertializationSpace::WorldRotation)
+			{
+				Prev1Transform.SetRotation(ComponentTransformGetRotationInverse * Prev1.ComponentTransform.GetRotation() * Prev1Transform.GetRotation());
+				Prev2Transform.SetRotation(ComponentTransformGetRotationInverse * Prev2.ComponentTransform.GetRotation() * Prev2Transform.GetRotation());
+			}
+		}
+
+		// Compute the bone translation difference
+		{
+			FVector TranslationDirection = FVector::ZeroVector;
+			float TranslationMagnitude = 0.0f;
+			float TranslationSpeed = 0.0f;
+
+			const FVector T = Prev1Transform.GetTranslation() - PoseTransform.GetTranslation();
+			TranslationMagnitude = T.Size();
+			if (TranslationMagnitude > UE_KINDA_SMALL_NUMBER)
+			{
+				TranslationDirection = T / TranslationMagnitude;
+			}
+
+			if (Prev1.DeltaTime > UE_KINDA_SMALL_NUMBER && TranslationMagnitude > UE_KINDA_SMALL_NUMBER)
+			{
+				const FVector PrevT = Prev2Transform.GetTranslation() - PoseTransform.GetTranslation();
+				const float PrevMagnitude = FVector::DotProduct(PrevT, TranslationDirection);
+				TranslationSpeed = (TranslationMagnitude - PrevMagnitude) / Prev1.DeltaTime;
+			}
+
+			BoneTranslationDiffDirection[InertializationBoneIndex] = (FVector3f)TranslationDirection;
+			BoneTranslationDiffMagnitude[InertializationBoneIndex] = TranslationMagnitude;
+			BoneTranslationDiffSpeed[InertializationBoneIndex] = TranslationSpeed;
+		}
+
+		// Compute the bone rotation difference
+		{
+			FVector RotationAxis = FVector::ZeroVector;
+			float RotationAngle = 0.0f;
+			float RotationSpeed = 0.0f;
+
+			const FQuat Q = Prev1Transform.GetRotation() * PoseTransform.GetRotation().Inverse();
+			Q.ToAxisAndAngle(RotationAxis, RotationAngle);
+			RotationAngle = FMath::UnwindRadians(RotationAngle);
+			if (RotationAngle < 0.0f)
+			{
+				RotationAxis = -RotationAxis;
+				RotationAngle = -RotationAngle;
+			}
+
+			if (Prev1.DeltaTime > UE_KINDA_SMALL_NUMBER && RotationAngle > UE_KINDA_SMALL_NUMBER)
+			{
+				const FQuat PrevQ = Prev2Transform.GetRotation() * PoseTransform.GetRotation().Inverse();
+				const float PrevAngle = PrevQ.GetTwistAngle(RotationAxis);
+				RotationSpeed = FMath::UnwindRadians(RotationAngle - PrevAngle) / Prev1.DeltaTime;
+			}
+
+			BoneRotationDiffAxis[InertializationBoneIndex] = (FVector3f)RotationAxis;
+			BoneRotationDiffAngle[InertializationBoneIndex] = RotationAngle;
+			BoneRotationDiffSpeed[InertializationBoneIndex] = RotationSpeed;
+		}
+
+		// Compute the bone scale difference
+		{
+			FVector ScaleAxis = FVector::ZeroVector;
+			float ScaleMagnitude = 0.0f;
+			float ScaleSpeed = 0.0f;
+
+			const FVector S = Prev1Transform.GetScale3D() - PoseTransform.GetScale3D();
+			ScaleMagnitude = S.Size();
+			if (ScaleMagnitude > UE_KINDA_SMALL_NUMBER)
+			{
+				ScaleAxis = S / ScaleMagnitude;
+			}
+
+			if (Prev1.DeltaTime > UE_KINDA_SMALL_NUMBER && ScaleMagnitude > UE_KINDA_SMALL_NUMBER)
+			{
+				const FVector PrevS = Prev2Transform.GetScale3D() - PoseTransform.GetScale3D();
+				const float PrevMagnitude = FVector::DotProduct(PrevS, ScaleAxis);
+				ScaleSpeed = (ScaleMagnitude - PrevMagnitude) / Prev1.DeltaTime;
+			}
+
+			BoneScaleDiffAxis[InertializationBoneIndex] = (FVector3f)ScaleAxis;
+			BoneScaleDiffMagnitude[InertializationBoneIndex] = ScaleMagnitude;
+			BoneScaleDiffSpeed[InertializationBoneIndex] = ScaleSpeed;
+		}
+	}
+
+	// Compute the curve differences
+	// First copy in current values
+	CurveDiffs.CopyFrom(InCurves);
+
+	// Compute differences
+	UE::Anim::FNamedValueArrayUtils::Union(CurveDiffs, Prev1.Curves.BlendedCurve,
+		[](FInertializationCurveDiffElement& OutResultElement, const UE::Anim::FCurveElement& InElement1, UE::Anim::ENamedValueUnionFlags InFlags)
+		{
+			OutResultElement.Delta = InElement1.Value - OutResultElement.Value;
+		});
+
+	// Compute derivatives
+	if (Prev1.DeltaTime > UE_KINDA_SMALL_NUMBER)
+	{
+		UE::Anim::FNamedValueArrayUtils::Union(CurveDiffs, Prev2.Curves.BlendedCurve,
+			[DeltaTime = Prev1.DeltaTime](FInertializationCurveDiffElement& OutResultElement, const UE::Anim::FCurveElement& InElement1, UE::Anim::ENamedValueUnionFlags InFlags)
+			{
+				const float Prev1Weight = OutResultElement.Delta - OutResultElement.Value;
+				const float Prev2Weight = InElement1.Value;
+				OutResultElement.Derivative = (Prev1Weight - Prev2Weight) / DeltaTime;
+			});
+	}
+
+	// Apply filtering to remove filtered curves from diffs. This does not actually
+	// prevent these curves from being inertialized, but does stop them appearing as empty
+	// in the output curves created by the Union in ApplyTo unless they are already in the
+	// destination animation.
+	if (CurveFilter.Num() > 0)
+	{
+		UE::Anim::FCurveUtils::Filter(CurveDiffs, CurveFilter);
+	}
+}
+
+void FAnimNode_Inertialization::ApplyTo(FCompactPose& InOutPose, FBlendedCurve& InOutCurves)
 {
-	if (!bPreallocateMemory)
+	const FBoneContainer& BoneContainer = InOutPose.GetBoneContainer();
+
+	// Apply pose difference
+	for (FCompactPoseBoneIndex BoneIndex : InOutPose.ForEachBoneIndex())
 	{
-		// InertializationPoseDiff::Reset actually empties and frees memory here
-		InertializationPoseDiff.Reset();
-		InertializationDurationPerBone.Empty();
+		const int32 SkeletonPoseBoneIndex = BoneContainer.GetSkeletonIndex(BoneIndex);
+
+		if (SkeletonPoseBoneIndex == INDEX_NONE || BoneIndices[SkeletonPoseBoneIndex] == INDEX_NONE || BoneFilter.Contains(BoneIndex))
+		{
+			continue;
+		}
+
+		const int32 InertializationBoneIndex = BoneIndices[SkeletonPoseBoneIndex];
+		check(InertializationBoneIndex != INDEX_NONE);
+
+		const float Duration = InertializationDurationPerBone[SkeletonPoseBoneIndex];
+
+		// Apply the bone translation difference
+		const FVector T = (FVector)BoneTranslationDiffDirection[InertializationBoneIndex] *
+			UE::Anim::Inertialization::Private::CalcInertialFloat(BoneTranslationDiffMagnitude[InertializationBoneIndex], BoneTranslationDiffSpeed[InertializationBoneIndex], InertializationElapsedTime, Duration);
+		InOutPose[BoneIndex].AddToTranslation(T);
+
+		// Apply the bone rotation difference
+		const FQuat Q = FQuat((FVector)BoneRotationDiffAxis[InertializationBoneIndex],
+			UE::Anim::Inertialization::Private::CalcInertialFloat(BoneRotationDiffAngle[InertializationBoneIndex], BoneRotationDiffSpeed[InertializationBoneIndex], InertializationElapsedTime, Duration));
+		InOutPose[BoneIndex].SetRotation(Q * InOutPose[BoneIndex].GetRotation());
+
+		// Apply the bone scale difference
+		const FVector S = (FVector)BoneScaleDiffAxis[InertializationBoneIndex] *
+			UE::Anim::Inertialization::Private::CalcInertialFloat(BoneScaleDiffMagnitude[InertializationBoneIndex], BoneScaleDiffSpeed[InertializationBoneIndex], InertializationElapsedTime, Duration);
+		InOutPose[BoneIndex].SetScale3D(S + InOutPose[BoneIndex].GetScale3D());
 	}
 
+	InOutPose.NormalizeRotations();
+
+	// Apply curve differences
+
+	PoseCurveData.CopyFrom(InOutCurves);
+
+	UE::Anim::FNamedValueArrayUtils::Union(InOutCurves, PoseCurveData, CurveDiffs, [this](
+		UE::Anim::FCurveElement& OutResultElement, 
+		const UE::Anim::FCurveElement& InElement0,
+		const FInertializationCurveDiffElement& InElement1,
+		UE::Anim::ENamedValueUnionFlags InFlags)
+		{
+			// For filtered Curves take destination value
+
+			if (FilteredCurves.Contains(OutResultElement.Name))
+			{
+				OutResultElement.Value = InElement0.Value;
+				OutResultElement.Flags = InElement0.Flags;
+				return;
+			}
+
+			// Otherwise take destination value plus offset
+
+			OutResultElement.Value = InElement0.Value + UE::Anim::Inertialization::Private::CalcInertialFloat(InElement1.Delta, InElement1.Derivative, InertializationElapsedTime, InertializationDuration);
+			OutResultElement.Flags = InElement0.Flags | InElement1.Flags;
+		});
+}
+
+void FAnimNode_Inertialization::Deactivate()
+{
 	InertializationState = EInertializationState::Inactive;
+
+	BoneIndices.Empty();
+	
+	BoneTranslationDiffDirection.Empty();
+	BoneTranslationDiffMagnitude.Empty();
+	BoneTranslationDiffSpeed.Empty();
+	
+	BoneRotationDiffAxis.Empty();
+	BoneRotationDiffAngle.Empty();
+	BoneRotationDiffSpeed.Empty();
+	
+	BoneScaleDiffAxis.Empty();
+	BoneScaleDiffMagnitude.Empty();
+	BoneScaleDiffSpeed.Empty();
+
+	CurveDiffs.Empty();
+
+	InertializationDurationPerBone.Empty();
+
 	InertializationElapsedTime = 0.0f;
 	InertializationDuration = 0.0f;
 	InertializationMaxDuration = 0.0f;
 	InertializationDeficit = 0.0f;
 }
 
+PRAGMA_DISABLE_DEPRECATION_WARNINGS
 
+// DEPRECATED: See FInertializationSparsePose::InitFrom
 void FInertializationPose::InitFrom(const FCompactPose& Pose, const FBlendedCurve& InCurves, const FTransform& InComponentTransform, const FName& InAttachParentName, float InDeltaTime)
 {
 	const FBoneContainer& BoneContainer = Pose.GetBoneContainer();
 
-	const int32 NumSkeletonBones = UE::Anim::GetNumSkeletonBones(BoneContainer);
+	const int32 NumSkeletonBones = UE::Anim::Inertialization::Private::GetNumSkeletonBones(BoneContainer);
 	BoneTransforms.Reset(NumSkeletonBones);
 	BoneTransforms.AddZeroed(NumSkeletonBones);
 	BoneStates.Reset(NumSkeletonBones);
@@ -625,16 +1076,18 @@ void FInertializationPose::InitFrom(const FCompactPose& Pose, const FBlendedCurv
 	DeltaTime = InDeltaTime;
 }
 
+PRAGMA_ENABLE_DEPRECATION_WARNINGS
+
 void FInertializationSparsePose::InitFrom(
 	const FCompactPose& Pose,
 	const FBlendedCurve& InCurves,
 	const FTransform& InComponentTransform,
-	const FName& InAttachParentName,
-	float InDeltaTime)
+	const FName InAttachParentName,
+	const float InDeltaTime)
 {
 	const FBoneContainer& BoneContainer = Pose.GetBoneContainer();
 
-	const int32 NumSkeletonBones = UE::Anim::GetNumSkeletonBones(BoneContainer);
+	const int32 NumSkeletonBones = UE::Anim::Inertialization::Private::GetNumSkeletonBones(BoneContainer);
 
 	// Allocate Bone Index Array
 
@@ -676,7 +1129,7 @@ void FInertializationSparsePose::InitFrom(
 
 		// Get the InertializationBoneIndex and write the transform data
 
-		const uint16 InertializationBoneIndex = BoneIndices[SkeletonPoseBoneIndex];
+		const int32 InertializationBoneIndex = BoneIndices[SkeletonPoseBoneIndex];
 		check(InertializationBoneIndex != INDEX_NONE);
 
 		const FTransform BoneTransform = Pose[BoneIndex];
@@ -708,14 +1161,9 @@ void FInertializationSparsePose::Empty()
 }
 
 
-// Initialize the pose difference from the current pose and the two previous snapshots
-//
-// Pose			the current frame's pose
-// Curves       the current frame's curves
-// Prev1		the previous frame's pose and curves
-// Prev2		the pose and curves from two frames before
-// DeltaTime	the time elapsed between Prev1 and Pose
-//
+PRAGMA_DISABLE_DEPRECATION_WARNINGS
+
+// DEPRECATED: See FAnimNode_Inertialization::InitFrom
 void FInertializationPoseDiff::InitFrom(const FCompactPose& Pose, const FBlendedCurve& Curves, const FTransform& ComponentTransform, const FName& AttachParentName, const FInertializationPose& Prev1, const FInertializationPose& Prev2, const UE::Anim::FCurveFilter& CurveFilter)
 {
 	const FBoneContainer& BoneContainer = Pose.GetBoneContainer();
@@ -742,7 +1190,7 @@ void FInertializationPoseDiff::InitFrom(const FCompactPose& Pose, const FBlended
 	}
 
 	// Compute the inertialization differences for each bone
-	const int32 NumSkeletonBones = UE::Anim::GetNumSkeletonBones(BoneContainer);
+	const int32 NumSkeletonBones = UE::Anim::Inertialization::Private::GetNumSkeletonBones(BoneContainer);
 	BoneDiffs.Empty(NumSkeletonBones);
 	BoneDiffs.AddZeroed(NumSkeletonBones);
 	for (FCompactPoseBoneIndex BoneIndex : Pose.ForEachBoneIndex())
@@ -914,9 +1362,7 @@ void FInertializationPoseDiff::InitFrom(const FCompactPose& Pose, const FBlended
 	}
 }
 
-
-// Apply this difference to a pose and a set of curves, decaying over time as InertializationElapsedTime approaches InertializationDuration
-//
+// DEPRECATED: See FAnimNode_Inertialization::ApplyTo
 void FInertializationPoseDiff::ApplyTo(FCompactPose& Pose, FBlendedCurve& Curves, float InertializationElapsedTime, float InertializationDuration, TArrayView<const float> InertializationDurationPerBone) const
 {
 	const FBoneContainer& BoneContainer = Pose.GetBoneContainer();
@@ -933,17 +1379,17 @@ void FInertializationPoseDiff::ApplyTo(FCompactPose& Pose, FBlendedCurve& Curves
 
 			// Apply the bone translation difference
 			const FVector T = BoneDiff.TranslationDirection *
-				CalcInertialFloat(BoneDiff.TranslationMagnitude, BoneDiff.TranslationSpeed, InertializationElapsedTime, Duration);
+				UE::Anim::Inertialization::Private::CalcInertialFloat(BoneDiff.TranslationMagnitude, BoneDiff.TranslationSpeed, InertializationElapsedTime, Duration);
 			Pose[BoneIndex].AddToTranslation(T);
 
 			// Apply the bone rotation difference
 			const FQuat Q = FQuat(BoneDiff.RotationAxis,
-				CalcInertialFloat(BoneDiff.RotationAngle, BoneDiff.RotationSpeed, InertializationElapsedTime, Duration));
+				UE::Anim::Inertialization::Private::CalcInertialFloat(BoneDiff.RotationAngle, BoneDiff.RotationSpeed, InertializationElapsedTime, Duration));
 			Pose[BoneIndex].SetRotation(Q * Pose[BoneIndex].GetRotation());
 
 			// Apply the bone scale difference
 			const FVector S = BoneDiff.ScaleAxis *
-				CalcInertialFloat(BoneDiff.ScaleMagnitude, BoneDiff.ScaleSpeed, InertializationElapsedTime, Duration);
+				UE::Anim::Inertialization::Private::CalcInertialFloat(BoneDiff.ScaleMagnitude, BoneDiff.ScaleSpeed, InertializationElapsedTime, Duration);
 			Pose[BoneIndex].SetScale3D(S + Pose[BoneIndex].GetScale3D());
 		}
 	}
@@ -954,141 +1400,12 @@ void FInertializationPoseDiff::ApplyTo(FCompactPose& Pose, FBlendedCurve& Curves
 	UE::Anim::FNamedValueArrayUtils::Union(Curves, CurveDiffs,
 		[&InertializationElapsedTime, &InertializationDuration](UE::Anim::FCurveElement& OutResultElement, const FInertializationCurveDiffElement& InParamElement, UE::Anim::ENamedValueUnionFlags InFlags)
 		{
-			OutResultElement.Value += CalcInertialFloat(InParamElement.Delta, InParamElement.Derivative, InertializationElapsedTime, InertializationDuration);
+			OutResultElement.Value += UE::Anim::Inertialization::Private::CalcInertialFloat(InParamElement.Delta, InParamElement.Derivative, InertializationElapsedTime, InertializationDuration);
 			OutResultElement.Flags |= InParamElement.Flags;
 		});
 }
 
+PRAGMA_ENABLE_DEPRECATION_WARNINGS
 
-// Calculate the "inertialized" value of a single float at time t
-//
-// @param x0	Initial value of the float (at time 0)
-// @param v0	Initial "velocity" (first derivative) of the float (at time 0)
-// @param t		Time at which to evaluate the float
-// @param t1	Ending inertialization time (ie: the time at which the curve must be zero)
-//
-// Evaluates a quintic polynomial curve with the specified initial conditions (x0, v0) which hits zero at time t1.  As well,
-// the curve is designed so that the first and second derivatives are also zero at time t1.
-//
-// The initial second derivative (a0) is chosen such that it is as close to zero as possible, but large enough to prevent any
-// overshoot (ie: enforce x >= 0 for t between 0 and t1).  If necessary, the ending time t1 will be adjusted (shortened) to
-// guarantee that there is no overshoot, even for very large initial velocities.
-//
-float FInertializationPoseDiff::CalcInertialFloat(float x0, float v0, float t, float t1)
-{
-	static_assert(INERTIALIZATION_TIME_EPSILON * INERTIALIZATION_TIME_EPSILON * INERTIALIZATION_TIME_EPSILON * INERTIALIZATION_TIME_EPSILON * INERTIALIZATION_TIME_EPSILON > FLT_MIN,
-		"INERTIALIZATION_TIME_EPSILON^5 must be greater than FLT_MIN to avoid denormalization (and potential division by zero) for very small values of t1");
-
-	if (t < 0.0f)
-	{
-		t = 0.0f;
-	}
-
-	if (t >= t1 - INERTIALIZATION_TIME_EPSILON)
-	{
-		return 0.0f;
-	}
-
-	// Assume that x0 >= 0... if this is not the case, then simply invert everything (both input and output)
-	float sign = 1.0f;
-	if (x0 < 0.0f)
-	{
-		x0 = -x0;
-		v0 = -v0;
-		sign = -1.0f;
-	}
-
-	// If v0 > 0, then the curve will overshoot away from zero, so clamp v0 here to guarantee that there is no overshoot
-	if (v0 > 0.0f)
-	{
-		v0 = 0.0f;
-	}
-
-	// Check for invalid values - this is only expected to occur if NaNs or other invalid values are coming into the node
-	if (!ensureMsgf(x0 >= 0.0f && v0 <= 0.0f && t >= 0.0f && t1 >= 0.0f,
-		TEXT("Invalid Value(s) in Inertialization - x0: %f, v0: %f, t: %f, t1: %f"), x0, v0, t, t1))
-	{
-		x0 = 0.0f;
-		v0 = 0.0f;
-		t = 0.0f;
-		t1 = 0.0f;
-	}
-
-	// Limit t1 such that the curve does not overshoot below zero (ensuring that x >= 0 for all t between 0 and t1).
-	//
-	// We observe that if the curve does overshoot below zero, it must have an inflection point somewhere between 0 and t1
-	// (since we know that both x0 and x1 are >= 0).  Therefore, we can prevent overshoot by adjusting t1 such that any
-	// inflection point is at t >= t1.
-	//
-	// Assuming that we are using the zero jerk version of the curve (see below) whose velocity has a triple root at t1, then
-	// we can prevent overshoot by forcing the remaining root to be at time t >= t1, or equivalently, we can set t1 to be the
-	// lesser of the original t1 or the value that gives us a solution with a quadruple velocity root at t1.
-	//
-	// The following Mathematica expression solves for t1 that gives us the quadruple velocity root:
-	//
-	//		v := q * (t-t1)^4
-	//		x := Integrate[Expand[v], t] + x0
-	//		eq1 := (v/.t->0)==v0
-	//		eq2 := (x/.t->t1)==0
-	//		Solve[{eq1 && eq2}, {q,t1}]
-	//
-	if (v0 < -UE_KINDA_SMALL_NUMBER)
-	{
-		t1 = FMath::Min(t1, -5.0f * x0 / v0);
-	}
-
-	if (t >= t1 - INERTIALIZATION_TIME_EPSILON)
-	{
-		return 0.0f;
-	}
-
-	const float t1_2 = t1 * t1;
-	const float t1_3 = t1 * t1_2;
-	const float t1_4 = t1 * t1_3;
-	const float t1_5 = t1 * t1_4;
-
-	// Compute the initial acceleration value (a0) for this curve.  Ideally we want to use an initial acceleration of zero, but
-	// if there is a large negative initial velocity, then we will need to use a larger acceleration in order to ensure that
-	// the curve does not overshoot below zero (ie: to ensure that x >= 0 for all t between 0 and t1).
-	//
-	// To compute a0, we first compute the a0 that we would get if we also specified that the third derivative (the "jerk" j)
-	// is zero at t1.  If this value of a0 is positive (and therefore opposing the initial velocity), then we use that.  If it
-	// is negative, then we simply use an initial a0 of zero.
-	//
-	// The following Mathematica expression solves for a0 that gives us zero jerk at t1:
-	//
-	//		x:= A*t^5 + B*t^4 + C*t^3 + D*t^2 + v0*t + x0
-	//		v:=Dt[x, t, Constants->{A,B,C,D,v0,x0}]
-	//		a:=Dt[v, t, Constants->{A,B,C,D,v0,x0}]
-	//		j:=Dt[a, t, Constants->{A,B,C,D,v0,x0}]
-	//		eq1:=(x/.t->t1)==0
-	//		eq2:=(v/.t->t1)==0
-	//		eq3:=(a/.t->t1)==0
-	//		eq4:=(j/.t->t1)==0
-	//		a0:=a/.t->0/.Solve[{eq1 && eq2 && eq3 && eq4}, {A,B,C,D}]
-	//		ExpandNumerator[a0]
-	//
-	const float a0 = FMath::Max(0.0f, (-8.0f*t1*v0 - 20.0f*x0) / t1_2);
-
-	// Compute the polynomial coefficients given the starting and ending conditions, solved from:
-	//
-	//		x:= A*t^5 + B*t^4 + C*t^3 + D*t^2 + v0*t + x0
-	//		v:=Dt[x, t, Constants->{A,B,C,D,v0,x0}]
-	//		a:=Dt[v, t, Constants->{A,B,C,D,v0,x0}]
-	//		eq1:=(x/.t->t1)==0
-	//		eq2:=(v/.t->t1)==0
-	//		eq3:=(a/.t->t1)==0
-	//		eq4:=(a/.t->0)==a0
-	//		Simplify[Solve[{eq1 && eq2 && eq3 && eq4}, {A,B,C,D}]]
-	//
-	const float A = -0.5f * (     a0*t1_2 +  6.0f*t1*v0 + 12.0f*x0) / t1_5;
-	const float B =  0.5f * (3.0f*a0*t1_2 + 16.0f*t1*v0 + 30.0f*x0) / t1_4;
-	const float C = -0.5f * (3.0f*a0*t1_2 + 12.0f*t1*v0 + 20.0f*x0) / t1_3;
-	const float D =  0.5f * a0;
-
-	const float x = (((((A*t) + B)*t + C)*t + D)*t + v0)*t + x0;
-
-	return x * sign;
-}
 
 #undef LOCTEXT_NAMESPACE
