Fix landscape layer not being recognized as such when used within material layers with the new HLSL translator enabled

diff --git a/Engine/Source/Runtime/Engine/Private/Materials/MaterialExpressionHLSL.cpp b/Engine/Source/Runtime/Engine/Private/Materials/MaterialExpressionHLSL.cpp
--- a/Engine/Source/Runtime/Engine/Private/Materials/MaterialExpressionHLSL.cpp
+++ b/Engine/Source/Runtime/Engine/Private/Materials/MaterialExpressionHLSL.cpp
@@ -4922,7 +4922,7 @@ bool GenerateStaticTerrainLayerWeightExpression(FName LayerName, float PreviewWe
 {
 	using namespace UE::HLSLTree;
 	const FExpression* TexCoordExpression = Generator.NewExternalInput(Material::EExternalInput::TexCoord3);
-	OutExpression = Generator.GetTree().NewExpression<Material::FExpressionStaticTerrainLayerWeight>(Generator.GetParameterInfo(LayerName), TexCoordExpression, PreviewWeight, bUseTextureArray);
+	OutExpression = Generator.GetTree().NewExpression<Material::FExpressionStaticTerrainLayerWeight>(FMaterialParameterInfo(LayerName), TexCoordExpression, PreviewWeight, bUseTextureArray);
 	return true;
 }
 
