World Partition: minor optimization when computing distance/angle ratio sorting priority.

diff --git a/Engine/Source/Runtime/Engine/Private/WorldPartition/WorldPartitionRuntimeCellData.cpp b/Engine/Source/Runtime/Engine/Private/WorldPartition/WorldPartitionRuntimeCellData.cpp
--- a/Engine/Source/Runtime/Engine/Private/WorldPartition/WorldPartitionRuntimeCellData.cpp
+++ b/Engine/Source/Runtime/Engine/Private/WorldPartition/WorldPartitionRuntimeCellData.cpp
@@ -64,8 +64,7 @@ void UWorldPartitionRuntimeCellData::AppendStreamingSourceInfo(const FWorldParti
 	const float SourceCosAngleRatio = SourceCosAngle * 0.5f + 0.5f;
 
 	// Compute final cell priority for this source
-	static float CosAngleContribution = 0.4f;
-	const double SortingPriority = SoureDistanceRatio * FMath::Cube(FMath::Pow(SourceCosAngleRatio, CosAngleContribution));
+	const double SortingPriority = SoureDistanceRatio * SourceCosAngleRatio;
 
 	// Update if lower
 	CachedMinSpatialSortingPriority = FMath::Min(CachedMinSpatialSortingPriority, SortingPriority);
