ActorDescContainer: fix bad events base class declaration.

diff --git a/Engine/Source/Runtime/Engine/Public/WorldPartition/ActorDescContainer.h b/Engine/Source/Runtime/Engine/Public/WorldPartition/ActorDescContainer.h
--- a/Engine/Source/Runtime/Engine/Public/WorldPartition/ActorDescContainer.h
+++ b/Engine/Source/Runtime/Engine/Public/WorldPartition/ActorDescContainer.h
@@ -74,10 +74,10 @@ public:
 
 	ENGINE_API bool IsActorDescHandled(const AActor* Actor) const;
 
-	DECLARE_EVENT_OneParam(UWorldPartition, FActorDescAddedEvent, FWorldPartitionActorDesc*);
+	DECLARE_EVENT_OneParam(UActorDescContainer, FActorDescAddedEvent, FWorldPartitionActorDesc*);
 	FActorDescAddedEvent OnActorDescAddedEvent;
 	
-	DECLARE_EVENT_OneParam(UWorldPartition, FActorDescRemovedEvent, FWorldPartitionActorDesc*);
+	DECLARE_EVENT_OneParam(UActorDescContainer, FActorDescRemovedEvent, FWorldPartitionActorDesc*);
 	FActorDescRemovedEvent OnActorDescRemovedEvent;
 
 	DECLARE_MULTICAST_DELEGATE_OneParam(FActorDescContainerInitializeDelegate, UActorDescContainer*);
