Gather unique shader compile errors a second time if we filtered all errors as warnings the first time on unsuccessful compile job.

#rnx

diff --git a/Engine/Source/Runtime/Engine/Private/ShaderCompiler/ShaderCompiler.cpp b/Engine/Source/Runtime/Engine/Private/ShaderCompiler/ShaderCompiler.cpp
--- a/Engine/Source/Runtime/Engine/Private/ShaderCompiler/ShaderCompiler.cpp
+++ b/Engine/Source/Runtime/Engine/Private/ShaderCompiler/ShaderCompiler.cpp
@@ -3331,6 +3331,26 @@ static bool CheckSingleJob(const FShaderCompileJob& SingleJob, TArray<FString>&
 };
 #endif // WITH_EDITOR
 
+static int32 AddErrorsForFailedJobFiltered(FShaderCompileJob& CurrentJob, FShaderErrorInfo& OutShaderErrorInfo, const TCHAR* FilterMessage)
+{
+	int32 NumAddedErrors = 0;
+
+	for (int32 ErrorIndex = 0; ErrorIndex < CurrentJob.Output.Errors.Num(); ErrorIndex++)
+	{
+		const FShaderCompilerError& CurrentError = CurrentJob.Output.Errors[ErrorIndex];
+
+		// Include warnings if LogShaders is unsuppressed, otherwise only include filtered messages
+		if (UE_LOG_ACTIVE(LogShaders, Log) || FilterMessage == nullptr || CurrentError.StrippedErrorMessage.Contains(FilterMessage))
+		{
+			OutShaderErrorInfo.UniqueErrors.AddUnique(CurrentJob.Output.Errors[ErrorIndex].GetErrorString());
+			OutShaderErrorInfo.ErrorJobs.AddUnique(&CurrentJob);
+			++NumAddedErrors;
+		}
+	}
+
+	return NumAddedErrors;
+}
+
 static void AddErrorsForFailedJob(FShaderCompileJob& CurrentJob, FShaderErrorInfo& OutShaderErrorInfo)
 {
 	OutShaderErrorInfo.ErrorPlatforms.AddUnique((EShaderPlatform)CurrentJob.Input.Target.Platform);
@@ -3342,16 +3362,11 @@ static void AddErrorsForFailedJob(FShaderCompileJob& CurrentJob, FShaderErrorInf
 		CurrentJob.Output.Errors.Add(Error);
 	}
 
-	for (int32 ErrorIndex = 0; ErrorIndex < CurrentJob.Output.Errors.Num(); ErrorIndex++)
+	// If we filter all error messages because they are interpreted as warnings, we have to assume all error messages are in fact errors and not warnings.
+	// In that case, add jobs again without a filter; e.g. when the stripped message starts with "Internal exception".
+	if (AddErrorsForFailedJobFiltered(CurrentJob, OutShaderErrorInfo, TEXT("error")) == 0)
 	{
-		const FShaderCompilerError& CurrentError = CurrentJob.Output.Errors[ErrorIndex];
-
-		// Include warnings if LogShaders is unsuppressed, otherwise only include errors
-		if (UE_LOG_ACTIVE(LogShaders, Log) || CurrentError.StrippedErrorMessage.Contains(TEXT("error")))
-		{
-			OutShaderErrorInfo.UniqueErrors.AddUnique(CurrentJob.Output.Errors[ErrorIndex].GetErrorString());
-			OutShaderErrorInfo.ErrorJobs.AddUnique(&CurrentJob);
-		}
+		AddErrorsForFailedJobFiltered(CurrentJob, OutShaderErrorInfo, nullptr);
 	}
 }
 
