Enable p.ClusterUnion.SyncBodiesMoveNewComponents by default

diff --git a/Engine/Source/Runtime/Engine/Private/PhysicsEngine/Experimental/PhysScene_Chaos.cpp b/Engine/Source/Runtime/Engine/Private/PhysicsEngine/Experimental/PhysScene_Chaos.cpp
--- a/Engine/Source/Runtime/Engine/Private/PhysicsEngine/Experimental/PhysScene_Chaos.cpp
+++ b/Engine/Source/Runtime/Engine/Private/PhysicsEngine/Experimental/PhysScene_Chaos.cpp
@@ -62,7 +62,7 @@ FAutoConsoleVariableRef CVar_KinematicDeferralLogInvalidBodies(TEXT("p.Kinematic
 float GReplicationCacheLingerForNSeconds = 3.f;
 FAutoConsoleVariableRef CVar_ReplicationCacheLingerForNSeconds(TEXT("np2.ReplicationCache.LingerForNSeconds"), GReplicationCacheLingerForNSeconds, TEXT("How long to keep data in the replication cache without the actor accessing it, after this we stop caching the actors state until it tries to access it again."));
 
-bool bGClusterUnionSyncBodiesMoveNewComponents = false;
+bool bGClusterUnionSyncBodiesMoveNewComponents = true;
 FAutoConsoleVariableRef CVar_GClusterUnionSyncBodiesCheckDirtyFlag(TEXT("p.ClusterUnion.SyncBodiesMoveNewComponents"), bGClusterUnionSyncBodiesMoveNewComponents, TEXT("Enable a fix to ensure new components in a cluster union are moved once on add (even if the cluster is not moving)."));
 
 DECLARE_CYCLE_STAT(TEXT("Update Kinematics On Deferred SkelMeshes"), STAT_UpdateKinematicsOnDeferredSkelMeshesChaos, STATGROUP_Physics);
