Material Translate:  Optimize out unused preshaders (meaning preshaders not referenced by any HLSL) in Translate_Legacy.  On a worst case shader, reduces preshaders by 4x, and parameter fetches and ops by 2.5x.  On a more typical shader, reduces preshaders and parameter fetches by 40%, and ops by 33%.

A scope class tracks when we are in a code path that's generating a uniform expression, and suppresses generation of preshaders for that case.  More detailed comments on FAddUniformExpressionScope class in the code.  Utility virtual function "GetChildren" added to FMaterialUniformExpression, which allows generalized traversal of a uniform expression tree, which I used for debug code (since removed), but others may find it useful.

#rnx

diff --git a/Engine/Shaders/Public/ShaderVersion.ush b/Engine/Shaders/Public/ShaderVersion.ush
--- a/Engine/Shaders/Public/ShaderVersion.ush
+++ b/Engine/Shaders/Public/ShaderVersion.ush
@@ -2,4 +2,4 @@
 // Each time the console command is executed it generates a new GUID. As this file is included
 // in Platform.ush (which should be included in any shader) it allows to invalidate the shader DDC.
 // 
-#pragma message("UESHADERMETADATA_VERSION A8F68E0C-FEF9-4330-8B8C-C1F8C4FFA445")
+#pragma message("UESHADERMETADATA_VERSION 78C3F233-CBA3-4DB9-8BF5-F1FE3E490DD3")
diff --git a/Engine/Source/Runtime/Engine/Private/Materials/HLSLMaterialTranslator.cpp b/Engine/Source/Runtime/Engine/Private/Materials/HLSLMaterialTranslator.cpp
--- a/Engine/Source/Runtime/Engine/Private/Materials/HLSLMaterialTranslator.cpp
+++ b/Engine/Source/Runtime/Engine/Private/Materials/HLSLMaterialTranslator.cpp
@@ -189,6 +189,12 @@ static FAutoConsoleVariableRef CVarPedanticErrorChecksEnabled(
 	GPedanticErrorChecksEnabled,
 	TEXT("Enables material compilation pedantic error checking"));
 
+static bool GCullIntermediateUniformExpressions = true;
+static FAutoConsoleVariableRef CVarCullIntermediateUniformExpressions(
+	TEXT("r.Material.CullIntermediateUniformExpressions"),
+	GCullIntermediateUniformExpressions,
+	TEXT("Enables culling of intermediate uniform expressions, reducing preshader count, saving performance"));
+
 /* Controls whether to use the new GetMaterialShaderCode() and GetMaterialEnvironment() implementations. */
 static bool GUseMaterialTranslationResultsGrouping = true;
 
@@ -206,7 +212,7 @@ UE::DerivedData::FValueId EnvironmentDefinesId = UE::DerivedData::FValueId::From
 
 /* This version number models the layout of the data stored on the DDC after a material translation (e.g. FEnvironmentDefines)
  * It must be bumped whenever a change is made requires re-translation of all materials. */
-static constexpr int MaterialTranslationDDCVersion = 1;
+static constexpr int MaterialTranslationDDCVersion = 2;
 
 /** Data structure used to cache a part of material translation results. It contains all the generated
  *  defines that will be declared during the compilation of the generated material shader.
@@ -650,6 +656,8 @@ FHLSLMaterialTranslator::FHLSLMaterialTranslator(FMaterial* InMaterial,
 ,	bMaterialIsSubstrate(false)
 ,	bUsesCurvature(false)
 ,	bUsesPerInstanceFadeAmount(false)
+,	bCullIntermediateUniformExpressions(GCullIntermediateUniformExpressions)
+,	AddingUniformExpression(0)
 ,	AllocatedUserTexCoords()
 ,	AllocatedUserVertexTexCoords()
 ,	DynamicParticleParameterMask(0)
@@ -4382,6 +4390,7 @@ int32 FHLSLMaterialTranslator::AddUniformExpressionInner(uint64 Hash, FMaterialU
 	if (!bFoundExistingExpression)
 	{
 		// Add an entry to the material-wide list of uniform expressions
+		UniformExpression->UniformIndex = UniformExpressions.Num();
 		new(UniformExpressions) FShaderCodeChunk(Hash, UniformExpression, FormattedCode, FormattedCode, Type, EDerivativeStatus::Zero);
 	}
 
@@ -4390,7 +4399,7 @@ int32 FHLSLMaterialTranslator::AddUniformExpressionInner(uint64 Hash, FMaterialU
 }
 
 // AddUniformExpression - Adds an input to the Code array and returns its index.
-int32 FHLSLMaterialTranslator::AddUniformExpression(FMaterialUniformExpression* UniformExpression,EMaterialValueType Type, const TCHAR* Format,...)
+int32 FHLSLMaterialTranslator::AddUniformExpression(FAddUniformExpressionScope& Scope, FMaterialUniformExpression* UniformExpression,EMaterialValueType Type, const TCHAR* Format,...)
 {
 	int32	BufferSize = 256;
 	TCHAR*	FormattedCode = NULL;
@@ -4421,6 +4430,14 @@ int32 FHLSLMaterialTranslator::AccessUniformExpression(int32 Index)
 	check(CodeChunk.UniformExpression && !CodeChunk.UniformExpression->IsConstant());
 	const bool bIsLWC = IsLWCType(CodeChunk.Type);
 
+	// If this is a uniform expression referenced from another uniform expression, return a unique stub expression, to
+	// avoid instantiating a preshader that's not actually needed.
+	if (bCullIntermediateUniformExpressions && AddingUniformExpression)
+	{
+		check(CodeChunk.UniformExpression->UniformIndex != INDEX_NONE);
+		return AddInlinedCodeChunkZeroDeriv(CodeChunk.Type, TEXT("UniformStub[$%d]"), CodeChunk.UniformExpression->UniformIndex);
+	}
+
 	FMaterialUniformExpressionTexture* TextureUniformExpression = CodeChunk.UniformExpression->GetTextureUniformExpression();
 	FMaterialUniformExpressionExternalTexture* ExternalTextureUniformExpression = CodeChunk.UniformExpression->GetExternalTextureUniformExpression();
 
@@ -5550,35 +5567,41 @@ int32 FHLSLMaterialTranslator::NumericParameter(EMaterialParameterType Parameter
 	FMaterialParameterInfo ParameterInfo = GetParameterAssociationInfo();
 	ParameterInfo.Name = ParameterName;
 
+	FAddUniformExpressionScope Scope(this);
 	const int32 ParameterIndex = MaterialCompilationOutput.UniformExpressionSet.FindOrAddNumericParameter(ParameterType, ParameterInfo, DefaultOffset);
-	return AddUniformExpression(new FMaterialUniformExpressionNumericParameter(ParameterInfo, ParameterIndex), GetMaterialValueType(ParameterType), TEXT(""));
+	return AddUniformExpression(Scope, new FMaterialUniformExpressionNumericParameter(ParameterInfo, ParameterIndex), GetMaterialValueType(ParameterType), TEXT(""));
 }
 
 int32 FHLSLMaterialTranslator::Constant(float X)
 {
-	return AddUniformExpression(new FMaterialUniformExpressionConstant(FLinearColor(X,X,X,X),MCT_Float),MCT_Float,TEXT("%0.8f"),X);
+	FAddUniformExpressionScope Scope(this);
+	return AddUniformExpression(Scope, new FMaterialUniformExpressionConstant(FLinearColor(X,X,X,X),MCT_Float),MCT_Float,TEXT("%0.8f"),X);
 }
 
 int32 FHLSLMaterialTranslator::Constant2(float X,float Y)
 {
-	return AddUniformExpression(new FMaterialUniformExpressionConstant(FLinearColor(X,Y,0,0),MCT_Float2),MCT_Float2,TEXT("MaterialFloat2(%0.8f,%0.8f)"),X,Y);
+	FAddUniformExpressionScope Scope(this);
+	return AddUniformExpression(Scope, new FMaterialUniformExpressionConstant(FLinearColor(X,Y,0,0),MCT_Float2),MCT_Float2,TEXT("MaterialFloat2(%0.8f,%0.8f)"),X,Y);
 }
 
 int32 FHLSLMaterialTranslator::Constant3(float X,float Y,float Z)
 {
-	return AddUniformExpression(new FMaterialUniformExpressionConstant(FLinearColor(X,Y,Z,0),MCT_Float3),MCT_Float3,TEXT("MaterialFloat3(%0.8f,%0.8f,%0.8f)"),X,Y,Z);
+	FAddUniformExpressionScope Scope(this);
+	return AddUniformExpression(Scope, new FMaterialUniformExpressionConstant(FLinearColor(X,Y,Z,0),MCT_Float3),MCT_Float3,TEXT("MaterialFloat3(%0.8f,%0.8f,%0.8f)"),X,Y,Z);
 }
 
 int32 FHLSLMaterialTranslator::Constant4(float X,float Y,float Z,float W)
 {
-	return AddUniformExpression(new FMaterialUniformExpressionConstant(FLinearColor(X,Y,Z,W),MCT_Float4),MCT_Float4,TEXT("MaterialFloat4(%0.8f,%0.8f,%0.8f,%0.8f)"),X,Y,Z,W);
+	FAddUniformExpressionScope Scope(this);
+	return AddUniformExpression(Scope, new FMaterialUniformExpressionConstant(FLinearColor(X,Y,Z,W),MCT_Float4),MCT_Float4,TEXT("MaterialFloat4(%0.8f,%0.8f,%0.8f,%0.8f)"),X,Y,Z,W);
 }
 
 int32 FHLSLMaterialTranslator::GenericConstant(const UE::Shader::FValue& Value)
 {
+	FAddUniformExpressionScope Scope(this);
 	TStringBuilder<1024> String;
 	Value.ToString(UE::Shader::EValueStringFormat::HLSL, String);
-	return AddUniformExpression(new FMaterialUniformExpressionGenericConstant(Value), GetMaterialValueType(Value.GetType()), String.ToString());
+	return AddUniformExpression(Scope, new FMaterialUniformExpressionGenericConstant(Value), GetMaterialValueType(Value.GetType()), String.ToString());
 }
 	
 int32 FHLSLMaterialTranslator::ViewProperty(EMaterialExposedViewProperty Property, bool InvProperty)
@@ -5709,7 +5732,8 @@ int32 FHLSLMaterialTranslator::PeriodicHint(int32 PeriodicCode)
 
 	if(GetParameterUniformExpression(PeriodicCode))
 	{
-		return AddUniformExpression(new FMaterialUniformExpressionPeriodic(GetParameterUniformExpression(PeriodicCode)),GetParameterType(PeriodicCode),TEXT("%s"),*GetParameterCode(PeriodicCode));
+		FAddUniformExpressionScope Scope(this);
+		return AddUniformExpression(Scope, new FMaterialUniformExpressionPeriodic(GetParameterUniformExpression(PeriodicCode)),GetParameterType(PeriodicCode),TEXT("%s"),*GetParameterCode(PeriodicCode));
 	}
 	else
 	{
@@ -5726,7 +5750,8 @@ int32 FHLSLMaterialTranslator::Sine(int32 X)
 
 	if(GetParameterUniformExpression(X))
 	{
-		return AddUniformExpression(new FMaterialUniformExpressionTrigMath(GetParameterUniformExpression(X),TMO_Sin),MCT_Float,TEXT("sin(%s)"),*CoerceParameter(X,MCT_Float));
+		FAddUniformExpressionScope Scope(this);
+		return AddUniformExpression(Scope, new FMaterialUniformExpressionTrigMath(GetParameterUniformExpression(X),TMO_Sin),MCT_Float,TEXT("sin(%s)"),*CoerceParameter(X,MCT_Float));
 	}
 	else
 	{
@@ -5750,7 +5775,8 @@ int32 FHLSLMaterialTranslator::Cosine(int32 X)
 
 	if(GetParameterUniformExpression(X))
 	{
-		return AddUniformExpression(new FMaterialUniformExpressionTrigMath(GetParameterUniformExpression(X),TMO_Cos),MCT_Float,TEXT("cos(%s)"),*CoerceParameter(X,MCT_Float));
+		FAddUniformExpressionScope Scope(this);
+		return AddUniformExpression(Scope, new FMaterialUniformExpressionTrigMath(GetParameterUniformExpression(X),TMO_Cos),MCT_Float,TEXT("cos(%s)"),*CoerceParameter(X,MCT_Float));
 	}
 	else
 	{
@@ -5774,7 +5800,8 @@ int32 FHLSLMaterialTranslator::Tangent(int32 X)
 
 	if(GetParameterUniformExpression(X))
 	{
-		return AddUniformExpression(new FMaterialUniformExpressionTrigMath(GetParameterUniformExpression(X),TMO_Tan),MCT_Float,TEXT("tan(%s)"),*CoerceParameter(X,MCT_Float));
+		FAddUniformExpressionScope Scope(this);
+		return AddUniformExpression(Scope, new FMaterialUniformExpressionTrigMath(GetParameterUniformExpression(X),TMO_Tan),MCT_Float,TEXT("tan(%s)"),*CoerceParameter(X,MCT_Float));
 	}
 	else
 	{
@@ -5798,7 +5825,8 @@ int32 FHLSLMaterialTranslator::Arcsine(int32 X)
 
 	if(GetParameterUniformExpression(X))
 	{
-		return AddUniformExpression(new FMaterialUniformExpressionTrigMath(GetParameterUniformExpression(X),TMO_Asin),MCT_Float,TEXT("asin(%s)"),*CoerceParameter(X,MCT_Float));
+		FAddUniformExpressionScope Scope(this);
+		return AddUniformExpression(Scope, new FMaterialUniformExpressionTrigMath(GetParameterUniformExpression(X),TMO_Asin),MCT_Float,TEXT("asin(%s)"),*CoerceParameter(X,MCT_Float));
 	}
 	else
 	{
@@ -5822,7 +5850,8 @@ int32 FHLSLMaterialTranslator::ArcsineFast(int32 X)
 
 	if(GetParameterUniformExpression(X))
 	{
-		return AddUniformExpression(new FMaterialUniformExpressionTrigMath(GetParameterUniformExpression(X),TMO_Asin),MCT_Float,TEXT("asinFast(%s)"),*CoerceParameter(X,MCT_Float));
+		FAddUniformExpressionScope Scope(this);
+		return AddUniformExpression(Scope, new FMaterialUniformExpressionTrigMath(GetParameterUniformExpression(X),TMO_Asin),MCT_Float,TEXT("asinFast(%s)"),*CoerceParameter(X,MCT_Float));
 	}
 	else
 	{
@@ -5846,7 +5875,8 @@ int32 FHLSLMaterialTranslator::Arccosine(int32 X)
 
 	if(GetParameterUniformExpression(X))
 	{
-		return AddUniformExpression(new FMaterialUniformExpressionTrigMath(GetParameterUniformExpression(X),TMO_Acos),MCT_Float,TEXT("acos(%s)"),*CoerceParameter(X,MCT_Float));
+		FAddUniformExpressionScope Scope(this);
+		return AddUniformExpression(Scope, new FMaterialUniformExpressionTrigMath(GetParameterUniformExpression(X),TMO_Acos),MCT_Float,TEXT("acos(%s)"),*CoerceParameter(X,MCT_Float));
 	}
 	else
 	{
@@ -5870,7 +5900,8 @@ int32 FHLSLMaterialTranslator::ArccosineFast(int32 X)
 
 	if(GetParameterUniformExpression(X))
 	{
-		return AddUniformExpression(new FMaterialUniformExpressionTrigMath(GetParameterUniformExpression(X),TMO_Acos),MCT_Float,TEXT("acosFast(%s)"),*CoerceParameter(X,MCT_Float));
+		FAddUniformExpressionScope Scope(this);
+		return AddUniformExpression(Scope, new FMaterialUniformExpressionTrigMath(GetParameterUniformExpression(X),TMO_Acos),MCT_Float,TEXT("acosFast(%s)"),*CoerceParameter(X,MCT_Float));
 	}
 	else
 	{
@@ -5894,7 +5925,8 @@ int32 FHLSLMaterialTranslator::Arctangent(int32 X)
 
 	if(GetParameterUniformExpression(X))
 	{
-		return AddUniformExpression(new FMaterialUniformExpressionTrigMath(GetParameterUniformExpression(X),TMO_Atan),MCT_Float,TEXT("atan(%s)"),*CoerceParameter(X,MCT_Float));
+		FAddUniformExpressionScope Scope(this);
+		return AddUniformExpression(Scope, new FMaterialUniformExpressionTrigMath(GetParameterUniformExpression(X),TMO_Atan),MCT_Float,TEXT("atan(%s)"),*CoerceParameter(X,MCT_Float));
 	}
 	else
 	{
@@ -5918,7 +5950,8 @@ int32 FHLSLMaterialTranslator::ArctangentFast(int32 X)
 
 	if(GetParameterUniformExpression(X))
 	{
-		return AddUniformExpression(new FMaterialUniformExpressionTrigMath(GetParameterUniformExpression(X),TMO_Atan),MCT_Float,TEXT("atanFast(%s)"),*CoerceParameter(X,MCT_Float));
+		FAddUniformExpressionScope Scope(this);
+		return AddUniformExpression(Scope, new FMaterialUniformExpressionTrigMath(GetParameterUniformExpression(X),TMO_Atan),MCT_Float,TEXT("atanFast(%s)"),*CoerceParameter(X,MCT_Float));
 	}
 	else
 	{
@@ -5942,7 +5975,8 @@ int32 FHLSLMaterialTranslator::Arctangent2(int32 Y, int32 X)
 
 	if(GetParameterUniformExpression(Y) && GetParameterUniformExpression(X))
 	{
-		return AddUniformExpression(new FMaterialUniformExpressionTrigMath(GetParameterUniformExpression(Y),GetParameterUniformExpression(X),TMO_Atan2),MCT_Float,TEXT("atan2(%s, %s)"),*CoerceParameter(Y,MCT_Float),*CoerceParameter(X,MCT_Float));
+		FAddUniformExpressionScope Scope(this);
+		return AddUniformExpression(Scope, new FMaterialUniformExpressionTrigMath(GetParameterUniformExpression(Y),GetParameterUniformExpression(X),TMO_Atan2),MCT_Float,TEXT("atan2(%s, %s)"),*CoerceParameter(Y,MCT_Float),*CoerceParameter(X,MCT_Float));
 	}
 	else
 	{
@@ -5966,7 +6000,8 @@ int32 FHLSLMaterialTranslator::Arctangent2Fast(int32 Y, int32 X)
 
 	if(GetParameterUniformExpression(Y) && GetParameterUniformExpression(X))
 	{
-		return AddUniformExpression(new FMaterialUniformExpressionTrigMath(GetParameterUniformExpression(Y),GetParameterUniformExpression(X),TMO_Atan2),MCT_Float,TEXT("atan2Fast(%s, %s)"),*CoerceParameter(Y,MCT_Float),*CoerceParameter(X,MCT_Float));
+		FAddUniformExpressionScope Scope(this);
+		return AddUniformExpression(Scope, new FMaterialUniformExpressionTrigMath(GetParameterUniformExpression(Y),GetParameterUniformExpression(X),TMO_Atan2),MCT_Float,TEXT("atan2Fast(%s, %s)"),*CoerceParameter(Y,MCT_Float),*CoerceParameter(X,MCT_Float));
 	}
 	else
 	{
@@ -5990,7 +6025,8 @@ int32 FHLSLMaterialTranslator::Floor(int32 X)
 
 	if(GetParameterUniformExpression(X))
 	{
-		return AddUniformExpression(new FMaterialUniformExpressionFloor(GetParameterUniformExpression(X)),GetParameterType(X),TEXT("floor(%s)"),*GetParameterCode(X));
+		FAddUniformExpressionScope Scope(this);
+		return AddUniformExpression(Scope, new FMaterialUniformExpressionFloor(GetParameterUniformExpression(X)),GetParameterType(X),TEXT("floor(%s)"),*GetParameterCode(X));
 	}
 	else
 	{
@@ -6016,7 +6052,8 @@ int32 FHLSLMaterialTranslator::Ceil(int32 X)
 
 	if(GetParameterUniformExpression(X))
 	{
-		return AddUniformExpression(new FMaterialUniformExpressionCeil(GetParameterUniformExpression(X)),GetParameterType(X),TEXT("ceil(%s)"),*GetParameterCode(X));
+		FAddUniformExpressionScope Scope(this);
+		return AddUniformExpression(Scope, new FMaterialUniformExpressionCeil(GetParameterUniformExpression(X)),GetParameterType(X),TEXT("ceil(%s)"),*GetParameterCode(X));
 	}
 	else
 	{
@@ -6042,7 +6079,8 @@ int32 FHLSLMaterialTranslator::Round(int32 X)
 
 	if(GetParameterUniformExpression(X))
 	{
-		return AddUniformExpression(new FMaterialUniformExpressionRound(GetParameterUniformExpression(X)),GetParameterType(X),TEXT("round(%s)"),*GetParameterCode(X));
+		FAddUniformExpressionScope Scope(this);
+		return AddUniformExpression(Scope, new FMaterialUniformExpressionRound(GetParameterUniformExpression(X)),GetParameterType(X),TEXT("round(%s)"),*GetParameterCode(X));
 	}
 	else
 	{
@@ -6068,7 +6106,8 @@ int32 FHLSLMaterialTranslator::Truncate(int32 X)
 
 	if(GetParameterUniformExpression(X))
 	{
-		return AddUniformExpression(new FMaterialUniformExpressionTruncate(GetParameterUniformExpression(X)),GetParameterType(X),TEXT("trunc(%s)"),*GetParameterCode(X));
+		FAddUniformExpressionScope Scope(this);
+		return AddUniformExpression(Scope, new FMaterialUniformExpressionTruncate(GetParameterUniformExpression(X)),GetParameterType(X),TEXT("trunc(%s)"),*GetParameterCode(X));
 	}
 	else
 	{
@@ -6094,7 +6133,8 @@ int32 FHLSLMaterialTranslator::Sign(int32 X)
 
 	if(GetParameterUniformExpression(X))
 	{
-		return AddUniformExpression(new FMaterialUniformExpressionSign(GetParameterUniformExpression(X)),GetParameterType(X),TEXT("sign(%s)"),*GetParameterCode(X));
+		FAddUniformExpressionScope Scope(this);
+		return AddUniformExpression(Scope, new FMaterialUniformExpressionSign(GetParameterUniformExpression(X)),GetParameterType(X),TEXT("sign(%s)"),*GetParameterCode(X));
 	}
 	else
 	{
@@ -6120,7 +6160,8 @@ int32 FHLSLMaterialTranslator::Frac(int32 X)
 
 	if(GetParameterUniformExpression(X))
 	{
-		return AddUniformExpression(new FMaterialUniformExpressionFrac(GetParameterUniformExpression(X)),GetParameterType(X),TEXT("frac(%s)"),*GetParameterCode(X));
+		FAddUniformExpressionScope Scope(this);
+		return AddUniformExpression(Scope, new FMaterialUniformExpressionFrac(GetParameterUniformExpression(X)),GetParameterType(X),TEXT("frac(%s)"),*GetParameterCode(X));
 	}
 	else
 	{
@@ -6153,7 +6194,8 @@ int32 FHLSLMaterialTranslator::Fmod(int32 A, int32 B)
 
 	if (GetParameterUniformExpression(A) && GetParameterUniformExpression(B))
 	{
-		return AddUniformExpression(new FMaterialUniformExpressionFmod(GetParameterUniformExpression(A),GetParameterUniformExpression(B)),
+		FAddUniformExpressionScope Scope(this);
+		return AddUniformExpression(Scope, new FMaterialUniformExpressionFmod(GetParameterUniformExpression(A),GetParameterUniformExpression(B)),
 			GetParameterType(A),TEXT("fmod(%s,%s)"),*GetParameterCode(A),*CoerceParameter(B,GetParameterType(A)));
 	}
 	else
@@ -6187,7 +6229,8 @@ int32 FHLSLMaterialTranslator::Abs(int32 X)
 
 	if (GetParameterUniformExpression(X))
 	{
-		return AddUniformExpression(new FMaterialUniformExpressionAbs(GetParameterUniformExpression(X)), GetParameterType(X), TEXT("abs(%s)"), *GetParameterCode(X));
+		FAddUniformExpressionScope Scope(this);
+		return AddUniformExpression(Scope, new FMaterialUniformExpressionAbs(GetParameterUniformExpression(X)), GetParameterType(X), TEXT("abs(%s)"), *GetParameterCode(X));
 	}
 	else
 	{
@@ -7822,7 +7865,8 @@ int32 FHLSLMaterialTranslator::TextureProperty(int32 TextureIndex, EMaterialExpo
 	}
 
 	const EMaterialValueType ValueType = (TextureType == MCT_VolumeTexture || TextureType == MCT_Texture2DArray || TextureType == MCT_SparseVolumeTexture) ? MCT_Float3 : MCT_Float2;
-	return AddUniformExpression(new FMaterialUniformExpressionTextureProperty(TextureExpression, Property), ValueType, TEXT(""));
+	FAddUniformExpressionScope Scope(this);
+	return AddUniformExpression(Scope, new FMaterialUniformExpressionTextureProperty(TextureExpression, Property), ValueType, TEXT(""));
 }
 
 int32 FHLSLMaterialTranslator::TextureDecalMipmapLevel(int32 TextureSizeInput)
@@ -8354,7 +8398,8 @@ int32 FHLSLMaterialTranslator::Texture(UTexture* InTexture, int32& TextureRefere
 		bVirtual = false;
 		ShaderType = MCT_Texture2D;
 	}
-	return AddUniformExpression(new FMaterialUniformExpressionTexture(TextureReferenceIndex, SamplerType, SamplerSource, bVirtual),ShaderType,TEXT(""));
+	FAddUniformExpressionScope Scope(this);
+	return AddUniformExpression(Scope, new FMaterialUniformExpressionTexture(TextureReferenceIndex, SamplerType, SamplerSource, bVirtual),ShaderType,TEXT(""));
 }
 
 int32 FHLSLMaterialTranslator::TextureParameter(FName ParameterName, UTexture* InDefaultValue, int32& TextureReferenceIndex, EMaterialSamplerType SamplerType, ESamplerSourceMode SamplerSource)
@@ -8388,7 +8433,8 @@ int32 FHLSLMaterialTranslator::TextureParameter(FName ParameterName, UTexture* I
 		bVirtual = false;
 		ShaderType = MCT_Texture2D;
 	}
-	return AddUniformExpression(new FMaterialUniformExpressionTextureParameter(ParameterInfo, TextureReferenceIndex, SamplerType, SamplerSource, bVirtual),ShaderType,TEXT(""));
+	FAddUniformExpressionScope Scope(this);
+	return AddUniformExpression(Scope, new FMaterialUniformExpressionTextureParameter(ParameterInfo, TextureReferenceIndex, SamplerType, SamplerSource, bVirtual),ShaderType,TEXT(""));
 }
 
 int32 FHLSLMaterialTranslator::VirtualTexture(URuntimeVirtualTexture* InTexture, int32 TextureLayerIndex, int32 PageTableLayerIndex, int32& TextureReferenceIndex, EMaterialSamplerType SamplerType) 
@@ -8401,7 +8447,8 @@ int32 FHLSLMaterialTranslator::VirtualTexture(URuntimeVirtualTexture* InTexture,
 	TextureReferenceIndex = Material->GetReferencedTextures().Find(InTexture);
 	checkf(TextureReferenceIndex != INDEX_NONE, TEXT("Material expression called Compiler->VirtualTexture() without implementing UMaterialExpression::GetReferencedTexture properly"));
 
-	return AddUniformExpression(new FMaterialUniformExpressionTexture(TextureReferenceIndex, TextureLayerIndex, PageTableLayerIndex, SamplerType), MCT_TextureVirtual, TEXT(""));
+	FAddUniformExpressionScope Scope(this);
+	return AddUniformExpression(Scope, new FMaterialUniformExpressionTexture(TextureReferenceIndex, TextureLayerIndex, PageTableLayerIndex, SamplerType), MCT_TextureVirtual, TEXT(""));
 }
 
 int32 FHLSLMaterialTranslator::VirtualTextureParameter(FName ParameterName, URuntimeVirtualTexture* InDefaultValue, int32 TextureLayerIndex, int32 PageTableLayerIndex, int32& TextureReferenceIndex, EMaterialSamplerType SamplerType)
@@ -8426,12 +8473,14 @@ int32 FHLSLMaterialTranslator::VirtualTextureParameter(FName ParameterName, URun
 	FMaterialParameterInfo ParameterInfo = GetParameterAssociationInfo();
 	ParameterInfo.Name = ParameterName;
 
-	return AddUniformExpression(new FMaterialUniformExpressionTextureParameter(ParameterInfo, TextureReferenceIndex, TextureLayerIndex, PageTableLayerIndex, SamplerType), MCT_TextureVirtual, TEXT(""));
+	FAddUniformExpressionScope Scope(this);
+	return AddUniformExpression(Scope, new FMaterialUniformExpressionTextureParameter(ParameterInfo, TextureReferenceIndex, TextureLayerIndex, PageTableLayerIndex, SamplerType), MCT_TextureVirtual, TEXT(""));
 }
 
 int32 FHLSLMaterialTranslator::VirtualTextureUniform(int32 TextureIndex, int32 VectorIndex, UE::Shader::EValueType Type)
 {
-	return AddUniformExpression(new FMaterialUniformExpressionRuntimeVirtualTextureUniform(TextureIndex, VectorIndex), GetMaterialValueType(Type), TEXT(""));
+	FAddUniformExpressionScope Scope(this);
+	return AddUniformExpression(Scope, new FMaterialUniformExpressionRuntimeVirtualTextureUniform(TextureIndex, VectorIndex), GetMaterialValueType(Type), TEXT(""));
 }
 
 int32 FHLSLMaterialTranslator::VirtualTextureUniform(FName ParameterName, int32 TextureIndex, int32 VectorIndex, UE::Shader::EValueType Type)
@@ -8439,7 +8488,8 @@ int32 FHLSLMaterialTranslator::VirtualTextureUniform(FName ParameterName, int32
 	FMaterialParameterInfo ParameterInfo = GetParameterAssociationInfo();
 	ParameterInfo.Name = ParameterName;
 
-	return AddUniformExpression(new FMaterialUniformExpressionRuntimeVirtualTextureUniform(ParameterInfo, TextureIndex, VectorIndex), GetMaterialValueType(Type), TEXT(""));
+	FAddUniformExpressionScope Scope(this);
+	return AddUniformExpression(Scope, new FMaterialUniformExpressionRuntimeVirtualTextureUniform(ParameterInfo, TextureIndex, VectorIndex), GetMaterialValueType(Type), TEXT(""));
 }
 
 int32 FHLSLMaterialTranslator::VirtualTextureWorldToUV(int32 WorldPositionIndex, int32 P0, int32 P1, int32 P2, EPositionOrigin PositionOrigin)
@@ -8538,7 +8588,8 @@ int32 FHLSLMaterialTranslator::ExternalTexture(const FGuid& ExternalTextureGuid)
 		return NonPixelShaderExpressionError();
 	}
 
-	return AddUniformExpression(new FMaterialUniformExpressionExternalTexture(ExternalTextureGuid), MCT_TextureExternal, TEXT(""));
+	FAddUniformExpressionScope Scope(this);
+	return AddUniformExpression(Scope, new FMaterialUniformExpressionExternalTexture(ExternalTextureGuid), MCT_TextureExternal, TEXT(""));
 }
 
 int32 FHLSLMaterialTranslator::ExternalTexture(UTexture* InTexture, int32& TextureReferenceIndex)
@@ -8553,7 +8604,8 @@ int32 FHLSLMaterialTranslator::ExternalTexture(UTexture* InTexture, int32& Textu
 	TextureReferenceIndex = Material->GetReferencedTextures().Find(InTexture);
 	checkf(TextureReferenceIndex != INDEX_NONE, TEXT("Material expression called Compiler->ExternalTexture() without implementing UMaterialExpression::GetReferencedTexture properly"));
 
-	return AddUniformExpression(new FMaterialUniformExpressionExternalTexture(TextureReferenceIndex), MCT_TextureExternal, TEXT(""));
+	FAddUniformExpressionScope Scope(this);
+	return AddUniformExpression(Scope, new FMaterialUniformExpressionExternalTexture(TextureReferenceIndex), MCT_TextureExternal, TEXT(""));
 }
 
 int32 FHLSLMaterialTranslator::ExternalTextureParameter(FName ParameterName, UTexture* DefaultValue, int32& TextureReferenceIndex)
@@ -8567,24 +8619,29 @@ int32 FHLSLMaterialTranslator::ExternalTextureParameter(FName ParameterName, UTe
 
 	TextureReferenceIndex = Material->GetReferencedTextures().Find(DefaultValue);
 	checkf(TextureReferenceIndex != INDEX_NONE, TEXT("Material expression called Compiler->ExternalTextureParameter() without implementing UMaterialExpression::GetReferencedTexture properly"));
-	return AddUniformExpression(new FMaterialUniformExpressionExternalTextureParameter(ParameterName, TextureReferenceIndex), MCT_TextureExternal, TEXT(""));
+	FAddUniformExpressionScope Scope(this);
+	return AddUniformExpression(Scope, new FMaterialUniformExpressionExternalTextureParameter(ParameterName, TextureReferenceIndex), MCT_TextureExternal, TEXT(""));
 }
 
 int32 FHLSLMaterialTranslator::ExternalTextureCoordinateScaleRotation(int32 TextureReferenceIndex, TOptional<FName> ParameterName)
 {
-	return AddUniformExpression(new FMaterialUniformExpressionExternalTextureCoordinateScaleRotation(TextureReferenceIndex, ParameterName), MCT_Float4, TEXT(""));
+	FAddUniformExpressionScope Scope(this);
+	return AddUniformExpression(Scope, new FMaterialUniformExpressionExternalTextureCoordinateScaleRotation(TextureReferenceIndex, ParameterName), MCT_Float4, TEXT(""));
 }
 int32 FHLSLMaterialTranslator::ExternalTextureCoordinateScaleRotation(const FGuid& ExternalTextureGuid)
 {
-	return AddUniformExpression(new FMaterialUniformExpressionExternalTextureCoordinateScaleRotation(ExternalTextureGuid), MCT_Float4, TEXT(""));
+	FAddUniformExpressionScope Scope(this);
+	return AddUniformExpression(Scope, new FMaterialUniformExpressionExternalTextureCoordinateScaleRotation(ExternalTextureGuid), MCT_Float4, TEXT(""));
 }
 int32 FHLSLMaterialTranslator::ExternalTextureCoordinateOffset(int32 TextureReferenceIndex, TOptional<FName> ParameterName)
 {
-	return AddUniformExpression(new FMaterialUniformExpressionExternalTextureCoordinateOffset(TextureReferenceIndex, ParameterName), MCT_Float2, TEXT(""));
+	FAddUniformExpressionScope Scope(this);
+	return AddUniformExpression(Scope, new FMaterialUniformExpressionExternalTextureCoordinateOffset(TextureReferenceIndex, ParameterName), MCT_Float2, TEXT(""));
 }
 int32 FHLSLMaterialTranslator::ExternalTextureCoordinateOffset(const FGuid& ExternalTextureGuid)
 {
-	return AddUniformExpression(new FMaterialUniformExpressionExternalTextureCoordinateOffset(ExternalTextureGuid), MCT_Float2, TEXT(""));
+	FAddUniformExpressionScope Scope(this);
+	return AddUniformExpression(Scope, new FMaterialUniformExpressionExternalTextureCoordinateOffset(ExternalTextureGuid), MCT_Float2, TEXT(""));
 }
 
 UObject* FHLSLMaterialTranslator::GetReferencedTexture(int32 Index)
@@ -8620,8 +8677,9 @@ int32 FHLSLMaterialTranslator::DynamicBoolParameter(FName ParameterName, bool bD
 		DefaultUniformValues.Add(DefaultValue, DefaultOffset);
 	}
 
+	FAddUniformExpressionScope Scope(this);
 	const int32 ParameterIndex = MaterialCompilationOutput.UniformExpressionSet.FindOrAddNumericParameter(EMaterialParameterType::StaticSwitch, ParameterInfo, DefaultOffset);
-	return AddUniformExpression(new FMaterialUniformExpressionStaticBoolParameter(ParameterInfo, ParameterIndex), MCT_Bool, TEXT(""));
+	return AddUniformExpression(Scope, new FMaterialUniformExpressionStaticBoolParameter(ParameterInfo, ParameterIndex), MCT_Bool, TEXT(""));
 }
 
 int32 FHLSLMaterialTranslator::StaticBoolParameter(FName ParameterName,bool bDefaultValue)
@@ -9146,7 +9204,8 @@ int32 FHLSLMaterialTranslator::Add(int32 A, int32 B)
 	const EMaterialValueType ResultType = GetArithmeticResultType(A, B);
 	if (ExpressionA && ExpressionB)
 	{
-		return AddUniformExpression(new FMaterialUniformExpressionFoldedMath(GetParameterUniformExpression(A), GetParameterUniformExpression(B), FMO_Add), ResultType, TEXT("(%s + %s)"), *GetParameterCode(A), *GetParameterCode(B));
+		FAddUniformExpressionScope Scope(this);
+		return AddUniformExpression(Scope, new FMaterialUniformExpressionFoldedMath(GetParameterUniformExpression(A), GetParameterUniformExpression(B), FMO_Add), ResultType, TEXT("(%s + %s)"), *GetParameterCode(A), *GetParameterCode(B));
 	}
 	else
 	{
@@ -9181,7 +9240,8 @@ int32 FHLSLMaterialTranslator::Sub(int32 A, int32 B)
 	const EMaterialValueType ResultType = GetArithmeticResultType(A, B);
 	if (ExpressionA && ExpressionB)
 	{
-		return AddUniformExpression(new FMaterialUniformExpressionFoldedMath(GetParameterUniformExpression(A), GetParameterUniformExpression(B), FMO_Sub), ResultType, TEXT("(%s - %s)"), *GetParameterCode(A), *GetParameterCode(B));
+		FAddUniformExpressionScope Scope(this);
+		return AddUniformExpression(Scope, new FMaterialUniformExpressionFoldedMath(GetParameterUniformExpression(A), GetParameterUniformExpression(B), FMO_Sub), ResultType, TEXT("(%s - %s)"), *GetParameterCode(A), *GetParameterCode(B));
 	}
 	else
 	{
@@ -9278,7 +9338,8 @@ int32 FHLSLMaterialTranslator::Mul(int32 A,int32 B)
 			return ConstResultValue(ResultType, ConstantValue);
 		}
 
-		return AddUniformExpression(new FMaterialUniformExpressionFoldedMath(ExpressionA, ExpressionB, FMO_Mul),GetArithmeticResultType(A,B),TEXT("(%s * %s)"),*GetParameterCode(A),*GetParameterCode(B));
+		FAddUniformExpressionScope Scope(this);
+		return AddUniformExpression(Scope, new FMaterialUniformExpressionFoldedMath(ExpressionA, ExpressionB, FMO_Mul),GetArithmeticResultType(A,B),TEXT("(%s * %s)"),*GetParameterCode(A),*GetParameterCode(B));
 	}
 	else
 	{
@@ -9332,13 +9393,19 @@ int32 FHLSLMaterialTranslator::Div(int32 A, int32 B)
 			return ConstResultValue(ResultType, ConstantValue);
 		}
 
-		return AddUniformExpression(new FMaterialUniformExpressionFoldedMath(ExpressionA, ExpressionB, FMO_Div), GetArithmeticResultType(A, B), TEXT("(%s / %s)"), *GetParameterCode(A), *GetParameterCode(B));
+		FAddUniformExpressionScope Scope(this);
+		return AddUniformExpression(Scope, new FMaterialUniformExpressionFoldedMath(ExpressionA, ExpressionB, FMO_Div), GetArithmeticResultType(A, B), TEXT("(%s / %s)"), *GetParameterCode(A), *GetParameterCode(B));
 	}
 	else if (ExpressionB && !ExpressionB->IsConstant())
 	{
 		// Division is often optimized as multiplication by reciprocal
 		// If the divisor is a uniform expression, we can fold the reciprocal into the preshader
-		const int32 RcpB = AddUniformExpression(new FMaterialUniformExpressionRcp(ExpressionB), GetParameterType(B), TEXT("rcp(%s)"), *GetParameterCode(B));
+		int32 RcpB;
+		{
+			// The reciprocal is a uniform expression -- the Mul below is not, so we need the scope limited to this call
+			FAddUniformExpressionScope Scope(this);
+			RcpB = AddUniformExpression(Scope, new FMaterialUniformExpressionRcp(ExpressionB), GetParameterType(B), TEXT("rcp(%s)"), *GetParameterCode(B));
+		}
 		return Mul(A, RcpB);
 	}
 	else
@@ -9374,26 +9441,27 @@ int32 FHLSLMaterialTranslator::Dot(int32 A,int32 B)
 	FMaterialUniformExpression* ExpressionB = GetParameterUniformExpression(B);
 	if(ExpressionA && ExpressionB)
 	{
+		FAddUniformExpressionScope Scope(this);
 		if (TypeA == MCT_Float && TypeB == MCT_Float)
 		{
-			return AddUniformExpression(new FMaterialUniformExpressionFoldedMath(ExpressionA,ExpressionB,FMO_Mul),MCT_Float,TEXT("(%s * %s)"),*GetParameterCode(A),*GetParameterCode(B));
+			return AddUniformExpression(Scope, new FMaterialUniformExpressionFoldedMath(ExpressionA,ExpressionB,FMO_Mul),MCT_Float,TEXT("(%s * %s)"),*GetParameterCode(A),*GetParameterCode(B));
 		}
 		else
 		{
 			if (TypeA == TypeB)
 			{
-				return AddUniformExpression(new FMaterialUniformExpressionFoldedMath(ExpressionA,ExpressionB,FMO_Dot,TypeA),MCT_Float,TEXT("dot(%s,%s)"),*GetParameterCode(A),*GetParameterCode(B));
+				return AddUniformExpression(Scope, new FMaterialUniformExpressionFoldedMath(ExpressionA,ExpressionB,FMO_Dot,TypeA),MCT_Float,TEXT("dot(%s,%s)"),*GetParameterCode(A),*GetParameterCode(B));
 			}
 			else
 			{
 				// Promote scalar (or truncate the bigger type)
 				if (TypeA == MCT_Float || (TypeB != MCT_Float && GetNumComponents(TypeA) > GetNumComponents(TypeB)))
 				{
-					return AddUniformExpression(new FMaterialUniformExpressionFoldedMath(ExpressionA,ExpressionB,FMO_Dot,TypeB),MCT_Float,TEXT("dot(%s,%s)"),*CoerceParameter(A, TypeB),*GetParameterCode(B));
+					return AddUniformExpression(Scope, new FMaterialUniformExpressionFoldedMath(ExpressionA,ExpressionB,FMO_Dot,TypeB),MCT_Float,TEXT("dot(%s,%s)"),*CoerceParameter(A, TypeB),*GetParameterCode(B));
 				}
 				else
 				{
-					return AddUniformExpression(new FMaterialUniformExpressionFoldedMath(ExpressionA,ExpressionB,FMO_Dot,TypeA),MCT_Float,TEXT("dot(%s,%s)"),*GetParameterCode(A),*CoerceParameter(B, TypeA));
+					return AddUniformExpression(Scope, new FMaterialUniformExpressionFoldedMath(ExpressionA,ExpressionB,FMO_Dot,TypeA),MCT_Float,TEXT("dot(%s,%s)"),*GetParameterCode(A),*CoerceParameter(B, TypeA));
 				}
 			}
 		}
@@ -9444,7 +9512,8 @@ int32 FHLSLMaterialTranslator::Cross(int32 A,int32 B)
 			return Errorf(TEXT("Cross product requires 3-component vector input."));
 		}
 
-		return AddUniformExpression(new FMaterialUniformExpressionFoldedMath(GetParameterUniformExpression(A),GetParameterUniformExpression(B),FMO_Cross,ResultType),MCT_Float3,TEXT("cross(%s,%s)"),*GetParameterCode(A),*GetParameterCode(B));
+		FAddUniformExpressionScope Scope(this);
+		return AddUniformExpression(Scope, new FMaterialUniformExpressionFoldedMath(GetParameterUniformExpression(A),GetParameterUniformExpression(B),FMO_Cross,ResultType),MCT_Float3,TEXT("cross(%s,%s)"),*GetParameterCode(A),*GetParameterCode(B));
 	}
 	else if(IsAnalyticDerivEnabled())
 	{
@@ -9536,7 +9605,8 @@ int32 FHLSLMaterialTranslator::Exponential(int32 X)
 
 	if(GetParameterUniformExpression(X))
 	{
-		return AddUniformExpression(new FMaterialUniformExpressionExponential(GetParameterUniformExpression(X)), GetParameterType(X), TEXT("exp(%s)"), *GetParameterCode(X));
+		FAddUniformExpressionScope Scope(this);
+		return AddUniformExpression(Scope, new FMaterialUniformExpressionExponential(GetParameterUniformExpression(X)), GetParameterType(X), TEXT("exp(%s)"), *GetParameterCode(X));
 	}
 	else if(IsAnalyticDerivEnabled())
 	{
@@ -9557,7 +9627,8 @@ int32 FHLSLMaterialTranslator::Exponential2(int32 X)
 
 	if(GetParameterUniformExpression(X))
 	{
-		return AddUniformExpression(new FMaterialUniformExpressionExponential2(GetParameterUniformExpression(X)), GetParameterType(X), TEXT("exp2(%s)"), *GetParameterCode(X));
+		FAddUniformExpressionScope Scope(this);
+		return AddUniformExpression(Scope, new FMaterialUniformExpressionExponential2(GetParameterUniformExpression(X)), GetParameterType(X), TEXT("exp2(%s)"), *GetParameterCode(X));
 	}
 	else if(IsAnalyticDerivEnabled())
 	{
@@ -9578,7 +9649,8 @@ int32 FHLSLMaterialTranslator::Logarithm(int32 X)
 
 	if(GetParameterUniformExpression(X))
 	{
-		return AddUniformExpression(new FMaterialUniformExpressionLogarithm(GetParameterUniformExpression(X)), GetParameterType(X), TEXT("log(%s)"), *GetParameterCode(X));
+		FAddUniformExpressionScope Scope(this);
+		return AddUniformExpression(Scope, new FMaterialUniformExpressionLogarithm(GetParameterUniformExpression(X)), GetParameterType(X), TEXT("log(%s)"), *GetParameterCode(X));
 	}
 	else if(IsAnalyticDerivEnabled())
 	{
@@ -9599,7 +9671,8 @@ int32 FHLSLMaterialTranslator::Logarithm2(int32 X)
 
 	if(GetParameterUniformExpression(X))
 	{
-		return AddUniformExpression(new FMaterialUniformExpressionLogarithm2(GetParameterUniformExpression(X)),GetParameterType(X),TEXT("log2(%s)"),*GetParameterCode(X));
+		FAddUniformExpressionScope Scope(this);
+		return AddUniformExpression(Scope, new FMaterialUniformExpressionLogarithm2(GetParameterUniformExpression(X)),GetParameterType(X),TEXT("log2(%s)"),*GetParameterCode(X));
 	}
 	else if (IsAnalyticDerivEnabled())
 	{
@@ -9620,7 +9693,8 @@ int32 FHLSLMaterialTranslator::Logarithm10(int32 X)
 
 	if(GetParameterUniformExpression(X))
 	{
-		return AddUniformExpression(new FMaterialUniformExpressionLogarithm10(GetParameterUniformExpression(X)),GetParameterType(X),TEXT("log10(%s)"),*GetParameterCode(X));
+		FAddUniformExpressionScope Scope(this);
+		return AddUniformExpression(Scope, new FMaterialUniformExpressionLogarithm10(GetParameterUniformExpression(X)),GetParameterType(X),TEXT("log10(%s)"),*GetParameterCode(X));
 	}
 	else if (IsAnalyticDerivEnabled())
 	{
@@ -9641,7 +9715,8 @@ int32 FHLSLMaterialTranslator::SquareRoot(int32 X)
 
 	if(GetParameterUniformExpression(X))
 	{
-		return AddUniformExpression(new FMaterialUniformExpressionSquareRoot(GetParameterUniformExpression(X)),GetParameterType(X),TEXT("sqrt(%s)"),*GetParameterCode(X));
+		FAddUniformExpressionScope Scope(this);
+		return AddUniformExpression(Scope, new FMaterialUniformExpressionSquareRoot(GetParameterUniformExpression(X)),GetParameterType(X),TEXT("sqrt(%s)"),*GetParameterCode(X));
 	}
 	else
 	{
@@ -9665,7 +9740,8 @@ int32 FHLSLMaterialTranslator::Length(int32 X)
 
 	if (GetParameterUniformExpression(X))
 	{
-		return AddUniformExpression(new FMaterialUniformExpressionLength(GetParameterUniformExpression(X), GetParameterType(X)), MCT_Float, TEXT("length(%s)"), *GetParameterCode(X));
+		FAddUniformExpressionScope Scope(this);
+		return AddUniformExpression(Scope, new FMaterialUniformExpressionLength(GetParameterUniformExpression(X), GetParameterType(X)), MCT_Float, TEXT("length(%s)"), *GetParameterCode(X));
 	}
 	else
 	{
@@ -9691,7 +9767,8 @@ int32 FHLSLMaterialTranslator::Normalize(int32 X)
 
 	if (GetParameterUniformExpression(X))
 	{
-		return AddUniformExpression(new FMaterialUniformExpressionNormalize(GetParameterUniformExpression(X)), ResultType, TEXT("normalize(%s)"), *GetParameterCode(X));
+		FAddUniformExpressionScope Scope(this);
+		return AddUniformExpression(Scope, new FMaterialUniformExpressionNormalize(GetParameterUniformExpression(X)), ResultType, TEXT("normalize(%s)"), *GetParameterCode(X));
 	}
 	else
 	{
@@ -10095,7 +10172,8 @@ int32 FHLSLMaterialTranslator::Min(int32 A,int32 B)
 			return ConstResultValue(TypeA, Result);
 		}
 		
-		return AddUniformExpression(new FMaterialUniformExpressionMin(ExpressionA, ExpressionB),GetParameterType(A),TEXT("min(%s,%s)"),*GetParameterCode(A),*CoerceParameter(B,GetParameterType(A)));
+		FAddUniformExpressionScope Scope(this);
+		return AddUniformExpression(Scope, new FMaterialUniformExpressionMin(ExpressionA, ExpressionB),GetParameterType(A),TEXT("min(%s,%s)"),*GetParameterCode(A),*CoerceParameter(B,GetParameterType(A)));
 	}
 	else
 	{
@@ -10141,7 +10219,8 @@ int32 FHLSLMaterialTranslator::Max(int32 A,int32 B)
 			return ConstResultValue(TypeA, Result);
 		}
 
-		return AddUniformExpression(new FMaterialUniformExpressionMax(ExpressionA, ExpressionB),GetParameterType(A),TEXT("max(%s,%s)"),*GetParameterCode(A),*CoerceParameter(B,GetParameterType(A)));
+		FAddUniformExpressionScope Scope(this);
+		return AddUniformExpression(Scope, new FMaterialUniformExpressionMax(ExpressionA, ExpressionB),GetParameterType(A),TEXT("max(%s,%s)"),*GetParameterCode(A),*CoerceParameter(B,GetParameterType(A)));
 	}
 	else
 	{
@@ -10190,7 +10269,8 @@ int32 FHLSLMaterialTranslator::Clamp(int32 X, int32 A, int32 B)
 
 	if (UniformX && UniformA && UniformB)
 	{
-		return AddUniformExpression(new FMaterialUniformExpressionClamp(UniformX, UniformA, UniformB), GetParameterType(X), TEXT("min(max(%s,%s),%s)"), *GetParameterCode(X), *CoerceParameter(A, GetParameterType(X)), *CoerceParameter(B, GetParameterType(X)));
+		FAddUniformExpressionScope Scope(this);
+		return AddUniformExpression(Scope, new FMaterialUniformExpressionClamp(UniformX, UniformA, UniformB), GetParameterType(X), TEXT("min(max(%s,%s),%s)"), *GetParameterCode(X), *CoerceParameter(A, GetParameterType(X)), *CoerceParameter(B, GetParameterType(X)));
 	}
 
 	if (IsAnalyticDerivEnabled())
@@ -10227,7 +10307,8 @@ int32 FHLSLMaterialTranslator::Saturate(int32 X)
 
 	if(UniformX)
 	{
-		return AddUniformExpression(new FMaterialUniformExpressionSaturate(UniformX),GetParameterType(X),TEXT("saturate(%s)"),*GetParameterCode(X));
+		FAddUniformExpressionScope Scope(this);
+		return AddUniformExpression(Scope, new FMaterialUniformExpressionSaturate(UniformX),GetParameterType(X),TEXT("saturate(%s)"),*GetParameterCode(X));
 	}
 	else
 	{
@@ -10336,7 +10417,6 @@ int32 FHLSLMaterialTranslator::ComponentMask(int32 Vector,bool R,bool G,bool B,b
 		A ? ((VectorType == MCT_Float || VectorType == MCT_LWCScalar) ? TEXT("r") : TEXT("a")) : TEXT("")
 		);
 
-	FString SourceString = GetParameterCode(Vector);
 	auto* Expression = GetParameterUniformExpression(Vector);
 	if (Expression)
 	{
@@ -10345,15 +10425,19 @@ int32 FHLSLMaterialTranslator::ComponentMask(int32 Vector,bool R,bool G,bool B,b
 		{
 			Mask[Index] = SwizzleComponentToIndex(MaskString[Index]);
 		}
+		FAddUniformExpressionScope Scope(this);
 		return AddUniformExpression(
+			Scope,
 			new FMaterialUniformExpressionComponentSwizzle(Expression, Mask[0], Mask[1], Mask[2], Mask[3]),
 			ResultType,
 			TEXT("%s.%s"),
-			*SourceString,
+			*GetParameterCode(Vector),
 			*MaskString
 			);
 	}
 
+	FString SourceString = GetParameterCode(Vector);
+
 	const EDerivativeStatus VectorDerivStatus = GetDerivativeStatus(Vector);
 	FString CodeFinite;
 	if (bIsLWC)
@@ -10423,7 +10507,8 @@ int32 FHLSLMaterialTranslator::AppendVector(int32 A,int32 B)
 
 	if(GetParameterUniformExpression(A) && GetParameterUniformExpression(B))
 	{
-		return AddUniformExpression(new FMaterialUniformExpressionAppendVector(GetParameterUniformExpression(A),GetParameterUniformExpression(B),GetNumComponents(GetParameterType(A))),ResultType,TEXT("MaterialFloat%u(%s,%s)"),NumResultComponents,*GetParameterCode(A),*GetParameterCode(B));
+		FAddUniformExpressionScope Scope(this);
+		return AddUniformExpression(Scope, new FMaterialUniformExpressionAppendVector(GetParameterUniformExpression(A),GetParameterUniformExpression(B),GetNumComponents(GetParameterType(A))),ResultType,TEXT("MaterialFloat%u(%s,%s)"),NumResultComponents,*GetParameterCode(A),*GetParameterCode(B));
 	}
 	else
 	{
@@ -15585,7 +15670,8 @@ int32 FHLSLMaterialTranslator::SparseVolumeTexture(USparseVolumeTexture* Texture
 	TextureReferenceIndex = Material->GetReferencedTextures().Find(Texture);
 	checkf(TextureReferenceIndex != INDEX_NONE, TEXT("Material expression called Compiler->SparseVolumeTexture() without implementing UMaterialExpression::GetReferencedTexture properly"));
 
-	return AddUniformExpression(new FMaterialUniformExpressionTexture(TextureReferenceIndex, SamplerType), MCT_SparseVolumeTexture, TEXT(""));
+	FAddUniformExpressionScope Scope(this);
+	return AddUniformExpression(Scope, new FMaterialUniformExpressionTexture(TextureReferenceIndex, SamplerType), MCT_SparseVolumeTexture, TEXT(""));
 }
 
 int32 FHLSLMaterialTranslator::SparseVolumeTextureParameter(FName ParameterName, USparseVolumeTexture* InDefaultTexture, int32& TextureReferenceIndex, EMaterialSamplerType SamplerType)
@@ -15605,19 +15691,22 @@ int32 FHLSLMaterialTranslator::SparseVolumeTextureParameter(FName ParameterName,
 	FMaterialParameterInfo ParameterInfo = GetParameterAssociationInfo();
 	ParameterInfo.Name = ParameterName;
 
-	return AddUniformExpression(new FMaterialUniformExpressionTextureParameter(ParameterInfo, TextureReferenceIndex, SamplerType), MCT_SparseVolumeTexture, TEXT(""));
+	FAddUniformExpressionScope Scope(this);
+	return AddUniformExpression(Scope, new FMaterialUniformExpressionTextureParameter(ParameterInfo, TextureReferenceIndex, SamplerType), MCT_SparseVolumeTexture, TEXT(""));
 }
 
 int32 FHLSLMaterialTranslator::SparseVolumeTextureUniform(int32 TextureIndex, int32 VectorIndex, UE::Shader::EValueType Type)
 {
-	return AddUniformExpression(new FMaterialUniformExpressionSparseVolumeTextureUniform(TextureIndex, VectorIndex), GetMaterialValueType(Type), TEXT(""));
+	FAddUniformExpressionScope Scope(this);
+	return AddUniformExpression(Scope, new FMaterialUniformExpressionSparseVolumeTextureUniform(TextureIndex, VectorIndex), GetMaterialValueType(Type), TEXT(""));
 }
 
 int32 FHLSLMaterialTranslator::SparseVolumeTextureUniformParameter(FName ParameterName, int32 TextureIndex, int32 VectorIndex, UE::Shader::EValueType Type)
 {
 	FMaterialParameterInfo ParameterInfo = GetParameterAssociationInfo();
 	ParameterInfo.Name = ParameterName;
-	return AddUniformExpression(new FMaterialUniformExpressionSparseVolumeTextureUniform(ParameterInfo, TextureIndex, VectorIndex), GetMaterialValueType(Type), TEXT(""));
+	FAddUniformExpressionScope Scope(this);
+	return AddUniformExpression(Scope, new FMaterialUniformExpressionSparseVolumeTextureUniform(ParameterInfo, TextureIndex, VectorIndex), GetMaterialValueType(Type), TEXT(""));
 }
 
 static const TCHAR* GetSparseVolumeTextureAddressMode(TextureAddress Address)
diff --git a/Engine/Source/Runtime/Engine/Private/Materials/HLSLMaterialTranslator.h b/Engine/Source/Runtime/Engine/Private/Materials/HLSLMaterialTranslator.h
--- a/Engine/Source/Runtime/Engine/Private/Materials/HLSLMaterialTranslator.h
+++ b/Engine/Source/Runtime/Engine/Private/Materials/HLSLMaterialTranslator.h
@@ -55,6 +55,7 @@
 #endif
 
 class Error;
+class FAddUniformExpressionScope;
 
 /**
  * Returns whether the specified class of material expression is permitted.
@@ -416,6 +417,11 @@ protected:
 	/** true if PerInstanceFadeAmount expression is used */
 	uint32 bUsesPerInstanceFadeAmount : 1;
 
+	uint32 bCullIntermediateUniformExpressions : 1;
+
+	/** Incremented and decremented by FAddUniformExpressionScope.  See comments on that class below */
+	int32 AddingUniformExpression;
+
 	/** Tracks the texture coordinates used by this material. */
 	TBitArray<> AllocatedUserTexCoords;
 	/** Tracks the texture coordinates used by the vertex shader in this material. */
@@ -765,7 +771,7 @@ protected:
 	int32 AddUniformExpressionInner(uint64 Hash, FMaterialUniformExpression* UniformExpression, EMaterialValueType Type, const TCHAR* FormattedCode);
 
 	// AddUniformExpression - Adds an input to the Code array and returns its index.
-	int32 AddUniformExpression(FMaterialUniformExpression* UniformExpression, EMaterialValueType Type, const TCHAR* Format, ...);
+	int32 AddUniformExpression(FAddUniformExpressionScope& Scope, FMaterialUniformExpression* UniformExpression, EMaterialValueType Type, const TCHAR* Format, ...);
 
 	// AccessUniformExpression - Adds code to access the value of a uniform expression to the Code array and returns its index.
 	int32 AccessUniformExpression(int32 Index);
@@ -1395,6 +1401,42 @@ protected:
 
 	/** The output material shader defines */
 	TUniquePtr<FEnvironmentDefines> EnvironmentDefines;
+
+	friend class FAddUniformExpressionScope;
+};
+
+/**
+ * The purpose of this class is to avoid generating unused preshaders, meaning preshaders not actually fetched in the HLSL code, saving performance.
+ * The translator normally generates preshaders when the function "AccessUniformExpression" is called the first time on a uniform expression.  We want
+ * to detect in that function whether the uniform expression access is from HLSL code or a parent uniform expression.  In the latter case, there is no
+ * need to generate a preshader.  The FAddUniformExpressionScope class provides that context.
+ *
+ * This class sets a flag in the translator when we enter a code path where we are adding a uniform expression (technically uses a count, so scopes can
+ * be nested).  To enforce usage of this scope class everywhere, the translator's AddUniformExpression function requires it to be passed in.  The reason
+ * we need a separate scope class, rather than setting the state in the AddUniformExpression function, is that the code that needs the context is actually
+ * the arguments to the AddUniformExpression function, not the function itself.  Calls to GetParameterCode and other utility functions that call
+ * GetParameterCode (CoerceParameter, ValidCast, ForceCast) are what call "AccessUniformExpression", and need to know whether they are being passed to
+ * "AddUniformExpression", versus functions like AddCodeChunk that emit the HLSL.
+ *
+ * If you make a mistake and place the scope somewhere it doesn't belong, it will generate stub code that will result in shader compile errors on
+ * symbols that look like "UniformStub[$0]".  Most likely, any new opcodes will be cut-and-paste from existing code, and thus already have the scope in
+ * the right place, so bugs should be rare.
+ */
+class FAddUniformExpressionScope
+{
+public:
+	FAddUniformExpressionScope(FHLSLMaterialTranslator* InTranslator)
+		: Translator(InTranslator)
+	{
+		Translator->AddingUniformExpression++;
+	}
+	~FAddUniformExpressionScope()
+	{
+		Translator->AddingUniformExpression--;
+		check(Translator->AddingUniformExpression >= 0);
+	}
+private:
+	FHLSLMaterialTranslator* Translator;
 };
 
 #endif // WITH_EDITORONLY_DATA
diff --git a/Engine/Source/Runtime/Engine/Private/Materials/MaterialUniformExpressions.h b/Engine/Source/Runtime/Engine/Private/Materials/MaterialUniformExpressions.h
--- a/Engine/Source/Runtime/Engine/Private/Materials/MaterialUniformExpressions.h
+++ b/Engine/Source/Runtime/Engine/Private/Materials/MaterialUniformExpressions.h
@@ -68,8 +68,13 @@ public:
 
 	virtual void GetNumberValue(const struct FMaterialRenderContext& Context, FLinearColor& OutValue) const;
 
+	virtual TArrayView<const FMaterialUniformExpression*> GetChildren() const { return TArrayView<const FMaterialUniformExpression*>(); }
+
 	/** Offset of this uniform, within the shader's uniform buffer array */
 	int32 UniformOffset = INDEX_NONE;
+
+	/** Index of this uniform in the material translator's list of unique expressions */
+	int32 UniformIndex = INDEX_NONE;
 };
 
 /**
@@ -529,6 +534,11 @@ public:
 		return X->IsIdentical(OtherSine->X) && bIsCosine == OtherSine->bIsCosine;
 	}
 
+	virtual TArrayView<const FMaterialUniformExpression*> GetChildren() const override
+	{
+		return TArrayView<const FMaterialUniformExpression*>((const FMaterialUniformExpression**)&X, 1);
+	}
+
 private:
 	TRefCountPtr<FMaterialUniformExpression> X;
 	bool bIsCosine;
@@ -602,6 +612,11 @@ public:
 		return X->IsIdentical(OtherTrig->X) && Y->IsIdentical(OtherTrig->Y) && Op == OtherTrig->Op;
 	}
 
+	virtual TArrayView<const FMaterialUniformExpression*> GetChildren() const override
+	{
+		return TArrayView<const FMaterialUniformExpression*>((const FMaterialUniformExpression**)&X, 2);
+	}
+
 private:
 	TRefCountPtr<FMaterialUniformExpression> X;
 	TRefCountPtr<FMaterialUniformExpression> Y;
@@ -640,6 +655,11 @@ public:
 		return X->IsIdentical(OtherSqrt->X);
 	}
 
+	virtual TArrayView<const FMaterialUniformExpression*> GetChildren() const override
+	{
+		return TArrayView<const FMaterialUniformExpression*>((const FMaterialUniformExpression**)&X, 1);
+	}
+
 private:
 	TRefCountPtr<FMaterialUniformExpression> X;
 };
@@ -676,6 +696,11 @@ public:
 		return X->IsIdentical(OtherRcp->X);
 	}
 
+	virtual TArrayView<const FMaterialUniformExpression*> GetChildren() const override
+	{
+		return TArrayView<const FMaterialUniformExpression*>((const FMaterialUniformExpression**)&X, 1);
+	}
+
 private:
 	TRefCountPtr<FMaterialUniformExpression> X;
 };
@@ -713,6 +738,11 @@ public:
 		return X->IsIdentical(OtherSqrt->X) && ValueType == OtherSqrt->ValueType;
 	}
 
+	virtual TArrayView<const FMaterialUniformExpression*> GetChildren() const override
+	{
+		return TArrayView<const FMaterialUniformExpression*>((const FMaterialUniformExpression**)&X, 1);
+	}
+
 private:
 	TRefCountPtr<FMaterialUniformExpression> X;
 	uint32 ValueType;
@@ -749,6 +779,11 @@ public:
 		return X->IsIdentical(OtherSqrt->X);
 	}
 
+	virtual TArrayView<const FMaterialUniformExpression*> GetChildren() const override
+	{
+		return TArrayView<const FMaterialUniformExpression*>((const FMaterialUniformExpression**)&X, 1);
+	}
+
 private:
 	TRefCountPtr<FMaterialUniformExpression> X;
 };
@@ -787,6 +822,11 @@ public:
 		return X->IsIdentical(OtherExp->X);
 	}
 
+	virtual TArrayView<const FMaterialUniformExpression*> GetChildren() const override
+	{
+		return TArrayView<const FMaterialUniformExpression*>((const FMaterialUniformExpression**)&X, 1);
+	}
+
 private:
 	TRefCountPtr<FMaterialUniformExpression> X;
 };
@@ -825,6 +865,11 @@ public:
 		return X->IsIdentical(OtherExp2->X);
 	}
 
+	virtual TArrayView<const FMaterialUniformExpression*> GetChildren() const override
+	{
+		return TArrayView<const FMaterialUniformExpression*>((const FMaterialUniformExpression**)&X, 1);
+	}
+
 private:
 	TRefCountPtr<FMaterialUniformExpression> X;
 };
@@ -863,6 +908,11 @@ public:
 		return X->IsIdentical(OtherLog->X);
 	}
 
+	virtual TArrayView<const FMaterialUniformExpression*> GetChildren() const override
+	{
+		return TArrayView<const FMaterialUniformExpression*>((const FMaterialUniformExpression**)&X, 1);
+	}
+
 private:
 	TRefCountPtr<FMaterialUniformExpression> X;
 };
@@ -900,6 +950,11 @@ public:
 		return X->IsIdentical(OtherLog->X);
 	}
 
+	virtual TArrayView<const FMaterialUniformExpression*> GetChildren() const override
+	{
+		return TArrayView<const FMaterialUniformExpression*>((const FMaterialUniformExpression**)&X, 1);
+	}
+
 private:
 	TRefCountPtr<FMaterialUniformExpression> X;
 };
@@ -937,6 +992,11 @@ public:
 		return X->IsIdentical(OtherLog->X);
 	}
 
+	virtual TArrayView<const FMaterialUniformExpression*> GetChildren() const override
+	{
+		return TArrayView<const FMaterialUniformExpression*>((const FMaterialUniformExpression**)&X, 1);
+	}
+
 private:
 	TRefCountPtr<FMaterialUniformExpression> X;
 };
@@ -998,6 +1058,11 @@ public:
 		return A->IsIdentical(OtherMath->A) && B->IsIdentical(OtherMath->B) && Op == OtherMath->Op && ValueType == OtherMath->ValueType;
 	}
 
+	virtual TArrayView<const FMaterialUniformExpression*> GetChildren() const override
+	{
+		return TArrayView<const FMaterialUniformExpression*>((const FMaterialUniformExpression**)&A, 2);
+	}
+
 private:
 	TRefCountPtr<FMaterialUniformExpression> A;
 	TRefCountPtr<FMaterialUniformExpression> B;
@@ -1038,6 +1103,11 @@ public:
 		return X->IsIdentical(OtherPeriodic->X);
 	}
 
+	virtual TArrayView<const FMaterialUniformExpression*> GetChildren() const override
+	{
+		return TArrayView<const FMaterialUniformExpression*>((const FMaterialUniformExpression**)&X, 1);
+	}
+
 private:
 	TRefCountPtr<FMaterialUniformExpression> X;
 };
@@ -1077,6 +1147,11 @@ public:
 		return A->IsIdentical(OtherAppend->A) && B->IsIdentical(OtherAppend->B) && NumComponentsA == OtherAppend->NumComponentsA;
 	}
 
+	virtual TArrayView<const FMaterialUniformExpression*> GetChildren() const override
+	{
+		return TArrayView<const FMaterialUniformExpression*>((const FMaterialUniformExpression**)&A, 2);
+	}
+
 private:
 	TRefCountPtr<FMaterialUniformExpression> A;
 	TRefCountPtr<FMaterialUniformExpression> B;
@@ -1117,6 +1192,11 @@ public:
 		return A->IsIdentical(OtherMin->A) && B->IsIdentical(OtherMin->B);
 	}
 
+	virtual TArrayView<const FMaterialUniformExpression*> GetChildren() const override
+	{
+		return TArrayView<const FMaterialUniformExpression*>((const FMaterialUniformExpression**)&A, 2);
+	}
+
 private:
 	TRefCountPtr<FMaterialUniformExpression> A;
 	TRefCountPtr<FMaterialUniformExpression> B;
@@ -1156,6 +1236,11 @@ public:
 		return A->IsIdentical(OtherMax->A) && B->IsIdentical(OtherMax->B);
 	}
 
+	virtual TArrayView<const FMaterialUniformExpression*> GetChildren() const override
+	{
+		return TArrayView<const FMaterialUniformExpression*>((const FMaterialUniformExpression**)&A, 2);
+	}
+
 private:
 	TRefCountPtr<FMaterialUniformExpression> A;
 	TRefCountPtr<FMaterialUniformExpression> B;
@@ -1197,6 +1282,11 @@ public:
 		return Input->IsIdentical(OtherClamp->Input) && Min->IsIdentical(OtherClamp->Min) && Max->IsIdentical(OtherClamp->Max);
 	}
 
+	virtual TArrayView<const FMaterialUniformExpression*> GetChildren() const override
+	{
+		return TArrayView<const FMaterialUniformExpression*>((const FMaterialUniformExpression**)&Input, 3);
+	}
+
 private:
 	TRefCountPtr<FMaterialUniformExpression> Input;
 	TRefCountPtr<FMaterialUniformExpression> Min;
@@ -1235,6 +1325,11 @@ public:
 		return Input->IsIdentical(OtherClamp->Input);
 	}
 
+	virtual TArrayView<const FMaterialUniformExpression*> GetChildren() const override
+	{
+		return TArrayView<const FMaterialUniformExpression*>((const FMaterialUniformExpression**)&Input, 1);
+	}
+
 private:
 	TRefCountPtr<FMaterialUniformExpression> Input;
 };
@@ -1303,6 +1398,11 @@ public:
 			IndexA == OtherSwizzle->IndexA;
 	}
 
+	virtual TArrayView<const FMaterialUniformExpression*> GetChildren() const override
+	{
+		return TArrayView<const FMaterialUniformExpression*>((const FMaterialUniformExpression**)&X, 1);
+	}
+
 private:
 	TRefCountPtr<FMaterialUniformExpression> X;
 	int8 IndexR;
@@ -1344,6 +1444,11 @@ public:
 		return X->IsIdentical(OtherFloor->X);
 	}
 
+	virtual TArrayView<const FMaterialUniformExpression*> GetChildren() const override
+	{
+		return TArrayView<const FMaterialUniformExpression*>((const FMaterialUniformExpression**)&X, 1);
+	}
+
 private:
 	TRefCountPtr<FMaterialUniformExpression> X;
 };
@@ -1380,6 +1485,11 @@ public:
 		return X->IsIdentical(OtherCeil->X);
 	}
 
+	virtual TArrayView<const FMaterialUniformExpression*> GetChildren() const override
+	{
+		return TArrayView<const FMaterialUniformExpression*>((const FMaterialUniformExpression**)&X, 1);
+	}
+
 private:
 	TRefCountPtr<FMaterialUniformExpression> X;
 };
@@ -1416,6 +1526,11 @@ public:
 		return X->IsIdentical(OtherRound->X);
 	}
 
+	virtual TArrayView<const FMaterialUniformExpression*> GetChildren() const override
+	{
+		return TArrayView<const FMaterialUniformExpression*>((const FMaterialUniformExpression**)&X, 1);
+	}
+
 private:
 	TRefCountPtr<FMaterialUniformExpression> X;
 };
@@ -1452,6 +1567,11 @@ public:
 		return X->IsIdentical(OtherTrunc->X);
 	}
 
+	virtual TArrayView<const FMaterialUniformExpression*> GetChildren() const override
+	{
+		return TArrayView<const FMaterialUniformExpression*>((const FMaterialUniformExpression**)&X, 1);
+	}
+
 private:
 	TRefCountPtr<FMaterialUniformExpression> X;
 };
@@ -1488,6 +1608,11 @@ public:
 		return X->IsIdentical(OtherSign->X);
 	}
 
+	virtual TArrayView<const FMaterialUniformExpression*> GetChildren() const override
+	{
+		return TArrayView<const FMaterialUniformExpression*>((const FMaterialUniformExpression**)&X, 1);
+	}
+
 private:
 	TRefCountPtr<FMaterialUniformExpression> X;
 };
@@ -1524,6 +1649,11 @@ public:
 		return X->IsIdentical(OtherFrac->X);
 	}
 
+	virtual TArrayView<const FMaterialUniformExpression*> GetChildren() const override
+	{
+		return TArrayView<const FMaterialUniformExpression*>((const FMaterialUniformExpression**)&X, 1);
+	}
+
 private:
 	TRefCountPtr<FMaterialUniformExpression> X;
 };
@@ -1562,6 +1692,11 @@ public:
 		return A->IsIdentical(OtherMax->A) && B->IsIdentical(OtherMax->B);
 	}
 
+	virtual TArrayView<const FMaterialUniformExpression*> GetChildren() const override
+	{
+		return TArrayView<const FMaterialUniformExpression*>((const FMaterialUniformExpression**)&A, 2);
+	}
+
 private:
 	TRefCountPtr<FMaterialUniformExpression> A;
 	TRefCountPtr<FMaterialUniformExpression> B;
@@ -1600,6 +1735,11 @@ public:
 		return X->IsIdentical(OtherAbs->X);
 	}
 
+	virtual TArrayView<const FMaterialUniformExpression*> GetChildren() const override
+	{
+		return TArrayView<const FMaterialUniformExpression*>((const FMaterialUniformExpression**)&X, 1);
+	}
+
 private:
 	TRefCountPtr<FMaterialUniformExpression> X;
 };
