Cameras: correctly initialize bOverrideAspectRatioAxisConstraint.

diff --git a/Engine/Source/Runtime/Engine/Private/Camera/CameraComponent.cpp b/Engine/Source/Runtime/Engine/Private/Camera/CameraComponent.cpp
--- a/Engine/Source/Runtime/Engine/Private/Camera/CameraComponent.cpp
+++ b/Engine/Source/Runtime/Engine/Private/Camera/CameraComponent.cpp
@@ -48,6 +48,7 @@ UCameraComponent::UCameraComponent(const FObjectInitializer& ObjectInitializer)
 	OrthoNearClipPlane = 0.0f;
 	OrthoFarClipPlane = UE_OLD_WORLD_MAX;
 	bConstrainAspectRatio = false;
+	bOverrideAspectRatioAxisConstraint = false;
 	bUseFieldOfViewForLOD = true;
 	PostProcessBlendWeight = 1.0f;
 	bUsePawnControlRotation = false;
