PR #10885: Fix IsValidLandingSpot for Large World Coordinates

#rnx

diff --git a/Engine/Source/Runtime/Engine/Private/CharacterMovementComponentAsync.cpp b/Engine/Source/Runtime/Engine/Private/CharacterMovementComponentAsync.cpp
--- a/Engine/Source/Runtime/Engine/Private/CharacterMovementComponentAsync.cpp
+++ b/Engine/Source/Runtime/Engine/Private/CharacterMovementComponentAsync.cpp
@@ -4033,7 +4033,7 @@ bool FCharacterMovementComponentAsyncInput::IsValidLandingSpot(const FVector& Ca
 		//CharacterOwner->GetCapsuleComponent()->GetScaledCapsuleSize(PawnRadius, PawnHalfHeight);
 
 		// Reject hits that are above our lower hemisphere (can happen when sliding down a vertical surface).
-		const float LowerHemisphereZ = Hit.Location.Z - PawnHalfHeight + PawnRadius;
+		const FVector::FReal LowerHemisphereZ = Hit.Location.Z - PawnHalfHeight + PawnRadius;
 		if (Hit.ImpactPoint.Z >= LowerHemisphereZ)
 		{
 			return false;
