Deprecate an unused flag from FLinearDriveConstraint. Add SetDriveParams to ConstraintInstance to allow the most common changes to be made without needing to make multiple calls and locks. Rename the FPBDJointConstraintHandle version of this so it matches.

diff --git a/Engine/Source/Runtime/Engine/Classes/PhysicsEngine/ConstraintDrives.h b/Engine/Source/Runtime/Engine/Classes/PhysicsEngine/ConstraintDrives.h
--- a/Engine/Source/Runtime/Engine/Classes/PhysicsEngine/ConstraintDrives.h
+++ b/Engine/Source/Runtime/Engine/Classes/PhysicsEngine/ConstraintDrives.h
@@ -65,7 +65,6 @@ USTRUCT()
 struct FLinearDriveConstraint
 {
 	GENERATED_BODY()
-	
 
 	/** Target position the linear drive.*/
 	UPROPERTY(EditAnywhere, Category = LinearMotor)
@@ -84,8 +83,12 @@ struct FLinearDriveConstraint
 	UPROPERTY(EditAnywhere, Category = LinearMotor)
 	FConstraintDrive ZDrive;
 
-	UPROPERTY(EditAnywhere, Category = LinearMotor)
-	uint8 bEnablePositionDrive : 1;
+#if WITH_EDITORONLY_DATA
+	PRAGMA_DISABLE_DEPRECATION_WARNINGS
+	UPROPERTY(config, meta = (DeprecatedProperty, DeprecationMessage = "Enable/disable of drives is done inside the individual constraint drives."))
+	uint8 bEnablePositionDrive_DEPRECATED : 1;
+	PRAGMA_ENABLE_DEPRECATION_WARNINGS
+#endif
 
 	ENGINE_API FLinearDriveConstraint();
 
diff --git a/Engine/Source/Runtime/Engine/Classes/PhysicsEngine/ConstraintInstance.h b/Engine/Source/Runtime/Engine/Classes/PhysicsEngine/ConstraintInstance.h
--- a/Engine/Source/Runtime/Engine/Classes/PhysicsEngine/ConstraintInstance.h
+++ b/Engine/Source/Runtime/Engine/Classes/PhysicsEngine/ConstraintInstance.h
@@ -866,6 +866,18 @@ public:
 		return ProfileInstance.AngularDrive.AngularDriveMode;
 	}
 
+	/** 
+	 * This allows the most common drive parameters to be set in one call, avoiding lock overheads etc.
+	 * Note that the individual drive elements will be enabled/disabled depending on the strength/damping
+	 * values passed in.
+	 * 
+	 * Angular values are passed in as (swing, twist, slerp)
+	 */
+	ENGINE_API void SetDriveParams(
+		const FVector& InPositionStrength, const FVector& InVelocityStrength, const FVector& InForceLimit,
+		const FVector& InAngularSpring, const FVector& InAngularDamping, const FVector& InTorqueLimit,
+		EAngularDriveMode::Type InAngularDriveMode);
+
 	/** Refreshes the physics engine joint's linear limits. Only applicable if the joint has been created already.*/
 	ENGINE_API void UpdateLinearLimit();
 
diff --git a/Engine/Source/Runtime/Engine/Private/PhysicsEngine/ConstraintDrives.cpp b/Engine/Source/Runtime/Engine/Private/PhysicsEngine/ConstraintDrives.cpp
--- a/Engine/Source/Runtime/Engine/Private/PhysicsEngine/ConstraintDrives.cpp
+++ b/Engine/Source/Runtime/Engine/Private/PhysicsEngine/ConstraintDrives.cpp
@@ -18,7 +18,6 @@ FConstraintDrive::FConstraintDrive()
 FLinearDriveConstraint::FLinearDriveConstraint()
 	: PositionTarget(ForceInit)
 	, VelocityTarget(ForceInit)
-	, bEnablePositionDrive(false)
 {
 }
 
diff --git a/Engine/Source/Runtime/Engine/Private/PhysicsEngine/ConstraintInstance.cpp b/Engine/Source/Runtime/Engine/Private/PhysicsEngine/ConstraintInstance.cpp
--- a/Engine/Source/Runtime/Engine/Private/PhysicsEngine/ConstraintInstance.cpp
+++ b/Engine/Source/Runtime/Engine/Private/PhysicsEngine/ConstraintInstance.cpp
@@ -1020,6 +1020,31 @@ void FConstraintInstance::SetAngularOrientationTarget(const FQuat& InOrientation
 	FPhysicsInterface::SetDriveOrientation(ConstraintHandle, InOrientationTarget);
 }
 
+void FConstraintInstance::SetDriveParams(
+	const FVector& InLinearSpring, const FVector& InLinearDamping, const FVector& InForceLimit,
+	const FVector& InAngularSpring, const FVector& InAngularDamping, const FVector& InTorqueLimit,
+	EAngularDriveMode::Type InAngularDriveMode)
+{
+	ProfileInstance.LinearDrive.SetDriveParams(InLinearSpring, InLinearDamping, InForceLimit);
+	ProfileInstance.LinearDrive.SetLinearPositionDrive(
+		InLinearSpring.X != 0, InLinearSpring.Y != 0, InLinearSpring.Z != 0);
+	ProfileInstance.LinearDrive.SetLinearVelocityDrive(
+		InLinearDamping.X != 0, InLinearDamping.Y != 0, InLinearDamping.Z != 0);
+
+	ProfileInstance.AngularDrive.SetDriveParams(InAngularSpring, InAngularDamping, InTorqueLimit);
+	ProfileInstance.AngularDrive.SetAngularDriveMode(InAngularDriveMode);
+	ProfileInstance.AngularDrive.SetOrientationDriveTwistAndSwing(InAngularSpring.Y != 0, InAngularSpring.X != 0);
+	ProfileInstance.AngularDrive.SetOrientationDriveSLERP(InAngularSpring.Z != 0);
+	ProfileInstance.AngularDrive.SetAngularVelocityDriveTwistAndSwing(InAngularDamping.Y != 0, InAngularDamping.X != 0);
+	ProfileInstance.AngularDrive.SetAngularVelocityDriveSLERP(InAngularDamping.Z != 0);
+
+	FPhysicsInterface::ExecuteOnUnbrokenConstraintReadWrite(ConstraintHandle, [this](const FPhysicsConstraintHandle& InUnbrokenConstraint)
+	{
+		FPhysicsInterface::UpdateLinearDrive_AssumesLocked(InUnbrokenConstraint, ProfileInstance.LinearDrive);
+		FPhysicsInterface::UpdateAngularDrive_AssumesLocked(InUnbrokenConstraint, ProfileInstance.AngularDrive);
+	});
+}
+
 float FConstraintInstance::GetCurrentSwing1() const
 {
 	return FPhysicsInterface::GetCurrentSwing1(ConstraintHandle);
diff --git a/Engine/Source/Runtime/Experimental/Chaos/Private/Chaos/PBDJointConstraints.cpp b/Engine/Source/Runtime/Experimental/Chaos/Private/Chaos/PBDJointConstraints.cpp
--- a/Engine/Source/Runtime/Experimental/Chaos/Private/Chaos/PBDJointConstraints.cpp
+++ b/Engine/Source/Runtime/Experimental/Chaos/Private/Chaos/PBDJointConstraints.cpp
@@ -210,7 +210,7 @@ namespace Chaos
 		GetConstraintSettingsInternal().AngularDriveMaxTorque = MaxTorque;
 	}
 
-	void FPBDJointConstraintHandle::SetDriveProperties(
+	void FPBDJointConstraintHandle::SetDriveParams(
 		const FVec3 LinearStiffness, const FVec3 LinearDamping, const FVec3 MaxForce,
 		const FVec3 AngularStiffness, const FVec3 AngularDamping, const FVec3 MaxTorque)
 	{
@@ -221,6 +221,18 @@ namespace Chaos
 		Settings.AngularDriveStiffness = AngularStiffness;
 		Settings.AngularDriveDamping = AngularDamping;
 		Settings.AngularDriveMaxTorque = MaxTorque;
+		Settings.bLinearPositionDriveEnabled[0] = LinearStiffness.X > 0;
+		Settings.bLinearPositionDriveEnabled[1] = LinearStiffness.Y > 0;
+		Settings.bLinearPositionDriveEnabled[2] = LinearStiffness.Z > 0;
+		Settings.bLinearVelocityDriveEnabled[0] = LinearDamping.X > 0;
+		Settings.bLinearVelocityDriveEnabled[1] = LinearDamping.Y > 0;
+		Settings.bLinearVelocityDriveEnabled[2] = LinearDamping.Z > 0;
+		Settings.bAngularSwingPositionDriveEnabled = AngularStiffness.X > 0;
+		Settings.bAngularTwistPositionDriveEnabled = AngularStiffness.Y > 0;
+		Settings.bAngularSLerpPositionDriveEnabled = AngularStiffness.Z > 0;
+		Settings.bAngularSwingVelocityDriveEnabled = AngularDamping.X > 0;
+		Settings.bAngularTwistVelocityDriveEnabled = AngularDamping.Y > 0;
+		Settings.bAngularSLerpVelocityDriveEnabled = AngularDamping.Z > 0;
 	}
 
 	const FPBDJointSettings& FPBDJointConstraintHandle::GetSettings() const
diff --git a/Engine/Source/Runtime/Experimental/Chaos/Public/Chaos/PBDJointConstraints.h b/Engine/Source/Runtime/Experimental/Chaos/Public/Chaos/PBDJointConstraints.h
--- a/Engine/Source/Runtime/Experimental/Chaos/Public/Chaos/PBDJointConstraints.h
+++ b/Engine/Source/Runtime/Experimental/Chaos/Public/Chaos/PBDJointConstraints.h
@@ -73,10 +73,25 @@ namespace Chaos
 		CHAOS_API void SetAngularDriveDamping(const FVec3 Damping);
 		CHAOS_API void SetAngularDriveMaxTorque(const FVec3 MaxTorque);
 
-		CHAOS_API void SetDriveProperties(
+		/**
+		 * This allows the most common drive parameters to be set in one call. Note that the 
+		 * individual drive elements will be enabled/disabled depending on the strength/damping
+		 * values passed in.
+		 *
+		 * Angular values are passed in as (swing, twist, slerp)
+		 */
+		CHAOS_API void SetDriveParams(
 			const FVec3 LinearStiffness, const FVec3 LinearDamping, const FVec3 MaxForce,
 			const FVec3 AngularStiffness, const FVec3 AngularDamping, const FVec3 MaxTorque);
 
+		// TODO remove - this is just a temporary workaround before fixing calling code
+		void SetDriveProperties(
+			const FVec3 LinearStiffness, const FVec3 LinearDamping, const FVec3 MaxForce,
+			const FVec3 AngularStiffness, const FVec3 AngularDamping, const FVec3 MaxTorque)
+		{
+			SetDriveParams(LinearStiffness, LinearDamping, MaxForce, AngularStiffness, AngularDamping, MaxTorque);
+		}
+
 		CHAOS_API TVec2<FGeometryParticleHandle*> GetConstrainedParticles() const override final;
 
 		static const FConstraintHandleTypeID& StaticType()
