Material Editor - Added a workaround that enables support for materials with parameter collections to be cached on the DDC.

diff --git a/Engine/Source/Runtime/Engine/Private/Materials/HLSLMaterialTranslator.cpp b/Engine/Source/Runtime/Engine/Private/Materials/HLSLMaterialTranslator.cpp
--- a/Engine/Source/Runtime/Engine/Private/Materials/HLSLMaterialTranslator.cpp
+++ b/Engine/Source/Runtime/Engine/Private/Materials/HLSLMaterialTranslator.cpp
@@ -57,6 +57,7 @@
 #include "Serialization/ObjectReader.h"
 #include "Serialization/BufferArchive.h"
 #include "Serialization/MemoryReader.h"
+#include "Serialization/ObjectAndNameAsStringProxyArchive.h"
 #include "DerivedDataCache.h"
 #include "DerivedDataRequestOwner.h"
 #include "MaterialCachedData.h"
@@ -206,7 +207,7 @@ UE::DerivedData::FValueId EnvironmentDefinesId = UE::DerivedData::FValueId::From
 
 /* This version number models the layout of the data stored on the DDC after a material translation (e.g. FEnvironmentDefines)
  * It must be bumped whenever a change is made requires re-translation of all materials. */
-static constexpr int MaterialTranslationDDCVersion = 1;
+static constexpr int MaterialTranslationDDCVersion = 3;
 
 /** Data structure used to cache a part of material translation results. It contains all the generated
  *  defines that will be declared during the compilation of the generated material shader.
@@ -304,6 +305,7 @@ struct FHLSLMaterialTranslator::FEnvironmentDefines
 	bool bSubstrateComplexSpecialPath;
 	bool bTextureSampleDebug;
 	TArray<TPair<FString, int32>> SubstrateDefines;
+	TArray<TObjectPtr<UMaterialParameterCollection>> ParameterCollections;
 
 	bool HasShadingModel(EMaterialShadingModel model) const
 	{
@@ -394,6 +396,7 @@ struct FHLSLMaterialTranslator::FEnvironmentDefines
 		Ar << bSubstrateComplexSpecialPath;
 		Ar << bTextureSampleDebug;
 		Ar << SubstrateDefines;
+		Ar << ParameterCollections;
 	}
 };
 
@@ -1102,6 +1105,16 @@ UE_TRACE_EVENT_END()
 
 bool FHLSLMaterialTranslator::Translate()
 {
+	// We call FindObject to serialize the array of Parameter Collections used by this material in EnvironmentDefines,
+	// but this can happen during save. FindObject is illegal during save because if the discovered objects
+	// are serialized into the package it will cause a crash on package load. But we are not storing the results
+	// of FindObject into the package so it is okay to remove the restriction.
+	TOptional<TGuardValue<bool>> IsSavingPackageGuard;
+	if (IsInGameThread())
+	{
+		IsSavingPackageGuard.Emplace(GIsSavingPackage, false);
+	}
+
 #if CPUPROFILERTRACE_ENABLED
 	FString TraceMaterialName;
 	if (UE_TRACE_CHANNELEXPR_IS_ENABLED(CpuChannel))
@@ -2440,14 +2453,14 @@ void FHLSLMaterialTranslator::GetMaterialEnvironment(EShaderPlatform InPlatform,
 		OutEnvironment.SetDefine(*FString::Printf(TEXT("VIRTUALTEXTURE_PAGETABLE_%d"), i), *PageTableValue);
 	}
 
-	for (int32 CollectionIndex = 0; CollectionIndex < ParameterCollections.Num(); CollectionIndex++)
+	for (int32 CollectionIndex = 0; CollectionIndex < EnvironmentDefines->ParameterCollections.Num(); CollectionIndex++)
 	{
 		// Add uniform buffer declarations for any parameter collections referenced
 		const FString CollectionName = FString::Printf(TEXT("MaterialCollection%u"), CollectionIndex);
 		// This can potentially become an issue for MaterialCollection Uniform Buffers if they ever get non-numeric resources (eg Textures), as
 		// OutEnvironment.ResourceTableMap has a map by name, and the N ParameterCollection Uniform Buffers ALL are names "MaterialCollection"
 		// (and the hlsl cbuffers are named MaterialCollection0, etc, so the names don't match the layout)
-		FShaderUniformBufferParameter::ModifyCompilationEnvironment(*CollectionName, ParameterCollections[CollectionIndex]->GetUniformBufferStruct(), InPlatform, OutEnvironment);
+		FShaderUniformBufferParameter::ModifyCompilationEnvironment(*CollectionName, EnvironmentDefines->ParameterCollections[CollectionIndex]->GetUniformBufferStruct(), InPlatform, OutEnvironment);
 	}
 
 	OutEnvironment.SetDefine(TEXT("IS_MATERIAL_SHADER"), TEXT("1"));
@@ -15196,6 +15209,7 @@ void FHLSLMaterialTranslator::PrepareEnvironmentDefines()
 	}
 
 	EnvironmentDefines->bTextureSampleDebug = IsDebugTextureSampleEnabled();
+	EnvironmentDefines->ParameterCollections = ParameterCollections;
 }
 
 bool FHLSLMaterialTranslator::QueryDDCCachedTranslationResults()
@@ -15241,7 +15255,8 @@ bool FHLSLMaterialTranslator::QueryDDCCachedTranslationResults()
 	{
 		// Read the environment defines
 		FMemoryReaderView EnvironmentDefinesBufferReader{ TArrayView<uint8>{ (uint8*)EnvironmentDefinesBuffer.GetData(), (int)EnvironmentDefinesBuffer.GetSize() } };
-		EnvironmentDefines->Serialize(EnvironmentDefinesBufferReader);
+		FObjectAndNameAsStringProxyArchive EnvironmentDefinesBufferReaderProxy{ EnvironmentDefinesBufferReader, true };
+		EnvironmentDefines->Serialize(EnvironmentDefinesBufferReaderProxy);
 
 		// Read the material compilation output
 		FShaderMapPointerTable PointerTable;
@@ -15263,10 +15278,7 @@ bool FHLSLMaterialTranslator::QueryDDCCachedTranslationResults()
 
 void FHLSLMaterialTranslator::PushResultsToDDCCache()
 {
-	// We currently don't support caching parameter collections on the DDC, as it is invalid at this point
-	// to serialize an array of object pointers. Only allow putting the translation results on the DDC
-	// if the this material does not use parameter collections.
-	if (GJobDisableMaterialTranslateDDC || !ParameterCollections.IsEmpty())
+	if (GJobDisableMaterialTranslateDDC)
 	{
 		GShaderCompilerStats->IncrementMaterialTranslationSkippedDDC();
 		return;
@@ -15299,7 +15311,8 @@ void FHLSLMaterialTranslator::PushResultsToDDCCache()
 
 	// Push the material shader defines
 	FBufferArchive EnvironmentDefinesBuffer;
-	EnvironmentDefines->Serialize(EnvironmentDefinesBuffer);
+	FObjectAndNameAsStringProxyArchive EnvironmentDefinesBufferProxy{ EnvironmentDefinesBuffer, true };
+	EnvironmentDefines->Serialize(EnvironmentDefinesBufferProxy);
 	RecordBuilder.AddValue(EnvironmentDefinesId, FSharedBuffer::MakeView(EnvironmentDefinesBuffer.GetData(), EnvironmentDefinesBuffer.Num()));
 
 	// Push the all the data to the DDC
