Avoid loading textures in material expressions in cooked games.

diff --git a/Engine/Source/Runtime/Engine/Private/Materials/MaterialExpressions.cpp b/Engine/Source/Runtime/Engine/Private/Materials/MaterialExpressions.cpp
--- a/Engine/Source/Runtime/Engine/Private/Materials/MaterialExpressions.cpp
+++ b/Engine/Source/Runtime/Engine/Private/Materials/MaterialExpressions.cpp
@@ -3934,6 +3934,7 @@ void UMaterialExpressionTextureSampleParameter::SetDefaultTexture()
 UMaterialExpressionTextureObjectParameter::UMaterialExpressionTextureObjectParameter(const FObjectInitializer& ObjectInitializer)
 	: Super(ObjectInitializer)
 {
+#if WITH_EDITORONLY_DATA
 	// Structure to hold one-time initialization
 	struct FConstructorStatics
 	{
@@ -3951,7 +3952,6 @@ UMaterialExpressionTextureObjectParameter::UMaterialExpressionTextureObjectParam
 
 	Texture = ConstructorStatics.DefaultTexture2D.Object;
 
-#if WITH_EDITORONLY_DATA
 	MenuCategories.Empty();
 	MenuCategories.Add(ConstructorStatics.NAME_Texture);
 	MenuCategories.Add(ConstructorStatics.NAME_Parameters);
@@ -4021,6 +4021,7 @@ int32 UMaterialExpressionTextureObjectParameter::CompilePreview(class FMaterialC
 UMaterialExpressionTextureObject::UMaterialExpressionTextureObject(const FObjectInitializer& ObjectInitializer)
 	: Super(ObjectInitializer)
 {
+#if WITH_EDITORONLY_DATA
 	// Structure to hold one-time initialization
 	struct FConstructorStatics
 	{
@@ -4038,7 +4039,6 @@ UMaterialExpressionTextureObject::UMaterialExpressionTextureObject(const FObject
 
 	Texture = ConstructorStatics.Object0.Object;
 
-#if WITH_EDITORONLY_DATA
 	MenuCategories.Add(ConstructorStatics.NAME_Texture);
 	MenuCategories.Add(ConstructorStatics.NAME_Functions);
 
@@ -4229,6 +4229,7 @@ bool UMaterialExpressionTextureProperty::MatchesSearchQuery(const TCHAR* SearchQ
 UMaterialExpressionTextureSampleParameter2D::UMaterialExpressionTextureSampleParameter2D(const FObjectInitializer& ObjectInitializer)
 	: Super(ObjectInitializer)
 {
+#if WITH_EDITORONLY_DATA
 	// Structure to hold one-time initialization
 	struct FConstructorStatics
 	{
@@ -4246,7 +4247,6 @@ UMaterialExpressionTextureSampleParameter2D::UMaterialExpressionTextureSamplePar
 
 	Texture = ConstructorStatics.DefaultTexture.Object;
 
-#if WITH_EDITORONLY_DATA
 	MenuCategories.Empty();
 	MenuCategories.Add(ConstructorStatics.NAME_Texture);
 	MenuCategories.Add(ConstructorStatics.NAME_Parameters);
@@ -4320,6 +4320,7 @@ void UMaterialExpressionTextureSampleParameter::SetEditableName(const FString& N
 UMaterialExpressionTextureSampleParameterCube::UMaterialExpressionTextureSampleParameterCube(const FObjectInitializer& ObjectInitializer)
 	: Super(ObjectInitializer)
 {
+#if WITH_EDITORONLY_DATA
 	// Structure to hold one-time initialization
 	struct FConstructorStatics
 	{
@@ -4337,7 +4338,6 @@ UMaterialExpressionTextureSampleParameterCube::UMaterialExpressionTextureSampleP
 
 	Texture = ConstructorStatics.DefaultTextureCube.Object;
 
-#if WITH_EDITORONLY_DATA
 	MenuCategories.Empty();
 	MenuCategories.Add(ConstructorStatics.NAME_Texture);
 	MenuCategories.Add(ConstructorStatics.NAME_Parameters);
@@ -4523,6 +4523,7 @@ const TCHAR* UMaterialExpressionTextureSampleParameterCubeArray::GetRequirements
 UMaterialExpressionTextureSampleParameterVolume::UMaterialExpressionTextureSampleParameterVolume(const FObjectInitializer& ObjectInitializer)
 	: Super(ObjectInitializer)
 {
+#if WITH_EDITORONLY_DATA
 	// Structure to hold one-time initialization
 	struct FConstructorStatics
 	{
@@ -4540,7 +4541,6 @@ UMaterialExpressionTextureSampleParameterVolume::UMaterialExpressionTextureSampl
 
 	Texture = ConstructorStatics.DefaultVolumeTexture.Object;
 
-#if WITH_EDITORONLY_DATA
 	MenuCategories.Empty();
 	MenuCategories.Add(ConstructorStatics.NAME_Texture);
 	MenuCategories.Add(ConstructorStatics.NAME_Parameters);
@@ -20706,19 +20706,16 @@ void UMaterialExpressionPerInstanceCustomData3Vector::GetCaption(TArray<FString>
 UMaterialExpressionAntialiasedTextureMask::UMaterialExpressionAntialiasedTextureMask(const FObjectInitializer& ObjectInitializer)
 	: Super(ObjectInitializer)
 {
+#if WITH_EDITORONLY_DATA
 	// Structure to hold one-time initialization
 	struct FConstructorStatics
 	{
 		ConstructorHelpers::FObjectFinder<UTexture2D> DefaultTexture;
-#if WITH_EDITORONLY_DATA
 		FText NAME_Utility;
-#endif
 		FName NAME_None;
 		FConstructorStatics()
 			: DefaultTexture(TEXT("/Engine/EngineResources/DefaultTexture"))
-#if WITH_EDITORONLY_DATA
 			, NAME_Utility(LOCTEXT( "Utility", "Utility" ))
-#endif
 			, NAME_None(TEXT("None"))
 		{
 		}
@@ -20727,13 +20724,12 @@ UMaterialExpressionAntialiasedTextureMask::UMaterialExpressionAntialiasedTexture
 
 	Texture = ConstructorStatics.DefaultTexture.Object;
 
-#if WITH_EDITORONLY_DATA
 	MenuCategories.Empty();
 	MenuCategories.Add(ConstructorStatics.NAME_Utility);
 #endif
 
 	Threshold = 0.5f;
-	ParameterName = ConstructorStatics.NAME_None;
+	ParameterName = NAME_None;
 	Channel = TCC_Alpha;
 
 #if WITH_EDITORONLY_DATA
