Fix an infinite loop caused by requesting a start time outside the range of the animation by exactly the input DeltaTime

#localization none
#tests Josselin Francois validated by spawning hundreds of characters

diff --git a/Engine/Source/Runtime/Engine/Private/Animation/AnimSequenceBase.cpp b/Engine/Source/Runtime/Engine/Private/Animation/AnimSequenceBase.cpp
--- a/Engine/Source/Runtime/Engine/Private/Animation/AnimSequenceBase.cpp
+++ b/Engine/Source/Runtime/Engine/Private/Animation/AnimSequenceBase.cpp
@@ -431,12 +431,11 @@ void UAnimSequenceBase::GetAnimNotifies(const float& StartTime, const float& Del
 		
 		GetAnimNotifiesFromDeltaPositions(PreviousPosition, CurrentPosition, NotifyContext);
 	
+		const float ActualDeltaMove = (CurrentPosition - PreviousPosition);
 		// If we've hit the end of the animation, and we're allowed to loop, keep going.
-		if( (AdvanceType == ETAA_Finished) &&  NotifyContext.TickRecord && NotifyContext.TickRecord->bLooping )
+		if( (AdvanceType == ETAA_Finished) &&  NotifyContext.TickRecord && NotifyContext.TickRecord->bLooping && ActualDeltaMove != 0.f)
 		{
-			const float ActualDeltaMove = (CurrentPosition - PreviousPosition);
-			DesiredDeltaMove -= ActualDeltaMove; 
-
+			DesiredDeltaMove -= ActualDeltaMove;
 			PreviousPosition = bPlayingBackwards ? GetPlayLength() : 0.f;
 			CurrentPosition = PreviousPosition;
 		}
