Changing FOnLocalPlayerEvent from DECLARE_EVENT to MULTICAST_DELEGATE

-review [REDACTED_HANDLE]
#p4v-preflight-copy 29891156

diff --git a/Engine/Source/Runtime/Engine/Classes/Engine/GameInstance.h b/Engine/Source/Runtime/Engine/Classes/Engine/GameInstance.h
--- a/Engine/Source/Runtime/Engine/Classes/Engine/GameInstance.h
+++ b/Engine/Source/Runtime/Engine/Classes/Engine/GameInstance.h
@@ -146,7 +146,7 @@ struct FGameInstancePIEParameters
 
 enum class EInputDeviceConnectionState : uint8;
 
-DECLARE_EVENT_OneParam(UGameInstance, FOnLocalPlayerEvent, ULocalPlayer*);
+DECLARE_MULTICAST_DELEGATE_OneParam(FOnLocalPlayerEvent, ULocalPlayer*);
 
 /**
  * GameInstance: high-level manager object for an instance of the running game.
