Remove invalid ensures. Compatible skeletons means we can no longer guarantee that all bones within the pose asset exist in the skeleton.

diff --git a/Engine/Source/Runtime/Engine/Private/Animation/PoseAsset.cpp b/Engine/Source/Runtime/Engine/Private/Animation/PoseAsset.cpp
--- a/Engine/Source/Runtime/Engine/Private/Animation/PoseAsset.cpp
+++ b/Engine/Source/Runtime/Engine/Private/Animation/PoseAsset.cpp
@@ -608,7 +608,6 @@ bool UPoseAsset::GetAnimationPose(struct FAnimationPoseData& OutAnimationPoseDat
 		FBlendedCurve& OutCurve = OutAnimationPoseData.GetCurve();
 
 		const FBoneContainer& RequiredBones = OutPose.GetBoneContainer();
-		const USkeletalMesh* RequiredBonesMesh = RequiredBones.GetSkeletalMeshAsset();
 		USkeleton* MySkeleton = GetSkeleton();
 		
 		check(PoseContainer.IsValid());
@@ -646,9 +645,6 @@ bool UPoseAsset::GetAnimationPose(struct FAnimationPoseData& OutAnimationPoseDat
 					for (int32 TrackIndex = 0; TrackIndex < TrackNum; ++TrackIndex)
 					{
 						const FSkeletonPoseBoneIndex SkeletonBoneIndex = FSkeletonPoseBoneIndex(SkeletonRemapping.IsValid() ? SkeletonRemapping.GetTargetSkeletonBoneIndex(PoseContainer.TrackBoneIndices[TrackIndex]) : PoseContainer.TrackBoneIndices[TrackIndex]);
-
-						ensureMsgf(RequiredBones.IsSkeletonPoseIndexValid(SkeletonBoneIndex), TEXT("PoseAsset [%s] with skeleton [%s] has bones not present in the evaluation container. Bone(%s) index(%d) not found in skeletal mesh %s."), *GetPathName(), MySkeleton ? *MySkeleton->GetPathName() : TEXT("<Skeleton Not Found>"), *PoseContainer.Tracks[TrackIndex].ToString(), SkeletonBoneIndex.GetInt(), RequiredBonesMesh ? *RequiredBonesMesh->GetName() : TEXT("<Skeletal Mesh Not Found>"));
-
 						const FCompactPoseBoneIndex CompactIndex = RequiredBones.GetCompactPoseIndexFromSkeletonPoseIndex(SkeletonBoneIndex);
 						
 						// If bone index is invalid, or not required for the pose - skip
@@ -710,9 +706,6 @@ bool UPoseAsset::GetAnimationPose(struct FAnimationPoseData& OutAnimationPoseDat
 			for(int32 TrackIndex = 0; TrackIndex < TrackNum; ++TrackIndex)
 			{
 				const FSkeletonPoseBoneIndex SkeletonBoneIndex = FSkeletonPoseBoneIndex(SkeletonRemapping.IsValid() ? SkeletonRemapping.GetTargetSkeletonBoneIndex(PoseContainer.TrackBoneIndices[TrackIndex]) : PoseContainer.TrackBoneIndices[TrackIndex]);
-
-				ensureMsgf(RequiredBones.IsSkeletonPoseIndexValid(SkeletonBoneIndex), TEXT("PoseAsset [%s] with skeleton [%s] has bones not present in the evaluation container. Bone(%s) index(%d) not found in skeletal mesh %s."), *GetPathName(), MySkeleton ? *MySkeleton->GetPathName() : TEXT("<Skeleton Not Found>"), *PoseContainer.Tracks[TrackIndex].ToString(), SkeletonBoneIndex.GetInt(), RequiredBonesMesh ? *RequiredBonesMesh->GetName() : TEXT("<Skeletal Mesh Not Found>"));
-
 				const FCompactPoseBoneIndex CompactIndex = RequiredBones.GetCompactPoseIndexFromSkeletonPoseIndex(SkeletonBoneIndex);
 
 				// we add even if it's invalid because we want it to match with track index
