Bring back computation of Chaos convex indices.

- Undo of CL 29578847 - these indices are also used by NavMesh.

diff --git a/Engine/Source/Runtime/Engine/Private/PhysicsEngine/BodySetup.cpp b/Engine/Source/Runtime/Engine/Private/PhysicsEngine/BodySetup.cpp
--- a/Engine/Source/Runtime/Engine/Private/PhysicsEngine/BodySetup.cpp
+++ b/Engine/Source/Runtime/Engine/Private/PhysicsEngine/BodySetup.cpp
@@ -1775,15 +1775,11 @@ void FKConvexElem::ResetChaosConvexMesh()
 
 ENGINE_API void FKConvexElem::ComputeChaosConvexIndices(bool bForceCompute)
 {
-	// these indices are not needed for simulation, but only get used for debug visualization
-#if (UE_ENABLE_DEBUG_DRAWING && !USE_NULL_RHI)	// FIXME: UE_ENABLE_DEBUG_DRAWING needs to be fixed to incorportate !USE_NULL_RHI
+	// these indices are not needed for simulation, but are also used by NavMesh
 	if (bForceCompute || IndexData.Num() == 0)
 	{
 		IndexData = GetChaosConvexIndices();
 	}
-#else
-	IndexData.Reset();
-#endif
 }
 
 TArray<int32> FKConvexElem::GetChaosConvexIndices() const
