Use the EChildActorComponentTreeViewVisualizationMode enum to hide the actor in the tree view
#rnx

diff --git a/Engine/Source/Runtime/Engine/Classes/Components/ChildActorComponent.h b/Engine/Source/Runtime/Engine/Classes/Components/ChildActorComponent.h
--- a/Engine/Source/Runtime/Engine/Classes/Components/ChildActorComponent.h
+++ b/Engine/Source/Runtime/Engine/Classes/Components/ChildActorComponent.h
@@ -74,6 +74,8 @@ enum class EChildActorComponentTreeViewVisualizationMode : uint8
 	ComponentWithChildActor,
 	/** Show only as a child actor hierarchy (i.e. do not show the outer component node as the root). */
 	ChildActorOnly,
+	/** Do not display the actor in the tree view. */
+	Hidden,
 };
 #endif
 
@@ -137,9 +139,6 @@ private:
 	/** Indicates how this component will be visualized for editing in a tree view. Users can change this setting per instance via the context menu in the Blueprint/SCS editor. */
 	UPROPERTY()
 	EChildActorComponentTreeViewVisualizationMode EditorTreeViewVisualizationMode;
-
-	UPROPERTY(EditDefaultsOnly, Category=ChildActorComponent)
-	uint8 bHideFromSceneOutliner:1;
 #endif
 
 	/**
@@ -218,9 +217,6 @@ public:
 	}
 
 	ENGINE_API void SetEditorTreeViewVisualizationMode(EChildActorComponentTreeViewVisualizationMode InMode);
-
-	bool GetHideFromSceneOutliner() const { return bHideFromSceneOutliner; }
-	void SetHideFromSceneOutliner(bool bHide) {bHideFromSceneOutliner = bHide;}
 #endif
 
 #if UE_WITH_IRIS
diff --git a/Engine/Source/Runtime/Engine/Private/Components/ChildActorComponent.cpp b/Engine/Source/Runtime/Engine/Private/Components/ChildActorComponent.cpp
--- a/Engine/Source/Runtime/Engine/Private/Components/ChildActorComponent.cpp
+++ b/Engine/Source/Runtime/Engine/Private/Components/ChildActorComponent.cpp
@@ -42,7 +42,6 @@ UChildActorComponent::UChildActorComponent(const FObjectInitializer& ObjectIniti
 
 #if WITH_EDITORONLY_DATA
 	EditorTreeViewVisualizationMode = EChildActorComponentTreeViewVisualizationMode::UseDefault;
-	bHideFromSceneOutliner = false;
 #endif
 }
 
@@ -801,7 +800,7 @@ void UChildActorComponent::CreateChildActor(TFunction<void(AActor*)> CustomizerF
 				Params.bCreateActorPackage = false;
 				Params.OverridePackage = (MyOwner ? MyOwner->GetExternalPackage() : nullptr);
 				Params.OverrideActorGuid = CachedInstanceData ? CachedInstanceData->ChildActorGUID : FGuid();
-				Params.bHideFromSceneOutliner = bHideFromSceneOutliner;
+				Params.bHideFromSceneOutliner = EditorTreeViewVisualizationMode == EChildActorComponentTreeViewVisualizationMode::Hidden;
 #endif
 				if (ChildActorTemplate && ChildActorTemplate->GetClass() == ChildActorClass)
 				{
