[NetPhyiscs] Implemente minimum extrapolation time for Predictive Interpolation

diff --git a/Engine/Source/Runtime/Engine/Private/PhysicsEngine/PhysicsReplication.cpp b/Engine/Source/Runtime/Engine/Private/PhysicsEngine/PhysicsReplication.cpp
--- a/Engine/Source/Runtime/Engine/Private/PhysicsEngine/PhysicsReplication.cpp
+++ b/Engine/Source/Runtime/Engine/Private/PhysicsEngine/PhysicsReplication.cpp
@@ -134,6 +134,9 @@ namespace PhysicsReplicationCVars
 		static float ExtrapolationTimeMultiplier = 3.0f;
 		static FAutoConsoleVariableRef CVarExtrapolationTimeMultiplier(TEXT("np2.PredictiveInterpolation.ExtrapolationTimeMultiplier"), ExtrapolationTimeMultiplier, TEXT("Multiplier to adjust the time to extrapolate the target forward over, the time is based on current send-rate."));
 
+		static float ExtrapolationMinTime = 0.75f;
+		static FAutoConsoleVariableRef CVarExtrapolationMinTime(TEXT("np2.PredictiveInterpolation.ExtrapolationMinTime"), ExtrapolationMinTime, TEXT("Clamps minimum extrapolation time. Value in seconds. Disable minimum clamp by setting to 0."));
+
 		static float MinExpectedDistanceCovered = 0.5f;
 		static FAutoConsoleVariableRef CVarMinExpectedDistanceCovered(TEXT("np2.PredictiveInterpolation.MinExpectedDistanceCovered"), MinExpectedDistanceCovered, TEXT("Value between 0-1, in percentage where 0.25 = 25%. How much of the expected distance based on replication velocity should the object have covered in a simulation tick to Not be considered stuck."));
 
@@ -1556,7 +1559,10 @@ bool FPhysicsReplicationAsync::PredictiveInterpolation(Chaos::FPBDRigidParticleH
 		Target.PrevPos = FVector(CurrentState.Position);
 
 		// --- Target Extrapolation ---
-		if (Target.TickCount <= FMath::CeilToInt(Target.ReceiveInterval * PhysicsReplicationCVars::PredictiveInterpolationCVars::ExtrapolationTimeMultiplier))
+		const int32 ExtrapolationTickLimit = FMath::Max(
+			FMath::CeilToInt(Target.ReceiveInterval * PhysicsReplicationCVars::PredictiveInterpolationCVars::ExtrapolationTimeMultiplier), // Extrapolate time based on receive interval * multiplier
+			FMath::CeilToInt(PhysicsReplicationCVars::PredictiveInterpolationCVars::ExtrapolationMinTime / DeltaSeconds)); // At least extrapolate for N seconds
+		if (Target.TickCount <= ExtrapolationTickLimit)
 		{
 			ExtrapolateTarget(Target, 1, DeltaSeconds);
 		}
