Material Editor - Disabled material translation DDC caching for now as it has a negative impact on performance.

diff --git a/Engine/Source/Runtime/Engine/Private/Materials/HLSLMaterialTranslator.cpp b/Engine/Source/Runtime/Engine/Private/Materials/HLSLMaterialTranslator.cpp
--- a/Engine/Source/Runtime/Engine/Private/Materials/HLSLMaterialTranslator.cpp
+++ b/Engine/Source/Runtime/Engine/Private/Materials/HLSLMaterialTranslator.cpp
@@ -194,7 +194,7 @@ static FAutoConsoleVariableRef CVarPedanticErrorChecksEnabled(
 static bool GUseMaterialTranslationResultsGrouping = true;
 
 /* Controls whether DDC caching of material translation results is enabled. */
-static bool GJobDisableMaterialTranslateDDC = false;
+static bool GJobDisableMaterialTranslateDDC = true;
 static FAutoConsoleVariableRef CVarJobDisableMaterialTranslateDDC(
 	TEXT("r.Material.DisableTranslateDDC"),
 	GJobDisableMaterialTranslateDDC,
@@ -1105,16 +1105,6 @@ UE_TRACE_EVENT_END()
 
 bool FHLSLMaterialTranslator::Translate()
 {
-	// We call FindObject to serialize the array of Parameter Collections used by this material in EnvironmentDefines,
-	// but this can happen during save. FindObject is illegal during save because if the discovered objects
-	// are serialized into the package it will cause a crash on package load. But we are not storing the results
-	// of FindObject into the package so it is okay to remove the restriction.
-	TOptional<TGuardValue<bool>> IsSavingPackageGuard;
-	if (IsInGameThread())
-	{
-		IsSavingPackageGuard.Emplace(GIsSavingPackage, false);
-	}
-
 #if CPUPROFILERTRACE_ENABLED
 	FString TraceMaterialName;
 	if (UE_TRACE_CHANNELEXPR_IS_ENABLED(CpuChannel))
@@ -15219,6 +15209,16 @@ bool FHLSLMaterialTranslator::QueryDDCCachedTranslationResults()
 		return false;
 	}
 
+	// We call FindObject to serialize the array of Parameter Collections used by this material in EnvironmentDefines,
+	// but this can happen during save. FindObject is illegal during save because if the discovered objects
+	// are serialized into the package it will cause a crash on package load. But we are not storing the results
+	// of FindObject into the package so it is okay to remove the restriction.
+	TOptional<TGuardValue<bool>> IsSavingPackageGuard;
+	if (IsInGameThread())
+	{
+		IsSavingPackageGuard.Emplace(GIsSavingPackage, false);
+	}
+
 	// Try fetching the translation results from the DDC using the generated key hash
 	FSharedBuffer MaterialCompilationOutputBuffer;
 	FSharedBuffer TranslationResultsBuffer;
@@ -15284,6 +15284,13 @@ void FHLSLMaterialTranslator::PushResultsToDDCCache()
 		return;
 	}
 
+	// See FHLSLMaterialTranslator::QueryDDCCachedTranslationResults()
+	TOptional<TGuardValue<bool>> IsSavingPackageGuard;
+	if (IsInGameThread())
+	{
+		IsSavingPackageGuard.Emplace(GIsSavingPackage, false);
+	}
+
 	UE::DerivedData::FCacheKey CacheKey{ MaterialTranslationDDCBucket, DDCKeyHash };
 	UE::DerivedData::FCacheRecordBuilder RecordBuilder{ CacheKey };
 
