Material Editor - Always sanitize expressions list in EditorOnly during a UMaterial::PostLoad(), regardless of whether it is running in editor (GIsEditor is true).

diff --git a/Engine/Source/Runtime/Engine/Private/Materials/Material.cpp b/Engine/Source/Runtime/Engine/Private/Materials/Material.cpp
--- a/Engine/Source/Runtime/Engine/Private/Materials/Material.cpp
+++ b/Engine/Source/Runtime/Engine/Private/Materials/Material.cpp
@@ -3898,16 +3898,14 @@ void UMaterial::PostLoad()
 #if WITH_EDITOR
 	// Create exec flow expressions, if needed
 	CreateExecutionFlowExpressions();
-	if (GIsEditor)
+
+	// Clean up any removed material expression classes. If running in editor, also release resources and mutate DDC key.
+	if (EditorOnly->ExpressionCollection.Expressions.Remove(nullptr) != 0 && GIsEditor)
 	{
-		// Clean up any removed material expression classes	
-		if (EditorOnly->ExpressionCollection.Expressions.Remove(nullptr) != 0)
-		{
-			// Force this material to recompile because its expressions have changed
-			// We're not providing a deterministic transformation guid because there could be many different ways expression
-			// could change. Each conversion code removing such expression would need its own guid.
-			ReleaseResourcesAndMutateDDCKey();
-		}
+		// Force this material to recompile because its expressions have changed
+		// We're not providing a deterministic transformation guid because there could be many different ways expression
+		// could change. Each conversion code removing such expression would need its own guid.
+		ReleaseResourcesAndMutateDDCKey();
 	}
 #endif // WITH_EDITOR
 
