World Grid Preview:
  - Added new world grid preview class to be used for any system that wants to show a world grid.
  - World Partition Runtime HashSet LHGrid can now show its debug world grid.
  - World Settings can now show the foliage debug world grid.

diff --git a/Engine/Source/Runtime/Engine/Classes/GameFramework/WorldSettings.h b/Engine/Source/Runtime/Engine/Classes/GameFramework/WorldSettings.h
--- a/Engine/Source/Runtime/Engine/Classes/GameFramework/WorldSettings.h
+++ b/Engine/Source/Runtime/Engine/Classes/GameFramework/WorldSettings.h
@@ -12,6 +12,7 @@
 #include "GameFramework/Info.h"
 #include "Sound/AudioVolume.h"
 #include "UObject/ConstructorHelpers.h"
+#include "WorldGridPreviewer.h"
 #include "WorldPartition/WorldPartitionEditorPerProjectUserSettings.h"
 #include "WorldSettings.generated.h"
 
@@ -624,6 +625,9 @@ public:
 	UPROPERTY(EditAnywhere, Category = Foliage)
 	uint32 InstancedFoliageGridSize;
 
+	UPROPERTY(EditAnywhere, Category = Foliage, Transient, SkipSerialization)
+	bool bShowInstancedFoliageGrid;
+
 	UPROPERTY(EditAnywhere, Category = Landscape)
 	uint32 LandscapeSplineMeshesGridSize;
 
@@ -644,6 +648,10 @@ public:
 	uint32 DefaultPlacementGridSize;
 #endif
 
+#if WITH_EDITOR
+	mutable TUniquePtr<FWorldGridPreviewer> InstancedFoliageGridGridPreviewer;
+#endif
+
 	/**
 	 * A list of runtime data layers that should be included in the base navmesh.
 	 * Editor data layers and actors outside data layers will be included.
diff --git a/Engine/Source/Runtime/Engine/Private/WorldGridPreviewer.cpp b/Engine/Source/Runtime/Engine/Private/WorldGridPreviewer.cpp
new file mode 100644
--- /dev/null
+++ b/Engine/Source/Runtime/Engine/Private/WorldGridPreviewer.cpp
@@ -0,0 +1,80 @@
+// Copyright Epic Games, Inc. All Rights Reserved.
+
+#include "WorldGridPreviewer.h"
+
+#if WITH_EDITOR
+#include "Materials/Material.h"
+#include "Engine/World.h"
+#include "Engine/PostProcessVolume.h"
+#include "Materials/MaterialInstanceConstant.h"
+
+FWorldGridPreviewer::FWorldGridPreviewer()
+	: CellSize(0)
+	, LoadingRange(0)
+	, GridColor(ForceInit)
+	, GridOffset(ForceInit)
+	, World(nullptr)
+{}
+
+FWorldGridPreviewer::FWorldGridPreviewer(UWorld* InWorld, bool bInIs2D)
+	: FWorldGridPreviewer()
+{
+	World = InWorld;
+
+	FActorSpawnParameters SpawnParameters;
+	SpawnParameters.ObjectFlags |= RF_Transient;
+	SpawnParameters.ObjectFlags &= ~RF_Transactional;
+	
+	PostProcessVolume = World->SpawnActor<APostProcessVolume>(SpawnParameters);
+	Material = PostProcessVolume.IsValid() ? LoadObject<UMaterial>(nullptr, TEXT("/Engine/EditorMaterials/WorldGridPreviewMaterial")) : nullptr;
+	MaterialInstance = NewObject<UMaterialInstanceConstant>(PostProcessVolume.Get());
+	MaterialInstance->Parent = Material;
+
+	if (MaterialInstance)
+	{
+		FStaticParameterSet StaticParams;
+		MaterialInstance->GetStaticParameterValues(StaticParams);
+
+		static FName NAME_IsGrid2D(TEXT("Is_Grid_2D"));
+		for (int ParamIdx = 0; ParamIdx < StaticParams.StaticSwitchParameters.Num(); ParamIdx++)
+		{
+			if (StaticParams.StaticSwitchParameters[ParamIdx].ParameterInfo.Name == NAME_IsGrid2D)
+			{
+				MaterialInstance->SetStaticSwitchParameterValueEditorOnly(StaticParams.StaticSwitchParameters[ParamIdx].ParameterInfo, bInIs2D);
+				break;
+			}
+		}
+
+		MaterialInstance->UpdateStaticPermutation();
+
+		PostProcessVolume->Settings.WeightedBlendables.Array.Emplace(1.0f, MaterialInstance);
+		PostProcessVolume->bUnbound = true;
+		PostProcessVolume->bEnabled = true;
+	}
+}
+
+FWorldGridPreviewer::~FWorldGridPreviewer()
+{
+	if (PostProcessVolume.IsValid())
+	{
+		World->DestroyActor(PostProcessVolume.Get());
+	}
+}
+
+void FWorldGridPreviewer::AddReferencedObjects(FReferenceCollector& InCollector)
+{
+	InCollector.AddReferencedObject(Material);
+	InCollector.AddReferencedObject(MaterialInstance);
+}
+
+void FWorldGridPreviewer::Update()
+{
+	if (MaterialInstance)
+	{
+		MaterialInstance->SetScalarParameterValueEditorOnly(TEXT("Grid_CellSize"), (float)CellSize);
+		MaterialInstance->SetScalarParameterValueEditorOnly(TEXT("Grid_LoadingRange"), (float)LoadingRange);
+		MaterialInstance->SetVectorParameterValueEditorOnly(TEXT("Grid_Color"), GridColor);
+		//MaterialInstance->SetVectorParameterValueEditorOnly(TEXT("Grid_Offset"), GridOffset);		
+	}
+}
+#endif
\ No newline at end of file
diff --git a/Engine/Source/Runtime/Engine/Private/WorldPartition/RuntimeHashSet/RuntimePartition.cpp b/Engine/Source/Runtime/Engine/Private/WorldPartition/RuntimeHashSet/RuntimePartition.cpp
--- a/Engine/Source/Runtime/Engine/Private/WorldPartition/RuntimeHashSet/RuntimePartition.cpp
+++ b/Engine/Source/Runtime/Engine/Private/WorldPartition/RuntimeHashSet/RuntimePartition.cpp
@@ -18,6 +18,7 @@ void URuntimePartition::SetDefaultValues()
 	bClientOnlyVisible = false;
 	Priority = 0;
 	LoadingRange = 25600;
+	DebugColor = FLinearColor::MakeRandomSeededColor(GetTypeHash(GetName()));
 	HLODIndex = INDEX_NONE;
 }
 
diff --git a/Engine/Source/Runtime/Engine/Private/WorldPartition/RuntimeHashSet/RuntimePartitionLHGrid.cpp b/Engine/Source/Runtime/Engine/Private/WorldPartition/RuntimeHashSet/RuntimePartitionLHGrid.cpp
--- a/Engine/Source/Runtime/Engine/Private/WorldPartition/RuntimeHashSet/RuntimePartitionLHGrid.cpp
+++ b/Engine/Source/Runtime/Engine/Private/WorldPartition/RuntimeHashSet/RuntimePartitionLHGrid.cpp
@@ -56,16 +56,50 @@ struct FCellCoord
 	}
 };
 
-void URuntimePartitionLHGrid::PostEditChangeProperty(FPropertyChangedEvent& PropertyChangedEvent)
+static bool GPackageWasDirty = false;
+void URuntimePartitionLHGrid::PreEditChange(FProperty* InPropertyAboutToChange)
 {
-	const FName PropertyName = PropertyChangedEvent.GetPropertyName();
+	GPackageWasDirty = GetPackage()->IsDirty();
+	Super::PreEditChange(InPropertyAboutToChange);
+}
+
+void URuntimePartitionLHGrid::PostEditChangeProperty(FPropertyChangedEvent& InPropertyChangedEvent)
+{
+	const FName PropertyName = InPropertyChangedEvent.GetPropertyName();
 
 	if (PropertyName == GET_MEMBER_NAME_CHECKED(URuntimePartitionLHGrid, CellSize))
 	{
 		CellSize = FMath::Max<int32>(CellSize, 1600);
 	}
+	else if (PropertyName == GET_MEMBER_NAME_CHECKED(URuntimePartitionLHGrid, bShowGridPreview))
+	{
+		if (bShowGridPreview)
+		{
+			check(!WorldGridPreviewer);
+			WorldGridPreviewer = MakeUnique<FWorldGridPreviewer>(GetTypedOuter<UWorld>(), false);
+		}
+		else
+		{
+			check(WorldGridPreviewer);
+			WorldGridPreviewer.Reset();
+		}
+
+		if (!GPackageWasDirty)
+		{
+			GetPackage()->ClearDirtyFlag();
+		}
+	}
+
+	if (WorldGridPreviewer)
+	{
+		WorldGridPreviewer->CellSize = CellSize;
+		WorldGridPreviewer->GridColor = DebugColor;
+		WorldGridPreviewer->GridOffset = FVector::ZeroVector;
+		WorldGridPreviewer->LoadingRange = LoadingRange;
+		WorldGridPreviewer->Update();
+	}
 
-	Super::PostEditChangeProperty(PropertyChangedEvent);
+	Super::PostEditChangeProperty(InPropertyChangedEvent);
 }
 
 void URuntimePartitionLHGrid::InitHLODRuntimePartitionFrom(const URuntimePartition* InRuntimePartition, int32 InHLODIndex)
diff --git a/Engine/Source/Runtime/Engine/Private/WorldSettings.cpp b/Engine/Source/Runtime/Engine/Private/WorldSettings.cpp
--- a/Engine/Source/Runtime/Engine/Private/WorldSettings.cpp
+++ b/Engine/Source/Runtime/Engine/Private/WorldSettings.cpp
@@ -764,6 +764,29 @@ void AWorldSettings::PostEditChangeProperty(FPropertyChangedEvent& PropertyChang
 		}
 	}
 
+	if (PropName == GET_MEMBER_NAME_CHECKED(AWorldSettings, bShowInstancedFoliageGrid))
+	{
+		if (bShowInstancedFoliageGrid)
+		{
+			check(!InstancedFoliageGridGridPreviewer);
+			InstancedFoliageGridGridPreviewer = MakeUnique<FWorldGridPreviewer>(GetTypedOuter<UWorld>(), true);
+		}
+		else
+		{
+			check(InstancedFoliageGridGridPreviewer);
+			InstancedFoliageGridGridPreviewer.Reset();
+		}
+
+		if (InstancedFoliageGridGridPreviewer)
+		{
+			InstancedFoliageGridGridPreviewer->CellSize = InstancedFoliageGridSize;
+			InstancedFoliageGridGridPreviewer->GridColor = FColor::White;
+			InstancedFoliageGridGridPreviewer->GridOffset = FVector::ZeroVector;
+			InstancedFoliageGridGridPreviewer->LoadingRange = MAX_int32;
+			InstancedFoliageGridGridPreviewer->Update();
+		}
+	}
+
 	Super::PostEditChangeProperty(PropertyChangedEvent);
 }
 
diff --git a/Engine/Source/Runtime/Engine/Public/WorldGridPreviewer.h b/Engine/Source/Runtime/Engine/Public/WorldGridPreviewer.h
new file mode 100644
--- /dev/null
+++ b/Engine/Source/Runtime/Engine/Public/WorldGridPreviewer.h
@@ -0,0 +1,42 @@
+// Copyright Epic Games, Inc. All Rights Reserved.
+#pragma once
+
+#include "CoreMinimal.h"
+#include "UObject/GCObject.h"
+#include "Engine/PostProcessVolume.h"
+
+#if WITH_EDITOR
+
+class UWorld;
+class UMaterial;
+class UMaterialInstanceConstant;
+class APostProcessVolume;
+
+struct FWorldGridPreviewer : public FGCObject
+{
+public:
+	ENGINE_API FWorldGridPreviewer();
+	ENGINE_API FWorldGridPreviewer(UWorld* InWorld, bool bInIs2D);
+	ENGINE_API ~FWorldGridPreviewer();
+
+	//~ Begin FGCObject interface
+	virtual void AddReferencedObjects(FReferenceCollector& Collector) override;
+	virtual FString GetReferencerName() const override { return TEXT("FWorldGridPreviewer"); }
+	//~ End FGCObject interface
+
+	ENGINE_API void Update();
+
+	FName Owner;
+	int32 CellSize;
+	int32 LoadingRange;
+	FLinearColor GridColor;
+	FVector GridOffset;
+
+private:
+	TObjectPtr<UWorld> World;
+	TObjectPtr<UMaterial> Material;
+	TObjectPtr<UMaterialInstanceConstant> MaterialInstance;
+	TWeakObjectPtr<APostProcessVolume> PostProcessVolume;
+};
+
+#endif
diff --git a/Engine/Source/Runtime/Engine/Public/WorldPartition/RuntimeHashSet/RuntimePartition.h b/Engine/Source/Runtime/Engine/Public/WorldPartition/RuntimeHashSet/RuntimePartition.h
--- a/Engine/Source/Runtime/Engine/Public/WorldPartition/RuntimeHashSet/RuntimePartition.h
+++ b/Engine/Source/Runtime/Engine/Public/WorldPartition/RuntimeHashSet/RuntimePartition.h
@@ -86,6 +86,9 @@ public:
 	UPROPERTY(EditAnywhere, Category = RuntimeSettings)
 	int32 LoadingRange;
 
+	UPROPERTY(EditAnywhere, Category = RuntimeSettings)
+	FLinearColor DebugColor;
+
 	UPROPERTY()
 	int32 HLODIndex;
 #endif
diff --git a/Engine/Source/Runtime/Engine/Public/WorldPartition/RuntimeHashSet/RuntimePartitionLHGrid.h b/Engine/Source/Runtime/Engine/Public/WorldPartition/RuntimeHashSet/RuntimePartitionLHGrid.h
--- a/Engine/Source/Runtime/Engine/Public/WorldPartition/RuntimeHashSet/RuntimePartitionLHGrid.h
+++ b/Engine/Source/Runtime/Engine/Public/WorldPartition/RuntimeHashSet/RuntimePartitionLHGrid.h
@@ -2,6 +2,7 @@
 #pragma once
 
 #include "CoreMinimal.h"
+#include "WorldGridPreviewer.h"
 #include "WorldPartition/RuntimeHashSet/RuntimePartition.h"
 #include "RuntimePartitionLHGrid.generated.h"
 
@@ -15,7 +16,8 @@ class URuntimePartitionLHGrid : public URuntimePartition
 public:
 #if WITH_EDITOR
 	//~ Begin UObject Interface.
-	virtual void PostEditChangeProperty(FPropertyChangedEvent& PropertyChangedEvent) override;
+	virtual void PreEditChange(FProperty* InPropertyAboutToChange) override;
+	virtual void PostEditChangeProperty(FPropertyChangedEvent& InPropertyChangedEvent) override;
 	//~ End UObject Interface.
 
 	//~ Begin URuntimePartition interface
@@ -31,5 +33,12 @@ private:
 #if WITH_EDITORONLY_DATA
 	UPROPERTY(EditAnywhere, Category = RuntimeSettings)
 	uint32 CellSize;
+
+	UPROPERTY(EditAnywhere, Category = RuntimeSettings, Transient, SkipSerialization)
+	bool bShowGridPreview = false;
+#endif
+
+#if WITH_EDITOR
+	TUniquePtr<FWorldGridPreviewer> WorldGridPreviewer;
 #endif
 };
\ No newline at end of file
