Fix PreskinnedNormal of skeletal mesh when skin cache is on

diff --git a/Engine/Shaders/Private/LocalVertexFactory.ush b/Engine/Shaders/Private/LocalVertexFactory.ush
--- a/Engine/Shaders/Private/LocalVertexFactory.ush
+++ b/Engine/Shaders/Private/LocalVertexFactory.ush
@@ -677,7 +677,16 @@ FMaterialVertexParameters GetMaterialVertexParameters(FVertexFactoryInput Input,
 #endif
 
 	Result.PreSkinnedPosition = Intermediates.PreSkinPosition.xyz;
-	Result.PreSkinnedNormal = TangentToLocal[2]; //TangentBias(Input.TangentZ.xyz);
+#if SUPPORT_GPUSKIN_PASSTHROUGH
+	if (IsGPUSkinPassThrough())
+	{
+		Result.PreSkinnedNormal = LocalVFLooseParameters.GPUSkinPassThroughPreSkinnedTangentBuffer[Input.VertexId * 2 + 1].xyz;
+	}
+	else
+#endif
+	{
+		Result.PreSkinnedNormal = TangentToLocal[2]; //TangentBias(Input.TangentZ.xyz);
+	}
 
 #if MANUAL_VERTEX_FETCH && NUM_MATERIAL_TEXCOORDS_VERTEX
 		const uint NumFetchTexCoords = LocalVF.VertexFetch_Parameters[VF_NumTexcoords_Index];
diff --git a/Engine/Source/Runtime/Engine/Private/GPUSkinVertexFactory.cpp b/Engine/Source/Runtime/Engine/Private/GPUSkinVertexFactory.cpp
--- a/Engine/Source/Runtime/Engine/Private/GPUSkinVertexFactory.cpp
+++ b/Engine/Source/Runtime/Engine/Private/GPUSkinVertexFactory.cpp
@@ -1346,6 +1346,7 @@ void FGPUSkinPassthroughVertexFactory::CreateLooseUniformBuffer(FRHICommandListB
 	Parameters.FrameNumber = InFrameNumber;
 	Parameters.GPUSkinPassThroughPositionBuffer = PositionSRV;
 	Parameters.GPUSkinPassThroughPreviousPositionBuffer = PrevPositionSRV;
+	Parameters.GPUSkinPassThroughPreSkinnedTangentBuffer = InSourceVertexFactory->GetTangentsSRV();
 	LooseParametersUniformBuffer.UpdateUniformBufferImmediate(RHICmdList, Parameters);
 }
 
diff --git a/Engine/Source/Runtime/Engine/Private/LocalVertexFactory.cpp b/Engine/Source/Runtime/Engine/Private/LocalVertexFactory.cpp
--- a/Engine/Source/Runtime/Engine/Private/LocalVertexFactory.cpp
+++ b/Engine/Source/Runtime/Engine/Private/LocalVertexFactory.cpp
@@ -460,6 +460,7 @@ void FLocalVertexFactory::InitRHI(FRHICommandListBase& RHICmdList)
 	LooseParameters.FrameNumber = -1;
 	LooseParameters.GPUSkinPassThroughPositionBuffer = GNullVertexBuffer.VertexBufferSRV;
 	LooseParameters.GPUSkinPassThroughPreviousPositionBuffer = GNullVertexBuffer.VertexBufferSRV;
+	LooseParameters.GPUSkinPassThroughPreSkinnedTangentBuffer = GNullVertexBuffer.VertexBufferSRV;
 	LooseParametersUniformBuffer = TUniformBufferRef<FLocalVertexFactoryLooseParameters>::CreateUniformBufferImmediate(LooseParameters, UniformBuffer_MultiFrame);
 
 	check(IsValidRef(GetDeclaration()));
diff --git a/Engine/Source/Runtime/Engine/Private/SkeletalRenderGPUSkin.cpp b/Engine/Source/Runtime/Engine/Private/SkeletalRenderGPUSkin.cpp
--- a/Engine/Source/Runtime/Engine/Private/SkeletalRenderGPUSkin.cpp
+++ b/Engine/Source/Runtime/Engine/Private/SkeletalRenderGPUSkin.cpp
@@ -2309,6 +2309,13 @@ bool FSkeletalMeshObjectGPUSkin::GetCachedGeometry(FCachedGeometry& OutCachedGeo
 			CachedSection.PositionBuffer = FGPUSkinCache::GetPositionBuffer(SkinCacheEntry, SectionIndex)->SRV;
 			CachedSection.PreviousPositionBuffer = FGPUSkinCache::GetPreviousPositionBuffer(SkinCacheEntry, SectionIndex)->SRV;
 
+			if (SectionIndex > 0)
+			{
+				ensureAlways(OutCachedGeometry.Sections[0].PositionBuffer == CachedSection.PositionBuffer);
+				ensureAlways(OutCachedGeometry.Sections[0].PreviousPositionBuffer == CachedSection.PreviousPositionBuffer);
+			}
+
+
 			FRWBuffer* TangentBuffer = FGPUSkinCache::GetTangentBuffer(SkinCacheEntry, SectionIndex);
 			CachedSection.TangentBuffer = TangentBuffer ? TangentBuffer->SRV : nullptr;
 		}
diff --git a/Engine/Source/Runtime/Engine/Public/LocalVertexFactory.h b/Engine/Source/Runtime/Engine/Public/LocalVertexFactory.h
--- a/Engine/Source/Runtime/Engine/Public/LocalVertexFactory.h
+++ b/Engine/Source/Runtime/Engine/Public/LocalVertexFactory.h
@@ -31,6 +31,7 @@ BEGIN_GLOBAL_SHADER_PARAMETER_STRUCT(FLocalVertexFactoryLooseParameters,)
 	SHADER_PARAMETER(uint32, FrameNumber)
 	SHADER_PARAMETER_SRV(Buffer<float>, GPUSkinPassThroughPositionBuffer)
 	SHADER_PARAMETER_SRV(Buffer<float>, GPUSkinPassThroughPreviousPositionBuffer)
+	SHADER_PARAMETER_SRV(Buffer<float4>, GPUSkinPassThroughPreSkinnedTangentBuffer)
 END_GLOBAL_SHADER_PARAMETER_STRUCT()
 
 extern ENGINE_API TUniformBufferRef<FLocalVertexFactoryUniformShaderParameters> CreateLocalVFUniformBuffer(
