Revert change that was checked-in by accident at 29917138

diff --git a/Engine/Source/Runtime/Engine/Private/SkeletalRenderGPUSkin.cpp b/Engine/Source/Runtime/Engine/Private/SkeletalRenderGPUSkin.cpp
--- a/Engine/Source/Runtime/Engine/Private/SkeletalRenderGPUSkin.cpp
+++ b/Engine/Source/Runtime/Engine/Private/SkeletalRenderGPUSkin.cpp
@@ -2309,13 +2309,6 @@ bool FSkeletalMeshObjectGPUSkin::GetCachedGeometry(FCachedGeometry& OutCachedGeo
 			CachedSection.PositionBuffer = FGPUSkinCache::GetPositionBuffer(SkinCacheEntry, SectionIndex)->SRV;
 			CachedSection.PreviousPositionBuffer = FGPUSkinCache::GetPreviousPositionBuffer(SkinCacheEntry, SectionIndex)->SRV;
 
-			if (SectionIndex > 0)
-			{
-				ensureAlways(OutCachedGeometry.Sections[0].PositionBuffer == CachedSection.PositionBuffer);
-				ensureAlways(OutCachedGeometry.Sections[0].PreviousPositionBuffer == CachedSection.PreviousPositionBuffer);
-			}
-
-
 			FRWBuffer* TangentBuffer = FGPUSkinCache::GetTangentBuffer(SkinCacheEntry, SectionIndex);
 			CachedSection.TangentBuffer = TangentBuffer ? TangentBuffer->SRV : nullptr;
 		}
