Fixed compilation error introduced in CL# 29919031

diff --git a/Engine/Source/Runtime/Engine/Private/Materials/HLSLMaterialTranslator.cpp b/Engine/Source/Runtime/Engine/Private/Materials/HLSLMaterialTranslator.cpp
--- a/Engine/Source/Runtime/Engine/Private/Materials/HLSLMaterialTranslator.cpp
+++ b/Engine/Source/Runtime/Engine/Private/Materials/HLSLMaterialTranslator.cpp
@@ -241,7 +241,7 @@ struct FHLSLMaterialTranslator::FEnvironmentDefines
 	bool bUsesEyeAdaptation;
 	bool bVirtualTextureOutput;
 	bool bUsesPerInstanceCustomData;
-	bool bUsesPerInstanceRandom;
+	bool bUsesPerInstanceRandomPS;
 	bool bUsesPerInstanceFadeAmount;
 	bool bUsesVertexInterpolator;
 	bool bUsesSkyAtmosphere;
@@ -333,7 +333,7 @@ struct FHLSLMaterialTranslator::FEnvironmentDefines
 		Ar << bUsesEyeAdaptation;
 		Ar << bVirtualTextureOutput;
 		Ar << bUsesPerInstanceCustomData;
-		Ar << bUsesPerInstanceRandom;
+		Ar << bUsesPerInstanceRandomPS;
 		Ar << bUsesPerInstanceFadeAmount;
 		Ar << bUsesVertexInterpolator;
 		Ar << bUsesSkyAtmosphere;
@@ -2388,7 +2388,6 @@ void FHLSLMaterialTranslator::GetMaterialEnvironment(EShaderPlatform InPlatform,
 	}
 
 	OutEnvironment.SetDefine(TEXT("USES_PER_INSTANCE_CUSTOM_DATA"), EnvironmentDefines->bUsesPerInstanceCustomData);
-	OutEnvironment.SetDefine(TEXT("USES_PER_INSTANCE_RANDOM"), EnvironmentDefines->bUsesPerInstanceRandom);
 	OutEnvironment.SetDefine(TEXT("USES_PER_INSTANCE_FADE_AMOUNT"), EnvironmentDefines->bUsesPerInstanceFadeAmount);
 	OutEnvironment.SetDefine(TEXT("USES_VERTEX_INTERPOLATOR"), EnvironmentDefines->bUsesVertexInterpolator);
 
@@ -2399,6 +2398,7 @@ void FHLSLMaterialTranslator::GetMaterialEnvironment(EShaderPlatform InPlatform,
 	OutEnvironment.SetDefine(TEXT("NEEDS_PARTICLE_WORLD_TO_LOCAL"), EnvironmentDefines->bUsesParticleWorldToLocal);
 	OutEnvironment.SetDefine(TEXT("NEEDS_INSTANCE_LOCAL_TO_WORLD_PS"), EnvironmentDefines->bUsesInstanceLocalToWorldPS);
 	OutEnvironment.SetDefine(TEXT("NEEDS_INSTANCE_WORLD_TO_LOCAL_PS"), EnvironmentDefines->bUsesInstanceWorldToLocalPS);
+	OutEnvironment.SetDefine(TEXT("NEEDS_PER_INSTANCE_RANDOM_PS"), EnvironmentDefines->bUsesPerInstanceRandomPS);
 	OutEnvironment.SetDefine(TEXT("USES_TRANSFORM_VECTOR"), EnvironmentDefines->bUsesTransformVector);
 	OutEnvironment.SetDefine(TEXT("WANT_PIXEL_DEPTH_OFFSET"), EnvironmentDefines->bUsesPixelDepthOffset);
 
@@ -14905,7 +14905,7 @@ void FHLSLMaterialTranslator::PrepareEnvironmentDefines()
 	EnvironmentDefines->bUsesEyeAdaptation = MaterialCompilationOutput.bUsesEyeAdaptation;
 	EnvironmentDefines->bVirtualTextureOutput = MaterialCompilationOutput.bHasRuntimeVirtualTextureOutputNode;
 	EnvironmentDefines->bUsesPerInstanceCustomData = MaterialCompilationOutput.bUsesPerInstanceCustomData && Material->IsUsedWithInstancedStaticMeshes();
-	EnvironmentDefines->bUsesPerInstanceRandom = MaterialCompilationOutput.bUsesPerInstanceRandom && Material->IsUsedWithInstancedStaticMeshes();
+	EnvironmentDefines->bUsesPerInstanceRandomPS = bUsesPerInstanceRandomPS && Material->IsUsedWithInstancedStaticMeshes();
 	EnvironmentDefines->bUsesPerInstanceFadeAmount = bUsesPerInstanceFadeAmount && Material->IsUsedWithInstancedStaticMeshes();
 	EnvironmentDefines->bUsesVertexInterpolator = MaterialCompilationOutput.bUsesVertexInterpolator;
 	EnvironmentDefines->bUsesSkyAtmosphere = bUsesSkyAtmosphere;
