Add support to the virtual texture SVT builder to disable cook per platform.
Initialize it with the per platform RVT volume settings when baking the SVT.

diff --git a/Engine/Source/Runtime/Engine/Classes/VT/VirtualTextureBuilder.h b/Engine/Source/Runtime/Engine/Classes/VT/VirtualTextureBuilder.h
--- a/Engine/Source/Runtime/Engine/Classes/VT/VirtualTextureBuilder.h
+++ b/Engine/Source/Runtime/Engine/Classes/VT/VirtualTextureBuilder.h
@@ -4,6 +4,7 @@
 
 #include "CoreMinimal.h"
 #include "Engine/Texture.h"
+#include "PerPlatformProperties.h"
 #include "VirtualTextureBuilder.generated.h"
 
 enum class EShadingPath;
@@ -68,6 +69,10 @@ public:
 	UPROPERTY(EditAnywhere, Category = Texture)
 	bool bSeparateTextureForMobile = false;
 
+	/** Per platform overrides for cooking the virtual texture. */
+	UPROPERTY(EditAnywhere, Category = Texture)
+	FPerPlatformBool EnableCookPerPlatform;
+
 #if WITH_EDITOR
 	/** Creates a new UVirtualTexture2D and stores it in the contained Texture. */
 	ENGINE_API void BuildTexture(EShadingPath ShadingPath, FVirtualTextureBuildDesc const& BuildDesc);
diff --git a/Engine/Source/Runtime/Engine/Private/Components/RuntimeVirtualTextureComponent.cpp b/Engine/Source/Runtime/Engine/Private/Components/RuntimeVirtualTextureComponent.cpp
--- a/Engine/Source/Runtime/Engine/Private/Components/RuntimeVirtualTextureComponent.cpp
+++ b/Engine/Source/Runtime/Engine/Private/Components/RuntimeVirtualTextureComponent.cpp
@@ -358,8 +358,9 @@ void URuntimeVirtualTextureComponent::InitializeStreamingTexture(EShadingPath Sh
 		BuildDesc.InSizeY = InSizeY;
 		BuildDesc.InData = InData;
 
-		StreamingTexture->Modify();
 		StreamingTexture->BuildTexture(ShadingPath, BuildDesc);
+		StreamingTexture->EnableCookPerPlatform = EnableInGamePerPlatform;
+		StreamingTexture->Modify();
 	}
 }
 
diff --git a/Engine/Source/Runtime/Engine/Private/VT/VirtualTextureBuilder.cpp b/Engine/Source/Runtime/Engine/Private/VT/VirtualTextureBuilder.cpp
--- a/Engine/Source/Runtime/Engine/Private/VT/VirtualTextureBuilder.cpp
+++ b/Engine/Source/Runtime/Engine/Private/VT/VirtualTextureBuilder.cpp
@@ -15,6 +15,7 @@
 
 UVirtualTextureBuilder::UVirtualTextureBuilder(const FObjectInitializer& ObjectInitializer)
 	: UObject(ObjectInitializer)
+	, EnableCookPerPlatform(true)
 {
 }
 
@@ -37,6 +38,13 @@ void UVirtualTextureBuilder::Serialize(FArchive& Ar)
 			TextureMobile = nullptr;
 		}
 
+		// Clear during cook for platforms that have explicitly disabled cooking in the asset settings.
+		if (!EnableCookPerPlatform.GetValueForPlatform(*Ar.CookingTarget()->PlatformName()))
+		{
+			Texture = nullptr;
+			TextureMobile = nullptr;
+		}
+
 		// Selectivly serialize VirtualTexture for platforms that support Deferred or/and Mobile rendering
 		if (!Ar.CookingTarget()->SupportsFeature(ETargetPlatformFeatures::DeferredRendering) && bSeparateTextureForMobile)
 		{
