Do not set an external package when we spawn a transient child actor, else the package will be marked as dirty
#rnx

diff --git a/Engine/Source/Runtime/Engine/Private/Components/ChildActorComponent.cpp b/Engine/Source/Runtime/Engine/Private/Components/ChildActorComponent.cpp
--- a/Engine/Source/Runtime/Engine/Private/Components/ChildActorComponent.cpp
+++ b/Engine/Source/Runtime/Engine/Private/Components/ChildActorComponent.cpp
@@ -795,10 +795,16 @@ void UChildActorComponent::CreateChildActor(TFunction<void(AActor*)> CustomizerF
 				}
 
 				Params.OverrideParentComponent = this;
+				
+				if (bChildActorIsTransient || HasAllFlags(RF_Transient) || (MyOwner && MyOwner->HasAllFlags(RF_Transient)))
+				{
+					// If this component or its owner are transient, set our created actor to transient. 
+					Params.ObjectFlags |= RF_Transient;
+				}
 
 #if WITH_EDITOR
 				Params.bCreateActorPackage = false;
-				Params.OverridePackage = (MyOwner ? MyOwner->GetExternalPackage() : nullptr);
+				Params.OverridePackage = (MyOwner && !(Params.ObjectFlags & RF_Transient)) ? MyOwner->GetExternalPackage() : nullptr;
 				Params.OverrideActorGuid = CachedInstanceData ? CachedInstanceData->ChildActorGUID : FGuid();
 				Params.bHideFromSceneOutliner = EditorTreeViewVisualizationMode == EChildActorComponentTreeViewVisualizationMode::Hidden;
 #endif
@@ -811,11 +817,7 @@ void UChildActorComponent::CreateChildActor(TFunction<void(AActor*)> CustomizerF
 				{
 					Params.ObjectFlags &= ~RF_Transactional;
 				}
-				if (bChildActorIsTransient || HasAllFlags(RF_Transient) || (MyOwner && MyOwner->HasAllFlags(RF_Transient)))
-				{
-					// If this component or its owner are transient, set our created actor to transient. 
-					Params.ObjectFlags |= RF_Transient;
-				}
+
 
 				// Spawn actor of desired class
 				ConditionalUpdateComponentToWorld();
