World Partition: when saving actor that are unregistered in their owning world, get their actor descriptor world data from the loaded path instead of the current destination.
This was breaking migration as the code was trying to retrieve the missing data from the destination instead of the source.

diff --git a/Engine/Source/Runtime/Engine/Private/WorldPartition/WorldPartitionActorDescUtils.cpp b/Engine/Source/Runtime/Engine/Private/WorldPartition/WorldPartitionActorDescUtils.cpp
--- a/Engine/Source/Runtime/Engine/Private/WorldPartition/WorldPartitionActorDescUtils.cpp
+++ b/Engine/Source/Runtime/Engine/Private/WorldPartition/WorldPartitionActorDescUtils.cpp
@@ -116,7 +116,7 @@ void FWorldPartitionActorDescUtils::AppendAssetDataTagsFromActor(const AActor* I
 
 			FARFilter Filter;
 			Filter.bIncludeOnlyOnDiskAssets = true;
-			Filter.PackageNames.Add(InActor->GetPackage()->GetFName());
+			Filter.PackageNames.Add(InActor->GetPackage()->GetLoadedPath().GetPackageFName());
 
 			TArray<FAssetData> Assets;
 			AssetRegistry.GetAssets(Filter, Assets);
