Static analysis warning issue 400315:  warning V595: The 'UniformExpressionSet' pointer was utilized before it was verified against nullptr.  Removed null check and changed higher level API to accept reference rather than pointer, to make it impossible at compile time for a null pointer to make it there.

#rnx

diff --git a/Engine/Source/Runtime/Engine/Private/Materials/MaterialShared.cpp b/Engine/Source/Runtime/Engine/Private/Materials/MaterialShared.cpp
--- a/Engine/Source/Runtime/Engine/Private/Materials/MaterialShared.cpp
+++ b/Engine/Source/Runtime/Engine/Private/Materials/MaterialShared.cpp
@@ -4355,7 +4355,7 @@ bool FMaterial::GetMaterialExpressionSource( FString& OutSource )
 					const FMaterialUniformPreshaderField& PreshaderField = NewCompilationOutput.UniformExpressionSet.UniformPreshaderFields[PreshaderIndex];
 
 					UE::Shader::FPreshaderDataContext PreshaderContext(PreshaderContextBase, PreshaderHeader.OpcodeOffset, PreshaderHeader.OpcodeSize);
-					FString PreshaderDebug = PreshaderGenerateDebugString(&NewCompilationOutput.UniformExpressionSet, MaterialContext, PreshaderContext, &ParameterReferences);
+					FString PreshaderDebug = PreshaderGenerateDebugString(NewCompilationOutput.UniformExpressionSet, MaterialContext, PreshaderContext, &ParameterReferences);
 
 					// If this is a numeric field, add a swizzle for it
 					const TCHAR* SwizzleSuffix = TEXT("");
@@ -4432,7 +4432,7 @@ void FMaterial::GetPreshaderStats(uint32& TotalParameters, uint32& TotalOps) con
 		for (const FMaterialUniformPreshaderHeader& PreshaderHeader : UniformExpressionSet.UniformPreshaders)
 		{
 			UE::Shader::FPreshaderDataContext PreshaderContext(PreshaderContextBase, PreshaderHeader.OpcodeOffset, PreshaderHeader.OpcodeSize);
-			PreshaderComputeDebugStats(&UniformExpressionSet, MaterialContext, PreshaderContext, TotalParameters, TotalOps);
+			PreshaderComputeDebugStats(UniformExpressionSet, MaterialContext, PreshaderContext, TotalParameters, TotalOps);
 		}
 	}
 #else
diff --git a/Engine/Source/Runtime/Engine/Private/Shader/Preshader.cpp b/Engine/Source/Runtime/Engine/Private/Shader/Preshader.cpp
--- a/Engine/Source/Runtime/Engine/Private/Shader/Preshader.cpp
+++ b/Engine/Source/Runtime/Engine/Private/Shader/Preshader.cpp
@@ -693,13 +693,6 @@ static void EvaluateParameter(FPreshaderDebugStack& Stack, const FUniformExpress
 
 	if (Stack.bGenerateString)
 	{
-		if (!UniformExpressionSet)
-		{
-			// return 0 for parameters if we don't have UniformExpressionSet
-			Stack.Push(TEXT("0.0"), EPreshaderDebugStackType::Singular);
-			return;
-		}
-
 		Stack.Push(FString::Printf(TEXT("Param[\"%s\"]"), *ParameterName), EPreshaderDebugStackType::Singular);
 	}
 }
@@ -1497,20 +1490,20 @@ static void EvaluatePreshaderDebug(const FUniformExpressionSet* UniformExpressio
 #undef EvaluateBinaryOpInPlace
 #undef EvaluateTernaryOp
 
-FString PreshaderGenerateDebugString(const FUniformExpressionSet* UniformExpressionSet, const FMaterialRenderContext& Context, FPreshaderDataContext& RESTRICT Data, TMap<FString, uint32>* ParameterReferences)
+FString PreshaderGenerateDebugString(const FUniformExpressionSet& UniformExpressionSet, const FMaterialRenderContext& Context, FPreshaderDataContext& RESTRICT Data, TMap<FString, uint32>* ParameterReferences)
 {
 	FPreshaderDebugStack Stack;
 	Stack.ParameterReferences = ParameterReferences;
-	EvaluatePreshaderDebug(UniformExpressionSet, Context, Stack, Data);
+	EvaluatePreshaderDebug(&UniformExpressionSet, Context, Stack, Data);
 
 	return Stack.Elements.Pop().Text;
 }
 
-void PreshaderComputeDebugStats(const FUniformExpressionSet* UniformExpressionSet, const FMaterialRenderContext& Context, FPreshaderDataContext& RESTRICT Data, uint32& TotalParameters, uint32& TotalOps)
+void PreshaderComputeDebugStats(const FUniformExpressionSet& UniformExpressionSet, const FMaterialRenderContext& Context, FPreshaderDataContext& RESTRICT Data, uint32& TotalParameters, uint32& TotalOps)
 {
 	FPreshaderDebugStack Stack;
 	Stack.bGenerateString = false;
-	EvaluatePreshaderDebug(UniformExpressionSet, Context, Stack, Data);
+	EvaluatePreshaderDebug(&UniformExpressionSet, Context, Stack, Data);
 
 	TotalParameters += Stack.ParameterCount;
 	TotalOps += Stack.OpCount;
diff --git a/Engine/Source/Runtime/Engine/Private/Shader/PreshaderEvaluate.h b/Engine/Source/Runtime/Engine/Private/Shader/PreshaderEvaluate.h
--- a/Engine/Source/Runtime/Engine/Private/Shader/PreshaderEvaluate.h
+++ b/Engine/Source/Runtime/Engine/Private/Shader/PreshaderEvaluate.h
@@ -28,8 +28,8 @@ struct FPreshaderDataContext
 
 ENGINE_API FPreshaderValue EvaluatePreshader(const FUniformExpressionSet* UniformExpressionSet, const FMaterialRenderContext& Context, FPreshaderStack& Stack, FPreshaderDataContext& RESTRICT Data);
 
-ENGINE_API FString PreshaderGenerateDebugString(const FUniformExpressionSet* UniformExpressionSet, const FMaterialRenderContext& Context, FPreshaderDataContext& RESTRICT Data, TMap<FString, uint32>* ParameterReferences = nullptr);
-ENGINE_API void PreshaderComputeDebugStats(const FUniformExpressionSet* UniformExpressionSet, const FMaterialRenderContext& Context, FPreshaderDataContext& RESTRICT Data, uint32& TotalParameters, uint32& TotalOps);
+ENGINE_API FString PreshaderGenerateDebugString(const FUniformExpressionSet& UniformExpressionSet, const FMaterialRenderContext& Context, FPreshaderDataContext& RESTRICT Data, TMap<FString, uint32>* ParameterReferences = nullptr);
+ENGINE_API void PreshaderComputeDebugStats(const FUniformExpressionSet& UniformExpressionSet, const FMaterialRenderContext& Context, FPreshaderDataContext& RESTRICT Data, uint32& TotalParameters, uint32& TotalOps);
 
 } // namespace Shader
 } // namespace UE
