followup deprecation of bIgnoreGammaConversions

diff --git a/Engine/Plugins/Runtime/nDisplay/Source/DisplayCluster/Private/Render/Viewport/RenderTarget/DisplayClusterRenderTargetResource.cpp b/Engine/Plugins/Runtime/nDisplay/Source/DisplayCluster/Private/Render/Viewport/RenderTarget/DisplayClusterRenderTargetResource.cpp
--- a/Engine/Plugins/Runtime/nDisplay/Source/DisplayCluster/Private/Render/Viewport/RenderTarget/DisplayClusterRenderTargetResource.cpp
+++ b/Engine/Plugins/Runtime/nDisplay/Source/DisplayCluster/Private/Render/Viewport/RenderTarget/DisplayClusterRenderTargetResource.cpp
@@ -14,10 +14,6 @@ FDisplayClusterViewportTextureResource::FDisplayClusterViewportTextureResource(c
 {
 	bSRGB = EnumHasAnyFlags(GetResourceSettings().GetResourceFlags(), EDisplayClusterViewportResourceSettingsFlags::ShouldUseSRGB);
 	bGreyScaleFormat = false;
-	
-	// I believe this setting of bIgnoreGammaConversions was probably a bad copy-paste and is probably not intentional
-	//	bIgnoreGammaConversions is deprecated.
-	bIgnoreGammaConversions = true;
 }
 
 void FDisplayClusterViewportTextureResource::InitRHI(FRHICommandListBase&)
diff --git a/Engine/Plugins/VirtualProduction/TextureShare/Source/TextureShare/Private/Resources/TextureShareResource.cpp b/Engine/Plugins/VirtualProduction/TextureShare/Source/TextureShare/Private/Resources/TextureShareResource.cpp
--- a/Engine/Plugins/VirtualProduction/TextureShare/Source/TextureShare/Private/Resources/TextureShareResource.cpp
+++ b/Engine/Plugins/VirtualProduction/TextureShare/Source/TextureShare/Private/Resources/TextureShareResource.cpp
@@ -25,10 +25,6 @@ FTextureShareResource::FTextureShareResource(const TSharedRef<ITextureShareCoreO
 	bSRGB = InResourceSettings.bShouldUseSRGB;
 
 	bGreyScaleFormat = false;
-	
-	// I believe this setting of bIgnoreGammaConversions was probably a bad copy-paste and is probably not intentional
-	//	bIgnoreGammaConversions is deprecated.
-	bIgnoreGammaConversions = true;
 }
 
 FTextureShareResource::~FTextureShareResource()
diff --git a/Engine/Source/Runtime/Engine/Private/Texture2DDynamic.cpp b/Engine/Source/Runtime/Engine/Private/Texture2DDynamic.cpp
--- a/Engine/Source/Runtime/Engine/Private/Texture2DDynamic.cpp
+++ b/Engine/Source/Runtime/Engine/Private/Texture2DDynamic.cpp
@@ -58,7 +58,6 @@ void FTexture2DDynamicResource::InitRHI(FRHICommandListBase&)
 	if (Owner->bIsResolveTarget)
 	{
 		Desc.AddFlags(ETextureCreateFlags::ResolveTargetable);
-		bIgnoreGammaConversions = true;		// Note, we're ignoring Owner->SRGB (it should be false).
 	}
 	else if (Owner->SRGB)
 	{
diff --git a/Engine/Source/Runtime/Engine/Private/UserInterface/Canvas.cpp b/Engine/Source/Runtime/Engine/Private/UserInterface/Canvas.cpp
--- a/Engine/Source/Runtime/Engine/Private/UserInterface/Canvas.cpp
+++ b/Engine/Source/Runtime/Engine/Private/UserInterface/Canvas.cpp
@@ -437,12 +437,6 @@ bool FCanvasBatchedElementRenderItem::Render_RenderThread(FCanvasRenderContext&
 			const FRenderTarget* CanvasRenderTarget = Canvas->GetRenderTarget();
 			float Gamma = 1.0f / CanvasRenderTarget->GetDisplayGamma();
 
-			// bIgnoreGammaConversions is deprecated; remove this
-			if (LocalData->Texture && LocalData->Texture->bIgnoreGammaConversions)
-			{
-				Gamma = 1.0f;
-			}
-
 			// draw batched items
 			LocalData->BatchedElements.Draw(
 				RHICmdList,
@@ -477,12 +471,6 @@ bool FCanvasBatchedElementRenderItem::Render_GameThread(const FCanvas* Canvas, F
 		const FRenderTarget* CanvasRenderTarget = Canvas->GetRenderTarget();
 		float Gamma = 1.0f / CanvasRenderTarget->GetDisplayGamma();
 		
-		// bIgnoreGammaConversions is deprecated; remove this
-		if ( Data->Texture && Data->Texture->bIgnoreGammaConversions )
-		{
-			Gamma = 1.0f;
-		}
-
 		// Render the batched elements.
 		struct FBatchedDrawParameters
 		{
diff --git a/Engine/Source/Runtime/RenderCore/Public/RenderResource.h b/Engine/Source/Runtime/RenderCore/Public/RenderResource.h
--- a/Engine/Source/Runtime/RenderCore/Public/RenderResource.h
+++ b/Engine/Source/Runtime/RenderCore/Public/RenderResource.h
@@ -381,9 +381,11 @@ public:
 	 * true if the texture is in the same gamma space as the intended rendertarget (e.g. screenshots).
 	 * The texture will have sRGB==false and bIgnoreGammaConversions==true, causing a non-sRGB texture lookup
 	 * and no gamma-correction in the shader.
-	 *
-	 * I believe this comment is no longer correct and bIgnoreGammaConversions is now deprecated and should not be used.
+	 * 
+	 * This was only ever checked in the Canvas renderer, not the standard Material shader path.
+	 * It is no longer set or checked.
 	 */
+	UE_DEPRECATED(5.5,"bIgnoreGammaConversions should not be used")
 	bool				bIgnoreGammaConversions = false;
 
 	/** 
