[NetPhysics] NetworkPhysicsComponent custom struct NetSerialize to work with Iris

diff --git a/Engine/Source/Runtime/Engine/Private/PhysicsEngine/NetworkPhysicsComponent.cpp b/Engine/Source/Runtime/Engine/Private/PhysicsEngine/NetworkPhysicsComponent.cpp
--- a/Engine/Source/Runtime/Engine/Private/PhysicsEngine/NetworkPhysicsComponent.cpp
+++ b/Engine/Source/Runtime/Engine/Private/PhysicsEngine/NetworkPhysicsComponent.cpp
@@ -12,6 +12,10 @@
 #include "PhysicsReplication.h"
 #include "PhysicsProxy/SingleParticlePhysicsProxy.h"
 
+#if UE_WITH_IRIS
+#include "Iris/ReplicationState/PropertyNetSerializerInfoRegistry.h"
+#endif // UE_WITH_IRIS
+
 #include UE_INLINE_GENERATED_CPP_BY_NAME(NetworkPhysicsComponent)
 
 /** These CVars are deprecated from UE 5.4, physics frame offset for networked physics prediction is now handled via PlayerController with automatic time dilation
@@ -115,6 +119,11 @@ FNetworkPhysicsRewindDataProxy& FNetworkPhysicsRewindDataProxy::operator=(const
 	return *this;
 }
 
+#if UE_WITH_IRIS
+UE_NET_IMPLEMENT_NAMED_STRUCT_LASTRESORT_NETSERIALIZER_AND_REGISTRY_DELEGATES(NetworkPhysicsRewindDataInputProxy);
+UE_NET_IMPLEMENT_NAMED_STRUCT_LASTRESORT_NETSERIALIZER_AND_REGISTRY_DELEGATES(NetworkPhysicsRewindDataStateProxy);
+#endif // UE_WITH_IRIS
+
 bool FNetworkPhysicsRewindDataInputProxy::NetSerialize(FArchive& Ar, class UPackageMap* Map, bool& bOutSuccess)
 {
 	return NetSerializeBase(Ar, Map, bOutSuccess, [this]() { return Owner->ReplicatedInputs.History->CreateNew(); });
