SpeedMultiplier property on InterpToMovementComponent so the timing of its movement can be changed without resetting the whole thing

#tests PIE

diff --git a/Engine/Source/Runtime/Engine/Classes/Components/InterpToMovementComponent.h b/Engine/Source/Runtime/Engine/Classes/Components/InterpToMovementComponent.h
--- a/Engine/Source/Runtime/Engine/Classes/Components/InterpToMovementComponent.h
+++ b/Engine/Source/Runtime/Engine/Classes/Components/InterpToMovementComponent.h
@@ -174,6 +174,12 @@ class UInterpToMovementComponent : public UMovementComponent
 	UPROPERTY(EditAnywhere, BlueprintReadWrite, meta=(ClampMin="0.0166", ClampMax="0.50", UIMin="0.0166", UIMax="0.50"), Category=Simulation)
 	float MaxSimulationTimeStep;
 
+	/**
+	 * Change the speed of movement.
+	 */
+	UPROPERTY(EditAnywhere, BlueprintReadWrite, Category = Simulation)
+	float SpeedMultiplier;
+
 	/**
 	 * Max number of iterations used for each discrete simulation step.
 	 * Increasing this value can address issues with fast-moving objects or complex collision scenarios, at the cost of performance.
diff --git a/Engine/Source/Runtime/Engine/Private/Components/InterpToMovementComponent.cpp b/Engine/Source/Runtime/Engine/Private/Components/InterpToMovementComponent.cpp
--- a/Engine/Source/Runtime/Engine/Private/Components/InterpToMovementComponent.cpp
+++ b/Engine/Source/Runtime/Engine/Private/Components/InterpToMovementComponent.cpp
@@ -77,6 +77,10 @@ void UInterpToMovementComponent::TickComponent(float DeltaTime, enum ELevelTick
 	UpdateControlPoints(false);
 
 	float RemainingTime = DeltaTime;
+	if (SpeedMultiplier > 0.f)
+	{
+		RemainingTime *= SpeedMultiplier;
+	}
 	int32 NumBounces = 0;
 	int32 Iterations = 0;
 	FHitResult Hit(1.f);
@@ -94,7 +98,7 @@ void UInterpToMovementComponent::TickComponent(float DeltaTime, enum ELevelTick
 		RemainingTime -= TimeTick;
 
 		// Calculate the current alpha with this tick iteration
-		const float TargetTime = FMath::Clamp(CurrentTime + ((TimeTick*TimeMultiplier)*CurrentDirection), 0.0f, 1.0f);		
+		const float TargetTime = FMath::Clamp(CurrentTime + ((TimeTick*TimeMultiplier)*CurrentDirection), 0.0f, 1.0f);
 		FVector MoveDelta = ComputeMoveDelta(TargetTime);
 		
 		// Update velocity
