World Grid Previewer: don't outer material instance to the actor to avoid keeping a reference to the level when unloading map.

diff --git a/Engine/Source/Runtime/Engine/Private/WorldGridPreviewer.cpp b/Engine/Source/Runtime/Engine/Private/WorldGridPreviewer.cpp
--- a/Engine/Source/Runtime/Engine/Private/WorldGridPreviewer.cpp
+++ b/Engine/Source/Runtime/Engine/Private/WorldGridPreviewer.cpp
@@ -27,7 +27,7 @@ FWorldGridPreviewer::FWorldGridPreviewer(UWorld* InWorld, bool bInIs2D)
 	
 	PostProcessVolume = World->SpawnActor<APostProcessVolume>(SpawnParameters);
 	Material = PostProcessVolume.IsValid() ? LoadObject<UMaterial>(nullptr, TEXT("/Engine/EditorMaterials/WorldGridPreviewMaterial")) : nullptr;
-	MaterialInstance = NewObject<UMaterialInstanceConstant>(PostProcessVolume.Get());
+	MaterialInstance = NewObject<UMaterialInstanceConstant>();
 	MaterialInstance->Parent = Material;
 
 	if (MaterialInstance)
