Allow to add additional flags to stereo layers

diff --git a/Engine/Source/Runtime/Engine/Classes/Components/StereoLayerComponent.h b/Engine/Source/Runtime/Engine/Classes/Components/StereoLayerComponent.h
--- a/Engine/Source/Runtime/Engine/Classes/Components/StereoLayerComponent.h
+++ b/Engine/Source/Runtime/Engine/Classes/Components/StereoLayerComponent.h
@@ -45,6 +45,16 @@ enum EStereoLayerShape : int
 	SLSH_MAX,
 };
 
+UCLASS()
+class UEditorFlagCollector : public UObject
+{
+	GENERATED_BODY()
+
+public:
+	UFUNCTION()
+	static TArray<FName> GetFlagNames();
+};
+
 UCLASS(EditInlineNew, Abstract, BlueprintType, CollapseCategories, MinimalAPI)
 class UStereoLayerShape : public UObject
 {
@@ -264,6 +274,7 @@ class UStereoLayerComponent : public USceneComponent
 public:
 
 	//~ Begin UObject Interface
+	ENGINE_API virtual void EndPlay(const EEndPlayReason::Type EndPlayReason) override;
 	ENGINE_API virtual void OnUnregister() override;
 	ENGINE_API virtual void Serialize(FArchive& Ar) override;
 	ENGINE_API virtual void PostLoad() override;
@@ -372,6 +383,10 @@ public:
 	UPROPERTY(EditAnywhere, BlueprintReadWrite, Category = "StereoLayer")
 	uint32 bQuadPreserveTextureRatio : 1;
 
+	/** Additional flags not included in IStereoLayers::ELayerFlags */
+	UPROPERTY(EditAnywhere, BlueprintReadOnly, Category = "StereoLayer", Meta=(GetOptions= "EditorFlagCollector.GetFlagNames"))
+	TArray<FName> AdditionalFlags;
+
 protected:
 	/** Size of the rendered stereo layer quad **/
 	UPROPERTY(EditAnywhere, BlueprintReadOnly, export, Category="StereoLayer")
diff --git a/Engine/Source/Runtime/Engine/Private/Components/StereoLayerComponent.cpp b/Engine/Source/Runtime/Engine/Private/Components/StereoLayerComponent.cpp
--- a/Engine/Source/Runtime/Engine/Private/Components/StereoLayerComponent.cpp
+++ b/Engine/Source/Runtime/Engine/Private/Components/StereoLayerComponent.cpp
@@ -6,6 +6,7 @@
 #include "TextureResource.h"
 #include "Engine/Texture.h"
 #include "StereoRendering.h"
+#include "StereoLayerAdditionalFlagsManager.h"
 #if WITH_EDITOR
 #include "SceneManagement.h"
 #else
@@ -37,6 +38,14 @@ UStereoLayerComponent::UStereoLayerComponent(const FObjectInitializer& ObjectIni
 	//Shape = ObjectInitializer.CreateDefaultSubobject<UStereoLayerShapeQuad>(this, TEXT("Shape"));
 }
 
+void UStereoLayerComponent::EndPlay(const EEndPlayReason::Type EndPlayReason)
+{
+	if (EndPlayReason == EEndPlayReason::EndPlayInEditor || EndPlayReason == EEndPlayReason::Quit)
+	{
+		FStereoLayerAdditionalFlagsManager::Destroy();
+	}
+}
+
 void UStereoLayerComponent::OnUnregister()
 {
 	Super::OnUnregister();
@@ -162,6 +171,12 @@ void UStereoLayerComponent::TickComponent(float DeltaTime, enum ELevelTick TickT
 		LayerDesc.Flags |= (bSupportsDepth) ? IStereoLayers::LAYER_FLAG_SUPPORT_DEPTH : 0;
 		LayerDesc.Flags |= (!bCurrVisible) ? IStereoLayers::LAYER_FLAG_HIDDEN : 0;
 
+		TSharedPtr<FStereoLayerAdditionalFlagsManager> FlagsManager = FStereoLayerAdditionalFlagsManager::Get();
+		for (FName& Flag : AdditionalFlags)
+		{
+			LayerDesc.Flags |= FlagsManager->GetFlagValue(Flag);
+		}
+
 		switch (StereoLayerType)
 		{
 		case SLT_WorldLocked:
@@ -419,3 +434,10 @@ bool FEquirectProps::operator==(const FEquirectProps& Other) const
 {
 	return (LeftUVRect == Other.LeftUVRect) && (RightUVRect == Other.RightUVRect) && (LeftScale == Other.LeftScale) && (RightScale == Other.RightScale) && (LeftBias == Other.LeftBias) && (RightBias == Other.RightBias) && (Radius == Other.Radius);
 }
+
+TArray<FName> UEditorFlagCollector::GetFlagNames()
+{
+	TSet<FName> UniqueFlags;
+	FStereoLayerAdditionalFlagsManager::CollectFlags(UniqueFlags);
+	return UniqueFlags.Array();
+}
diff --git a/Engine/Source/Runtime/Engine/Private/StereoLayerAdditionalFlagsManager.cpp b/Engine/Source/Runtime/Engine/Private/StereoLayerAdditionalFlagsManager.cpp
new file mode 100644
--- /dev/null
+++ b/Engine/Source/Runtime/Engine/Private/StereoLayerAdditionalFlagsManager.cpp
@@ -0,0 +1,70 @@
+// Copyright Epic Games, Inc. All Rights Reserved.
+
+#include "StereoLayerAdditionalFlagsManager.h"
+
+#include "IStereoLayers.h"
+#include "IStereoLayersFlagsSupplier.h"
+#include "Features/IModularFeatures.h"
+
+DEFINE_LOG_CATEGORY(LogStereoLayerFlags);
+
+TSharedPtr<FStereoLayerAdditionalFlagsManager> FStereoLayerAdditionalFlagsManager::Instance = NULL;
+
+TSharedPtr<FStereoLayerAdditionalFlagsManager> FStereoLayerAdditionalFlagsManager::Get()
+{
+	if (!Instance.IsValid())
+	{
+		Instance = MakeShared<FStereoLayerAdditionalFlagsManager>();
+		Instance->CollectFlags();
+		Instance->CreateRuntimeFlagsMap();
+	}
+
+	return Instance;
+}
+
+void FStereoLayerAdditionalFlagsManager::Destroy()
+{
+	if (Instance.IsValid())
+	{
+		Instance.Reset();
+	}
+}
+
+void FStereoLayerAdditionalFlagsManager::CollectFlags(TSet<FName>& OutFlags)
+{
+	TArray<IStereoLayersFlagsSupplier*> FlagsSuppliers = IModularFeatures::Get().GetModularFeatureImplementations<IStereoLayersFlagsSupplier>(IStereoLayersFlagsSupplier::GetModularFeatureName());
+	for (IStereoLayersFlagsSupplier* FlagSupplier : FlagsSuppliers)
+	{
+		if (FlagSupplier)
+		{
+			FlagSupplier->EnumerateFlags(OutFlags);
+		}
+	}
+}
+
+void FStereoLayerAdditionalFlagsManager::CollectFlags()
+{
+	CollectFlags(UniqueFlags);
+}
+
+void FStereoLayerAdditionalFlagsManager::CreateRuntimeFlagsMap()
+{
+	RuntimeFlags.Empty();
+	uint32 Value = IStereoLayers::ELayerFlags::LAYER_FLAG_MAX_VALUE << 1;
+	for (FName& Flag : UniqueFlags)
+	{
+		RuntimeFlags.Add(Flag, Value);
+		Value = Value << 1;
+	}
+}
+
+uint32 FStereoLayerAdditionalFlagsManager::GetFlagValue(const FName Flag) const
+{
+	const uint32* Value = RuntimeFlags.Find(Flag);
+	if (Value)
+	{
+		return *Value;
+	}
+	UE_LOG(LogStereoLayerFlags, Log, TEXT("Flag %s was not found."), *(Flag.ToString()));
+	return 0;
+}
\ No newline at end of file
diff --git a/Engine/Source/Runtime/Engine/Public/IStereoLayers.h b/Engine/Source/Runtime/Engine/Public/IStereoLayers.h
--- a/Engine/Source/Runtime/Engine/Public/IStereoLayers.h
+++ b/Engine/Source/Runtime/Engine/Public/IStereoLayers.h
@@ -75,6 +75,8 @@ public:
 		LAYER_FLAG_TEX_EXTERNAL = 0x00000010,
 		// When set, this layer will not be rendered.
 		LAYER_FLAG_HIDDEN = 0x00000020,
+		// Max flag value, update this when new flags are added!
+		LAYER_FLAG_MAX_VALUE = LAYER_FLAG_HIDDEN << 1,
 	};
 
 
diff --git a/Engine/Source/Runtime/Engine/Public/IStereoLayersFlagsSupplier.h b/Engine/Source/Runtime/Engine/Public/IStereoLayersFlagsSupplier.h
new file mode 100644
--- /dev/null
+++ b/Engine/Source/Runtime/Engine/Public/IStereoLayersFlagsSupplier.h
@@ -0,0 +1,30 @@
+// Copyright Epic Games, Inc. All Rights Reserved.
+
+#pragma once
+
+#include "Features/IModularFeature.h"
+
+class IStereoLayersFlagsSupplier : public IModularFeature
+{
+public:
+	static FName GetModularFeatureName()
+	{
+		static FName FeatureName = FName(TEXT("StereoLayersFlagsSupplier"));
+		return FeatureName;
+	}
+
+	/**
+	 * Adds stereo layer flags provided by this supplier to the 
+	 * list of all flags. Flags are unique identifiers and, if two suppliers
+	 * provide a flag with the same name, only one will be kept.
+	 * IStereoLayersFlagsSuppliers must not remove flags from the
+	 * OutFlags parameter but should only add them.
+	 * 
+	 * The flags are intended to be added to StereoLayerComponents to 
+	 * add functionalities.
+	 * 
+	 * @param OutFlags		The set containing all flags
+	 * 
+	 */
+	virtual void EnumerateFlags(TSet<FName>& OutFlags) = 0;
+};
diff --git a/Engine/Source/Runtime/Engine/Public/StereoLayerAdditionalFlagsManager.h b/Engine/Source/Runtime/Engine/Public/StereoLayerAdditionalFlagsManager.h
new file mode 100644
--- /dev/null
+++ b/Engine/Source/Runtime/Engine/Public/StereoLayerAdditionalFlagsManager.h
@@ -0,0 +1,73 @@
+// Copyright Epic Games, Inc. All Rights Reserved.
+
+#pragma once
+
+#include "Containers/Map.h"
+#include "Containers/Set.h"
+#include "Logging/LogMacros.h"
+#include "Templates/SharedPointer.h"
+
+/**
+ * To add flags to StereoLayerComponents from plugins, please follow this procedure:
+ * 1. Implement the IStereoLayersFlagsSupplier interface.
+ * 
+ * 2. Register the modular feature:
+ *	  IModularFeatures::Get().RegisterModularFeature(IStereoLayersFlagsSupplier::GetModularFeatureName(), this);
+ * 
+ * 3. Unregister the modular feature when the plugin is torn down:
+ *    IModularFeatures::Get().UnregisterModularFeature(IStereoLayersFlagsSupplier::GetModularFeatureName(), this);
+ * 
+ * 4. Implement the EnumerateFlags(TSet<FName>& OutFlags) function, where you can add your flags. A simple 
+ *	  implementation might be:
+ *		void MyFlagSupplier::EnumerateFlags(TSet<FName>& OutFlags)
+ *		{
+ *			OutFlags.Add(FName("MY_FLAG_1"));
+ *			OutFlags.Add(FName("MY_FLAG_2"));
+ *		}
+ * 5. It's up to the developers to make sure the flag values are not overflowing.
+ * 6. At runtime, use 
+ *		TSharedPtr<FStereoLayerAdditionalFlagsManager> FlagsManager = FStereoLayerAdditionalFlagsManager::Get();
+ *		FlagsManager->GetFlagValue(Flag);
+ *	  to get flags value to check against the flags in your StereoLayerComponents.
+ */
+
+DECLARE_LOG_CATEGORY_EXTERN(LogStereoLayerFlags, Log, All);
+
+class FStereoLayerAdditionalFlagsManager final
+{
+	// Required for creating a TSharedPtr because this class has a private constructor.
+	template<typename, ESPMode>
+	friend class SharedPointerInternals::TIntrusiveReferenceController;
+
+public:
+	/** Returns instance of this class. If it doesn't exist yet, it creates one and collects all the flags. */
+	static TSharedPtr<FStereoLayerAdditionalFlagsManager> Get();
+
+	/** Collect all flags to use in StereoLayerComponents. */
+	static void CollectFlags(TSet<FName>& OutFlags);
+
+	/** Destroys manager when a game or PIE session is ending. */
+	static void Destroy();
+
+	/** Returns the value of the flag if it has been found, 0 otherwise. */
+	uint32 GetFlagValue(const FName Flag) const;
+
+private:
+	FStereoLayerAdditionalFlagsManager() {};
+
+	/** Create runtime mapping between flag names and values. */
+	void CreateRuntimeFlagsMap();
+
+	/** Parameterless private function to collect flags and store them in the UniqueFlags. */
+	void CollectFlags();
+
+	static TSharedPtr<FStereoLayerAdditionalFlagsManager> Instance;
+
+	// All the unique flags provided by suppliers.
+	TSet<FName> UniqueFlags;
+	// Maps flags with their runtime values.
+	TMap<FName, uint32> RuntimeFlags;
+};
+
+
+
