Added support for IsSelected() in FLightSceneProxy.

diff --git a/Engine/Source/Runtime/Engine/Classes/Components/LightComponent.h b/Engine/Source/Runtime/Engine/Classes/Components/LightComponent.h
--- a/Engine/Source/Runtime/Engine/Classes/Components/LightComponent.h
+++ b/Engine/Source/Runtime/Engine/Classes/Components/LightComponent.h
@@ -324,6 +324,11 @@ public:
 	/** The light's scene info. */
 	class FLightSceneProxy* SceneProxy;
 
+	/**
+	 * Pushes new selection state to the render thread light proxy
+	 */
+	ENGINE_API void PushSelectionToProxy();
+
 	FStaticShadowDepthMap StaticShadowDepthMap;
 
 	/** Fence used to track progress of render resource destruction. */
diff --git a/Engine/Source/Runtime/Engine/Private/Actor.cpp b/Engine/Source/Runtime/Engine/Private/Actor.cpp
--- a/Engine/Source/Runtime/Engine/Private/Actor.cpp
+++ b/Engine/Source/Runtime/Engine/Private/Actor.cpp
@@ -50,6 +50,7 @@
 #endif // UE_WITH_IRIS
 #include "LevelUtils.h"
 #include "Components/DecalComponent.h"
+#include "Components/LightComponent.h"
 #include "GameFramework/InputSettings.h"
 #include "Algo/AnyOf.h"
 #include "PrimitiveSceneProxy.h"
@@ -5204,6 +5205,11 @@ void AActor::PushSelectionToProxies()
 			{
 				DecalComponent->PushSelectionToProxy();
 			}
+
+			if (ULightComponent* LightComponent = Cast<ULightComponent>(Component))
+			{
+				LightComponent->PushSelectionToProxy();
+			}
 		}
 	}
 
diff --git a/Engine/Source/Runtime/Engine/Private/Components/LightComponent.cpp b/Engine/Source/Runtime/Engine/Private/Components/LightComponent.cpp
--- a/Engine/Source/Runtime/Engine/Private/Components/LightComponent.cpp
+++ b/Engine/Source/Runtime/Engine/Private/Components/LightComponent.cpp
@@ -1433,6 +1433,11 @@ void ULightComponent::SetMaterial(int32 ElementIndex, UMaterialInterface* InMate
 	}
 }
 
+void ULightComponent::PushSelectionToProxy()
+{
+	MarkRenderStateDirty();
+}
+
 /** Stores a light and a channel it has been assigned to. */
 struct FLightAndChannel
 {
diff --git a/Engine/Source/Runtime/Engine/Private/LightSceneProxy.cpp b/Engine/Source/Runtime/Engine/Private/LightSceneProxy.cpp
--- a/Engine/Source/Runtime/Engine/Private/LightSceneProxy.cpp
+++ b/Engine/Source/Runtime/Engine/Private/LightSceneProxy.cpp
@@ -53,6 +53,7 @@ FLightSceneProxy::FLightSceneProxy(const ULightComponent* InLightComponent)
 	, bUseVirtualShadowMaps(false)	// See below
 	, bCastModulatedShadows(false)
 	, bUseWholeSceneCSMForMovableObjects(false)
+	, bSelected(InLightComponent->GetOwner() ? InLightComponent->GetOwner()->IsActorOrSelectionParentSelected() : false)
 	, AtmosphereSunLightIndex(InLightComponent->GetAtmosphereSunLightIndex())
 	, AtmosphereSunDiskColorScale(InLightComponent->GetAtmosphereSunDiskColorScale())
 	, LightType(InLightComponent->GetLightType())
diff --git a/Engine/Source/Runtime/Engine/Public/LightSceneProxy.h b/Engine/Source/Runtime/Engine/Public/LightSceneProxy.h
--- a/Engine/Source/Runtime/Engine/Public/LightSceneProxy.h
+++ b/Engine/Source/Runtime/Engine/Public/LightSceneProxy.h
@@ -212,6 +212,8 @@ public:
 	inline uint8 GetLightingChannelMask() const { return LightingChannelMask; }
 	inline FName GetComponentFName() const { return ComponentName; }
 
+	inline bool IsSelected() const { return bSelected; }
+
 	/**
 	 * Use to get the owning actor label (or component name as fallback, if the owner is null or ENABLE_DEBUG_LABELS is off) for diagnostic messages, debug or profiling.
 	 * The actor label is what is shown in the UI (as opposed to the the FName).
@@ -428,6 +430,8 @@ protected:
 	/** Whether to render csm shadows for movable objects only (mobile). */
 	uint8 bUseWholeSceneCSMForMovableObjects : 1;
 
+	const uint8 bSelected : 1;
+
 	/** The index of the atmospheric light. Multiple lights can be considered when computing the sky/atmospheric scattering. */
 	const uint8 AtmosphereSunLightIndex;
 
