Remove non-WPO position from interpolators if pixel shader does not need it

diff --git a/Engine/Shaders/Private/MaterialTemplate.ush b/Engine/Shaders/Private/MaterialTemplate.ush
--- a/Engine/Shaders/Private/MaterialTemplate.ush
+++ b/Engine/Shaders/Private/MaterialTemplate.ush
@@ -36,7 +36,7 @@
 
 // Update this GUID to force all material shaders to recompile
 // Merge conflicts on this line should be resolved by generating a new GUID
-#pragma message("UESHADERMETADATA_VERSION 2475CED6-EBA9-4332-AFD0-4EA70825BABD")
+#pragma message("UESHADERMETADATA_VERSION 8899111D-F1E2-45E7-AB89-E1169159870B")
 
 #if USES_SPEEDTREE
 	#include "/Engine/Private/SpeedTreeCommon.ush"
diff --git a/Engine/Source/Runtime/Engine/Private/Materials/HLSLMaterialTranslator.cpp b/Engine/Source/Runtime/Engine/Private/Materials/HLSLMaterialTranslator.cpp
--- a/Engine/Source/Runtime/Engine/Private/Materials/HLSLMaterialTranslator.cpp
+++ b/Engine/Source/Runtime/Engine/Private/Materials/HLSLMaterialTranslator.cpp
@@ -6561,7 +6561,7 @@ int32 FHLSLMaterialTranslator::WorldPosition(EWorldPositionIncludedOffsets World
 
 	case WPT_ExcludeAllShaderOffsets:
 		{
-			bNeedsWorldPositionExcludingShaderOffsets = true;
+			bNeedsWorldPositionExcludingShaderOffsets |= (ShaderFrequency == SF_Pixel);
 			FunctionNamePattern = TEXT("Get<PREV>WorldPosition<NO_MATERIAL_OFFSETS>");
 			Type = MCT_LWCVector3;
 			break;
@@ -6576,7 +6576,7 @@ int32 FHLSLMaterialTranslator::WorldPosition(EWorldPositionIncludedOffsets World
 
 	case WPT_CameraRelativeNoOffsets:
 		{
-			bNeedsWorldPositionExcludingShaderOffsets = true;
+			bNeedsWorldPositionExcludingShaderOffsets |= (ShaderFrequency == SF_Pixel);
 			FunctionNamePattern = TEXT("Get<PREV>TranslatedWorldPosition<NO_MATERIAL_OFFSETS>");
 			Type = MCT_Float3;
 			break;
diff --git a/Engine/Source/Runtime/Engine/Private/Materials/MaterialHLSLEmitter.cpp b/Engine/Source/Runtime/Engine/Private/Materials/MaterialHLSLEmitter.cpp
--- a/Engine/Source/Runtime/Engine/Private/Materials/MaterialHLSLEmitter.cpp
+++ b/Engine/Source/Runtime/Engine/Private/Materials/MaterialHLSLEmitter.cpp
@@ -399,10 +399,10 @@ static void GetMaterialEnvironment(EShaderPlatform InPlatform,
 
 	const bool bNeedsWorldPositionExcludingShaderOffsets =
 		EmitContext.bUsesWorldPositionExcludingShaderOffsets || 
-		EmitMaterialData.IsExternalInputUsed(Material::EExternalInput::WorldPosition_NoOffsets) ||
-		EmitMaterialData.IsExternalInputUsed(Material::EExternalInput::TranslatedWorldPosition_NoOffsets) ||
-		EmitMaterialData.IsExternalInputUsed(Material::EExternalInput::PrevWorldPosition_NoOffsets) ||
-		EmitMaterialData.IsExternalInputUsed(Material::EExternalInput::PrevTranslatedWorldPosition_NoOffsets);
+		EmitMaterialData.IsExternalInputUsed(SF_Pixel, Material::EExternalInput::WorldPosition_NoOffsets) ||
+		EmitMaterialData.IsExternalInputUsed(SF_Pixel, Material::EExternalInput::TranslatedWorldPosition_NoOffsets) ||
+		EmitMaterialData.IsExternalInputUsed(SF_Pixel, Material::EExternalInput::PrevWorldPosition_NoOffsets) ||
+		EmitMaterialData.IsExternalInputUsed(SF_Pixel, Material::EExternalInput::PrevTranslatedWorldPosition_NoOffsets);
 
 	if (bNeedsWorldPositionExcludingShaderOffsets)
 	{
