Substrate - Fix for material function not being duplicated properly for reusability within a material graph.

diff --git a/Engine/Source/Runtime/Engine/Private/Materials/HLSLMaterialTranslator.cpp b/Engine/Source/Runtime/Engine/Private/Materials/HLSLMaterialTranslator.cpp
--- a/Engine/Source/Runtime/Engine/Private/Materials/HLSLMaterialTranslator.cpp
+++ b/Engine/Source/Runtime/Engine/Private/Materials/HLSLMaterialTranslator.cpp
@@ -4927,9 +4927,8 @@ int32 FHLSLMaterialTranslator::CallExpression(FMaterialExpressionKey ExpressionK
 	// Substrate BSDF expression should not be de-duplicated using expression output hash. 
 	// This is automatically handled via the compiler SubstrateTreeStack.
 	// It means that a node can be blended at multiple point of the graph (allowing acyclic graph instead of tree, e.g. a Slab can be used into multiple input).
-	// It is worth noting that only Substrate BSDF can be duplicated today according to SubstrateTreeStack.
-	const bool bExpressionIsSubstrateBSDF = ExpressionKey.Expression && ExpressionKey.Expression->IsA<UMaterialExpressionSubstrateSlabBSDF>();
-	const bool bExpressionIsSubstrateFullySimplified = CurrentSubstrateCompilationContext == ESubstrateCompilationContext::SCC_FullySimplified && ExpressionKey.Expression && ExpressionKey.Expression->IsResultSubstrateMaterial(ExpressionKey.OutputIndex);
+	// We do this for all SubstrateData which can be output from BSDF nodes (slabs) or other nodes such as material functions.
+	const bool bExpressionIsSubstrate = ExpressionKey.Expression && ExpressionKey.Expression->IsResultSubstrateMaterial(ExpressionKey.OutputIndex);
 
 	// Check if this expression has already been translated.
 	check(ShaderFrequency < SF_NumFrequencies);
@@ -4937,7 +4936,7 @@ int32 FHLSLMaterialTranslator::CallExpression(FMaterialExpressionKey ExpressionK
 	FMaterialFunctionCompileState* CurrentFunctionState = CurrentFunctionStack.Last();
 
 	static bool sDebugCacheDuplicateCode = true;
-	int32* ExistingCodeIndex = sDebugCacheDuplicateCode && !bExpressionIsSubstrateBSDF && !bExpressionIsSubstrateFullySimplified ? CurrentFunctionState->ExpressionCodeMap.Find(ExpressionKey) : nullptr;
+	int32* ExistingCodeIndex = sDebugCacheDuplicateCode && !bExpressionIsSubstrate ? CurrentFunctionState->ExpressionCodeMap.Find(ExpressionKey) : nullptr;
 	int32 Result = INDEX_NONE;
 	if (ExistingCodeIndex)
 	{
diff --git a/Engine/Source/Runtime/Engine/Private/Materials/MaterialExpressions.cpp b/Engine/Source/Runtime/Engine/Private/Materials/MaterialExpressions.cpp
--- a/Engine/Source/Runtime/Engine/Private/Materials/MaterialExpressions.cpp
+++ b/Engine/Source/Runtime/Engine/Private/Materials/MaterialExpressions.cpp
@@ -24327,7 +24327,8 @@ UMaterialExpressionSubstrateSlabBSDF::UMaterialExpressionSubstrateSlabBSDF(const
 FName CreateSpecularProfileParameterName(USpecularProfile* InProfile);
 int32 UMaterialExpressionSubstrateSlabBSDF::Compile(class FMaterialCompiler* Compiler, int32 OutputIndex)
 {
-	FSubstrateOperator& SubstrateOperator = Compiler->SubstrateCompilationGetOperator(Compiler->SubstrateTreeStackGetPathUniqueId());
+	FGuid PathUniqueId = Compiler->SubstrateTreeStackGetPathUniqueId();
+	FSubstrateOperator& SubstrateOperator = Compiler->SubstrateCompilationGetOperator(PathUniqueId);
 
 	// We also cannot ignore the tangent when using the default Tangent because GetTangentBasis
 	// used in SubstrateGetBSDFSharedBasis cannot be relied on for smooth tangent used for lighting on any mesh.
