Add workaround to skeletal mesh unlimited bone influence settings.

This workaround fix the correct loading of this setting during skeletal DDC key building. Prior to this change, the CVar value would be incorrectly loaded because it is loaded from a helper thread, which will read the settings value from a shadowed value incorrectly initialiialized.

diff --git a/Engine/Source/Runtime/Engine/Private/GPUSkinVertexFactory.cpp b/Engine/Source/Runtime/Engine/Private/GPUSkinVertexFactory.cpp
--- a/Engine/Source/Runtime/Engine/Private/GPUSkinVertexFactory.cpp
+++ b/Engine/Source/Runtime/Engine/Private/GPUSkinVertexFactory.cpp
@@ -404,10 +404,19 @@ bool FGPUBaseSkinVertexFactory::UseUnlimitedBoneInfluences(uint32 MaxBoneInfluen
 #if ALLOW_OTHER_PLATFORM_CONFIG
 	if (TargetPlatform)
 	{
-		TSharedPtr<IConsoleVariable> VariablePtr = CVarUnlimitedBoneInfluencesThreshold->GetPlatformValueVariable(*TargetPlatform->IniPlatformName());
-		if (VariablePtr.IsValid())
+		const ITargetPlatform* RunningPlatform = GetTargetPlatformManagerRef().GetRunningTargetPlatform();
+		const bool bIsRunningPlatform = RunningPlatform == TargetPlatform;
+		if (bIsRunningPlatform)
 		{
-			UnlimitedBoneInfluencesThreshold = (uint32)VariablePtr->GetInt();
+			UnlimitedBoneInfluencesThreshold = CVarUnlimitedBoneInfluencesThreshold->GetInt();
+		}
+		else
+		{
+			TSharedPtr<IConsoleVariable> VariablePtr = CVarUnlimitedBoneInfluencesThreshold->GetPlatformValueVariable(*TargetPlatform->IniPlatformName());
+			if (VariablePtr.IsValid())
+			{
+				UnlimitedBoneInfluencesThreshold = (uint32)VariablePtr->GetInt();
+			}
 		}
 	}
 #endif
@@ -420,10 +429,19 @@ bool FGPUBaseSkinVertexFactory::GetUnlimitedBoneInfluences(const ITargetPlatform
 #if ALLOW_OTHER_PLATFORM_CONFIG
 	if (TargetPlatform)
 	{
-		TSharedPtr<IConsoleVariable> VariablePtr = CVarUnlimitedBoneInfluences->GetPlatformValueVariable(*TargetPlatform->IniPlatformName());
-		if (VariablePtr.IsValid())
+		const ITargetPlatform* RunningPlatform = GetTargetPlatformManagerRef().GetRunningTargetPlatform();
+		const bool bIsRunningPlatform = RunningPlatform == TargetPlatform;
+		if (bIsRunningPlatform)
 		{
-			return VariablePtr->GetBool();
+			return CVarUnlimitedBoneInfluences->GetBool();
+		}
+		else
+		{
+			TSharedPtr<IConsoleVariable> VariablePtr = CVarUnlimitedBoneInfluences->GetPlatformValueVariable(*TargetPlatform->IniPlatformName());
+			if (VariablePtr.IsValid())
+			{
+				return VariablePtr->GetBool();
+			}
 		}
 	}
 #endif
