RetargetSource in AnimSeqeuence doesn't affect DDC key
#fix incorporate RetargetSourceAsset path hash into DDC-key/hash, and recompress when property changes

diff --git a/Engine/Source/Runtime/Engine/Private/Animation/AnimSequence.cpp b/Engine/Source/Runtime/Engine/Private/Animation/AnimSequence.cpp
--- a/Engine/Source/Runtime/Engine/Private/Animation/AnimSequence.cpp
+++ b/Engine/Source/Runtime/Engine/Private/Animation/AnimSequence.cpp
@@ -274,7 +274,7 @@ FString GetAnimSequenceSpecificCacheKeySuffix(const UAnimSequence& Seq, bool bPe
 	Seq.CurveCompressionSettings->PopulateDDCKey(ArcToHexString.Ar);
 	Seq.VariableFrameStrippingSettings->PopulateDDCKey(UE::Anim::Compression::FAnimDDCKeyArgs(Seq, TargetPlatform), ArcToHexString.Ar);
 
-	FString Ret = FString::Printf(TEXT("%i_%s%s%s_%c%c%i_%s_%s_%i"),
+	FString Ret = FString::Printf(TEXT("%i_%s%s%s_%c%c%i_%s_%s_%i_%i"),
 		Seq.CompressCommandletVersion,
 		PRAGMA_DISABLE_DEPRECATION_WARNINGS
 		*Seq.GetDataModel()->GenerateGuid().ToString(),
@@ -288,7 +288,8 @@ FString GetAnimSequenceSpecificCacheKeySuffix(const UAnimSequence& Seq, bool bPe
 		(bIsValidAdditive && Seq.RefPoseSeq) ? *Seq.RefPoseSeq->GetDataModel()->GenerateGuid().ToString() : TEXT("NoAdditiveGuid"),
 		PRAGMA_ENABLE_DEPRECATION_WARNINGS
 		*ArcToHexString.MakeString(),
-		Seq.GetSamplingFrameRate().Numerator / Seq.GetSamplingFrameRate().Denominator
+		Seq.GetSamplingFrameRate().Numerator / Seq.GetSamplingFrameRate().Denominator,
+		GetTypeHash(Seq.RetargetSourceAsset->GetPathName())
 	);
 
 	return Ret;
@@ -421,7 +422,7 @@ void UAnimSequence::GetAssetRegistryTags(TArray<FAssetRegistryTag>& OutTags) con
 		OutTags.Add( FAssetRegistryTag(SourceFileTagName(), AssetImportData->GetSourceData().ToJson(), FAssetRegistryTag::TT_Hidden) );
 	}
 
-	if (DataModelInterface.GetObject() && DataModelInterface->HasBeenPopulated())
+	if (DataModelInterface.GetObject() && DataModelInterface->HasBeenPopulated() && DataModelInterface.GetObject()->HasAnyFlags(RF_WasLoaded))
 	{
 		OutTags.Add(FAssetRegistryTag(TEXT("Compression Ratio"), FString::Printf(TEXT("%.03f"), (float)GetApproxCompressedSize() / (float)GetUncompressedRawSize()), FAssetRegistryTag::TT_Numerical));
 		OutTags.Add(FAssetRegistryTag(TEXT("Source Frame Rate"), FString::Printf(TEXT("%.2f"), DataModelInterface->GetFrameRate().AsDecimal()), FAssetRegistryTag::TT_Numerical));
@@ -1104,11 +1105,6 @@ void UAnimSequence::PostEditChangeProperty(FPropertyChangedEvent& PropertyChange
 
 	if(PropertyChangedEvent.Property)
 	{
-		if (PropertyChangedEvent.Property->GetFName() == GET_MEMBER_NAME_CHECKED(UAnimSequence, RetargetSourceAsset))
-		{
-			UpdateRetargetSourceAsset();
-		}
-
 		const bool bChangedRefFrameIndex = PropertyChangedEvent.Property->GetFName() == GET_MEMBER_NAME_CHECKED(UAnimSequence, RefFrameIndex);
 
 		if ((bChangedRefFrameIndex && PropertyChangedEvent.ChangeType != EPropertyChangeType::Interactive) ||
@@ -1132,6 +1128,12 @@ void UAnimSequence::PostEditChangeProperty(FPropertyChangedEvent& PropertyChange
 											  || PropertyChangedEvent.Property->GetFName() == GET_MEMBER_NAME_CHECKED(UAnimSequence, VariableFrameStrippingSettings);
 
 		bShouldResample = PropertyChangedEvent.GetMemberPropertyName() == GET_MEMBER_NAME_CHECKED(UAnimSequence, PlatformTargetFrameRate) || bChangedRefFrameIndex;
+
+		if (PropertyChangedEvent.Property->GetFName() == GET_MEMBER_NAME_CHECKED(UAnimSequence, RetargetSourceAsset))
+		{
+			UpdateRetargetSourceAsset();
+			bShouldResample = true;
+		}
 	}
 
 	if (bShouldResample)
@@ -4525,7 +4527,7 @@ FIoHash UAnimSequence::CreateDerivedDataKeyHash(const ITargetPlatform* TargetPla
 	
 	const FFrameRate FrameRate = PlatformTargetFrameRate.GetValueForPlatform(TargetPlatform->GetPlatformInfo().IniPlatformName);
 
-	FString Ret = FString::Printf(TEXT("%i_%s%s%s_%c%c%i_%s_%s_%i_%i_%s"),
+	FString Ret = FString::Printf(TEXT("%i_%s%s%s_%c%c%i_%s_%s_%i_%i_%s_%i"),
 		CompressCommandletVersion,
 		PRAGMA_DISABLE_DEPRECATION_WARNINGS
 		*GetDataModel()->GenerateGuid().ToString(),
@@ -4541,7 +4543,8 @@ FIoHash UAnimSequence::CreateDerivedDataKeyHash(const ITargetPlatform* TargetPla
 		*ArcToHexString.MakeString(),
 		FrameRate.Numerator,
 		FrameRate.Denominator,
-		*UE::Anim::Compression::AnimationCompressionVersionString
+		*UE::Anim::Compression::AnimationCompressionVersionString,
+		GetTypeHash(RetargetSourceAsset->GetPathName())
 	);
 
 	Writer << Ret;
