Move initialization of engine level UOL fragments earlier

The built-in types will now be registered as soon as the object system has been initialized. This ensures that preloading assets before the engine has been initialized can still succeed.

diff --git a/Engine/Source/Runtime/Engine/Private/UnrealEngine.cpp b/Engine/Source/Runtime/Engine/Private/UnrealEngine.cpp
--- a/Engine/Source/Runtime/Engine/Private/UnrealEngine.cpp
+++ b/Engine/Source/Runtime/Engine/Private/UnrealEngine.cpp
@@ -122,6 +122,7 @@ UnrealEngine.cpp: Implements the UEngine class and helpers.
 #include "Chaos/TriangleMeshImplicitObject.h"
 #include "HAL/PlatformMemoryHelpers.h"
 #include "RenderCore.h"
+#include "Misc/DelayedAutoRegister.h"
 #include <exception>
 
 #if UE_WITH_IRIS
@@ -273,6 +274,8 @@ void OnChangeEngineCVarRequiringRecreateRenderState(IConsoleVariable* Var)
 
 void FEngineModule::StartupModule()
 {
+	using namespace UE::UniversalObjectLocator;
+
 	// Setup delegate callback for ProfilingHelpers to access current map name
 	extern const FString GetMapNameStatic();
 	GGetMapNameDelegate.BindStatic(&GetMapNameStatic);
@@ -309,6 +312,33 @@ void FEngineModule::StartupModule()
 	FSkinWeightProfileManager::OnStartup();
 
 	UE::Anim::FSkeletonRemappingRegistry::Init();
+
+	IUniversalObjectLocatorModule& UolModule = FModuleManager::Get().LoadModuleChecked<IUniversalObjectLocatorModule>("UniversalObjectLocator");
+
+	FDelayedAutoRegisterHelper(EDelayedRegisterRunPhase::ObjectSystemReady,
+		[&UolModule]
+		{
+			// Initialize built-in UOL fragment types as soon as the object system is ready
+			{
+				FFragmentTypeParameters FragmentTypeParams("actor", NSLOCTEXT("Engine", "ActorLocatorFragment", "Actor"));
+				FragmentTypeParams.PrimaryEditorType = "Actor";
+				FActorLocatorFragment::FragmentType = UolModule.RegisterFragmentType<FActorLocatorFragment>(FragmentTypeParams);
+			}
+			{
+				FFragmentTypeParameters FragmentTypeParams("asset", NSLOCTEXT("Engine", "AssetLocatorFragment", "Asset"));
+				FragmentTypeParams.PrimaryEditorType = "Asset";
+				FAssetLocatorFragment::FragmentType = UolModule.RegisterFragmentType<FAssetLocatorFragment>(FragmentTypeParams);
+			}
+			{
+				FFragmentTypeParameters FragmentTypeParams("animinst", NSLOCTEXT("Engine", "AnimInstanceLocatorFragment", "AnimInstance"));
+				FragmentTypeParams.PrimaryEditorType = "Component";
+				FAnimInstanceLocatorFragment::FragmentType = UolModule.RegisterFragmentType<FAnimInstanceLocatorFragment>(FragmentTypeParams);
+			}
+			{
+				FActorLocatorFragmentResolveParameter::ParameterType = UolModule.RegisterParameterType<FActorLocatorFragmentResolveParameter>();
+			}
+		}
+	);
 }
 
 void FEngineModule::ShutdownModule()
@@ -2035,32 +2065,6 @@ void UEngine::Init(IEngineLoop* InEngineLoop)
 
 	InitializeObjectReferences();
 
-	// Initialize built-in UOL fragment types
-	{
-		using namespace UE::UniversalObjectLocator;
-
-		IUniversalObjectLocatorModule& UolModule = FModuleManager::Get().LoadModuleChecked<IUniversalObjectLocatorModule>("UniversalObjectLocator");
-
-		{
-			FFragmentTypeParameters FragmentTypeParams("actor", NSLOCTEXT("Engine", "ActorLocatorFragment", "Actor"));
-			FragmentTypeParams.PrimaryEditorType = "Actor";
-			FActorLocatorFragment::FragmentType = UolModule.RegisterFragmentType<FActorLocatorFragment>(FragmentTypeParams);
-		}
-		{
-			FFragmentTypeParameters FragmentTypeParams("asset", NSLOCTEXT("Engine", "AssetLocatorFragment", "Asset"));
-			FragmentTypeParams.PrimaryEditorType = "Asset";
-			FAssetLocatorFragment::FragmentType = UolModule.RegisterFragmentType<FAssetLocatorFragment>(FragmentTypeParams);
-		}
-		{
-			FFragmentTypeParameters FragmentTypeParams("animinst", NSLOCTEXT("Engine", "AnimInstanceLocatorFragment", "AnimInstance"));
-			FragmentTypeParams.PrimaryEditorType = "Component";
-			FAnimInstanceLocatorFragment::FragmentType = UolModule.RegisterFragmentType<FAnimInstanceLocatorFragment>(FragmentTypeParams);
-		}
-		{
-			FActorLocatorFragmentResolveParameter::ParameterType = UolModule.RegisterParameterType<FActorLocatorFragmentResolveParameter>();
-		}
-	}
-
 	if (GConfig)
 	{
 		bool bTemp = true;
diff --git a/Engine/Source/Runtime/UniversalObjectLocator/Private/UniversalObjectLocatorModule.cpp b/Engine/Source/Runtime/UniversalObjectLocator/Private/UniversalObjectLocatorModule.cpp
--- a/Engine/Source/Runtime/UniversalObjectLocator/Private/UniversalObjectLocatorModule.cpp
+++ b/Engine/Source/Runtime/UniversalObjectLocator/Private/UniversalObjectLocatorModule.cpp
@@ -8,6 +8,7 @@
 #include "UniversalObjectLocatorParameterTypeHandle.h"
 //#include "Modules/VisualizerDebuggingState.h"
 #include "Containers/Ticker.h"
+#include "Misc/DelayedAutoRegister.h"
 
 #include "DirectPathObjectLocator.h"
 
@@ -21,16 +22,22 @@ public:
 	
 	void StartupModule() override
 	{
-		{
-			FFragmentTypeParameters FragmentTypeParams("subobj", NSLOCTEXT("SubObjectLocator", "Object", "Object"));
-			FragmentTypeParams.PrimaryEditorType = "SubObject";
-			FSubObjectLocator::FragmentType = RegisterFragmentType<FSubObjectLocator>(FragmentTypeParams);
-		}
-
-		{
-			FFragmentTypeParameters FragmentTypeParams("uobj", NSLOCTEXT("DirectPathObjectLocator", "Object", "Object"));
-			FDirectPathObjectLocator::FragmentType = RegisterFragmentType<FDirectPathObjectLocator>(FragmentTypeParams);
-		}
+		// Register fragment types as soon as the object system is ready
+		FDelayedAutoRegisterHelper(EDelayedRegisterRunPhase::ObjectSystemReady,
+			[this]
+			{
+				{
+					FFragmentTypeParameters FragmentTypeParams("subobj", NSLOCTEXT("SubObjectLocator", "Object", "Object"));
+					FragmentTypeParams.PrimaryEditorType = "SubObject";
+					FSubObjectLocator::FragmentType = this->RegisterFragmentType<FSubObjectLocator>(FragmentTypeParams);
+				}
+
+				{
+					FFragmentTypeParameters FragmentTypeParams("uobj", NSLOCTEXT("DirectPathObjectLocator", "Object", "Object"));
+					FDirectPathObjectLocator::FragmentType = this->RegisterFragmentType<FDirectPathObjectLocator>(FragmentTypeParams);
+				}
+			}
+		);
 
 		TickerDelegate = FTSTicker::GetCoreTicker().AddTicker(
 			FTickerDelegate::CreateRaw(this, &FUniversalObjectLocatorModule::PurgeVisualizers), 60.f);
