Various #ifdef WITH_EDITOR* fixes

#rnx

diff --git a/Engine/Plugins/Experimental/Water/Source/Runtime/Private/WaterBodyOceanComponent.cpp b/Engine/Plugins/Experimental/Water/Source/Runtime/Private/WaterBodyOceanComponent.cpp
--- a/Engine/Plugins/Experimental/Water/Source/Runtime/Private/WaterBodyOceanComponent.cpp
+++ b/Engine/Plugins/Experimental/Water/Source/Runtime/Private/WaterBodyOceanComponent.cpp
@@ -465,9 +465,6 @@ void UWaterBodyOceanComponent::Reset()
 void UWaterBodyOceanComponent::PostLoad()
 {
 	Super::PostLoad();
-
-#ifdef WITH_EDITORONLY_DATA
-#endif // WITH_EDITORONLY_DATA
 }
 
 void UWaterBodyOceanComponent::OnPostRegisterAllComponents()
diff --git a/Engine/Source/Editor/EditorStyle/Private/StarshipStyle.cpp b/Engine/Source/Editor/EditorStyle/Private/StarshipStyle.cpp
--- a/Engine/Source/Editor/EditorStyle/Private/StarshipStyle.cpp
+++ b/Engine/Source/Editor/EditorStyle/Private/StarshipStyle.cpp
@@ -125,8 +125,7 @@ FStarshipEditorStyle::FStyle::FStyle( const TWeakObjectPtr< UEditorStyleSettings
 
 FStarshipEditorStyle::FStyle::~FStyle()
 {
-
-#ifdef WITH_EDITOR
+#if WITH_EDITOR
 	if (Settings.IsValid())
 	{
 		Settings->OnSettingChanged().Remove(SettingChangedHandler);
@@ -258,7 +257,7 @@ void FStarshipEditorStyle::FStyle::Initialize()
 	
 	SyncSettings();
 
-#ifdef WITH_EDITOR
+#if WITH_EDITOR
 	if (Settings.IsValid())
 	{
 		SettingChangedHandler = Settings->OnSettingChanged().AddRaw(this, &FStarshipEditorStyle::FStyle::SettingsChanged);
diff --git a/Engine/Source/Runtime/Engine/Private/Components/PrimitiveComponent.cpp b/Engine/Source/Runtime/Engine/Private/Components/PrimitiveComponent.cpp
--- a/Engine/Source/Runtime/Engine/Private/Components/PrimitiveComponent.cpp
+++ b/Engine/Source/Runtime/Engine/Private/Components/PrimitiveComponent.cpp
@@ -581,7 +581,7 @@ void UPrimitiveComponent::CreateRenderState_Concurrent(FRegisterComponentContext
 
 	// If the primitive isn't hidden and the detail mode setting allows it, add it to the scene.
 	if (ShouldComponentAddToScene()
-#ifdef WITH_EDITOR
+#if WITH_EDITOR
 		// [HOTFIX] When force deleting an asset, a SceneProxy is set to null from a different thread unsafely, causing the old stale value of SceneProxy being read here from the cache.
 		// We need to better investigate why this happens, but for now this prevents a crash from occurring.
 		&& SceneProxy == nullptr
