Add SkeletalMeshRenderData building function to USkelMesh for a given platform.

This is similar to what UStaticMesh expose, and allow external system to build/get RenderData for a given platform.

diff --git a/Engine/Source/Runtime/Engine/Classes/Engine/SkeletalMesh.h b/Engine/Source/Runtime/Engine/Classes/Engine/SkeletalMesh.h
--- a/Engine/Source/Runtime/Engine/Classes/Engine/SkeletalMesh.h
+++ b/Engine/Source/Runtime/Engine/Classes/Engine/SkeletalMesh.h
@@ -62,6 +62,8 @@ class FSkinnedAssetPostLoadContext;
 class USkeletalMeshLODSettings;
 class USkeleton;
 class UThumbnailInfo;
+class FEvent;
+class USkeletalMesh;
 struct FMeshUVChannelInfo;
 struct FSkeletalMaterial;
 struct FSkinWeightProfileInfo;
@@ -410,6 +412,22 @@ namespace NSSkeletalMeshSourceFileLabels
 }
 #endif
 
+/* Scoped SkeletalMeshRenderData
+ * All to get access to SkeletalMeshRenderData for a given platform an ensuring the data are thread-safe during the entire scope
+ */
+struct FScopedSkeletalMeshRenderData
+{
+public:
+	FScopedSkeletalMeshRenderData(USkeletalMesh* Mesh);
+	~FScopedSkeletalMeshRenderData();
+	const FSkeletalMeshRenderData* GetData() const;
+private:
+	FEvent* Lock = nullptr;
+	USkeletalMesh* Mesh = nullptr;
+	FSkeletalMeshRenderData* Data = nullptr;
+
+	friend class USkeletalMesh;
+};
 
 /**
  * SkeletalMesh is geometry bound to a hierarchical skeleton of bones which can be animated for the purpose of deforming the mesh.
@@ -2612,6 +2630,9 @@ public:
 
 	/** Generate the derived data key used to fetch derived data */
 	ENGINE_API FString GetDerivedDataKey();
+
+	/** Get a skeletal mesh render data for requested platform. */
+	ENGINE_API static void GetPlatformSkeletalMeshRenderData(const ITargetPlatform* Platform, FScopedSkeletalMeshRenderData& Out);
 #endif 
 
 private:
diff --git a/Engine/Source/Runtime/Engine/Private/SkeletalMesh.cpp b/Engine/Source/Runtime/Engine/Private/SkeletalMesh.cpp
--- a/Engine/Source/Runtime/Engine/Private/SkeletalMesh.cpp
+++ b/Engine/Source/Runtime/Engine/Private/SkeletalMesh.cpp
@@ -1507,6 +1507,46 @@ static FSkeletalMeshRenderData& GetPlatformSkeletalMeshRenderData(USkeletalMesh*
 	check(PlatformRenderData);
 	return *PlatformRenderData;
 }
+
+FScopedSkeletalMeshRenderData::FScopedSkeletalMeshRenderData(USkeletalMesh* Mesh)
+{
+	if (Mesh)
+	{
+		// Lock the skeletalmesh properties since we call USkeletalMesh::Cache() function (through GetPlatformSkeletalMeshRenderData -> CachePlatform -> Cache) 
+		// and which could be called by other threads at the same time
+		Lock = FPlatformProcess::GetSynchEventFromPool();
+		Mesh->LockPropertiesUntil(Lock);
+	}
+}
+
+FScopedSkeletalMeshRenderData::~FScopedSkeletalMeshRenderData()
+{
+	if (Mesh)
+	{
+		check(Lock);
+
+		Lock->Trigger();
+		FPlatformProcess::ReturnSynchEventToPool(Lock);
+		Data = nullptr;
+		Mesh = nullptr;
+		Lock = nullptr;
+	}
+}
+
+const FSkeletalMeshRenderData* FScopedSkeletalMeshRenderData::GetData() const
+{
+	return Data;
+}
+
+void USkeletalMesh::GetPlatformSkeletalMeshRenderData(const ITargetPlatform* TargetPlatform, FScopedSkeletalMeshRenderData& Out)
+{
+	check(Out.Mesh);
+	check(Out.Lock);
+
+	// Copy the return FSkeletalMeshRenderData to ensure it won't be modified externally
+	constexpr bool bIsSerializeSaving = false;
+	Out.Data = &::GetPlatformSkeletalMeshRenderData(Out.Mesh, TargetPlatform, bIsSerializeSaving);
+}
 #endif
 
 LLM_DEFINE_TAG(SkeletalMesh_Serialize); // This is an important test case for LLM_DEFINE_TAG
@@ -1596,13 +1636,13 @@ void USkeletalMesh::Serialize( FArchive& Ar )
 					constexpr bool bIsSerializeSaving = true;
 					if (ArchiveCookingTarget)
 					{
-						LocalSkeletalMeshRenderData = &GetPlatformSkeletalMeshRenderData(this, ArchiveCookingTarget, bIsSerializeSaving);
+						LocalSkeletalMeshRenderData = &::GetPlatformSkeletalMeshRenderData(this, ArchiveCookingTarget, bIsSerializeSaving);
 					}
 					else
 					{
 						//Fall back in case we use an archive that the cooking target has not been set (i.e. Duplicate archive)
 						check(RunningPlatform != NULL);
-						LocalSkeletalMeshRenderData = &GetPlatformSkeletalMeshRenderData(this, RunningPlatform, bIsSerializeSaving);
+						LocalSkeletalMeshRenderData = &::GetPlatformSkeletalMeshRenderData(this, RunningPlatform, bIsSerializeSaving);
 					}
 #endif
 					int32 MaxBonesPerChunk = LocalSkeletalMeshRenderData->GetMaxBonesPerSection();
@@ -4379,7 +4419,7 @@ void USkeletalMesh::BeginCacheForCookedPlatformData(const ITargetPlatform* Targe
 	LLM_SCOPE(ELLMTag::SkeletalMesh);
 	// Make sure to cache platform data so it doesn't happen lazily during serialization of the skeletal mesh
 	constexpr bool bIsSerializeSaving = false;
-	GetPlatformSkeletalMeshRenderData(this, TargetPlatform, bIsSerializeSaving);
+	::GetPlatformSkeletalMeshRenderData(this, TargetPlatform, bIsSerializeSaving);
 	ValidateBoneWeights(TargetPlatform);
 }
 
