Request full garbage collection when FPreviewScene is destroyed to ensure that the PreviewWorld is destroyed immediately, together with any GPU resources it might still hold through FScene.

The preview world is a heavy-weight object and may hold a significant amount of resources, including various GPU render targets and buffers required for rendering the scene.
Since UWorld is garbage-collected, this memory may not be cleaned for an indeterminate amount of time. By forcing garbage collection explicitly, we allow memory to be reused immediately.

diff --git a/Engine/Source/Runtime/Engine/Private/PreviewScene.cpp b/Engine/Source/Runtime/Engine/Private/PreviewScene.cpp
--- a/Engine/Source/Runtime/Engine/Private/PreviewScene.cpp
+++ b/Engine/Source/Runtime/Engine/Private/PreviewScene.cpp
@@ -132,6 +132,13 @@ FPreviewScene::~FPreviewScene()
 		GEngine->DestroyWorldContext(GetWorld());
 		// Release PhysicsScene for fixing big fbx importing bug
 		PreviewWorld->ReleasePhysicsScene();
+
+		// The preview world is a heavy-weight object and may hold a significant amount of resources,
+		// including various GPU render targets and buffers required for rendering the scene.
+		// Since UWorld is garbage-collected, this memory may not be cleaned for an indeterminate amount of time.
+		// By forcing garbage collection explicitly, we allow memory to be reused immediately.
+		PreviewWorld = nullptr;
+		GEngine->ForceGarbageCollection(true /*bFullPurge*/);
 	}
 }
 
