Fix URO-interpolated curves 'flickering'

Legacy-clearing of anim curves was left behind by the curve refactor (it was added to deal with dangling UID table references and eveolved in a number of ways). It is no longer needed

diff --git a/Engine/Source/Runtime/Engine/Private/Components/SkeletalMeshComponent.cpp b/Engine/Source/Runtime/Engine/Private/Components/SkeletalMeshComponent.cpp
--- a/Engine/Source/Runtime/Engine/Private/Components/SkeletalMeshComponent.cpp
+++ b/Engine/Source/Runtime/Engine/Private/Components/SkeletalMeshComponent.cpp
@@ -2340,8 +2340,6 @@ void USkeletalMeshComponent::RefreshBoneTransforms(FActorComponentTickFunction*
 	AnimEvaluationContext.AnimInstance = AnimScriptInstance;
 	AnimEvaluationContext.PostProcessAnimInstance = (ShouldEvaluatePostProcessInstance())? ToRawPtr(PostProcessAnimInstance): nullptr;
 
-	AnimCurves.Empty();
-
 	AnimEvaluationContext.bDoEvaluation = bShouldDoEvaluation;
 	AnimEvaluationContext.bDoInterpolation = bShouldDoInterpolation;
 	AnimEvaluationContext.bDuplicateToCacheBones = bInvalidCachedBones || (bDoEvaluationRateOptimization && AnimEvaluationContext.bDoEvaluation && !AnimEvaluationContext.bDoInterpolation);
