Fixed the server not respecting the runtime Data Layers logic operator.

diff --git a/Engine/Source/Runtime/Engine/Private/WorldPartition/WorldPartitionStreamingPolicy.cpp b/Engine/Source/Runtime/Engine/Private/WorldPartition/WorldPartitionStreamingPolicy.cpp
--- a/Engine/Source/Runtime/Engine/Private/WorldPartition/WorldPartitionStreamingPolicy.cpp
+++ b/Engine/Source/Runtime/Engine/Private/WorldPartition/WorldPartitionStreamingPolicy.cpp
@@ -353,28 +353,17 @@ void UWorldPartitionStreamingPolicy::UpdateStreamingState()
 
 			bUpdateServerEpoch = true;
 
-			const UDataLayerManager* DataLayerManager = WorldPartition->GetDataLayerManager();
-			TSet<FName> EffectiveActiveDataLayerNames = DataLayerManager->GetEffectiveActiveDataLayerNames();
-			TSet<FName> EffectiveLoadedDataLayerNames = DataLayerManager->GetEffectiveLoadedDataLayerNames();
-
 			auto CanServerDeactivateOrUnloadDataLayerCell = [&ServerDisallowStreamingOutDataLayers](const UWorldPartitionRuntimeCell* Cell)
 			{
-				if (Cell->HasDataLayers())
-				{
-					if (Cell->HasAnyDataLayer(ServerDisallowStreamingOutDataLayers))
-					{
-						return false;
-					}
-				}
-
-				return true;
+				return !Cell->HasDataLayers() || !Cell->HasAnyDataLayer(ServerDisallowStreamingOutDataLayers);
 			};
 
-			auto AddServerFrameCell = [this, CanServerDeactivateOrUnloadDataLayerCell, &EffectiveLoadedDataLayerNames, &EffectiveActiveDataLayerNames](const UWorldPartitionRuntimeCell* Cell)
+			auto AddServerFrameCell = [this, CanServerDeactivateOrUnloadDataLayerCell](const UWorldPartitionRuntimeCell* Cell)
 			{
 				// Keep Data Layer cells in their current state if server cannot deactivate/unload data layer cells
 				if (!CanServerDeactivateOrUnloadDataLayerCell(Cell))
 				{
+					// If cell was activated, keep it activated
 					if (ActivatedCells.Contains(Cell))
 					{
 						FrameActivateCells.Add(Cell);
@@ -382,8 +371,9 @@ void UWorldPartitionStreamingPolicy::UpdateStreamingState()
 					}
 					else
 					{
-						// Allow a cell with data layer(s) to switch from loaded to activated. Do not early return here in that case.
-						const bool bIsAnActivatedDataLayerCell = EffectiveActiveDataLayerNames.Num() && Cell->HasAnyDataLayer(EffectiveActiveDataLayerNames);
+						// If cell was loaded, keep it loaded except if it should become activated.
+						// In the second case, let the standard code path process it and add it to FrameActivateCells.
+						const bool bIsAnActivatedDataLayerCell = Cell->HasDataLayers() && (Cell->GetCellEffectiveWantedState() == EDataLayerRuntimeState::Activated); 
 						if (LoadedCells.Contains(Cell) && !bIsAnActivatedDataLayerCell)
 						{
 							FrameLoadCells.Add(Cell);
@@ -392,15 +382,18 @@ void UWorldPartitionStreamingPolicy::UpdateStreamingState()
 					}
 				}
 				
-				// Non Data Layer Cells + Active Data Layers
-				if (!Cell->HasDataLayers() || (EffectiveActiveDataLayerNames.Num() && Cell->HasAnyDataLayer(EffectiveActiveDataLayerNames)))
-				{
-					FrameActivateCells.Add(Cell);
-				}
-				// Loaded Data Layers Cells only
-				else if (Cell->HasDataLayers() && EffectiveLoadedDataLayerNames.Num() && Cell->HasAnyDataLayer(EffectiveLoadedDataLayerNames))
+				switch (Cell->GetCellEffectiveWantedState())
 				{
+				case EDataLayerRuntimeState::Loaded:
 					FrameLoadCells.Add(Cell);
+					break;
+				case EDataLayerRuntimeState::Activated:
+					FrameActivateCells.Add(Cell);
+					break;
+				case EDataLayerRuntimeState::Unloaded:
+					break;
+				default:
+					checkNoEntry();
 				}
 			};
 
