fix incorrect scaling of tangents w/ non-uniform BuildScale in StaticMeshBuilder, but support still using incorrect scaling as a legacy option on static mesh (enabled on existing meshes that have non-uniform build scales, and disabled on all new meshes or meshes with uniform build scale). Add the legacy option flag to the static mesh DDC key suffix so that the mesh is rebuilt when the flag changes.

diff --git a/Engine/Source/Developer/MeshBuilder/Private/StaticMeshBuilder.cpp b/Engine/Source/Developer/MeshBuilder/Private/StaticMeshBuilder.cpp
--- a/Engine/Source/Developer/MeshBuilder/Private/StaticMeshBuilder.cpp
+++ b/Engine/Source/Developer/MeshBuilder/Private/StaticMeshBuilder.cpp
@@ -972,6 +972,8 @@ void BuildVertexBuffer(
 
 	float VertexComparisonThreshold = BuildSettings.bRemoveDegenerates ? THRESH_POINTS_ARE_SAME : 0.0f;
 
+	bool bUseLegacyTangentScaling = StaticMesh->GetLegacyTangentScaling();
+
 	int32 WedgeIndex = 0;
 	for (const FTriangleID TriangleID : MeshDescription.Triangles().GetElementIDs())
 	{
@@ -1011,8 +1013,18 @@ void BuildVertexBuffer(
 			PendingVertex.Position = VertexPosition * BuildScale;
 			if (bNeedTangents)
 			{
-				PendingVertex.TangentX = VertexInstanceTangent / BuildScale;
-				PendingVertex.TangentY = ( (VertexInstanceNormal ^ VertexInstanceTangent) * VertexInstanceBinormalSign ) / BuildScale;
+				if (bUseLegacyTangentScaling)
+				{
+					// Apply incorrect inverse scale to tangents to match an old bug, for legacy assets only
+					PendingVertex.TangentX = VertexInstanceTangent / BuildScale;
+					PendingVertex.TangentY = ((VertexInstanceNormal ^ VertexInstanceTangent) * VertexInstanceBinormalSign) / BuildScale;
+				}
+				else
+				{
+					// Tangents should transform by directly applying the same scale as the geometry; it's only the normal that needs an inverse scale
+					PendingVertex.TangentX = VertexInstanceTangent * BuildScale;
+					PendingVertex.TangentY = ((VertexInstanceNormal ^ VertexInstanceTangent) * VertexInstanceBinormalSign) * BuildScale;
+				}
 				PendingVertex.TangentX.Normalize();
 				PendingVertex.TangentY.Normalize();
 			}
diff --git a/Engine/Source/Runtime/Core/Public/UObject/FortniteMainBranchObjectVersion.h b/Engine/Source/Runtime/Core/Public/UObject/FortniteMainBranchObjectVersion.h
--- a/Engine/Source/Runtime/Core/Public/UObject/FortniteMainBranchObjectVersion.h
+++ b/Engine/Source/Runtime/Core/Public/UObject/FortniteMainBranchObjectVersion.h
@@ -393,6 +393,9 @@ struct FFortniteMainBranchObjectVersion
 		// Add anim notify rate scaling, defaults to on for new content, off for old content
 		AnimNotifyAddRateScale,
 
+		// Fix tangents for non-uniform build scales, and add a flag to optionally match the previous (incorrect) tangents
+		FixedTangentTransformForNonuniformBuildScale,
+
 		// -----<new versions can be added above this line>-------------------------------------------------
 		VersionPlusOne,
 		LatestVersion = VersionPlusOne - 1
diff --git a/Engine/Source/Runtime/Engine/Classes/Engine/StaticMesh.h b/Engine/Source/Runtime/Engine/Classes/Engine/StaticMesh.h
--- a/Engine/Source/Runtime/Engine/Classes/Engine/StaticMesh.h
+++ b/Engine/Source/Runtime/Engine/Classes/Engine/StaticMesh.h
@@ -72,6 +72,7 @@ enum class EStaticMeshAsyncProperties : uint32
 	LightMapCoordinateIndex = 1 << 16,
 	LightMapResolution      = 1 << 17,
 	HiResSourceModel		= 1 << 18,
+	UseLegacyTangentScaling = 1 << 19,
 
 	All                     = MAX_uint32
 };
@@ -118,6 +119,8 @@ inline const TCHAR* ToString(EStaticMeshAsyncProperties Value)
 			return TEXT("LightMapResolution");
 		case EStaticMeshAsyncProperties::HiResSourceModel:
 			return TEXT("HiResSourceModel");
+		case EStaticMeshAsyncProperties::UseLegacyTangentScaling:
+			return TEXT("UseLegacyTangentScaling");
 		default: 
 			check(false); 
 			return TEXT("Unknown");
@@ -1072,6 +1075,31 @@ public:
 	UPROPERTY(EditAnywhere, AdvancedDisplay, Category = StaticMesh)
 	uint8 bSupportPhysicalMaterialMasks : 1;
 
+#if WITH_EDITORONLY_DATA
+private:
+	// If true, will incorrectly scale tangents when applying a non-uniform BuildScale to match what legacy code did. Only use for consistency on old assets.
+	UE_DEPRECATED(5.4, "Please do not access this member directly; use UStaticMesh::GetLegacyTangentScaling() or UStaticMesh::SetLegacyTangentScaling().")
+	UPROPERTY(EditAnywhere, AdvancedDisplay, Category = StaticMesh)
+	uint8 bUseLegacyTangentScaling : 1;
+public:
+
+	bool GetLegacyTangentScaling() const
+	{
+		WaitUntilAsyncPropertyReleased(EStaticMeshAsyncProperties::UseLegacyTangentScaling);
+		PRAGMA_DISABLE_DEPRECATION_WARNINGS
+		return bUseLegacyTangentScaling;
+		PRAGMA_ENABLE_DEPRECATION_WARNINGS
+	}
+
+	void SetLegacyTangentScaling(bool bInUseLegacyTangentScaling)
+	{
+		WaitUntilAsyncPropertyReleased(EStaticMeshAsyncProperties::UseLegacyTangentScaling);
+		PRAGMA_DISABLE_DEPRECATION_WARNINGS
+		bUseLegacyTangentScaling = bInUseLegacyTangentScaling;
+		PRAGMA_ENABLE_DEPRECATION_WARNINGS
+	}
+#endif // WITH_EDITORONLY_DATA
+
 	/**
 	 * If true, a ray tracing acceleration structure will be built for this mesh and it may be used in ray tracing effects
 	 */
diff --git a/Engine/Source/Runtime/Engine/Private/StaticMesh.cpp b/Engine/Source/Runtime/Engine/Private/StaticMesh.cpp
--- a/Engine/Source/Runtime/Engine/Private/StaticMesh.cpp
+++ b/Engine/Source/Runtime/Engine/Private/StaticMesh.cpp
@@ -22,6 +22,7 @@
 #include "UObject/DevObjectVersion.h"
 #include "UObject/DebugSerializationFlags.h"
 #include "UObject/UE5ReleaseStreamObjectVersion.h"
+#include "UObject/FortniteMainBranchObjectVersion.h"
 #include "UObject/UObjectAnnotation.h"
 #include "EngineUtils.h"
 #include "Engine/AssetUserData.h"
@@ -2693,9 +2694,14 @@ static FString BuildStaticMeshDerivedDataKeySuffix(const ITargetPlatform* Target
 	}
 
 	int32 NumLODs = Mesh->GetNumSourceModels();
+	bool bHasNonUniformBuildScale = false;
 	for (int32 LODIndex = 0; LODIndex < NumLODs; ++LODIndex)
 	{
 		FStaticMeshSourceModel& SrcModel = Mesh->GetSourceModel(LODIndex);
+		if (!SrcModel.BuildSettings.BuildScale3D.AllComponentsEqual())
+		{
+			bHasNonUniformBuildScale = true;
+		}
 		
 		check(SrcModel.RawMeshBulkData->IsEmpty());
 		if (!SrcModel.GetMeshDescriptionBulkData()->IsEmpty())
@@ -2730,6 +2736,15 @@ static FString BuildStaticMeshDerivedDataKeySuffix(const ITargetPlatform* Target
 		}
 	}
 
+	// Note: this ifdef is for consistency/generality but this whole function is part of a giant multi-thousand-line editor-only block
+#if WITH_EDITORONLY_DATA
+	if (bHasNonUniformBuildScale) // intentionally only affect key suffix if there is actual non-uniform scaling; otherwise legacy tangent scaling has no effect
+	{
+		KeySuffix += "LTS";
+		KeySuffix.AppendChar(Mesh->GetLegacyTangentScaling() ? TEXT('1') : TEXT('0'));
+	}
+#endif
+
 	// Add hi-res mesh description into DDC key
 	if (!Mesh->GetHiResSourceModel().GetMeshDescriptionBulkData()->IsEmpty())
 	{
@@ -3333,6 +3348,7 @@ UStaticMesh::UStaticMesh(const FObjectInitializer& ObjectInitializer)
 	NumStreamedLODs.Default = -1;
 	GetHiResSourceModel().StaticMeshDescriptionBulkData = CreateDefaultSubobject<UStaticMeshDescriptionBulkData>(TEXT("HiResMeshDescription"));
 	GetHiResSourceModel().StaticMeshDescriptionBulkData->SetFlags(RF_Transactional);
+	SetLegacyTangentScaling(false);
 #endif // #if WITH_EDITORONLY_DATA
 	SetLightMapResolution(4);
 	SetMinLOD(0);
@@ -5558,6 +5574,7 @@ void UStaticMesh::Serialize(FArchive& Ar)
 	Ar.UsingCustomVersion(FReleaseObjectVersion::GUID);
 	Ar.UsingCustomVersion(FUE5MainStreamObjectVersion::GUID);
 	Ar.UsingCustomVersion(FUE5ReleaseStreamObjectVersion::GUID);
+	Ar.UsingCustomVersion(FFortniteMainBranchObjectVersion::GUID);
 
 	FStripDataFlags StripFlags( Ar );
 
@@ -5833,6 +5850,24 @@ void UStaticMesh::Serialize(FArchive& Ar)
 		}
 	}
 #endif // WITH_EDITOR
+
+#if WITH_EDITORONLY_DATA
+	if (Ar.IsLoading() && Ar.CustomVer(FFortniteMainBranchObjectVersion::GUID) < FFortniteMainBranchObjectVersion::FixedTangentTransformForNonuniformBuildScale)
+	{
+		bool bHasNonUniformSourceModel = false;
+		int32 NumSourceModels = GetNumSourceModels();
+		for (int32 LODIndex = 0; LODIndex < NumSourceModels; ++LODIndex)
+		{
+			const FStaticMeshSourceModel& SourceModel = GetSourceModel(LODIndex);
+			if (!SourceModel.BuildSettings.BuildScale3D.AllComponentsEqual())
+			{
+				bHasNonUniformSourceModel = true;
+			}
+		}
+		// Only set the flag to use incorrect tangents if the asset had non-uniform scaling on a source model
+		SetLegacyTangentScaling(bHasNonUniformSourceModel);
+	}
+#endif // WITH_EDITORONLY_DATA
 }
 
 bool UStaticMesh::IsPostLoadThreadSafe() const
