Apply -NoShaderDDC argument to per-shader DDC.

#rnx

diff --git a/Engine/Source/Runtime/Engine/Private/ShaderCompiler/ShaderCompiler.cpp b/Engine/Source/Runtime/Engine/Private/ShaderCompiler/ShaderCompiler.cpp
--- a/Engine/Source/Runtime/Engine/Private/ShaderCompiler/ShaderCompiler.cpp
+++ b/Engine/Source/Runtime/Engine/Private/ShaderCompiler/ShaderCompiler.cpp
@@ -10450,9 +10450,12 @@ FShaderJobCacheRef FShaderJobCache::FindOrAdd(const FJobInputHash& Hash, EShader
 #if WITH_EDITOR
 	else
 	{
+		// If NoShaderDDC then don't check for a material the first time we encounter it to simulate a cold DDC
+		static bool bNoShaderDDC = FParse::Param(FCommandLine::Get(), TEXT("noshaderddc"));
+
 		// If we didn't find it in memory search the DDC if it's enabled.
 		// Don't search if this isn't the first job with this hash (JobInFlight already set), or there's already a request in flight.
-		const bool bCachePerShaderDDC = IsShaderJobCacheDDCEnabled() && bCheckDDC;
+		const bool bCachePerShaderDDC = IsShaderJobCacheDDCEnabled() && bCheckDDC && !bNoShaderDDC;
 		if (bCachePerShaderDDC && (JobData.JobInFlight == nullptr) && !InoutRequestOwner)
 		{
 			TRACE_COUNTER_INCREMENT(Shaders_JobCacheDDCRequests);
