Change priority of subobject and actor locator fragments to ensure that actor fragments are used preferentially for actors

diff --git a/Engine/Source/Runtime/Engine/Private/UniversalObjectLocators/ActorLocatorFragment.cpp b/Engine/Source/Runtime/Engine/Private/UniversalObjectLocators/ActorLocatorFragment.cpp
--- a/Engine/Source/Runtime/Engine/Private/UniversalObjectLocators/ActorLocatorFragment.cpp
+++ b/Engine/Source/Runtime/Engine/Private/UniversalObjectLocators/ActorLocatorFragment.cpp
@@ -188,11 +188,11 @@ UE::UniversalObjectLocator::FInitializeResult FActorLocatorFragment::Initialize(
 
 uint32 FActorLocatorFragment::ComputePriority(const UObject* ObjectToReference, const UObject* Context)
 {
-	// Can only reference actors and components
+	// Can only reference actors
 	if (ObjectToReference->IsA<AActor>())
 	{
-		// Let internal references be used first if possible
-		return 100;
+		// This locator should always be used over subobject locators in order to ensure they are used even if the context is a level
+		return 2000;
 	}
 	return 0;
 }
diff --git a/Engine/Source/Runtime/UniversalObjectLocator/Private/SubObjectLocator.cpp b/Engine/Source/Runtime/UniversalObjectLocator/Private/SubObjectLocator.cpp
--- a/Engine/Source/Runtime/UniversalObjectLocator/Private/SubObjectLocator.cpp
+++ b/Engine/Source/Runtime/UniversalObjectLocator/Private/SubObjectLocator.cpp
@@ -61,7 +61,7 @@ uint32 FSubObjectLocator::ComputePriority(const UObject* ObjectToReference, cons
 	// Can only reference objects that are relative to the context
 	if (Context && ObjectToReference->IsIn(Context))
 	{
-		return 2000;
+		return 1000;
 	}
 
 	// We can't use this at all
