Fix C4702 warnings
#rnx

diff --git a/Engine/Plugins/FX/Niagara/Source/Niagara/Private/NiagaraModule.cpp b/Engine/Plugins/FX/Niagara/Source/Niagara/Private/NiagaraModule.cpp
--- a/Engine/Plugins/FX/Niagara/Source/Niagara/Private/NiagaraModule.cpp
+++ b/Engine/Plugins/FX/Niagara/Source/Niagara/Private/NiagaraModule.cpp
@@ -1461,8 +1461,6 @@ FNiagaraTypeDefinition FNiagaraTypeDefinition::GetNumericOutputType(TConstArrayV
 	{
 		return SortedTypeDefinitions[0];
 	}
-
-	return FNiagaraTypeDefinition::GetGenericNumericDef();
 }
 
 bool FNiagaraTypeDefinition::Serialize(FArchive& Ar)
diff --git a/Engine/Plugins/Runtime/GeometryProcessing/Source/GeometryAlgorithms/Private/FTetWildWrapper.cpp b/Engine/Plugins/Runtime/GeometryProcessing/Source/GeometryAlgorithms/Private/FTetWildWrapper.cpp
--- a/Engine/Plugins/Runtime/GeometryProcessing/Source/GeometryAlgorithms/Private/FTetWildWrapper.cpp
+++ b/Engine/Plugins/Runtime/GeometryProcessing/Source/GeometryAlgorithms/Private/FTetWildWrapper.cpp
@@ -3,7 +3,9 @@
 
 #include "FTetWildWrapper.h"
 
+THIRD_PARTY_INCLUDES_START
 #include "ThirdParty/fTetWild/FloatTetwild.h"
+THIRD_PARTY_INCLUDES_END
 
 namespace UE::Geometry::FTetWild
 {
diff --git a/Engine/Source/Runtime/AudioMixer/Private/SoundWaveDecoder.cpp b/Engine/Source/Runtime/AudioMixer/Private/SoundWaveDecoder.cpp
--- a/Engine/Source/Runtime/AudioMixer/Private/SoundWaveDecoder.cpp
+++ b/Engine/Source/Runtime/AudioMixer/Private/SoundWaveDecoder.cpp
@@ -574,8 +574,6 @@ namespace Audio
 			check(InitData.SoundWave->GetPrecacheState() == ESoundWavePrecacheState::Done);
 			return InitDecodingSourceInternal(InitData);
 		}
-
-		return false;
 	}
 
 	void FSoundSourceDecoder::RemoveDecodingSource(const FDecodingSoundSourceHandle& Handle)
diff --git a/Engine/Source/Runtime/Core/Public/MSVC/MSVCPlatformCompilerPreSetup.h b/Engine/Source/Runtime/Core/Public/MSVC/MSVCPlatformCompilerPreSetup.h
--- a/Engine/Source/Runtime/Core/Public/MSVC/MSVCPlatformCompilerPreSetup.h
+++ b/Engine/Source/Runtime/Core/Public/MSVC/MSVCPlatformCompilerPreSetup.h
@@ -195,6 +195,7 @@
 			__pragma(warning(disable: 4125))  /* decimal digit terminates octal escape sequence. */ \
 			__pragma(warning(disable: 4510))  /* '<class>': default constructor could not be generated. */ \
 			__pragma(warning(disable: 4610))  /* object '<class>' can never be instantiated - user-defined constructor required. */ \
+			__pragma(warning(disable: 4702))  /* unreachable code. */ \
 			__pragma(warning(disable: 4800))  /* Implicit conversion from '<type>' to bool. Possible information loss. */ \
 			__pragma(warning(disable: 4946))  /* reinterpret_cast used between related classes: '<class1>' and '<class2>' */ \
 			__pragma(warning(disable: 4996))  /* '<obj>' was declared deprecated. */ \
diff --git a/Engine/Source/Runtime/Engine/Private/SkeletalMesh.cpp b/Engine/Source/Runtime/Engine/Private/SkeletalMesh.cpp
--- a/Engine/Source/Runtime/Engine/Private/SkeletalMesh.cpp
+++ b/Engine/Source/Runtime/Engine/Private/SkeletalMesh.cpp
@@ -3846,18 +3846,7 @@ USkeletalMeshSocket* USkeletalMesh::FindSocketAndIndex(FName InSocketName, int32
 		return nullptr;
 	}
 
-#if !WITH_EDITOR
-	check(!HasAnyFlags(RF_NeedPostLoad));
-
-	const FSocketInfo* FoundSocketInfo = SocketMap.Find(InSocketName);
-	if (FoundSocketInfo)
-	{
-		OutIndex = FoundSocketInfo->SocketIndex;
-		return FoundSocketInfo->Socket;
-	}
-	return nullptr;
-#endif
-
+#if WITH_EDITOR
 	for (int32 i = 0; i < Sockets.Num(); i++)
 	{
 		USkeletalMeshSocket* Socket = Sockets[i];
@@ -3878,6 +3867,16 @@ USkeletalMeshSocket* USkeletalMesh::FindSocketAndIndex(FName InSocketName, int32
 		}
 		return SkeletonSocket;
 	}
+#else
+	check(!HasAnyFlags(RF_NeedPostLoad));
+
+	const FSocketInfo* FoundSocketInfo = SocketMap.Find(InSocketName);
+	if (FoundSocketInfo)
+	{
+		OutIndex = FoundSocketInfo->SocketIndex;
+		return FoundSocketInfo->Socket;
+	}
+#endif
 
 	return nullptr;
 }
@@ -3893,20 +3892,7 @@ USkeletalMeshSocket* USkeletalMesh::FindSocketInfo(FName InSocketName, FTransfor
 		return nullptr;
 	}
 
-#if !WITH_EDITOR
-	check(!HasAnyFlags(RF_NeedPostLoad));
-
-	const FSocketInfo* FoundSocketInfo = SocketMap.Find(InSocketName);
-	if (FoundSocketInfo)
-	{
-		OutTransform = FoundSocketInfo->SocketLocalTransform;
-		OutIndex = FoundSocketInfo->SocketIndex;
-		OutBoneIndex = FoundSocketInfo->SocketBoneIndex;
-		return FoundSocketInfo->Socket;
-	}
-	return nullptr;
-#endif
-
+#if WITH_EDITOR
 	for (int32 i = 0; i < Sockets.Num(); i++)
 	{
 		USkeletalMeshSocket* Socket = Sockets[i];
@@ -3931,6 +3917,18 @@ USkeletalMeshSocket* USkeletalMesh::FindSocketInfo(FName InSocketName, FTransfor
 		}
 		return SkeletonSocket;
 	}
+#else
+	check(!HasAnyFlags(RF_NeedPostLoad));
+
+	const FSocketInfo* FoundSocketInfo = SocketMap.Find(InSocketName);
+	if (FoundSocketInfo)
+	{
+		OutTransform = FoundSocketInfo->SocketLocalTransform;
+		OutIndex = FoundSocketInfo->SocketIndex;
+		OutBoneIndex = FoundSocketInfo->SocketBoneIndex;
+		return FoundSocketInfo->Socket;
+	}
+#endif
 
 	return nullptr;
 }
diff --git a/Engine/Source/Runtime/Engine/Private/StaticMesh.cpp b/Engine/Source/Runtime/Engine/Private/StaticMesh.cpp
--- a/Engine/Source/Runtime/Engine/Private/StaticMesh.cpp
+++ b/Engine/Source/Runtime/Engine/Private/StaticMesh.cpp
@@ -7453,19 +7453,22 @@ bool UStaticMesh::ContainsPhysicsTriMeshDataCheckComplex(bool bInUseAllTriData,
 	{
 		// Get the LOD level to use for collision
 		const FStaticMeshLODResources& LOD = GetRenderData()->LODResources[UseLODIndex];
+#if WITH_EDITORONLY_DATA
 		for (int32 SectionIndex = 0; SectionIndex < LOD.Sections.Num(); ++SectionIndex)
 		{
 			const FStaticMeshSection& Section = LOD.Sections[SectionIndex];
-#if WITH_EDITORONLY_DATA
 			// we can only use GetSectionInfoMap() in WITH_EDITORONLY_DATA mode, otherwise, assume bInUseAllTriData :
 			if ((bInUseAllTriData || GetSectionInfoMap().Get(UseLODIndex, SectionIndex).bEnableCollision) && Section.NumTriangles > 0)
 			{
 				return true;
 			}
+		}
 #else // #if WITH_EDITORONLY_DATA
+		if (LOD.Sections.Num() > 0)
+		{
 			return true;
-#endif // #if WITH_EDITORONLY_DATA
 		}
+#endif // #if WITH_EDITORONLY_DATA
 	}
 	return false; 
 }
diff --git a/Engine/Source/Runtime/Experimental/GeometryCollectionEngine/Private/GeometryCollection/GeometryCollectionComponent.cpp b/Engine/Source/Runtime/Experimental/GeometryCollectionEngine/Private/GeometryCollection/GeometryCollectionComponent.cpp
--- a/Engine/Source/Runtime/Experimental/GeometryCollectionEngine/Private/GeometryCollection/GeometryCollectionComponent.cpp
+++ b/Engine/Source/Runtime/Experimental/GeometryCollectionEngine/Private/GeometryCollection/GeometryCollectionComponent.cpp
@@ -5727,8 +5727,6 @@ AChaosSolverActor* UGeometryCollectionComponent::GetPhysicsSolverActor() const
 		FPhysScene_Chaos const* const Scene = GetInnerChaosScene();
 		return Scene ? Cast<AChaosSolverActor>(Scene->GetSolverActor()) : nullptr;
 	}
-
-	return nullptr;
 }
 
 #define GEOMETRY_COLLECTION_CHECK_FOR_NANS_IN_TRANSFORMS 0
