Issue 402271:  Shader Compiler workaround for platform specific register allocation bug.  The failure was introduced when unused elements were optimized out of the preshader uniform buffer.  Arbitrarily introducing some unused elements back into the buffer fixes the issue (r.Material.PreshaderGapInterval specifies the interval at which unused elements should occur). The theory is that the shader compiler attempts to bulk fetch consecutive items read from the uniform buffer, and if the uniform buffer is large enough, there aren't enough registers to fit all the items in one fetch.  The workaround can be removed once the developers of the compiler solve the issue.

#rnx

diff --git a/Engine/Shaders/Public/ShaderVersion.ush b/Engine/Shaders/Public/ShaderVersion.ush
--- a/Engine/Shaders/Public/ShaderVersion.ush
+++ b/Engine/Shaders/Public/ShaderVersion.ush
@@ -3,4 +3,4 @@
 // Each time the console command is executed it generates a new GUID. As this file is included
 // in Platform.ush (which should be included in any shader) it allows to invalidate the shader DDC.
 // 
-#pragma message("UESHADERMETADATA_VERSION ACBE8DAC-4C47-0CED-3253-9EA5B31A795E")
\ No newline at end of file
+#pragma message("UESHADERMETADATA_VERSION 6169D26A-4079-3DED-6B9D-25878110CFAC")
\ No newline at end of file
diff --git a/Engine/Source/Runtime/Engine/Private/Materials/HLSLMaterialTranslator.cpp b/Engine/Source/Runtime/Engine/Private/Materials/HLSLMaterialTranslator.cpp
--- a/Engine/Source/Runtime/Engine/Private/Materials/HLSLMaterialTranslator.cpp
+++ b/Engine/Source/Runtime/Engine/Private/Materials/HLSLMaterialTranslator.cpp
@@ -196,6 +196,12 @@ static FAutoConsoleVariableRef CVarCullIntermediateUniformExpressions(
 	GCullIntermediateUniformExpressions,
 	TEXT("Enables culling of intermediate uniform expressions, reducing preshader count, saving performance"));
 
+static int32 GPreshaderGapInterval = 32;
+static FAutoConsoleVariableRef CVarPreshaderGapInterval(
+	TEXT("r.Material.PreshaderGapInterval"),
+	GPreshaderGapInterval,
+	TEXT("Insert an empty element in the preshader buffer every specified number of elements in the buffer.  Workaround for a shader compiler bug."));
+
 /* Controls whether to use the new GetMaterialShaderCode() and GetMaterialEnvironment() implementations. */
 static bool GUseMaterialTranslationResultsGrouping = true;
 
@@ -4516,6 +4522,18 @@ int32 FHLSLMaterialTranslator::AccessUniformExpression(int32 Index)
 				UniformPreshaderOffset = Align(UniformPreshaderOffset, 4u);
 			}
 
+			// Optionally insert an empty vector element (gap) at the given interval, to work around a platform specific shader compiler bug
+			if (GPreshaderGapInterval > 0)
+			{
+				// Check if we are on the vector element that is the end of the gap interval, or we will pass the end of the gap interval if the
+				// preshader is larger than a single vector.  The divides will change value when a modulus gap interval boundary is crossed.
+				uint32 PreshaderGapOffset = FMath::Max(4u, bIsLWC ? NumComponents * 2u : NumComponents);
+				if (UniformPreshaderOffset / (GPreshaderGapInterval*4) != (UniformPreshaderOffset + PreshaderGapOffset) / (GPreshaderGapInterval*4))
+				{
+					UniformPreshaderOffset += 4u;
+				}
+			}
+
 			const UE::Shader::EValueComponentType ComponentType = bIsLWC ? UE::Shader::EValueComponentType::Double : UE::Shader::EValueComponentType::Float;
 
 			FMaterialUniformPreshaderHeader& Preshader = MaterialCompilationOutput.UniformExpressionSet.UniformPreshaders.AddDefaulted_GetRef();
