Fix issue where a delayed movement adjustment would have its movement mode overwritten with the current mode, sometimes resulting in a mismatch, such as when correcting during a transition between falling and walking.

diff --git a/Engine/Source/Runtime/Engine/Private/Components/CharacterMovementComponent.cpp b/Engine/Source/Runtime/Engine/Private/Components/CharacterMovementComponent.cpp
--- a/Engine/Source/Runtime/Engine/Private/Components/CharacterMovementComponent.cpp
+++ b/Engine/Source/Runtime/Engine/Private/Components/CharacterMovementComponent.cpp
@@ -10592,7 +10592,6 @@ void UCharacterMovementComponent::SendClientAdjustment()
 
 			if (ShouldUsePackedMovementRPCs())
 			{
-				ServerData->PendingAdjustment.MovementMode = PackNetworkMovementMode();
 				ServerSendMoveResponse(ServerData->PendingAdjustment);
 			}
 			else
@@ -10615,7 +10614,7 @@ void UCharacterMovementComponent::SendClientAdjustment()
 						ServerData->PendingAdjustment.NewBaseBoneName,
 						ServerData->PendingAdjustment.NewBase != NULL,
 						ServerData->PendingAdjustment.bBaseRelativePosition,
-						PackNetworkMovementMode()
+						ServerData->PendingAdjustment.MovementMode
 					);
 				}
 				else if (bIsPlayingNetworkedRootMotionMontage)
@@ -10633,7 +10632,7 @@ void UCharacterMovementComponent::SendClientAdjustment()
 						ServerData->PendingAdjustment.NewBaseBoneName,
 						ServerData->PendingAdjustment.NewBase != NULL,
 						ServerData->PendingAdjustment.bBaseRelativePosition,
-						PackNetworkMovementMode()
+						ServerData->PendingAdjustment.MovementMode
 					);
 				}
 				else if (ServerData->PendingAdjustment.NewVel.IsZero())
@@ -10646,7 +10645,7 @@ void UCharacterMovementComponent::SendClientAdjustment()
 						ServerData->PendingAdjustment.NewBaseBoneName,
 						ServerData->PendingAdjustment.NewBase != NULL,
 						ServerData->PendingAdjustment.bBaseRelativePosition,
-						PackNetworkMovementMode()
+						ServerData->PendingAdjustment.MovementMode
 					);
 				}
 				else
@@ -10660,7 +10659,7 @@ void UCharacterMovementComponent::SendClientAdjustment()
 						ServerData->PendingAdjustment.NewBaseBoneName,
 						ServerData->PendingAdjustment.NewBase != NULL,
 						ServerData->PendingAdjustment.bBaseRelativePosition,
-						PackNetworkMovementMode()
+						ServerData->PendingAdjustment.MovementMode
 					);
 				}
 			}
