* Expose scoped skeletal render data
* Handle null skelmesh

diff --git a/Engine/Source/Runtime/Engine/Classes/Engine/SkeletalMesh.h b/Engine/Source/Runtime/Engine/Classes/Engine/SkeletalMesh.h
--- a/Engine/Source/Runtime/Engine/Classes/Engine/SkeletalMesh.h
+++ b/Engine/Source/Runtime/Engine/Classes/Engine/SkeletalMesh.h
@@ -415,7 +415,7 @@ namespace NSSkeletalMeshSourceFileLabels
 /* Scoped SkeletalMeshRenderData
  * All to get access to SkeletalMeshRenderData for a given platform an ensuring the data are thread-safe during the entire scope
  */
-struct FScopedSkeletalMeshRenderData
+struct ENGINE_API FScopedSkeletalMeshRenderData
 {
 public:
 	FScopedSkeletalMeshRenderData(USkeletalMesh* Mesh);
diff --git a/Engine/Source/Runtime/Engine/Private/SkeletalMesh.cpp b/Engine/Source/Runtime/Engine/Private/SkeletalMesh.cpp
--- a/Engine/Source/Runtime/Engine/Private/SkeletalMesh.cpp
+++ b/Engine/Source/Runtime/Engine/Private/SkeletalMesh.cpp
@@ -1508,8 +1508,9 @@ static FSkeletalMeshRenderData& GetPlatformSkeletalMeshRenderData(USkeletalMesh*
 	return *PlatformRenderData;
 }
 
-FScopedSkeletalMeshRenderData::FScopedSkeletalMeshRenderData(USkeletalMesh* Mesh)
+FScopedSkeletalMeshRenderData::FScopedSkeletalMeshRenderData(USkeletalMesh* InMesh)
 {
+	Mesh = InMesh;
 	if (Mesh)
 	{
 		// Lock the skeletalmesh properties since we call USkeletalMesh::Cache() function (through GetPlatformSkeletalMeshRenderData -> CachePlatform -> Cache) 
@@ -1540,12 +1541,11 @@ const FSkeletalMeshRenderData* FScopedSkeletalMeshRenderData::GetData() const
 
 void USkeletalMesh::GetPlatformSkeletalMeshRenderData(const ITargetPlatform* TargetPlatform, FScopedSkeletalMeshRenderData& Out)
 {
-	check(Out.Mesh);
-	check(Out.Lock);
-
-	// Copy the return FSkeletalMeshRenderData to ensure it won't be modified externally
-	constexpr bool bIsSerializeSaving = false;
-	Out.Data = &::GetPlatformSkeletalMeshRenderData(Out.Mesh, TargetPlatform, bIsSerializeSaving);
+	if (Out.Mesh && Out.Lock)
+	{
+		constexpr bool bIsSerializeSaving = false;
+		Out.Data = &::GetPlatformSkeletalMeshRenderData(Out.Mesh, TargetPlatform, bIsSerializeSaving);
+	}
 }
 #endif
 
