[NetPhysics] Replication fixes
- Enable time-dilation without using resim to keep client and server physics aligned for better physics replication
- Allow larger replication target alignment, since it will be kept in check via time-dilation
- Allow early out while target has velocity and extrapolate target correctly while performing early-out
- Take receive rate into account for Pos- and RotCorrection
- Don't send RepMovement from server if using ReplicationCache and PT has not ticked

diff --git a/Engine/Source/Runtime/Engine/Private/ActorReplication.cpp b/Engine/Source/Runtime/Engine/Private/ActorReplication.cpp
--- a/Engine/Source/Runtime/Engine/Private/ActorReplication.cpp
+++ b/Engine/Source/Runtime/Engine/Private/ActorReplication.cpp
@@ -396,9 +396,7 @@ void AActor::GatherCurrentMovement()
 		UPrimitiveComponent* RootPrimComp = Cast<UPrimitiveComponent>(GetRootComponent());
 		if (RootPrimComp && RootPrimComp->IsSimulatingPhysics())
 		{
-#if UE_WITH_IRIS
 			const bool bPrevRepPhysics = ReplicatedMovement.bRepPhysics;
-#endif // UE_WITH_IRIS
 
 			const bool bShouldUsePhysicsReplicationCache = GetPhysicsReplicationMode() != EPhysicsReplicationMode::Default;
 			bool bFoundInCache = false;
@@ -409,9 +407,13 @@ void AActor::GatherCurrentMovement()
 			{
 				if (FPhysScene_Chaos* Scene = static_cast<FPhysScene_Chaos*>(World->GetPhysicsScene()))
 				{
-					if (const FRigidBodyState* FoundState = Scene->GetStateFromReplicationCache(RootPrimComp, ServerFrame))
+					if (const FRigidBodyState* FoundState = Scene->GetStateFromReplicationCache(RootPrimComp, /*OUT*/ServerFrame))
 					{
-						ReplicatedMovement.FillFrom(*FoundState, this, Scene->ReplicationCache.ServerFrame);
+						if (ReplicatedMovement.ServerFrame != ServerFrame)
+						{
+							ReplicatedMovement.FillFrom(*FoundState, this, ServerFrame);
+							bWasRepMovementModified = true;
+						}
 						bFoundInCache = true;
 					}
 				}
@@ -423,22 +425,21 @@ void AActor::GatherCurrentMovement()
 				FRigidBodyState RBState;
 				RootPrimComp->GetRigidBodyState(RBState);
 				ReplicatedMovement.FillFrom(RBState, this, ServerFrame);
+				bWasRepMovementModified = true;
 			}
 
 			// Don't replicate movement if we're welded to another parent actor.
 			// Their replication will affect our position indirectly since we are attached.
 			ReplicatedMovement.bRepPhysics = !RootPrimComp->IsWelded();
-			
-			// Technically, the values might have stayed the same, but we'll just assume they've changed.
-			bWasRepMovementModified = true;
 
-#if UE_WITH_IRIS
 			// If RepPhysics has changed value then notify the ReplicationSystem
 			if (bPrevRepPhysics != ReplicatedMovement.bRepPhysics)
 			{
+#if UE_WITH_IRIS
 				UpdateReplicatePhysicsCondition();
-			}
 #endif // UE_WITH_IRIS
+				bWasRepMovementModified = true;
+			}
 		}
 		else if (RootComponent != nullptr)
 		{
diff --git a/Engine/Source/Runtime/Engine/Private/PhysicsEngine/PhysicsReplication.cpp b/Engine/Source/Runtime/Engine/Private/PhysicsEngine/PhysicsReplication.cpp
--- a/Engine/Source/Runtime/Engine/Private/PhysicsEngine/PhysicsReplication.cpp
+++ b/Engine/Source/Runtime/Engine/Private/PhysicsEngine/PhysicsReplication.cpp
@@ -107,7 +107,7 @@ namespace PhysicsReplicationCVars
 		static float PosCorrectionTimeBase = 0.0f;
 		static FAutoConsoleVariableRef CVarPosCorrectionTimeBase(TEXT("np2.PredictiveInterpolation.PosCorrectionTimeBase"), PosCorrectionTimeBase, TEXT("Base time to correct positional offset over. RTT * PosCorrectionTimeMultiplier are added on top of this."));
 
-		static float PosCorrectionTimeMin = 0.13f;
+		static float PosCorrectionTimeMin = 0.1f;
 		static FAutoConsoleVariableRef CVarPosCorrectionTimeMin(TEXT("np2.PredictiveInterpolation.PosCorrectionTimeMin"), PosCorrectionTimeMin, TEXT("Min time time to correct positional offset over. DeltaSeconds is added on top of this."));
 
 		static float PosCorrectionTimeMultiplier = 1.0f;
@@ -116,7 +116,7 @@ namespace PhysicsReplicationCVars
 		static float RotCorrectionTimeBase = 0.0f;
 		static FAutoConsoleVariableRef CVarRotCorrectionTimeBase(TEXT("np2.PredictiveInterpolation.RotCorrectionTimeBase"), RotCorrectionTimeBase, TEXT("Base time to correct positional offset over. RTT * PosCorrectionTimeMultiplier are added on top of this."));
 
-		static float RotCorrectionTimeMin = 0.13f;
+		static float RotCorrectionTimeMin = 0.1f;
 		static FAutoConsoleVariableRef CVarRotCorrectionTimeMin(TEXT("np2.PredictiveInterpolation.RotCorrectionTimeMin"), RotCorrectionTimeMin, TEXT("Min time time to correct rotational offset over. DeltaSeconds is added on top of this."));
 
 		static float RotCorrectionTimeMultiplier = 1.0f;
@@ -161,12 +161,15 @@ namespace PhysicsReplicationCVars
 		static float SoftSnapRotStrength = 0.5f;
 		static FAutoConsoleVariableRef CVarSoftSnapRotStrength(TEXT("np2.PredictiveInterpolation.SoftSnapRotStrength"), SoftSnapRotStrength, TEXT("Value in percent between 0.0 - 1.0 representing how much to softsnap each tick of the remaining distance."));
 
-		static float EarlyOutDistanceSqr = 2.f;
+		static float EarlyOutDistanceSqr = 1.0f;
 		static FAutoConsoleVariableRef CVarEarlyOutDistanceSqr(TEXT("np2.PredictiveInterpolation.EarlyOutDistanceSqr"), EarlyOutDistanceSqr, TEXT("Squared value. If object is within this distance from the source target, early out from replication and apply sleep if replicated."));
 		
-		static float EarlyOutAngle = 1.f;
+		static float EarlyOutAngle = 0.75f;
 		static FAutoConsoleVariableRef CVarEarlyOutAngle(TEXT("np2.PredictiveInterpolation.EarlyOutAngle"), EarlyOutAngle, TEXT("If object is within this rotational angle (in degrees) from the source target, early out from replication and apply sleep if replicated."));
 		
+		static bool bEarlyOutWithVelocity = true;
+		static FAutoConsoleVariableRef CVarEarlyOutWithVelocity(TEXT("np2.PredictiveInterpolation.EarlyOutWithVelocity"), bEarlyOutWithVelocity, TEXT("If true, allow replication logic to early out if current velocities are driving replication well enough. If false, only early out if target velocity is zero."));
+
 		static bool bPostResimWaitForUpdate = false;
 		static FAutoConsoleVariableRef CVarPostResimWaitForUpdate(TEXT("np2.PredictiveInterpolation.PostResimWaitForUpdate"), bPostResimWaitForUpdate, TEXT("After a resimulation, wait for replicated states that correspond to post-resim state before processing replication again."));
 		
@@ -194,7 +197,7 @@ namespace PhysicsReplicationCVars
 		static float SleepSecondsClearTarget = 15.0f;
 		static FAutoConsoleVariableRef CVarSleepSecondsClearTarget(TEXT("np2.PredictiveInterpolation.SleepSecondsClearTarget"), SleepSecondsClearTarget, TEXT("Wait for the object to sleep for this many seconds before clearing the replication target, to ensure nothing wakes up the object just after it goes to sleep on the client."));
 		
-		static int32 TargetTickAlignmentClampMultiplier = 2;
+		static int32 TargetTickAlignmentClampMultiplier = 20;
 		static FAutoConsoleVariableRef CVarTargetTickAlignmentClampMultiplier(TEXT("np2.PredictiveInterpolation.TargetTickAlignmentClampMultiplier"), TargetTickAlignmentClampMultiplier, TEXT("Multiplier to adjust clamping of target alignment via TickCount. Multiplier is performed on AverageReceiveInterval."));
 		
 		static bool LegacyTargetUpdateCheck = false;
@@ -907,7 +910,7 @@ void FPhysicsReplicationAsync::UpdateAsyncTarget(const FPhysicsRepAsyncInputData
 			// If we extrapolated the previous target past the receive interval, extrapolate this target by the overshoot
 			if (!Target->bWaiting)
 			{
-				ExtrapolateTarget(*Target, Target->TickCount, GetDeltaTime_Internal());
+				FPhysicsReplicationAsync::ExtrapolateTarget(*Target, Target->TickCount, GetDeltaTime_Internal());
 			}
 		}
 	}
@@ -1344,27 +1347,41 @@ bool FPhysicsReplicationAsync::PredictiveInterpolation(Chaos::FPBDRigidParticleH
 	Target.AccumulatedSleepSeconds = bIsSleeping ? (Target.AccumulatedSleepSeconds + DeltaSeconds) : 0.0f;
 	
 	// Helper for sleep and target clearing at replication end
-	auto EndReplicationHelper = [RigidsSolver, Handle, bCanSimulate, bIsSleeping](FReplicatedPhysicsTargetAsync& Target, bool bOkToClear) -> bool
+	auto EndReplicationHelper = [RigidsSolver, Handle, bCanSimulate, bIsSleeping, DeltaSeconds](FReplicatedPhysicsTargetAsync& Target, bool bOkToClear) -> bool
 	{
 		const bool bShouldSleep = (Target.TargetState.Flags & ERigidBodyFlags::Sleeping) != 0;
 		const bool bReplicatingPhysics = (Target.TargetState.Flags & ERigidBodyFlags::RepPhysics) != 0;
 
+		// --- Set Sleep State ---
 		if (bOkToClear && bShouldSleep && bCanSimulate)
 		{
 			RigidsSolver->GetEvolution()->SetParticleObjectState(Handle, Chaos::EObjectStateType::Sleeping);
 		}
 
+		// --- Should replication stop? ---
 		const bool bClearTarget =
 			(!bCanSimulate
 				|| (bOkToClear && bShouldSleep && Target.AccumulatedSleepSeconds >= PhysicsReplicationCVars::PredictiveInterpolationCVars::SleepSecondsClearTarget) // Don't clear the target due to sleeping until the object both should sleep and is sleeping for n seconds
 				|| (bOkToClear && !bReplicatingPhysics))
 			&& !PhysicsReplicationCVars::PredictiveInterpolationCVars::bDontClearTarget;
-	
+
+		// --- Target Prediction ---
+		if (!bClearTarget)
+		{
+			const int32 ExtrapolationTickLimit = FMath::Max(
+				FMath::CeilToInt(Target.ReceiveInterval * PhysicsReplicationCVars::PredictiveInterpolationCVars::ExtrapolationTimeMultiplier), // Extrapolate time based on receive interval * multiplier
+				FMath::CeilToInt(PhysicsReplicationCVars::PredictiveInterpolationCVars::ExtrapolationMinTime / DeltaSeconds)); // At least extrapolate for N seconds
+			if (Target.TickCount <= ExtrapolationTickLimit)
+			{
+				FPhysicsReplicationAsync::ExtrapolateTarget(Target, 1, DeltaSeconds);
+			}
+		}
+
 		return bClearTarget;
 	};
 
 	// If target velocity is low enough, check the distance from the current position to the source position of our target to see if it's low enough to early out of replication
-	const bool bXCanEarlyOut = Target.TargetState.LinVel.SizeSquared() < UE_KINDA_SMALL_NUMBER &&
+	const bool bXCanEarlyOut = (PhysicsReplicationCVars::PredictiveInterpolationCVars::bEarlyOutWithVelocity || Target.TargetState.LinVel.SizeSquared() < UE_KINDA_SMALL_NUMBER) &&
 		(Target.PrevPosTarget - Handle->X()).SizeSquared() < PhysicsReplicationCVars::PredictiveInterpolationCVars::EarlyOutDistanceSqr;
 
 	// Early out if we are within range of target, also apply target sleep state
@@ -1477,13 +1494,16 @@ bool FPhysicsReplicationAsync::PredictiveInterpolation(Chaos::FPBDRigidParticleH
 	}
 	else // Velocity-based Replication
 	{
+		// Calculate interpolation time based on current average receive rate
+		const float AverageReceiveIntervalSeconds = Target.AverageReceiveInterval * DeltaSeconds;
+		const float InterpolationTime = AverageReceiveIntervalSeconds * PhysicsReplicationCVars::PredictiveInterpolationCVars::PosInterpolationTimeMultiplier;
+
 		// Calculate position correction time based on current Round Trip Time
 		const float RTT = LatencyOneWay * 2.f;
-		const float PosCorrectionTime = FMath::Max(PhysicsReplicationCVars::PredictiveInterpolationCVars::PosCorrectionTimeBase + RTT * PhysicsReplicationCVars::PredictiveInterpolationCVars::PosCorrectionTimeMultiplier, DeltaSeconds + PhysicsReplicationCVars::PredictiveInterpolationCVars::PosCorrectionTimeMin);
-		const float RotCorrectionTime = FMath::Max(PhysicsReplicationCVars::PredictiveInterpolationCVars::RotCorrectionTimeBase + RTT * PhysicsReplicationCVars::PredictiveInterpolationCVars::RotCorrectionTimeMultiplier, DeltaSeconds + PhysicsReplicationCVars::PredictiveInterpolationCVars::RotCorrectionTimeMin);
-
-		// Calculate interpolation time based on current average receive rate
-		const float InterpolationTime = Target.AverageReceiveInterval * DeltaSeconds * PhysicsReplicationCVars::PredictiveInterpolationCVars::PosInterpolationTimeMultiplier;
+		const float PosCorrectionTime = FMath::Max(PhysicsReplicationCVars::PredictiveInterpolationCVars::PosCorrectionTimeBase + AverageReceiveIntervalSeconds + RTT * PhysicsReplicationCVars::PredictiveInterpolationCVars::PosCorrectionTimeMultiplier,
+			DeltaSeconds + PhysicsReplicationCVars::PredictiveInterpolationCVars::PosCorrectionTimeMin);
+		const float RotCorrectionTime = FMath::Max(PhysicsReplicationCVars::PredictiveInterpolationCVars::RotCorrectionTimeBase + AverageReceiveIntervalSeconds + RTT * PhysicsReplicationCVars::PredictiveInterpolationCVars::RotCorrectionTimeMultiplier,
+			DeltaSeconds + PhysicsReplicationCVars::PredictiveInterpolationCVars::RotCorrectionTimeMin);
 
 		if (!bXCanEarlyOut)
 		{	// --- Velocity Replication ---
@@ -1558,15 +1578,6 @@ bool FPhysicsReplicationAsync::PredictiveInterpolation(Chaos::FPBDRigidParticleH
 		// Cache data for next replication
 		Target.PrevPos = FVector(CurrentState.Position);
 
-		// --- Target Extrapolation ---
-		const int32 ExtrapolationTickLimit = FMath::Max(
-			FMath::CeilToInt(Target.ReceiveInterval * PhysicsReplicationCVars::PredictiveInterpolationCVars::ExtrapolationTimeMultiplier), // Extrapolate time based on receive interval * multiplier
-			FMath::CeilToInt(PhysicsReplicationCVars::PredictiveInterpolationCVars::ExtrapolationMinTime / DeltaSeconds)); // At least extrapolate for N seconds
-		if (Target.TickCount <= ExtrapolationTickLimit)
-		{
-			ExtrapolateTarget(Target, 1, DeltaSeconds);
-		}
-
 		if (bSoftSnap)
 		{
 			const FVector SoftSnapPos = FMath::Lerp(FVector(CurrentState.Position), Target.PrevPosTarget, FMath::Clamp(PhysicsReplicationCVars::PredictiveInterpolationCVars::SoftSnapPosStrength, 0.0f, 1.0f));
@@ -1581,6 +1592,7 @@ bool FPhysicsReplicationAsync::PredictiveInterpolation(Chaos::FPBDRigidParticleH
 	return EndReplicationHelper(Target, false);
 }
 
+/** Static function to extrapolate a target for N ticks using X DeltaSeconds */
 void FPhysicsReplicationAsync::ExtrapolateTarget(FReplicatedPhysicsTargetAsync& Target, const int32 ExtrapolateFrames, const float DeltaSeconds)
 {
 	const float ExtrapolationTime = DeltaSeconds * static_cast<float>(ExtrapolateFrames);
diff --git a/Engine/Source/Runtime/Engine/Private/PlayerController.cpp b/Engine/Source/Runtime/Engine/Private/PlayerController.cpp
--- a/Engine/Source/Runtime/Engine/Private/PlayerController.cpp
+++ b/Engine/Source/Runtime/Engine/Private/PlayerController.cpp
@@ -129,7 +129,7 @@ namespace NetworkPhysicsCvars
 	FAutoConsoleVariableRef CVarPredictionAsyncFrameBuffer(TEXT("np2.PredictionAsyncFrameBuffer"), PredictionAsyncFrameBuffer, TEXT("Additional frame offset to be added to the local to server offset used by network prediction"));
 
 	int32 TickOffsetUpdateInterval = 10;
-	FAutoConsoleVariableRef CVarTickOffsetUpdateInterval(TEXT("np2.TickOffsetUpdateInterval"), TickOffsetUpdateInterval, TEXT("How many physics ticks to wait between each tick offset update. Lowest viable value = 1, which means update each tick."));
+	FAutoConsoleVariableRef CVarTickOffsetUpdateInterval(TEXT("np2.TickOffsetUpdateInterval"), TickOffsetUpdateInterval, TEXT("How many physics ticks to wait between each tick offset update. Lowest viable value = 1, which means update each tick. Deactivate physics offset updates by setting to 0 or negative value."));
 	
 	float TimeDilationAmount = 0.01f;
 	FAutoConsoleVariableRef CVarTimeDilationAmount(TEXT("np2.TimeDilationAmount"), TimeDilationAmount, TEXT("Server-side CVar, Disable TimeDilation by setting to 0 | Default: 0.01 | Value is in percent where 0.01 = 1% dilation. Example: 1.0/0.01 = 100, meaning that over the time it usually takes to tick 100 physics steps we will tick 99 or 101 depending on if we dilate up or down."));
@@ -6186,25 +6186,8 @@ FAsyncPhysicsTimestamp APlayerController::GetPhysicsTimestamp(float DeltaSeconds
 
 void APlayerController::UpdateServerAsyncPhysicsTickOffset()
 {
-	using namespace Chaos;
-
-	if (UWorld* World = GetWorld())
-	{
-		if (FPhysScene* PhysScene = World->GetPhysicsScene())
-		{
-			if (FPBDRigidsSolver* Solver = static_cast<FPBDRigidsSolver*>(PhysScene->GetSolver()))
-			{
-				// Make sure that the client is actively using rewind data and that we are not currently resimulating
-				if (!Solver->GetRewindData() || Solver->GetEvolution()->IsResimming())
-				{
-					return;
-				}
-			}
-		}
-	}
-
 	FAsyncPhysicsTimestamp Timestamp = GetPhysicsTimestamp();
-	if(ClientLatestAsyncPhysicsStepSent + NetworkPhysicsCvars::TickOffsetUpdateInterval > Timestamp.LocalFrame)
+	if (NetworkPhysicsCvars::TickOffsetUpdateInterval <= 0 || ClientLatestAsyncPhysicsStepSent + NetworkPhysicsCvars::TickOffsetUpdateInterval > Timestamp.LocalFrame)
 	{
 		//Only send a new timestamp if enough physics ticks have passed, based on CVar.
 		//If GT is running faster than physics sim the physics timestep will not have changed, so no need to send another update to server
diff --git a/Engine/Source/Runtime/Engine/Public/PhysicsReplication.h b/Engine/Source/Runtime/Engine/Public/PhysicsReplication.h
--- a/Engine/Source/Runtime/Engine/Public/PhysicsReplication.h
+++ b/Engine/Source/Runtime/Engine/Public/PhysicsReplication.h
@@ -169,7 +169,7 @@ private:
 	void UpdateAsyncTarget(const FPhysicsRepAsyncInputData& Input, Chaos::FPBDRigidsSolver* RigidsSolver);
 	void UpdateRewindDataTarget(const FPhysicsRepAsyncInputData& Input);
 	void CacheResimInteractions();
-	void ExtrapolateTarget(FReplicatedPhysicsTargetAsync& Target, const int32 ExtrapolateFrames, const float DeltaSeconds);
+	static void ExtrapolateTarget(FReplicatedPhysicsTargetAsync& Target, const int32 ExtrapolateFrames, const float DeltaSeconds);
 
 public:
 	void Setup(FRigidBodyErrorCorrection ErrorCorrection)
