[NetPhysics] Create fake target for predictive interpolation replication on Awake event to stop client from desyncing from server by waking up and simulating freely.

diff --git a/Engine/Source/Runtime/Engine/Classes/GameFramework/Actor.h b/Engine/Source/Runtime/Engine/Classes/GameFramework/Actor.h
--- a/Engine/Source/Runtime/Engine/Classes/GameFramework/Actor.h
+++ b/Engine/Source/Runtime/Engine/Classes/GameFramework/Actor.h
@@ -2747,6 +2747,11 @@ public:
 	/** Update and smooth simulated physic state, replaces PostNetReceiveLocation() and PostNetReceiveVelocity() */
 	ENGINE_API virtual void PostNetReceivePhysicState();
 
+	/** Set the current state as a faked networked physics state for physics replication
+	* Limited for use with actors using EPhysicsReplicationMode::PredictiveInterpolation only.
+	* @param bShouldSleep  Should the replication force the object to sleep */
+	void SetFakeNetPhysicsState(bool bShouldSleep);
+
 protected:
 	/** Sync IsSimulatingPhysics() with ReplicatedMovement.bRepPhysics */
 	ENGINE_API void SyncReplicatedPhysicsSimulation();
diff --git a/Engine/Source/Runtime/Engine/Private/ActorReplication.cpp b/Engine/Source/Runtime/Engine/Private/ActorReplication.cpp
--- a/Engine/Source/Runtime/Engine/Private/ActorReplication.cpp
+++ b/Engine/Source/Runtime/Engine/Private/ActorReplication.cpp
@@ -309,6 +309,45 @@ void AActor::PostNetReceivePhysicState()
 	}
 }
 
+void AActor::SetFakeNetPhysicsState(bool bShouldSleep)
+{
+	if (!IsNetMode(ENetMode::NM_Client))
+	{
+		return;
+	}
+
+	if (!IsReplicatingMovement())
+	{
+		return;
+	}
+
+	if (GetPhysicsReplicationMode() != EPhysicsReplicationMode::PredictiveInterpolation)
+	{
+		return;
+	}
+
+	UPrimitiveComponent* RootPrimComp = Cast<UPrimitiveComponent>(RootComponent);
+	if (RootPrimComp)
+	{
+		FRigidBodyState RBState;
+		RootPrimComp->GetRigidBodyState(RBState);
+
+		RBState.Flags |= ERigidBodyFlags::NeedsUpdate;
+		RBState.Flags |= ERigidBodyFlags::RepPhysics;
+
+		if (bShouldSleep)
+		{
+			RBState.Flags |= ERigidBodyFlags::Sleeping;
+
+			// Force no velocity for sleeping objects else it won't be able to go to sleep
+			RBState.AngVel = FVector::ZeroVector;
+			RBState.LinVel = FVector::ZeroVector;
+		}
+
+		RootPrimComp->SetRigidBodyReplicatedTarget(RBState, NAME_None, /*ServerFrame*/ INDEX_NONE, INDEX_NONE);
+	}
+}
+
 void AActor::SyncReplicatedPhysicsSimulation()
 {
 	const FRepMovement& LocalRepMovement = GetReplicatedMovement();
diff --git a/Engine/Source/Runtime/Engine/Private/Components/PrimitiveComponent.cpp b/Engine/Source/Runtime/Engine/Private/Components/PrimitiveComponent.cpp
--- a/Engine/Source/Runtime/Engine/Private/Components/PrimitiveComponent.cpp
+++ b/Engine/Source/Runtime/Engine/Private/Components/PrimitiveComponent.cpp
@@ -113,6 +113,15 @@ namespace PrimitiveComponentCVars
 	static FAutoConsoleVariableRef CVarEnableFastOverlapCheck(TEXT("p.EnableFastOverlapCheck"), bEnableFastOverlapCheck, TEXT("Enable fast overlap check against sweep hits, avoiding UpdateOverlaps (for the swept component)."));
 }
 
+namespace PhysicsReplicationCVars
+{
+	namespace PredictiveInterpolationCVars
+	{
+		static bool bFakeTargetOnClientWakeUp = false;
+		static FAutoConsoleVariableRef CVarFakeTargetOnClientWakeUp(TEXT("np2.PredictiveInterpolation.FakeTargetOnClientWakeUp"), bFakeTargetOnClientWakeUp, TEXT("When true, predictive interpolation will fake a replication target at the current transform marked as asleep, this target only apply if the client doesn't receive targets from the server. This stops the client from desyncing from the server if being woken up by mistake"));
+	}
+}
+
 #if !(UE_BUILD_SHIPPING || UE_BUILD_TEST)
 int32 CVarShowInitialOverlaps = 0;
 FAutoConsoleVariableRef CVarRefShowInitialOverlaps(
@@ -2959,6 +2968,18 @@ void UPrimitiveComponent::DispatchWakeEvents(ESleepEvent WakeEvent, FName BoneNa
 			}
 		}
 	}
+
+	if (PhysicsReplicationCVars::PredictiveInterpolationCVars::bFakeTargetOnClientWakeUp)
+	{
+		if (WakeEvent == ESleepEvent::SET_Wakeup && IsSimulatingPhysics())
+		{
+			AActor* Owner = GetOwner();
+			if (Owner && Owner->GetRootComponent() == this)
+			{
+				Owner->SetFakeNetPhysicsState(/*bShouldSleep*/ true);
+			}
+		}
+	}
 }
 
 bool UPrimitiveComponent::ShouldDispatchWakeEvents(FName BoneName) const
diff --git a/Engine/Source/Runtime/Engine/Private/PhysicsEngine/PhysicsReplication.cpp b/Engine/Source/Runtime/Engine/Private/PhysicsEngine/PhysicsReplication.cpp
--- a/Engine/Source/Runtime/Engine/Private/PhysicsEngine/PhysicsReplication.cpp
+++ b/Engine/Source/Runtime/Engine/Private/PhysicsEngine/PhysicsReplication.cpp
@@ -871,7 +871,7 @@ void FPhysicsReplicationAsync::UpdateAsyncTarget(const FPhysicsRepAsyncInputData
 	// Note: If target is waiting, check PrevServerFrame since ServerFrame has been modified and we still want to cache the input target
 	if (PhysicsReplicationCVars::PredictiveInterpolationCVars::LegacyTargetUpdateCheck
 		? (Input.ServerFrame >= Target->PrevServerFrame) // Legacy check
-		: (Input.ServerFrame == 0 || (Target->bWaiting ? Input.ServerFrame > Target->PrevServerFrame : Input.ServerFrame > Target->ServerFrame)))
+		: (bFirstTarget || Input.ServerFrame == 0 || (Target->bWaiting ? Input.ServerFrame > Target->PrevServerFrame : Input.ServerFrame > Target->ServerFrame)))
 	{
 		const int32 PrevTickCount = Target->TickCount;
 		const int32 PrevReceiveInterval = Target->ReceiveInterval;
