Streaming API additions
- Added UMeshComponent::PrestreamMeshLODs() to request the streaming system starts streaming in all LODs for the mesh; analogous to PrestreamTextures()
- Added UStaticMeshComponent::RegisterLODStreamingCallback() to match USkinnedMeshComponent
- FastForceFullyResident() returns true if the asset meets the conditions to be fast fully resident, and false otherwise
- Passing -1 for the LOD in RegisterLODStreamingCallback() asks for a callback when the maximum available LOD is streamed in.
Streaming API behavior improvements
- If FramesForFullUpdate is not 1, we still call UpdatePendingStates() every frame, to merge any just-loaded assets into the array, otherwise streaming completion callbacks won't work for just-loaded assets.
- FRenderAssetStreamingManager::TickFastResponseAssets() no longer checks for RequestedMips as bForceMiplevelsToBeResident doesn't set it
- Completion callbacks are issued for fast-force-resident assets immediately after the fast streaming completes. Previously it waited until the next FramesForFullUpdate Tick.

diff --git a/Engine/Source/Runtime/Engine/Classes/Components/MeshComponent.h b/Engine/Source/Runtime/Engine/Classes/Components/MeshComponent.h
--- a/Engine/Source/Runtime/Engine/Classes/Components/MeshComponent.h
+++ b/Engine/Source/Runtime/Engine/Classes/Components/MeshComponent.h
@@ -115,10 +115,19 @@ public:
 	UFUNCTION(BlueprintCallable, Category = "Rendering")
 	ENGINE_API virtual void PrestreamTextures( float Seconds, bool bPrioritizeCharacterTextures, int32 CinematicTextureGroups = 0 );
 
+	/**
+	 *	Tell the streaming system to start streaming in all LODs for the mesh.
+	*	Note: this function may set bIgnoreStreamingMipBias on this component enable the FastForceResident system.
+	 *  @return bool							True if streaming was successfully requested
+	 *	@param Seconds							Number of seconds to force all LODs to be resident
+	 */
+	UFUNCTION(BlueprintCallable, Category = "Rendering")
+	ENGINE_API virtual bool PrestreamMeshLODs( float Seconds ) { return false; }
+
 	/**
 	 * Register a one-time callback that will be called when criteria met
 	 * @param Callback
-	 * @param LODIdx		The LOD index expected
+	 * @param LODIdx		The LOD index expected. Specify -1 for the MinLOD.
 	 * @param TimeoutSecs	Timeout in seconds
 	 * @param bOnStreamIn	To get notified when the expected LOD is streamed in or out
 	 */
diff --git a/Engine/Source/Runtime/Engine/Classes/Components/SkinnedMeshComponent.h b/Engine/Source/Runtime/Engine/Classes/Components/SkinnedMeshComponent.h
--- a/Engine/Source/Runtime/Engine/Classes/Components/SkinnedMeshComponent.h
+++ b/Engine/Source/Runtime/Engine/Classes/Components/SkinnedMeshComponent.h
@@ -1069,6 +1069,7 @@ public:
 	//~ Begin UMeshComponent Interface
 	ENGINE_API virtual void RegisterLODStreamingCallback(FLODStreamingCallback&& Callback, int32 LODIdx, float TimeoutSecs, bool bOnStreamIn) override;
 	ENGINE_API virtual void RegisterLODStreamingCallback(FLODStreamingCallback&& CallbackStreamingStart, FLODStreamingCallback&& CallbackStreamingDone, float TimeoutStartSecs, float TimeoutDoneSecs) override;
+	ENGINE_API virtual bool PrestreamMeshLODs(float Seconds) override;
 	//~ End UMeshComponent Interface
 
 	/** Get the pre-skinning local space bounds for this component. */
diff --git a/Engine/Source/Runtime/Engine/Classes/Components/StaticMeshComponent.h b/Engine/Source/Runtime/Engine/Classes/Components/StaticMeshComponent.h
--- a/Engine/Source/Runtime/Engine/Classes/Components/StaticMeshComponent.h
+++ b/Engine/Source/Runtime/Engine/Classes/Components/StaticMeshComponent.h
@@ -190,6 +190,9 @@ protected:
 	/** Initial value of bEvaluateWorldPositionOffset when BeginPlay() was called. Can be useful if we want to reset to initial state. */
 	uint8 bInitialEvaluateWorldPositionOffset : 1;
 
+	/** Whether mip callbacks have been registered and need to be removed on destroy */
+	uint8 bMipLevelCallbackRegistered : 1;
+
 public:
 
 #if WITH_EDITORONLY_DATA
@@ -572,6 +575,12 @@ public:
 #endif
 //~ End UPrimitiveComponent Interface.
 
+	//~ Begin UMeshComponent Interface
+	ENGINE_API virtual void RegisterLODStreamingCallback(FLODStreamingCallback&& Callback, int32 LODIdx, float TimeoutSecs, bool bOnStreamIn) override;
+	ENGINE_API virtual void RegisterLODStreamingCallback(FLODStreamingCallback&& CallbackStreamingStart, FLODStreamingCallback&& CallbackStreamingDone, float TimeoutStartSecs, float TimeoutDoneSecs) override;
+	ENGINE_API virtual bool PrestreamMeshLODs(float Seconds) override;
+	//~ End UMeshComponent Interface
+
 	//~ Begin INavRelevantInterface Interface.
 	ENGINE_API virtual bool IsNavigationRelevant() const override;
 	ENGINE_API virtual FBox GetNavigationBounds() const override;
diff --git a/Engine/Source/Runtime/Engine/Private/Components/SkinnedMeshComponent.cpp b/Engine/Source/Runtime/Engine/Private/Components/SkinnedMeshComponent.cpp
--- a/Engine/Source/Runtime/Engine/Private/Components/SkinnedMeshComponent.cpp
+++ b/Engine/Source/Runtime/Engine/Private/Components/SkinnedMeshComponent.cpp
@@ -4030,6 +4030,10 @@ void USkinnedMeshComponent::RegisterLODStreamingCallback(FLODStreamingCallback&&
 {
 	if (GetSkinnedAsset())
 	{
+		if (LODIdx < 0)
+		{
+			LODIdx = GetSkinnedAsset()->GetMinLodIdx(true);
+		}
 		GetSkinnedAsset()->RegisterMipLevelChangeCallback(this, LODIdx, TimeoutSecs, bOnStreamIn, MoveTemp(Callback));
 		bMipLevelCallbackRegistered = true;
 	}
@@ -4044,6 +4048,18 @@ void USkinnedMeshComponent::RegisterLODStreamingCallback(FLODStreamingCallback&&
 	}
 }
 
+bool USkinnedMeshComponent::PrestreamMeshLODs(float Seconds)
+{
+	if (USkinnedAsset* Asset = GetSkinnedAsset())
+	{
+		static IConsoleVariable* CVarAllowFastForceResident = IConsoleManager::Get().FindConsoleVariable(TEXT("r.Streaming.AllowFastForceResident"));
+		Asset->bIgnoreStreamingMipBias = CVarAllowFastForceResident && CVarAllowFastForceResident->GetInt();
+		Asset->SetForceMipLevelsToBeResident(Seconds);
+		return IStreamingManager::Get().GetRenderAssetStreamingManager().FastForceFullyResident(Asset);
+	}
+	return false;
+}
+
 void USkinnedMeshComponent::BeginDestroy()
 {
 	if (GetSkinnedAsset() && bMipLevelCallbackRegistered)
diff --git a/Engine/Source/Runtime/Engine/Private/Components/StaticMeshComponent.cpp b/Engine/Source/Runtime/Engine/Private/Components/StaticMeshComponent.cpp
--- a/Engine/Source/Runtime/Engine/Private/Components/StaticMeshComponent.cpp
+++ b/Engine/Source/Runtime/Engine/Private/Components/StaticMeshComponent.cpp
@@ -208,6 +208,7 @@ UStaticMeshComponent::UStaticMeshComponent(const FObjectInitializer& ObjectIniti
 	bWorldPositionOffsetWritesVelocity = true;
 	bEvaluateWorldPositionOffsetInRayTracing = false;
 	bInitialEvaluateWorldPositionOffset = false;
+	bMipLevelCallbackRegistered = false;
 	DistanceFieldIndirectShadowMinVisibility = .1f;
 	GetBodyInstance()->bAutoWeld = true;	//static mesh by default has auto welding
 
@@ -1861,6 +1862,12 @@ void UStaticMeshComponent::ReleaseResources()
 
 void UStaticMeshComponent::BeginDestroy()
 {
+	if (GetStaticMesh() && bMipLevelCallbackRegistered)
+	{
+		GetStaticMesh()->RemoveMipLevelChangeCallback(this);
+		bMipLevelCallbackRegistered = false;
+	}
+
 	Super::BeginDestroy();
 	ReleaseResources();
 
@@ -3123,6 +3130,40 @@ UMaterialInterface* UStaticMeshComponent::GetMaterialFromCollisionFaceIndex(int3
 }
 
 
+void UStaticMeshComponent::RegisterLODStreamingCallback(FLODStreamingCallback&& Callback, int32 LODIdx, float TimeoutSecs, bool bOnStreamIn)
+{
+	if (UStaticMesh* Mesh = GetStaticMesh())
+	{
+		if (LODIdx < 0)
+		{
+			LODIdx = Mesh->GetMinLODIdx(true);
+		}
+		Mesh->RegisterMipLevelChangeCallback(this, LODIdx, TimeoutSecs, bOnStreamIn, MoveTemp(Callback));
+		bMipLevelCallbackRegistered = true;
+	}
+}
+
+void UStaticMeshComponent::RegisterLODStreamingCallback(FLODStreamingCallback&& CallbackStreamingStart, FLODStreamingCallback&& CallbackStreamingDone, float TimeoutStartSecs, float TimeoutDoneSecs)
+{
+	if (UStaticMesh* Mesh = GetStaticMesh())
+	{
+		Mesh->RegisterMipLevelChangeCallback(this, TimeoutStartSecs, MoveTemp(CallbackStreamingStart), TimeoutDoneSecs, MoveTemp(CallbackStreamingDone));
+		bMipLevelCallbackRegistered = true;
+	}
+}
+
+bool UStaticMeshComponent::PrestreamMeshLODs(float Seconds)
+{
+	if (UStaticMesh* Mesh = GetStaticMesh())
+	{
+		static IConsoleVariable* CVarAllowFastForceResident = IConsoleManager::Get().FindConsoleVariable(TEXT("r.Streaming.AllowFastForceResident"));
+		Mesh->bIgnoreStreamingMipBias = CVarAllowFastForceResident && CVarAllowFastForceResident->GetInt();
+		Mesh->SetForceMipLevelsToBeResident(Seconds);
+		return IStreamingManager::Get().GetRenderAssetStreamingManager().FastForceFullyResident(Mesh);
+	}
+	return false;
+}
+
 bool UStaticMeshComponent::IsNavigationRelevant() const
 {
 	if (const UStaticMesh* Mesh = GetStaticMesh())
diff --git a/Engine/Source/Runtime/Engine/Private/Streaming/StreamingManagerTexture.cpp b/Engine/Source/Runtime/Engine/Private/Streaming/StreamingManagerTexture.cpp
--- a/Engine/Source/Runtime/Engine/Private/Streaming/StreamingManagerTexture.cpp
+++ b/Engine/Source/Runtime/Engine/Private/Streaming/StreamingManagerTexture.cpp
@@ -467,7 +467,7 @@ void FRenderAssetStreamingManager::TickFastResponseAssets()
 
 		Asset.UpdateStreamingStatus(false);
 
-		if (Asset.ResidentMips != Asset.RequestedMips && Asset.ResidentMips < Asset.MaxAllowedMips)
+		if (Asset.ResidentMips < Asset.MaxAllowedMips)
 		{
 			RenderAsset->StreamIn(Asset.MaxAllowedMips, true);
 			RenderAsset->bHasStreamingUpdatePending = true;
@@ -486,6 +486,7 @@ void FRenderAssetStreamingManager::TickFastResponseAssets()
 
 			Asset.UpdateStreamingStatus(false);
 			TrackRenderAssetEvent(&Asset, RenderAsset, true, this);
+			RenderAsset->TickMipLevelChangeCallbacks(nullptr);
 		}
 	}
 }
@@ -1246,7 +1247,7 @@ void FRenderAssetStreamingManager::UpdateIndividualRenderAsset( UStreamableRende
 	StreamingRenderAsset->StreamWantedMips(*this);
 }
 
-void FRenderAssetStreamingManager::FastForceFullyResident(UStreamableRenderAsset* RenderAsset)
+bool FRenderAssetStreamingManager::FastForceFullyResident(UStreamableRenderAsset* RenderAsset)
 {
 	check(IsInGameThread());
 	TArray<FStreamingRenderAsset>& StreamingRenderAssets = GetStreamingRenderAssetsAsyncSafe();
@@ -1268,8 +1269,10 @@ void FRenderAssetStreamingManager::FastForceFullyResident(UStreamableRenderAsset
 		if (Asset.ResidentMips < Asset.MaxAllowedMips)
 		{
 			FastResponseRenderAssets.Add(Asset.RenderAsset);
+			return true;
 		}
 	}
+	return false;
 }
 
 /**
@@ -1809,6 +1812,8 @@ void FRenderAssetStreamingManager::UpdateResourceStreaming( float DeltaTime, boo
 	}
 	else if (ProcessingStage <= NumRenderAssetProcessingStages)
 	{
+		UpdatePendingStates(false);
+
 		STAT(int32 StartTime = (int32)FPlatformTime::Cycles();)
 
 		if (ProcessingStage == 1)
@@ -1848,6 +1853,8 @@ void FRenderAssetStreamingManager::UpdateResourceStreaming( float DeltaTime, boo
 	{
 		STAT(GatheredStats.StreamRenderAssetsCycles = -(int32)FPlatformTime::Cycles();)
 
+		UpdatePendingStates(false);
+
 		// Since this step is lightweight, tick each texture inflight here, to accelerate the state changes.
 		for (int32 TextureIndex : InflightRenderAssets)
 		{
diff --git a/Engine/Source/Runtime/Engine/Private/Streaming/StreamingManagerTexture.h b/Engine/Source/Runtime/Engine/Private/Streaming/StreamingManagerTexture.h
--- a/Engine/Source/Runtime/Engine/Private/Streaming/StreamingManagerTexture.h
+++ b/Engine/Source/Runtime/Engine/Private/Streaming/StreamingManagerTexture.h
@@ -58,8 +58,9 @@ struct FRenderAssetStreamingManager final : public IRenderAssetStreamingManager
 	 * Either bForceMiplevelsToBeResident or ForceMipLevelsToBeResidentTimestamp need to be set on the asset.
 	 *
 	 * @param RenderAsset The asset to register
+	 * @return bool True if the streaming request is successful
 	 */
-	virtual void FastForceFullyResident(UStreamableRenderAsset* RenderAsset) override;
+	virtual bool FastForceFullyResident(UStreamableRenderAsset* RenderAsset) override;
 
 	/**
 	 * Blocks till all pending requests are fulfilled.
diff --git a/Engine/Source/Runtime/Engine/Public/ContentStreaming.h b/Engine/Source/Runtime/Engine/Public/ContentStreaming.h
--- a/Engine/Source/Runtime/Engine/Public/ContentStreaming.h
+++ b/Engine/Source/Runtime/Engine/Public/ContentStreaming.h
@@ -433,8 +433,8 @@ struct IRenderAssetStreamingManager : public IStreamingManager
 	*/
 	virtual void UpdateIndividualRenderAsset(UStreamableRenderAsset* RenderAsset) = 0;
 
-	/** Stream in non-resident mips for an asset ASAP. */
-	virtual void FastForceFullyResident(UStreamableRenderAsset* RenderAsset) = 0;
+	/** Stream in non-resident mips for an asset ASAP. Returns true if streaming request will be successful. */
+	virtual bool FastForceFullyResident(UStreamableRenderAsset* RenderAsset) = 0;
 
 	/**
 	* Temporarily boosts the streaming distance factor by the specified number.
