[shaders] when recompiling shaders, block on global shader completion after each platform. since there's only a single global shadermap ID, the compile job keys collide when more than one platform is in flight at once, resulting in missing shaders for platforms triggered beyond the first.

diff --git a/Engine/Source/Runtime/Engine/Private/ShaderCompiler/ShaderCompiler.cpp b/Engine/Source/Runtime/Engine/Private/ShaderCompiler/ShaderCompiler.cpp
--- a/Engine/Source/Runtime/Engine/Private/ShaderCompiler/ShaderCompiler.cpp
+++ b/Engine/Source/Runtime/Engine/Private/ShaderCompiler/ShaderCompiler.cpp
@@ -8733,15 +8733,15 @@ bool RecompileShaders(const TCHAR* Cmd, FOutputDevice& Ar)
 
 				UpdateReferencedUniformBufferNames(OutdatedShaderTypes, OutdatedFactoryTypes, OutdatedShaderPipelineTypes);
 
+				// Kick off global shader recompiles
 				UMaterialInterface::IterateOverActiveFeatureLevels([&](ERHIFeatureLevel::Type InFeatureLevel) {
 					auto ShaderPlatform = GShaderPlatformForFeatureLevel[InFeatureLevel];
 					BeginRecompileGlobalShaders(OutdatedShaderTypes, OutdatedShaderPipelineTypes, ShaderPlatform);
+					// Block on global shader compilation. Do this for each feature level/platform compiled as otherwise global shader compile job IDs collide.
+					FinishRecompileGlobalShaders();
 				});
 
-				// Block on global shaders
-				FinishRecompileGlobalShaders();
-
-				// Kick off global shader recompiles
+				// Kick off material shader recompiles
 				UMaterialInterface::IterateOverActiveFeatureLevels([&](ERHIFeatureLevel::Type InFeatureLevel) {
 					auto ShaderPlatform = GShaderPlatformForFeatureLevel[InFeatureLevel];
 					UMaterial::UpdateMaterialShaders(OutdatedShaderTypes, OutdatedShaderPipelineTypes, OutdatedFactoryTypes, ShaderPlatform);
