World Partition Actor Descriptors: directional light now returns an invalid streaming bounds to avoid having huge level instance bounds when it is not marked as main level only.

diff --git a/Engine/Source/Runtime/Engine/Classes/Engine/DirectionalLight.h b/Engine/Source/Runtime/Engine/Classes/Engine/DirectionalLight.h
--- a/Engine/Source/Runtime/Engine/Classes/Engine/DirectionalLight.h
+++ b/Engine/Source/Runtime/Engine/Classes/Engine/DirectionalLight.h
@@ -38,6 +38,7 @@ public:
 	virtual void LoadedFromAnotherClass(const FName& OldClassName) override;
 	virtual void PostEditChangeProperty(FPropertyChangedEvent& PropertyChangedEvent) override;
 	virtual bool CanChangeIsSpatiallyLoadedFlag() const override { return false; }
+	virtual FBox GetStreamingBounds() const override { return FBox(ForceInit); }
 #endif
 
 public:
