Compressed animations with uniform identify scaling break when Skeleton contains scaling
#fix ensure that scaling tracks _are_ compressed if their single key entry does not equal the skeleton reference pose transform

diff --git a/Engine/Source/Runtime/Engine/Classes/Animation/AnimSequenceHelpers.h b/Engine/Source/Runtime/Engine/Classes/Animation/AnimSequenceHelpers.h
--- a/Engine/Source/Runtime/Engine/Classes/Animation/AnimSequenceHelpers.h
+++ b/Engine/Source/Runtime/Engine/Classes/Animation/AnimSequenceHelpers.h
@@ -198,8 +198,24 @@ namespace Anim {
 		*
 		* @return Whether or not the operation was succesful
 		*/
+		UE_DEPRECATED(5.4, "CompressAnimationDataTracks has been deprecated use other signature with skeleton and bone mappings")
 		ENGINE_API bool CompressAnimationDataTracks(TArray<FRawAnimSequenceTrack>& RawAnimationData, int32 NumberOfKeys, FName ErrorName, float MaxPosDiff = 0.0001f, float MaxAngleDiff = 0.0003f);
 
+		/**
+		* Compress the provided tracks using two methods.
+		* 1. Checking for uniform key data according to the provided error metrics
+		* 2. Setting number of scale keys to 0 if it is uniform, equal to FVector::OneVector and matches the Skeleton reference bone pose(s)
+		*
+		* @param	RawAnimationData	Set of tracks to apply the reduction to
+		* @param	NumberOfKeys		Expected number of keys for each component within the tracks
+		* @param	ErrorName			Identifier for outer callsite when failing to reduce the data
+		* @param	MaxPosDiff			Maximum positional delta when determining whether or not two positional keys are identical
+		* @param	MaxAngleDiff		Maximum rotational delta when determining whether or not two rotational keys are identical
+		*
+		* @return Whether or not the operation was succesful
+		*/
+		ENGINE_API bool CompressAnimationDataTracks(const USkeleton* Skeleton, const TArray<FTrackToSkeletonMap>& TrackToSkeleton,TArray<FRawAnimSequenceTrack>& RawAnimationData, int32 NumberOfKeys, FName ErrorName, float MaxPosDiff = 0.0001f, float MaxAngleDiff = 0.0003f);
+
 		/**
 		* Compress the provided track by checking for uniform key data according to the provided error metrics
 		*
diff --git a/Engine/Source/Runtime/Engine/Private/Animation/AnimCompressionTypes.cpp b/Engine/Source/Runtime/Engine/Private/Animation/AnimCompressionTypes.cpp
--- a/Engine/Source/Runtime/Engine/Private/Animation/AnimCompressionTypes.cpp
+++ b/Engine/Source/Runtime/Engine/Private/Animation/AnimCompressionTypes.cpp
@@ -858,8 +858,8 @@ void FCompressibleAnimData::FetchData(const ITargetPlatform* InPlatform)
 	// Apply any key reduction if possible
 	if (RawAnimationData.Num())
 	{ 
-		UE::Anim::Compression::CompressAnimationDataTracks(RawAnimationData, NumberOfKeys, AnimSequence->GetFName(), -1.f, -1.f);
-		UE::Anim::Compression::CompressAnimationDataTracks(RawAnimationData, NumberOfKeys, AnimSequence->GetFName());
+		UE::Anim::Compression::CompressAnimationDataTracks(Skeleton, TrackToSkeletonMapTable, RawAnimationData, NumberOfKeys, AnimSequence->GetFName(), -1.f, -1.f);
+		UE::Anim::Compression::CompressAnimationDataTracks(Skeleton, TrackToSkeletonMapTable, RawAnimationData, NumberOfKeys, AnimSequence->GetFName());
 	}
 
 	auto IsKeyArrayValidForRemoval = [](const auto& Keys, const auto& IdentityValue) -> bool
diff --git a/Engine/Source/Runtime/Engine/Private/Animation/AnimSequenceHelpers.cpp b/Engine/Source/Runtime/Engine/Private/Animation/AnimSequenceHelpers.cpp
--- a/Engine/Source/Runtime/Engine/Private/Animation/AnimSequenceHelpers.cpp
+++ b/Engine/Source/Runtime/Engine/Private/Animation/AnimSequenceHelpers.cpp
@@ -385,7 +385,59 @@ bool Compression::CompressAnimationDataTracks(TArray<FRawAnimSequenceTrack>& Raw
 			if (RawData.ScaleKeys.Num() > 0)
 			{
 				// if scale key exists, see if we can just empty it
-				if ((RawData.ScaleKeys.Num() > 1) || (RawData.ScaleKeys[0].Equals(FVector3f::OneVector) == false))
+				const bool bHaveMultipleKeys = RawData.ScaleKeys.Num() > 1;
+				const bool bHaveNonIdentityKey = !RawData.ScaleKeys[0].Equals(FVector3f::OneVector);
+				if (bHaveMultipleKeys|| bHaveNonIdentityKey)
+				{
+					bCompressScaleKeys = true;
+					break;
+				}
+			}
+		}
+
+		// if we don't have scale, we should delete all scale keys
+		// if you have one track that has scale, we still should support scale, so compress scale
+		if (!bCompressScaleKeys)
+		{
+			// then remove all scale keys
+			for (int32 TrackIndex = 0; TrackIndex < RawAnimationData.Num(); TrackIndex++)
+			{
+				FRawAnimSequenceTrack& RawData = RawAnimationData[TrackIndex];
+				RawData.ScaleKeys.Empty();
+			}
+		}
+	}
+#endif
+	return bRemovedKeys;
+}
+
+bool Compression::CompressAnimationDataTracks(const USkeleton* Skeleton, const TArray<FTrackToSkeletonMap>& TrackToSkeleton, TArray<FRawAnimSequenceTrack>& RawAnimationData, int32 NumberOfKeys, FName ErrorName, float MaxPosDiff /*= 0.0001f*/, float MaxAngleDiff /*= 0.0003f*/)
+{
+	bool bRemovedKeys = false;
+
+#if WITH_EDITORONLY_DATA
+	if (ensureMsgf(RawAnimationData.Num() > 0, TEXT("%s is trying to compress while raw animation is missing"), * ErrorName.ToString()))
+	{
+		// This removes trivial keys, and this has to happen before the removing tracks
+		for (int32 TrackIndex = 0; TrackIndex < RawAnimationData.Num(); TrackIndex++)
+		{
+			bRemovedKeys |= CompressRawAnimSequenceTrack(RawAnimationData[TrackIndex], NumberOfKeys, ErrorName, MaxPosDiff, MaxAngleDiff);
+		}
+
+		bool bCompressScaleKeys = false;
+		// go through remove keys if not needed
+		for (int32 TrackIndex = 0; TrackIndex < RawAnimationData.Num(); TrackIndex++)
+		{
+			FRawAnimSequenceTrack const& RawData = RawAnimationData[TrackIndex];
+			if (RawData.ScaleKeys.Num() > 0)
+			{
+				// if scale key exists, see if we can just empty it
+				const bool bHaveMultipleKeys = RawData.ScaleKeys.Num() > 1;
+				const bool bHaveNonIdentityKey = !RawData.ScaleKeys[0].Equals(FVector3f::OneVector);
+				const FReferenceSkeleton& ReferenceSkeleton = Skeleton->GetReferenceSkeleton();
+				const int32 TrackBoneIndex = TrackToSkeleton[TrackIndex].BoneTreeIndex;				
+				const bool bHaveNonRefPoseKey = ReferenceSkeleton.IsValidIndex(TrackBoneIndex) && !RawData.ScaleKeys[0].Equals(FVector3f(ReferenceSkeleton.GetRefBonePose()[TrackBoneIndex].GetScale3D()));
+				if (bHaveMultipleKeys|| bHaveNonIdentityKey || bHaveNonRefPoseKey)
 				{
 					bCompressScaleKeys = true;
 					break;
