Making the StaticMeshComponent's OnStaticMeshChanged event available in non editor builds

diff --git a/Engine/Source/Runtime/Engine/Classes/Components/StaticMeshComponent.h b/Engine/Source/Runtime/Engine/Classes/Components/StaticMeshComponent.h
--- a/Engine/Source/Runtime/Engine/Classes/Components/StaticMeshComponent.h
+++ b/Engine/Source/Runtime/Engine/Classes/Components/StaticMeshComponent.h
@@ -792,16 +792,13 @@ public:
 	virtual FOnGetNaniteResources& OnGetNaniteResources() { return OnGetNaniteResourcesEvent; }
 	virtual const FOnGetNaniteResources& OnGetNaniteResources() const { return OnGetNaniteResourcesEvent; }
 
-#if WITH_EDITOR
 	/** Called when the static mesh changes  */
 	DECLARE_EVENT_OneParam(UStaticMeshComponent, FOnStaticMeshChanged, UStaticMeshComponent*);
 	virtual FOnStaticMeshChanged& OnStaticMeshChanged() { return OnStaticMeshChangedEvent; }
-#endif
 
 private:
-#if WITH_EDITOR
 	FOnStaticMeshChanged OnStaticMeshChangedEvent;
-#endif
+
 	FOnGetNaniteResources OnGetNaniteResourcesEvent;
 
 	friend class FStaticMeshComponentRecreateRenderStateContext;
diff --git a/Engine/Source/Runtime/Engine/Private/Components/StaticMeshComponent.cpp b/Engine/Source/Runtime/Engine/Private/Components/StaticMeshComponent.cpp
--- a/Engine/Source/Runtime/Engine/Private/Components/StaticMeshComponent.cpp
+++ b/Engine/Source/Runtime/Engine/Private/Components/StaticMeshComponent.cpp
@@ -2355,10 +2355,8 @@ bool UStaticMeshComponent::SetStaticMesh(UStaticMesh* NewMesh)
 	// Mark cached material parameter names dirty
 	MarkCachedMaterialParameterNameIndicesDirty();
 
-#if WITH_EDITOR
 	// Broadcast that the static mesh has changed
 	OnStaticMeshChangedEvent.Broadcast(this);
-#endif
 
 #if WITH_EDITORONLY_DATA
 	if (GetStaticMesh())
