Always defer ray tracing geometry builds in FRayTracingGeometryManager.
- even if MaxBuiltPrimitivesPerFrame limit is not set or build priority is Immediate.
- this makes it possible to batch the build commands early in the frame.
- also adding GPU stat to track cost of BLAS builds and separated dynamic BLAS in a separate stat.

diff --git a/Engine/Source/Runtime/Engine/Private/Rendering/RayTracingGeometryManager.cpp b/Engine/Source/Runtime/Engine/Private/Rendering/RayTracingGeometryManager.cpp
--- a/Engine/Source/Runtime/Engine/Private/Rendering/RayTracingGeometryManager.cpp
+++ b/Engine/Source/Runtime/Engine/Private/Rendering/RayTracingGeometryManager.cpp
@@ -51,8 +51,7 @@ static float GetInitialBuildPriority(ERTAccelerationStructureBuildPriority InBui
 	case ERTAccelerationStructureBuildPriority::Skip:
 	default:
 	{
-		// should not get here
-		check(false);
+		checkNoEntry();
 		return 0.0f;
 	}
 	}
@@ -60,29 +59,19 @@ static float GetInitialBuildPriority(ERTAccelerationStructureBuildPriority InBui
 
 FRayTracingGeometryManager::BuildRequestIndex FRayTracingGeometryManager::RequestBuildAccelerationStructure(FRHICommandList& RHICmdList, FRayTracingGeometry* InGeometry, ERTAccelerationStructureBuildPriority InPriority, EAccelerationStructureBuildMode InBuildMode)
 {
-	// If immediate then enqueue command directly on the immediate command list
-	if (GRayTracingMaxBuiltPrimitivesPerFrame <= 0 || InPriority == ERTAccelerationStructureBuildPriority::Immediate)
-	{
-		check(InBuildMode == EAccelerationStructureBuildMode::Build);
-		RHICmdList.BuildAccelerationStructure(InGeometry->RayTracingGeometryRHI);
-		return INDEX_NONE;
-	}
-	else
-	{
-		BuildRequest Request;
-		Request.BuildPriority = GetInitialBuildPriority(InPriority);
-		Request.Owner = InGeometry;
-		Request.BuildMode = EAccelerationStructureBuildMode::Build;
+	BuildRequest Request;
+	Request.BuildPriority = GetInitialBuildPriority(InPriority);
+	Request.Owner = InGeometry;
+	Request.BuildMode = EAccelerationStructureBuildMode::Build;
 
-		FScopeLock ScopeLock(&RequestCS);
-		BuildRequestIndex RequestIndex = GeometryBuildRequests.Add(Request);
-		GeometryBuildRequests[RequestIndex].RequestIndex = RequestIndex;
+	FScopeLock ScopeLock(&RequestCS);
+	BuildRequestIndex RequestIndex = GeometryBuildRequests.Add(Request);
+	GeometryBuildRequests[RequestIndex].RequestIndex = RequestIndex;
 
-		INC_DWORD_STAT(STAT_RayTracingPendingBuilds);
-		INC_DWORD_STAT_BY(STAT_RayTracingPendingBuildPrimitives, InGeometry->Initializer.TotalPrimitiveCount);
+	INC_DWORD_STAT(STAT_RayTracingPendingBuilds);
+	INC_DWORD_STAT_BY(STAT_RayTracingPendingBuildPrimitives, InGeometry->Initializer.TotalPrimitiveCount);
 
-		return RequestIndex;
-	}
+	return RequestIndex;
 }
 
 void FRayTracingGeometryManager::RemoveBuildRequest(BuildRequestIndex InRequestIndex)
@@ -212,62 +201,93 @@ void FRayTracingGeometryManager::ProcessBuildRequests(FRHIComputeCommandList& In
 		return;
 	}
 
-	SortedRequests.Empty(FMath::Max(SortedRequests.Max(), GeometryBuildRequests.Num()));
+	checkf(BuildParams.IsEmpty(), TEXT("Unexpected entries in BuildParams. The array should've been reset at the end of the previous call."));
+	checkf(SortedRequests.IsEmpty(), TEXT("Unexpected entries in SortedRequests. The array should've been reset at the end of the previous call."));
+
+	BuildParams.Empty(FMath::Max(BuildParams.Max(), GeometryBuildRequests.Num()));
 
+	if (GRayTracingMaxBuiltPrimitivesPerFrame <= 0)
 	{
-		TRACE_CPUPROFILER_EVENT_SCOPE(SortRequests);
+		// no limit -> no need to sort
 
-		// Is there a fast way to extract all entries from sparse array?
-		for (const BuildRequest& Request : GeometryBuildRequests)
+		SortedRequests.Empty(); // free potentially allocated memory
+
+		for (BuildRequest& Request : GeometryBuildRequests)
 		{
-			SortedRequests.Add(Request);
+			const bool bRemoveFromRequestArray = false; // can't modify array while iterating over it
+			SetupBuildParams(Request, BuildParams, bRemoveFromRequestArray);
 		}
-		SortedRequests.Sort([](const BuildRequest& InLHS, const BuildRequest& InRHS)
-			{				
-				return InLHS.BuildPriority > InRHS.BuildPriority;
-			});
-	}
 
-	BuildParams.Empty(FMath::Max(BuildParams.Max(), SortedRequests.Num()));
-
-	// process n requests each 'frame'
-	uint64 PrimitivesBuild = 0;
-	bool bAddBuildRequest = true;
-	for (BuildRequest& Request : SortedRequests)
+		// after setting up build params can clear the whole array
+		GeometryBuildRequests.Reset();
+	}
+	else
 	{
-		if (bAddBuildRequest)
+		SortedRequests.Empty(FMath::Max(SortedRequests.Max(), GeometryBuildRequests.Num()));
+
 		{
-			SetupBuildParams(Request, BuildParams);
+			TRACE_CPUPROFILER_EVENT_SCOPE(SortRequests);
 
-			// Requested enough?
-			PrimitivesBuild += Request.Owner->Initializer.TotalPrimitiveCount;
-			if (!bInBuildAll && PrimitivesBuild > GRayTracingMaxBuiltPrimitivesPerFrame)
-				bAddBuildRequest = false;
+			// Is there a fast way to extract all entries from sparse array?
+			for (const BuildRequest& Request : GeometryBuildRequests)
+			{
+				SortedRequests.Add(Request);
+			}
+
+			SortedRequests.Sort([](const BuildRequest& InLHS, const BuildRequest& InRHS)
+				{
+					return InLHS.BuildPriority > InRHS.BuildPriority;
+				});
 		}
-		else
+
+		// process n requests each 'frame'
+		uint64 PrimitivesBuild = 0;
+		bool bAddBuildRequest = true;
+		for (BuildRequest& Request : SortedRequests)
 		{
-			// Increment priority to make sure requests don't starve
-			Request.BuildPriority += GRayTracingPendingBuildPriorityBoostPerFrame;
+			if (bAddBuildRequest || Request.BuildPriority >= 1.0f) // always build immediate requests
+			{
+				SetupBuildParams(Request, BuildParams);
+
+				// Requested enough?
+				PrimitivesBuild += Request.Owner->Initializer.TotalPrimitiveCount;
+				if (!bInBuildAll && (PrimitivesBuild > GRayTracingMaxBuiltPrimitivesPerFrame))
+				{
+					bAddBuildRequest = false;
+				}
+			}
+			else
+			{
+				// Increment priority to make sure requests don't starve
+				Request.BuildPriority += GRayTracingPendingBuildPriorityBoostPerFrame;
+			}
 		}
+
+		SortedRequests.Reset();
 	}
 
 	// kick actual build request to RHI command list
 	InCmdList.BuildAccelerationStructures(BuildParams);
+
+	BuildParams.Reset();
 }
 
-void FRayTracingGeometryManager::SetupBuildParams(const BuildRequest& InBuildRequest, TArray<FRayTracingGeometryBuildParams>& InBuildParams)
+void FRayTracingGeometryManager::SetupBuildParams(const BuildRequest& InBuildRequest, TArray<FRayTracingGeometryBuildParams>& InBuildParams, bool bRemoveFromRequestArray)
 {
-	// Setup the actual build params
+	check(InBuildRequest.RequestIndex != INDEX_NONE && InBuildRequest.Owner->RayTracingBuildRequestIndex != INDEX_NONE);
+
 	FRayTracingGeometryBuildParams BuildParam;
 	BuildParam.Geometry = InBuildRequest.Owner->RayTracingGeometryRHI;
 	BuildParam.BuildMode = InBuildRequest.BuildMode;
 	InBuildParams.Add(BuildParam);
 
-	// Remove from pending array and update the geometry that data is valid
-	check(InBuildRequest.RequestIndex != INDEX_NONE);
-	GeometryBuildRequests.RemoveAt(InBuildRequest.RequestIndex);
 	InBuildRequest.Owner->RayTracingBuildRequestIndex = INDEX_NONE;
 
+	if (bRemoveFromRequestArray)
+	{
+		GeometryBuildRequests.RemoveAt(InBuildRequest.RequestIndex);
+	}
+
 	DEC_DWORD_STAT(STAT_RayTracingPendingBuilds);
 	DEC_DWORD_STAT_BY(STAT_RayTracingPendingBuildPrimitives, InBuildRequest.Owner->Initializer.TotalPrimitiveCount);
 }
diff --git a/Engine/Source/Runtime/Engine/Public/Rendering/RayTracingGeometryManager.h b/Engine/Source/Runtime/Engine/Public/Rendering/RayTracingGeometryManager.h
--- a/Engine/Source/Runtime/Engine/Public/Rendering/RayTracingGeometryManager.h
+++ b/Engine/Source/Runtime/Engine/Public/Rendering/RayTracingGeometryManager.h
@@ -45,7 +45,7 @@ private:
 		EAccelerationStructureBuildMode BuildMode;
 	};
 
-	void SetupBuildParams(const BuildRequest& InBuildRequest, TArray<FRayTracingGeometryBuildParams>& InBuildParams);
+	void SetupBuildParams(const BuildRequest& InBuildRequest, TArray<FRayTracingGeometryBuildParams>& InBuildParams, bool bRemoveFromRequestArray = true);
 
 	FCriticalSection RequestCS;
 
diff --git a/Engine/Source/Runtime/Renderer/Private/DeferredShadingRenderer.cpp b/Engine/Source/Runtime/Renderer/Private/DeferredShadingRenderer.cpp
--- a/Engine/Source/Runtime/Renderer/Private/DeferredShadingRenderer.cpp
+++ b/Engine/Source/Runtime/Renderer/Private/DeferredShadingRenderer.cpp
@@ -431,6 +431,7 @@ DECLARE_CYCLE_STAT(TEXT("DeferredShadingSceneRenderer RenderFinish"), STAT_FDefe
 DECLARE_GPU_STAT(RayTracingUpdate);
 DECLARE_GPU_STAT(RayTracingScene);
 DECLARE_GPU_STAT(RayTracingGeometry);
+DECLARE_GPU_STAT(RayTracingDynamicGeometry);
 
 DEFINE_GPU_STAT(Postprocessing);
 DECLARE_GPU_STAT(VisibilityCommands);
@@ -2096,7 +2097,7 @@ bool FDeferredShadingSceneRenderer::DispatchRayTracingWorldUpdates(FRDGBuilder&
 			GraphBuilder.AddPass(RDG_EVENT_NAME("RayTracingDynamicUpdate"), PassParams, ComputePassFlags | ERDGPassFlags::NeverCull | ERDGPassFlags::NeverParallel,
 				[this, PassParams, bRayTracingAsyncBuild](FRHICommandListImmediate& RHICmdList)
 			{
-				SCOPED_GPU_STAT(RHICmdList, RayTracingGeometry);
+				SCOPED_GPU_STAT(RHICmdList, RayTracingDynamicGeometry);
 				FRHIBuffer* DynamicGeometryScratchBuffer = PassParams->DynamicGeometryScratchBuffer ? PassParams->DynamicGeometryScratchBuffer->GetRHI() : nullptr;
 				Scene->GetRayTracingDynamicGeometryCollection()->DispatchUpdates(RHICmdList, DynamicGeometryScratchBuffer);
 			});
@@ -2669,6 +2670,7 @@ void FDeferredShadingSceneRenderer::Render(FRDGBuilder& GraphBuilder)
 	if (RendererOutput != FSceneRenderer::ERendererOutput::DepthPrepassOnly)
 	{
 		// TODO: should only process build requests once per frame
+		SCOPED_GPU_STAT(GraphBuilder.RHICmdList, RayTracingGeometry);
 		GRayTracingGeometryManager->ProcessBuildRequests(GraphBuilder.RHICmdList);
 	}
 #endif
