World Partition - HLOD: Allow line traces on HLOD actors in the editor
* Will enable many editor features to work with HLOD actors (placement tool, play from here, go here, etc)
* Build collision data for HLOD static meshes
* Enable queries collisions in editor
* Strip all collision data during cook

diff --git a/Engine/Plugins/Editor/WorldPartitionHLODUtilities/Source/Private/WorldPartition/HLOD/Utilities/WorldPartitionHLODUtilities.cpp b/Engine/Plugins/Editor/WorldPartitionHLODUtilities/Source/Private/WorldPartition/HLOD/Utilities/WorldPartitionHLODUtilities.cpp
--- a/Engine/Plugins/Editor/WorldPartitionHLODUtilities/Source/Private/WorldPartition/HLOD/Utilities/WorldPartitionHLODUtilities.cpp
+++ b/Engine/Plugins/Editor/WorldPartitionHLODUtilities/Source/Private/WorldPartition/HLOD/Utilities/WorldPartitionHLODUtilities.cpp
@@ -45,7 +45,7 @@ static uint32 ComputeHLODHash(AWorldPartitionHLOD* InHLODActor, const TArray<UAc
 	FArchiveCrc32 Ar;
 
 	// Base key, changing this will force a rebuild of all HLODs
-	FString HLODBaseKey = "DFCCA659C5964FB29D31A4DF2682D4D6";
+	FString HLODBaseKey = "75D23ACFD6C549069945B08D2A977661";
 	Ar << HLODBaseKey;
 
 	// HLOD Source Actors
@@ -820,12 +820,23 @@ uint32 FWorldPartitionHLODUtilities::BuildHLOD(AWorldPartitionHLOD* InHLODActor)
 								StaticMeshComponent->bRayTracingFarField = true;
 							}
 
-							// Disable collisions
+							// Disable navigation data on HLODs
 							StaticMesh->MarkAsNotHavingNavigationData();
-							if (UBodySetup* BodySetup = StaticMesh->GetBodySetup())
+
+							// Ensure we can perform line trace on HLOD static meshes. This is useful for all kind of editor features on HLODs (Actor placement, Play from here, Go Here, etc)
+							// Collision data will be stripped at cook.
+							StaticMesh->CreateBodySetup();
+							UBodySetup* BodySetup = StaticMesh->GetBodySetup();
+							BodySetup->DefaultInstance.SetCollisionEnabled(ECollisionEnabled::QueryOnly);
+							BodySetup->DefaultInstance.SetCollisionProfileName(UCollisionProfile::BlockAll_ProfileName);
+							for (int32 LODIndex = 0; LODIndex < StaticMesh->GetNumLODs(); ++LODIndex)
 							{
-								BodySetup->DefaultInstance.SetCollisionProfileName(UCollisionProfile::NoCollision_ProfileName);
-								BodySetup->CollisionTraceFlag = CTF_UseSimpleAsComplex;
+								for (int32 SectionIndex = 0; SectionIndex < StaticMesh->GetNumSections(LODIndex); ++SectionIndex)
+								{
+									FMeshSectionInfo MeshSectionInfo = StaticMesh->GetSectionInfoMap().Get(LODIndex, SectionIndex);
+									MeshSectionInfo.bEnableCollision = true;
+									StaticMesh->GetSectionInfoMap().Set(LODIndex, SectionIndex, MeshSectionInfo);
+								}
 							}
 
 							// Rename owned static mesh
diff --git a/Engine/Source/Runtime/Engine/Private/WorldPartition/HLOD/HLODActor.cpp b/Engine/Source/Runtime/Engine/Private/WorldPartition/HLOD/HLODActor.cpp
--- a/Engine/Source/Runtime/Engine/Private/WorldPartition/HLOD/HLODActor.cpp
+++ b/Engine/Source/Runtime/Engine/Private/WorldPartition/HLOD/HLODActor.cpp
@@ -16,9 +16,11 @@
 #include "WorldPartition/WorldPartitionHelpers.h"
 
 #if WITH_EDITOR
+#include "Components/StaticMeshComponent.h"
 #include "Editor.h"
 #include "Modules/ModuleManager.h"
 #include "Misc/ArchiveMD5.h"
+#include "PhysicsEngine/BodySetup.h"
 #include "Subsystems/AssetEditorSubsystem.h"
 #include "UObject/FortniteMainBranchObjectVersion.h"
 #include "WorldPartition/HLOD/HLODActorDesc.h"
@@ -249,6 +251,47 @@ void AWorldPartitionHLOD::PostLoad()
 #endif
 }
 
+#if WITH_EDITOR
+void AWorldPartitionHLOD::PreSave(FObjectPreSaveContext ObjectSaveContext)
+{
+	Super::PreSave(ObjectSaveContext);
+
+	// Always disable collisions on HLODs
+	SetActorEnableCollision(false);
+
+	ForEachComponent<UPrimitiveComponent>(false, [this, &ObjectSaveContext](UPrimitiveComponent* PrimitiveComponent)
+	{
+		// Disable collision on HLOD components
+		PrimitiveComponent->SetCollisionEnabled(ECollisionEnabled::NoCollision);
+
+		// When cooking, get rid of collision data
+		if (ObjectSaveContext.IsCooking())
+		{
+			if (UStaticMeshComponent* StaticMeshComponent = Cast<UStaticMeshComponent>(PrimitiveComponent))
+			{
+				if (UStaticMesh* StaticMesh = StaticMeshComponent->GetStaticMesh())
+				{
+					// If the HLOD process did create this static mesh
+					if (StaticMesh->GetPackage() == GetPackage())
+					{
+						if (UBodySetup* BodySetup = StaticMesh->GetBodySetup())
+						{
+ 							// To ensure a deterministic cook, save the current GUID and restore it below
+							FGuid PreviousBodySetupGuid = BodySetup->BodySetupGuid;
+							BodySetup->DefaultInstance.SetCollisionProfileName(UCollisionProfile::NoCollision_ProfileName);
+							BodySetup->bNeverNeedsCookedCollisionData = true;
+							BodySetup->bHasCookedCollisionData = false;
+							BodySetup->InvalidatePhysicsData();
+							BodySetup->BodySetupGuid = PreviousBodySetupGuid;
+						}
+					}
+				}
+			}
+		}
+	});
+}
+#endif
+
 void AWorldPartitionHLOD::PreRegisterAllComponents()
 {
 	Super::PreRegisterAllComponents();
@@ -261,6 +304,25 @@ void AWorldPartitionHLOD::PreRegisterAllComponents()
 		});
 	}
 
+#if WITH_EDITOR
+	// In editor, turn on collision on HLODs in order to enable some useful editor features on HLODs (Actor placement, Play from here, Go Here, etc)
+	// In PIE, collisions should be disabled
+	if (GetWorld() && !IsRunningCommandlet() && !FApp::IsUnattended())
+	{
+		bool bShouldEnableCollision = !GetWorld()->IsGameWorld();
+		if (GetActorEnableCollision() != bShouldEnableCollision)
+		{
+			SetActorEnableCollision(bShouldEnableCollision);
+			ForEachComponent<UPrimitiveComponent>(false, [bShouldEnableCollision](UPrimitiveComponent* PrimitiveComponent)
+			{
+				PrimitiveComponent->SetCollisionEnabled(bShouldEnableCollision ? ECollisionEnabled::QueryOnly : ECollisionEnabled::NoCollision);
+				PrimitiveComponent->SetCollisionResponseToChannel(ECC_Visibility, bShouldEnableCollision ? ECR_Block : ECR_Ignore);
+				PrimitiveComponent->SetCollisionResponseToChannel(ECC_Camera, bShouldEnableCollision ? ECR_Block : ECR_Ignore);
+			});
+		}
+	}	
+#endif
+
 	// If world is instanced, we need to recompute our bounds since they are in the instanced-world space
 	if (UWorldPartition* WorldPartition = FWorldPartitionHelpers::GetWorldPartition(this))
 	{
diff --git a/Engine/Source/Runtime/Engine/Public/WorldPartition/HLOD/HLODActor.h b/Engine/Source/Runtime/Engine/Public/WorldPartition/HLOD/HLODActor.h
--- a/Engine/Source/Runtime/Engine/Public/WorldPartition/HLOD/HLODActor.h
+++ b/Engine/Source/Runtime/Engine/Public/WorldPartition/HLOD/HLODActor.h
@@ -8,6 +8,11 @@
 #include "WorldPartition/WorldPartitionRuntimeCell.h"
 #include "WorldPartition/WorldPartitionActorDesc.h"
 #include "WorldPartition/HLOD/HLODStats.h"
+
+#if WITH_EDITOR
+#include "UObject/ObjectSaveContext.h"
+#endif // WITH_EDITOR
+
 #include "HLODActor.generated.h"
 
 class UHLODLayer;
@@ -65,6 +70,7 @@ protected:
 	ENGINE_API virtual bool NeedsLoadForServer() const override;
 	ENGINE_API virtual void PostLoad() override;
 #if WITH_EDITOR
+	ENGINE_API virtual void PreSave(FObjectPreSaveContext ObjectSaveContext) override;
 	ENGINE_API virtual void RerunConstructionScripts() override;
 	virtual bool CanEditChange(const FProperty* InProperty) const override { return false; }
 	virtual bool CanEditChangeComponent(const UActorComponent* Component, const FProperty* InProperty) const override { return false; }
diff --git a/Engine/Source/Runtime/Landscape/Private/LandscapeHLODBuilder.cpp b/Engine/Source/Runtime/Landscape/Private/LandscapeHLODBuilder.cpp
--- a/Engine/Source/Runtime/Landscape/Private/LandscapeHLODBuilder.cpp
+++ b/Engine/Source/Runtime/Landscape/Private/LandscapeHLODBuilder.cpp
@@ -314,16 +314,6 @@ TArray<UActorComponent*> ULandscapeHLODBuilder::Build(const FHLODBuildContext& I
 
 	UStaticMesh::BatchBuild(StaticMeshes);
 
-	// Once static meshes have been built, their BodySetup should be available : 
-	for (UStaticMesh* StaticMesh : StaticMeshes)
-	{
-		// We won't ever enable collisions, ensure we don't even cook or load any collision data on this mesh: 
-		if (UBodySetup* BodySetup = StaticMesh->GetBodySetup())
-		{
-			BodySetup->bNeverNeedsCookedCollisionData = true;
-		}
-	}
-
 	return HLODComponents;
 }
 
