Fix r.vt.PoolSizeScale so that it can be changed at runtime and pools will immediately adjust.

diff --git a/Engine/Config/BaseEngine.ini b/Engine/Config/BaseEngine.ini
--- a/Engine/Config/BaseEngine.ini
+++ b/Engine/Config/BaseEngine.ini
@@ -3343,7 +3343,7 @@ DefaultChaosSolverActorClass=/Script/ChaosSolverEngine.ChaosSolverActor
 ; "r.VT.ListPhysicalPool" list in-depth details on the allocated physical pools
 ; "r.VT.DumpPoolUsage" list in-depth details on asset usage of physical pools
 ; The default pool size and behavior is set here. It can be overridden in project config files.
-+Pools=(SizeInMegabyte=64, bAllowSizeScale=False, bEnableResidencyMipMapBias=True)
++Pools=(SizeInMegabyte=64, bAllowSizeScale=True, bEnableResidencyMipMapBias=True)
 ; Specific pools can be configured per tile size or format in project config files. Examples commented below:
 ;+Pools=(Formats=(PF_DXT1), SizeInMegabyte=36, bAllowSizeScale=True, ScalabilityGroup=1, bEnableResidencyMipMapBias=False)
 ;+Pools=(Formats=(PF_DXT5, PF_DXT5), SizeInMegabyte=100)
diff --git a/Engine/Source/Runtime/Engine/Private/LightMap.cpp b/Engine/Source/Runtime/Engine/Private/LightMap.cpp
--- a/Engine/Source/Runtime/Engine/Private/LightMap.cpp
+++ b/Engine/Source/Runtime/Engine/Private/LightMap.cpp
@@ -3424,6 +3424,7 @@ static void OnVirtualTextureDestroyed(const FVirtualTextureProducerHandle& InHan
 {
 	FLightmapResourceCluster* Cluster = static_cast<FLightmapResourceCluster*>(Baton);
 	Cluster->ReleaseAllocatedVT();
+	Cluster->ConditionalCreateAllocatedVT();
 	Cluster->UpdateUniformBuffer();
 }
 
diff --git a/Engine/Source/Runtime/Engine/Private/VT/VirtualTextureScalability.cpp b/Engine/Source/Runtime/Engine/Private/VT/VirtualTextureScalability.cpp
--- a/Engine/Source/Runtime/Engine/Private/VT/VirtualTextureScalability.cpp
+++ b/Engine/Source/Runtime/Engine/Private/VT/VirtualTextureScalability.cpp
@@ -199,37 +199,20 @@ namespace VirtualTextureScalability
 
 		if (bUpdate)
 		{
-			// Temporarily release runtime virtual textures
-			for (TObjectIterator<URuntimeVirtualTexture> It(RF_ClassDefaultObject, false, EInternalObjectFlags::Garbage); It; ++It)
-			{
-				It->Release();
-			}
+			// Release all VT render resources here.
+			// Assuming all virtual textures are released, then virtual texture pools will reach a zero ref count and release, which is needed for any pool size scale to be effective.
+			// Note that for pool size scale changes, there will be a transition period (with high memory watermark) where new pools are created before old pools are released.
 
-			// Release streaming virtual textures
-			TArray<UTexture2D*> ReleasedVirtualTextures;
+			// Reinit streaming virtual textures.
 			for (TObjectIterator<UTexture2D> It(RF_ClassDefaultObject, false, EInternalObjectFlags::Garbage); It; ++It)
 			{
 				if (It->IsCurrentlyVirtualTextured() && It->GetResource() != nullptr)
 				{
-					ReleasedVirtualTextures.Add(*It);
-					BeginReleaseResource(It->GetResource());
+					BeginUpdateResourceRHI(It->GetResource());
 				}
 			}
 
-			// Force garbage collect of pools
-			ENQUEUE_RENDER_COMMAND(VirtualTextureScalability_Release)([](FRHICommandList& RHICmdList)
-			{
-				GetRendererModule().ReleaseVirtualTexturePendingResources();
-			});
-
-			// Now all pools should be flushed...
-			// Reinit streaming virtual textures
-			for (UTexture2D* Texture : ReleasedVirtualTextures)
-			{
-				BeginInitResource(Texture->GetResource());
-			}
-
-			// Reinit runtime virtual textures
+			// Reinit runtime virtual textures.
 			for (TObjectIterator<URuntimeVirtualTextureComponent> It(RF_ClassDefaultObject, false, EInternalObjectFlags::Garbage); It; ++It)
 			{
 				It->MarkRenderStateDirty();
diff --git a/Engine/Source/Runtime/Renderer/Private/VT/VirtualTexturePhysicalSpace.h b/Engine/Source/Runtime/Renderer/Private/VT/VirtualTexturePhysicalSpace.h
--- a/Engine/Source/Runtime/Renderer/Private/VT/VirtualTexturePhysicalSpace.h
+++ b/Engine/Source/Runtime/Renderer/Private/VT/VirtualTexturePhysicalSpace.h
@@ -12,6 +12,7 @@
 
 struct FVTPhysicalSpaceDescription
 {
+	uint32 PoolSettingsHash;
 	uint32 TileSize;
 	uint8 Dimensions;
 	uint8 NumLayers;
@@ -22,7 +23,8 @@ struct FVTPhysicalSpaceDescription
 
 inline bool operator==(const FVTPhysicalSpaceDescription& Lhs, const FVTPhysicalSpaceDescription& Rhs)
 {
-	if (Lhs.TileSize != Rhs.TileSize || 
+	if (Lhs.PoolSettingsHash != Rhs.PoolSettingsHash ||
+		Lhs.TileSize != Rhs.TileSize ||
 		Lhs.NumLayers != Rhs.NumLayers || 
 		Lhs.Dimensions != Rhs.Dimensions || 
 		Lhs.bContinuousUpdate != Rhs.bContinuousUpdate)
@@ -51,7 +53,8 @@ inline bool operator!=(const FVTPhysicalSpaceDescription& Lhs, const FVTPhysical
 
 inline uint32 GetTypeHash(const FVTPhysicalSpaceDescription& Desc)
 {
-	uint32 Hash = GetTypeHash(Desc.TileSize);
+	uint32 Hash = Desc.PoolSettingsHash;
+	Hash = HashCombine(Hash, GetTypeHash(Desc.TileSize));
 	Hash = HashCombine(Hash, GetTypeHash(Desc.Dimensions));
 	Hash = HashCombine(Hash, GetTypeHash(Desc.NumLayers));
 	Hash = HashCombine(Hash, GetTypeHash(Desc.bContinuousUpdate));
diff --git a/Engine/Source/Runtime/Renderer/Private/VT/VirtualTextureProducer.cpp b/Engine/Source/Runtime/Renderer/Private/VT/VirtualTextureProducer.cpp
--- a/Engine/Source/Runtime/Renderer/Private/VT/VirtualTextureProducer.cpp
+++ b/Engine/Source/Runtime/Renderer/Private/VT/VirtualTextureProducer.cpp
@@ -1,8 +1,10 @@
 // Copyright Epic Games, Inc. All Rights Reserved.
 
 #include "VirtualTextureProducer.h"
-#include "VirtualTextureSystem.h"
-#include "VirtualTexturePhysicalSpace.h"
+
+#include "VT/VirtualTexturePhysicalSpace.h"
+#include "VT/VirtualTextureScalability.h"
+#include "VT/VirtualTextureSystem.h"
 
 FVirtualTextureProducer::~FVirtualTextureProducer()
 {
@@ -71,6 +73,7 @@ FVirtualTextureProducerHandle FVirtualTextureProducerCollection::RegisterProduce
 	Entry.DestroyedCallbacksIndex = AcquireCallback();
 
 	FVTPhysicalSpaceDescription PhysicalSpaceDesc;
+	PhysicalSpaceDesc.PoolSettingsHash = VirtualTextureScalability::GetPhysicalPoolSettingsHash();
 	PhysicalSpaceDesc.Dimensions = InDesc.Dimensions;
 	PhysicalSpaceDesc.TileSize = InDesc.TileSize + InDesc.TileBorderSize * 2u;
 	PhysicalSpaceDesc.bContinuousUpdate = InDesc.bContinuousUpdate;
