Move reading of r.VT.RVT.EnableVolumes cvar to scene proxy creation time. Needs to happen late enough that the cvar is set if it is being driven from configuration systems.

diff --git a/Engine/Source/Runtime/Engine/Private/Components/RuntimeVirtualTextureComponent.cpp b/Engine/Source/Runtime/Engine/Private/Components/RuntimeVirtualTextureComponent.cpp
--- a/Engine/Source/Runtime/Engine/Private/Components/RuntimeVirtualTextureComponent.cpp
+++ b/Engine/Source/Runtime/Engine/Private/Components/RuntimeVirtualTextureComponent.cpp
@@ -101,7 +101,7 @@ void URuntimeVirtualTextureComponent::GetHidePrimitiveSettings(bool& OutHidePrim
 bool URuntimeVirtualTextureComponent::ShouldCreateRenderState() const
 {
 	// Make sure to have the component do nothing if VT is disabled or if the world is not compatible with RVT
-	return Super::ShouldCreateRenderState() && IsActiveInWorld() && CVarRVTEnableVolumes.GetValueOnGameThread() != 0 && UseVirtualTexturing(GetScene()->GetShaderPlatform());
+	return Super::ShouldCreateRenderState() && IsActiveInWorld() && UseVirtualTexturing(GetScene()->GetShaderPlatform());
 }
 
 void URuntimeVirtualTextureComponent::ApplyWorldOffset(const FVector& InOffset, bool bWorldShift)
@@ -145,6 +145,17 @@ void URuntimeVirtualTextureComponent::DestroyRenderState_Concurrent()
 
 bool URuntimeVirtualTextureComponent::IsEnabledInScene() const
 {
+	if (CVarRVTEnableVolumes.GetValueOnGameThread() == 0)
+	{
+		return false;
+	}
+
+	const bool bUseNanite = UseNanite(GetScene()->GetShaderPlatform());
+	if (bEnableForNaniteOnly && !bUseNanite)
+	{
+		return false;
+	}
+
 	if (UWorld* World = GetWorld())
 	{
 		if (World->WorldType == EWorldType::Game || World->WorldType == EWorldType::PIE)
@@ -161,12 +172,6 @@ bool URuntimeVirtualTextureComponent::IsEnabledInScene() const
 			}
 #endif
 		}
-
-		const bool bUseNanite = UseNanite(GetFeatureLevelShaderPlatform(World->GetFeatureLevel()));
-		if (bEnableForNaniteOnly && !bUseNanite)
-		{
-			return false;
-		}
 	}
 
 	return true;
