Fix for crash when unlink layers is called w/ null mesh
#rnx

diff --git a/Engine/Source/Runtime/Engine/Private/Animation/AnimInstance.cpp b/Engine/Source/Runtime/Engine/Private/Animation/AnimInstance.cpp
--- a/Engine/Source/Runtime/Engine/Private/Animation/AnimInstance.cpp
+++ b/Engine/Source/Runtime/Engine/Private/Animation/AnimInstance.cpp
@@ -3104,10 +3104,14 @@ void UAnimInstance::PerformLinkedLayerOverlayOperation(TSubclassOf<UAnimInstance
 		auto InitializeAndCacheBonesForLinkedRoot = [](FAnimNode_LinkedAnimLayer* InLayerNode, FAnimInstanceProxy& InThisProxy, UAnimInstance* InLinkedInstance, FAnimInstanceProxy& InLinkedProxy)
 		{
 			InLinkedProxy.InitializeObjects(InLinkedInstance);
-			FAnimationInitializeContext InitContext(&InThisProxy);
-			InLayerNode->InitializeSubGraph_AnyThread(InitContext);
-			FAnimationCacheBonesContext CacheBonesContext(&InThisProxy);
-			InLayerNode->CacheBonesSubGraph_AnyThread(CacheBonesContext);
+
+			if (InLinkedInstance->GetSkelMeshComponent()->GetSkeletalMeshAsset() != nullptr)
+			{
+				FAnimationInitializeContext InitContext(&InThisProxy);
+				InLayerNode->InitializeSubGraph_AnyThread(InitContext);
+				FAnimationCacheBonesContext CacheBonesContext(&InThisProxy);
+				InLayerNode->CacheBonesSubGraph_AnyThread(CacheBonesContext);
+			}
 		};
 
 		for (TPair<UClass*, TMap<FName, TArray<FAnimNode_LinkedAnimLayer*, TInlineAllocator<4>>, TInlineSetAllocator<4>>> ClassLayerNodesToSet : LayerNodesToSet)
