Moved skylight  illuminance visualisation to a showflag.

diff --git a/Engine/Source/Runtime/Engine/Classes/Components/SkyLightComponent.h b/Engine/Source/Runtime/Engine/Classes/Components/SkyLightComponent.h
--- a/Engine/Source/Runtime/Engine/Classes/Components/SkyLightComponent.h
+++ b/Engine/Source/Runtime/Engine/Classes/Components/SkyLightComponent.h
@@ -203,13 +203,6 @@ class USkyLightComponent : public ULightComponentBase
 	/** Controls how occlusion from Distance Field Ambient Occlusion is combined with Screen Space Ambient Occlusion. */
 	UPROPERTY(EditAnywhere, BlueprintReadOnly, Category=DistanceFieldAmbientOcclusion)
 	TEnumAsByte<enum EOcclusionCombineMode> OcclusionCombineMode;
-	
-	/**
-	 * When enabled, visualize on screen the sky cube map as well as the integrate half hemisphere of luminance, together with the sky illuminance contribution at the position of the skylight component.
-	 * It also show the atmospheric directional light illuminance post SkyAtmosphere transmittance at the position on the skylight component.
-	 */
-	UPROPERTY(EditAnywhere, BlueprintReadOnly, Category = Visualization)
-	uint32 bShowIlluminanceMeter : 1;
 		
 	ENGINE_API class FSkyLightSceneProxy* CreateSceneProxy() const;
 
diff --git a/Engine/Source/Runtime/Engine/Private/Components/SkyLightComponent.cpp b/Engine/Source/Runtime/Engine/Private/Components/SkyLightComponent.cpp
--- a/Engine/Source/Runtime/Engine/Private/Components/SkyLightComponent.cpp
+++ b/Engine/Source/Runtime/Engine/Private/Components/SkyLightComponent.cpp
@@ -243,7 +243,6 @@ FSkyLightSceneProxy::FSkyLightSceneProxy(const USkyLightComponent* InLightCompon
 	, bCaptureSkyLightWaitingForShaders(false)
 	, bCaptureSkyLightWaitingForMeshesOrTextures(false)
 #endif
-	, bShowIlluminanceMeter(InLightComponent->bShowIlluminanceMeter)
 	, LightColor(FLinearColor(InLightComponent->LightColor) * InLightComponent->Intensity)
 	, bMovable(InLightComponent->IsMovable())
 {
@@ -307,7 +306,6 @@ USkyLightComponent::USkyLightComponent(const FObjectInitializer& ObjectInitializ
 	CloudAmbientOcclusionStrength = 1.0f;
 	CloudAmbientOcclusionMapResolutionScale = 1.0f;
 	CloudAmbientOcclusionApertureScale = 0.05f;
-	bShowIlluminanceMeter = false;
 
 #if WITH_EDITOR
 	CaptureStatus = ESkyLightCaptureStatus::SLCS_Uninitialized;
diff --git a/Engine/Source/Runtime/Engine/Public/SceneManagement.h b/Engine/Source/Runtime/Engine/Public/SceneManagement.h
--- a/Engine/Source/Runtime/Engine/Public/SceneManagement.h
+++ b/Engine/Source/Runtime/Engine/Public/SceneManagement.h
@@ -1176,8 +1176,6 @@ public:
 	bool bCaptureSkyLightWaitingForMeshesOrTextures;
 #endif
 
-	uint8 bShowIlluminanceMeter : 1;
-
 private:
 	FLinearColor LightColor;
 	const uint8 bMovable : 1;
diff --git a/Engine/Source/Runtime/Engine/Public/ShowFlags.h b/Engine/Source/Runtime/Engine/Public/ShowFlags.h
--- a/Engine/Source/Runtime/Engine/Public/ShowFlags.h
+++ b/Engine/Source/Runtime/Engine/Public/ShowFlags.h
@@ -370,6 +370,7 @@ private:
 
 		// The following code sets what should be off by default.
 		SetVisualizeHDR(false);
+		SetVisualizeSkyLightIlluminance(false);
 		SetVisualizeLocalExposure(false);
 		SetVisualizeShadingModels(false);
 		SetOverrideDiffuseAndSpecular(false);
diff --git a/Engine/Source/Runtime/Engine/Public/ShowFlagsValues.inl b/Engine/Source/Runtime/Engine/Public/ShowFlagsValues.inl
--- a/Engine/Source/Runtime/Engine/Public/ShowFlagsValues.inl
+++ b/Engine/Source/Runtime/Engine/Public/ShowFlagsValues.inl
@@ -36,6 +36,8 @@ SHOWFLAG_ALWAYS_ACCESSIBLE(AmbientCubemap, SFG_LightingFeatures, NSLOCTEXT("Unre
 SHOWFLAG_ALWAYS_ACCESSIBLE(EyeAdaptation, SFG_PostProcess, NSLOCTEXT("UnrealEd", "EyeAdaptationSF", "Eye Adaptation"))
 /** Display a histogram of the scene HDR color */
 SHOWFLAG_FIXED_IN_SHIPPING(0, VisualizeHDR, SFG_Visualize, NSLOCTEXT("UnrealEd", "VisualizeHDRSF", "HDR (Eye Adaptation)"))
+/** Display the illuminance debug view for the skylight */
+SHOWFLAG_FIXED_IN_SHIPPING(0, VisualizeSkyLightIlluminance, SFG_Visualize, NSLOCTEXT("UnrealEd", "VisualizeSkyLightIlluminance", "Visualize SkyLight Illuminance"))
 /** Helper to tweak local expsoure settings */
 SHOWFLAG_FIXED_IN_SHIPPING(0, VisualizeLocalExposure, SFG_Visualize, NSLOCTEXT("UnrealEd", "VisualizeLocalExposureSF", "Local Exposure"))
 /** Image based lens flares (Simulate artifact of reflections within a camera system) */
diff --git a/Engine/Source/Runtime/Renderer/Private/IlluminanceMeter.cpp b/Engine/Source/Runtime/Renderer/Private/IlluminanceMeter.cpp
--- a/Engine/Source/Runtime/Renderer/Private/IlluminanceMeter.cpp
+++ b/Engine/Source/Runtime/Renderer/Private/IlluminanceMeter.cpp
@@ -99,7 +99,7 @@ IMPLEMENT_GLOBAL_SHADER(FPrintIlluminanceMeterPS, "/Engine/Private/ReflectionEnv
 void FScene::ProcessAndRenderIlluminanceMeter(FRDGBuilder& GraphBuilder, TArrayView<FViewInfo> Views, FRDGTextureRef SceneColorTexture)
 {
 	FViewInfo& MainView = Views[0];
-	if (!SkyLight || !SkyLight->bShowIlluminanceMeter || MainView.bIsSceneCapture || MainView.bIsReflectionCapture || MainView.bIsPlanarReflection ||
+	if (!SkyLight || !MainView.Family->EngineShowFlags.VisualizeSkyLightIlluminance || MainView.bIsSceneCapture || MainView.bIsReflectionCapture || MainView.bIsPlanarReflection ||
 		!CanIlluminanceMeterDisplayOnPlatform(MainView.GetShaderPlatform()))
 	{
 		return;
diff --git a/Engine/Source/Runtime/Renderer/Private/SceneRendering.cpp b/Engine/Source/Runtime/Renderer/Private/SceneRendering.cpp
--- a/Engine/Source/Runtime/Renderer/Private/SceneRendering.cpp
+++ b/Engine/Source/Runtime/Renderer/Private/SceneRendering.cpp
@@ -4206,7 +4206,7 @@ FSceneRenderer* FSceneRenderer::CreateSceneRenderer(const FSceneViewFamily* InVi
 
 bool FSceneRenderer::ShouldCompositeEditorPrimitives(const FViewInfo& View)
 {
-	if (View.Family->EngineShowFlags.VisualizeHDR || View.Family->UseDebugViewPS())
+	if (View.Family->EngineShowFlags.VisualizeHDR || View.Family->EngineShowFlags.VisualizeSkyLightIlluminance || View.Family->UseDebugViewPS())
 	{
 		// certain visualize modes get obstructed too much
 		return false;
