New HLSL Translator fixes
Crash fix in FExpressionSceneTexture
Correct defaults for material function named reroutes
Repair debug preshader eval of conditionals

#tests materials, editor, PIE, win64

diff --git a/Engine/Source/Runtime/Engine/Private/Materials/MaterialHLSLGenerator.cpp b/Engine/Source/Runtime/Engine/Private/Materials/MaterialHLSLGenerator.cpp
--- a/Engine/Source/Runtime/Engine/Private/Materials/MaterialHLSLGenerator.cpp
+++ b/Engine/Source/Runtime/Engine/Private/Materials/MaterialHLSLGenerator.cpp
@@ -567,7 +567,13 @@ const UE::HLSLTree::FExpression* FMaterialHLSLGenerator::AcquireExpression(UE::H
 	else
 	{
 		check(!Expression);
-		Expression = GetTree().NewExpression<UE::HLSLTree::FExpressionError>(AcquireError());
+		FStringView CurrentError = AcquireError();
+		if (!CurrentError.IsEmpty())
+		{
+			// if we have an error, generate an error expression
+			// if not, continue to return nullptr so caller can default missing inputs
+			Expression = GetTree().NewExpression<UE::HLSLTree::FExpressionError>(CurrentError);
+		}
 	}
 
 	return Expression;
diff --git a/Engine/Source/Runtime/Engine/Private/Materials/MaterialHLSLTree.cpp b/Engine/Source/Runtime/Engine/Private/Materials/MaterialHLSLTree.cpp
--- a/Engine/Source/Runtime/Engine/Private/Materials/MaterialHLSLTree.cpp
+++ b/Engine/Source/Runtime/Engine/Private/Materials/MaterialHLSLTree.cpp
@@ -2309,7 +2309,7 @@ void FExpressionSceneTexture::EmitValueShader(FEmitContext& Context, FEmitScope&
 	}
 	else
 	{
-		EmitLookup = Context.EmitExpression(Scope, Shader::EValueType::Float4, TEXT("MobileSceneTextureLookup(Parameters, %, %, %)"), (int)SceneTextureId, EmitTexCoord);
+		EmitLookup = Context.EmitExpression(Scope, Shader::EValueType::Float4, TEXT("MobileSceneTextureLookup(Parameters, %, %)"), (int)SceneTextureId, EmitTexCoord);
 	}
 
 	if (SceneTextureId == PPI_PostProcessInput0 && Context.Material->GetMaterialDomain() == MD_PostProcess && Context.Material->GetBlendableLocation() != BL_AfterTonemapping)
diff --git a/Engine/Source/Runtime/Engine/Private/Shader/Preshader.cpp b/Engine/Source/Runtime/Engine/Private/Shader/Preshader.cpp
--- a/Engine/Source/Runtime/Engine/Private/Shader/Preshader.cpp
+++ b/Engine/Source/Runtime/Engine/Private/Shader/Preshader.cpp
@@ -1311,10 +1311,9 @@ static void EvaluateJumpIfFalse(FPreshaderDebugStack& Stack, FPreshaderDataConte
 	// It's complicated to include conditional jumps in a stack evaluator.  Ideally we'd want to evaluate and display both
 	// branches, but that requires evaluating both branches forward to where the branches reach the same instruction
 	// again, then converting stack entries that diverge between the two code paths to include a conditional expression.
-	// I don't have any examples of bytecode with conditional branches to be able to test an implementation of that (I
-	// think only the new translator generates these), so we'll just ignore the condition completely for now, and it won't
-	// show up in the debug expression.
-	FString ConditionValue = Stack.Elements.Pop().Text;
+
+	// The new translator can generate conditionals, until as above is implemented, ignore the conditional.
+	// This may elevate the reported count as actually exclusive branches may both be debug evaluated.
 }
 
 FPreshaderValue EvaluatePreshader(const FUniformExpressionSet* UniformExpressionSet, const FMaterialRenderContext& Context, FPreshaderStack& Stack, FPreshaderDataContext& RESTRICT Data)
