World Partition Content Bundles: fix loot containers that cannot be opened in cooked builds when loaded via content bundles.

diff --git a/Engine/Source/Runtime/Engine/Classes/GameFramework/Actor.h b/Engine/Source/Runtime/Engine/Classes/GameFramework/Actor.h
--- a/Engine/Source/Runtime/Engine/Classes/GameFramework/Actor.h
+++ b/Engine/Source/Runtime/Engine/Classes/GameFramework/Actor.h
@@ -49,6 +49,7 @@ class UDEPRECATED_DataLayer;
 class UDataLayerAsset;
 class UDataLayerInstance;
 class AWorldDataLayers;
+class IWorldPartitionCell;
 #if UE_WITH_IRIS
 struct FActorBeginReplicationParams;
 #endif // UE_WITH_IRIS
@@ -1377,10 +1378,11 @@ public:
 	ENGINE_API bool ContainsDataLayer(const UDataLayerAsset* DataLayerAsset) const;
 	ENGINE_API bool ContainsDataLayer(const UDataLayerInstance* DataLayerInstance) const;
 	ENGINE_API bool HasDataLayers() const;
+	ENGINE_API bool HasContentBundle() const;
 
 private:
 	ENGINE_API TArray<const UDataLayerInstance*> GetDataLayerInstancesInternal(bool bUseLevelContext, bool bIncludeParentDataLayers = true) const;
-	ENGINE_API bool UseWorldPartitionRuntimeCellDataLayers() const;
+	ENGINE_API const IWorldPartitionCell* GetWorldPartitionRuntimeCell() const;
 
 	//~=============================================================================
 	// General functions.
diff --git a/Engine/Source/Runtime/Engine/Private/Actor.cpp b/Engine/Source/Runtime/Engine/Private/Actor.cpp
--- a/Engine/Source/Runtime/Engine/Private/Actor.cpp
+++ b/Engine/Source/Runtime/Engine/Private/Actor.cpp
@@ -6218,10 +6218,9 @@ TArray<const UDataLayerInstance*> AActor::GetDataLayerInstances() const
 // to resolve valid datalayers for this particular level.
 TArray<const UDataLayerInstance*> AActor::GetDataLayerInstancesInternal(bool bUseLevelContext, bool bIncludeParentDataLayers) const
 {
-	if (UseWorldPartitionRuntimeCellDataLayers())
+	if (const IWorldPartitionCell* Cell = GetWorldPartitionRuntimeCell())
 	{
-		const IWorldPartitionCell* Cell = GetLevel()->GetWorldPartitionRuntimeCell();
-		return Cell ? Cell->GetDataLayerInstances() : TArray<const UDataLayerInstance*>();
+		return Cell->GetDataLayerInstances();
 	}
 
 #if WITH_EDITOR
@@ -6258,10 +6257,9 @@ TArray<const UDataLayerInstance*> AActor::GetDataLayerInstancesInternal(bool bUs
 
 bool AActor::ContainsDataLayer(const UDataLayerInstance* DataLayerInstance) const
 {
-	if (UseWorldPartitionRuntimeCellDataLayers())
+	if (const IWorldPartitionCell* Cell = GetWorldPartitionRuntimeCell())
 	{
-		const IWorldPartitionCell* Cell = GetLevel()->GetWorldPartitionRuntimeCell();
-		return Cell ? Cell->ContainsDataLayer(DataLayerInstance) : false;
+		return Cell->ContainsDataLayer(DataLayerInstance);
 	}
 
 #if WITH_EDITOR
@@ -6278,10 +6276,9 @@ bool AActor::ContainsDataLayer(const UDataLayerAsset* DataLayerAsset) const
 		return false;
 	}
 
-	if (UseWorldPartitionRuntimeCellDataLayers())
+	if (const IWorldPartitionCell* Cell = GetWorldPartitionRuntimeCell())
 	{
-		const IWorldPartitionCell* Cell = GetLevel()->GetWorldPartitionRuntimeCell();
-		return Cell ? Cell->ContainsDataLayer(DataLayerAsset) : false;
+		return Cell->ContainsDataLayer(DataLayerAsset);
 	}
 
 #if WITH_EDITOR
@@ -6301,10 +6298,9 @@ bool AActor::ContainsDataLayer(const UDataLayerAsset* DataLayerAsset) const
 
 bool AActor::HasDataLayers() const
 {
-	if (UseWorldPartitionRuntimeCellDataLayers())
+	if (const IWorldPartitionCell* Cell = GetWorldPartitionRuntimeCell())
 	{
-		const IWorldPartitionCell* Cell = GetLevel()->GetWorldPartitionRuntimeCell();
-		return Cell ? Cell->HasDataLayers() : false;
+		return Cell->HasDataLayers();
 	}
 
 #if WITH_EDITOR
@@ -6314,10 +6310,27 @@ bool AActor::HasDataLayers() const
 #endif
 }
 
-bool AActor::UseWorldPartitionRuntimeCellDataLayers() const
+bool AActor::HasContentBundle() const
 {
-	ULevel* Level = GetLevel();
-	return Level != nullptr && Level->IsWorldPartitionRuntimeCell();
+	if (const IWorldPartitionCell* Cell = GetWorldPartitionRuntimeCell())
+	{
+		return Cell->HasContentBundle();
+	}
+
+#if WITH_EDITOR
+	return GetContentBundleGuid().IsValid();
+#else
+	return false;
+#endif
+}
+
+const IWorldPartitionCell* AActor::GetWorldPartitionRuntimeCell() const
+{
+	if (ULevel* Level = GetLevel())
+	{
+		return Level->GetWorldPartitionRuntimeCell();
+	}
+	return nullptr;
 }
 
 
diff --git a/Engine/Source/Runtime/Engine/Private/WorldPartition/WorldPartitionRuntimeCell.cpp b/Engine/Source/Runtime/Engine/Private/WorldPartition/WorldPartitionRuntimeCell.cpp
--- a/Engine/Source/Runtime/Engine/Private/WorldPartition/WorldPartitionRuntimeCell.cpp
+++ b/Engine/Source/Runtime/Engine/Private/WorldPartition/WorldPartitionRuntimeCell.cpp
@@ -178,6 +178,11 @@ bool UWorldPartitionRuntimeCell::ContainsDataLayer(const UDataLayerAsset* DataLa
 	return DataLayerInstance ? ContainsDataLayer(DataLayerInstance) : false;
 }
 
+bool UWorldPartitionRuntimeCell::HasContentBundle() const
+{
+	return GetContentBundleID().IsValid();
+}
+
 bool UWorldPartitionRuntimeCell::ContainsDataLayer(const UDataLayerInstance* DataLayerInstance) const
 {
 	return GetDataLayers().Contains(DataLayerInstance->GetDataLayerFName());
diff --git a/Engine/Source/Runtime/Engine/Public/WorldPartition/WorldPartitionActorDesc.h b/Engine/Source/Runtime/Engine/Public/WorldPartition/WorldPartitionActorDesc.h
--- a/Engine/Source/Runtime/Engine/Public/WorldPartition/WorldPartitionActorDesc.h
+++ b/Engine/Source/Runtime/Engine/Public/WorldPartition/WorldPartitionActorDesc.h
@@ -291,9 +291,9 @@ protected:
 	// Persistent
 	FGuid							Guid;
 	FTopLevelAssetPath				BaseClass;
-	FTopLevelAssetPath				NativeClass;
-	FName							ActorPackage;
-	FSoftObjectPath					ActorPath;
+	FTopLevelAssetPath				NativeClass;	// Not serialized, comes from initialization data
+	FName							ActorPackage;	// Not serialized, comes from initialization data
+	FSoftObjectPath					ActorPath;		// Not serialized, comes from initialization data
 	FName							ActorLabel;
 	FVector							BoundsLocation;
 	FVector							BoundsExtent;
diff --git a/Engine/Source/Runtime/Engine/Public/WorldPartition/WorldPartitionRuntimeCell.h b/Engine/Source/Runtime/Engine/Public/WorldPartition/WorldPartitionRuntimeCell.h
--- a/Engine/Source/Runtime/Engine/Public/WorldPartition/WorldPartitionRuntimeCell.h
+++ b/Engine/Source/Runtime/Engine/Public/WorldPartition/WorldPartitionRuntimeCell.h
@@ -212,6 +212,7 @@ class UWorldPartitionRuntimeCell : public UObject, public IWorldPartitionCell
 	ENGINE_API virtual TArray<const UDataLayerInstance*> GetDataLayerInstances() const override;
 	ENGINE_API virtual bool ContainsDataLayer(const UDataLayerAsset* DataLayerAsset) const override;
 	ENGINE_API virtual bool ContainsDataLayer(const UDataLayerInstance* DataLayerInstance) const override;
+	ENGINE_API virtual bool HasContentBundle() const override;
 	virtual const TArray<FName>& GetDataLayers() const override  { return DataLayers; }
 	virtual bool HasAnyDataLayer(const TSet<FName>& InDataLayers) const override
 	{
diff --git a/Engine/Source/Runtime/Engine/Public/WorldPartition/WorldPartitionRuntimeCellInterface.h b/Engine/Source/Runtime/Engine/Public/WorldPartition/WorldPartitionRuntimeCellInterface.h
--- a/Engine/Source/Runtime/Engine/Public/WorldPartition/WorldPartitionRuntimeCellInterface.h
+++ b/Engine/Source/Runtime/Engine/Public/WorldPartition/WorldPartitionRuntimeCellInterface.h
@@ -27,6 +27,8 @@ public:
 	virtual bool ContainsDataLayer(const UDataLayerInstance* DataLayerInstance) const = 0;
 	/** Returns whether the cell content is associated with data layers or not. */
 	bool HasDataLayers() const { return !GetDataLayers().IsEmpty(); }
+	/** Returns whether the cell content is associated with content bundle or not. */
+	virtual bool HasContentBundle() const = 0;
 	/** Returns the cell content associated data layers. */
 	virtual const TArray<FName>& GetDataLayers() const = 0;
 	/** Returns whether the cell data layers referenced any of the provided data layer or not. */
