Added RedirectToFile dev command
- Usage: RedirectToFile <FILE_PATH> <COMMAND> [COMMAND_OPTIONS]...
- It allows output the result of the COMMAND to the file specified by FILEPATH
- For security purposes FILE_PATH is sandboxed: Only subdirectories inside Project/Saved directory are allowed
- Command will create required directory structure before creating of the file
- RedirectToFile does not work for DEFER commands and has no effect for MEMREPORT

#tests PIE

diff --git a/Engine/Source/Runtime/Engine/Classes/Engine/Engine.h b/Engine/Source/Runtime/Engine/Classes/Engine/Engine.h
--- a/Engine/Source/Runtime/Engine/Classes/Engine/Engine.h
+++ b/Engine/Source/Runtime/Engine/Classes/Engine/Engine.h
@@ -2334,6 +2334,7 @@ public:
 	ENGINE_API bool HandleConfigHashCommand( const TCHAR* Cmd, FOutputDevice& Ar );						
 	ENGINE_API bool HandleConfigMemCommand( const TCHAR* Cmd, FOutputDevice& Ar );	
 	ENGINE_API bool HandleGetIniCommand(const TCHAR* Cmd, FOutputDevice& Ar);
+	ENGINE_API bool HandleRedirectOutputCommand(const TCHAR* Cmd, FOutputDevice& Ar, UWorld* InWorld);
 #endif // !UE_BUILD_SHIPPING
 
 	/** Update everything. */
diff --git a/Engine/Source/Runtime/Engine/Private/UnrealEngine.cpp b/Engine/Source/Runtime/Engine/Private/UnrealEngine.cpp
--- a/Engine/Source/Runtime/Engine/Private/UnrealEngine.cpp
+++ b/Engine/Source/Runtime/Engine/Private/UnrealEngine.cpp
@@ -4986,6 +4986,10 @@ bool UEngine::Exec_Dev( UWorld* InWorld, const TCHAR* Cmd, FOutputDevice& Ar )
 	{
 		return HandleDumpConsoleCommandsCommand( Cmd, Ar, InWorld );
 	}
+	else if (FParse::Command(&Cmd, TEXT("REDIRECTTOFILE")))
+	{
+		return HandleRedirectOutputCommand( Cmd, Ar, InWorld );
+	}
 	else if (FParse::Command(&Cmd, TEXT("DUMPAVAILABLERESOLUTIONS")))
 	{
 		return HandleDumpAvailableResolutionsCommand( Cmd, Ar );
@@ -9960,6 +9964,92 @@ bool UEngine::HandleGetIniCommand(const TCHAR* Cmd, FOutputDevice& Ar)
 
 	return true;
 }
+
+bool UEngine::HandleRedirectOutputCommand(const TCHAR* Cmd, FOutputDevice& Ar, UWorld* InWorld)
+{
+#if ALLOW_DEBUG_FILES
+	FString FilePath = FParse::Token(Cmd, true);
+	FPaths::NormalizeFilename(FilePath);
+	if (!FPaths::CollapseRelativeDirectories(FilePath))
+	{
+		Ar.Logf(ELogVerbosity::Error, TEXT("Redirect failed: can't redirect outside of Saved folder: %s"), *FilePath);
+		return false;
+	}
+	FPaths::RemoveDuplicateSlashes(FilePath);
+
+	FText ErrorReason;
+	if (!FPaths::ValidatePath(FilePath, &ErrorReason))
+	{
+		Ar.Logf(ELogVerbosity::Error, TEXT("Redirect failed: %s"), *ErrorReason.ToString());
+		return false;
+	}
+
+	// Skip preceeding spaces and tabs.
+	while (FChar::IsWhitespace(*Cmd))
+	{
+		Cmd++;
+	}
+
+	if (*Cmd == 0)
+	{
+		Ar.Logf(ELogVerbosity::Warning, TEXT("Redirect skipped: the command is empty"));
+		return false;
+	}
+
+	static const TCHAR DeferCommandName[] = TEXT("DEFER");
+	static const size_t DeferCommandSize = UE_ARRAY_COUNT(DeferCommandName) - 1;
+
+	if (FCString::Strnicmp(Cmd, DeferCommandName, DeferCommandSize) == 0)
+	{
+		if (FChar::IsWhitespace(*(Cmd + DeferCommandSize)))
+		{
+			Ar.Logf(ELogVerbosity::Error, TEXT("Redirect failed: can't redirect deferred command"));
+			return false;
+		}
+	}
+
+	FilePath = FPaths::ProjectSavedDir() + FilePath;
+
+	const FString OutputDir = FPaths::GetPath(FilePath);
+	if (!IFileManager::Get().MakeDirectory(*OutputDir, true))
+	{
+		Ar.Logf(ELogVerbosity::Error, TEXT("Redirect failed: can't create path %s"), *OutputDir);
+		return false;
+	}
+
+	TUniquePtr<FArchive> CommandFile(IFileManager::Get().CreateDebugFileWriter(*FilePath));
+	if (!CommandFile)
+	{
+		Ar.Logf(ELogVerbosity::Error, TEXT("Redirect failed: Can't create file writer. See logs."));
+		return false;
+	}
+
+	FOutputDeviceArchiveWrapper CommandFileAr(CommandFile.Get());
+
+	FOutputDeviceRedirector Redirector;
+	Redirector.AddOutputDevice(&CommandFileAr);
+	Redirector.AddOutputDevice(GLog);
+
+	bool bHandled = false;
+	if (ULocalPlayer* Player = GetDebugLocalPlayer())
+	{
+		bHandled = Player->Exec(InWorld, Cmd, Redirector);
+	}
+	else
+	{
+		bHandled = Exec(InWorld, Cmd, Redirector);
+	}
+
+	Redirector.Flush();
+
+	// Propagate status of inner command to upstream code.
+	return bHandled;
+#else
+	Ar.Logf(ELogVerbosity::Error, TEXT("Redirect failed: Writing to files is not available."));
+	return false;
+#endif
+}
+
 #endif // !UE_BUILD_SHIPPING
 
 /** Helper function to cause a stack overflow crash */
