#ue5 - Remove incorrect deprecation comment in implementation, function is not deprecated.

#tests compiled

diff --git a/Engine/Source/Runtime/Engine/Private/Components/CharacterMovementComponent.cpp b/Engine/Source/Runtime/Engine/Private/Components/CharacterMovementComponent.cpp
--- a/Engine/Source/Runtime/Engine/Private/Components/CharacterMovementComponent.cpp
+++ b/Engine/Source/Runtime/Engine/Private/Components/CharacterMovementComponent.cpp
@@ -9726,7 +9726,6 @@ void UCharacterMovementComponent::ServerMove_Implementation(
 }
 
 
-///// DEPRECATED /////
 void UCharacterMovementComponent::ServerMoveHandleClientError(float ClientTimeStamp, float DeltaTime, const FVector& Accel, const FVector& RelativeClientLoc, UPrimitiveComponent* ClientMovementBase, FName ClientBaseBoneName, uint8 ClientMovementMode)
 {
 	if (!ShouldUsePackedMovementRPCs())
