Add blueprint array sort functions for native types.

Adds in-place sort functions. Supports options for stable sort and ascending/descending.

#tests unit tests

diff --git a/Engine/Source/Runtime/Engine/Classes/Kismet/KismetArrayLibrary.h b/Engine/Source/Runtime/Engine/Classes/Kismet/KismetArrayLibrary.h
--- a/Engine/Source/Runtime/Engine/Classes/Kismet/KismetArrayLibrary.h
+++ b/Engine/Source/Runtime/Engine/Classes/Kismet/KismetArrayLibrary.h
@@ -15,6 +15,13 @@
 
 class AActor;
 
+UENUM(BlueprintType)
+enum class EArraySortOrder : uint8
+{
+	Ascending,
+	Descending
+};
+
 UCLASS(meta=(BlueprintThreadSafe), MinimalAPI)
 class UKismetArrayLibrary : public UBlueprintFunctionLibrary
 {
@@ -266,6 +273,67 @@ class UKismetArrayLibrary : public UBlueprintFunctionLibrary
 	UFUNCTION(BlueprintPure, CustomThunk, meta=(DisplayName = "Random Array Item from Stream", ArrayParm = "TargetArray", ArrayTypeDependentParams = "OutItem"), Category="Utilities|Array")
 	static ENGINE_API void Array_RandomFromStream(const TArray<int32>& TargetArray, UPARAM(Ref) FRandomStream& RandomStream, int32& OutItem, int32& OutIndex);
 
+	/**
+	 * Sorts an array of strings alphabetically.
+	 *
+	 * @param	TargetArray		The array to sort.
+	 * @param	bStableSort		If a stable sort should be used. This preserves the order of identical elements, but is slower.
+	 * @param	SortOrder		If the array should be sorted in ascending or descending order.
+	 */
+	UFUNCTION(BlueprintCallable, Category = "Utilities|Array|Sort", meta = (AdvancedDisplay = "bStableSort,SortOrder"))
+	static ENGINE_API void SortStringArray(UPARAM(Ref) TArray<FString>& TargetArray, bool bStableSort = false, EArraySortOrder SortOrder = EArraySortOrder::Ascending);
+
+	/**
+	 * Sorts an array of FNames.
+	 *
+	 * @param	TargetArray		The array to sort.
+	 * @param	bStableSort		If a stable sort should be used. This preserves the order of identical elements, but is slower.
+	 * @param	bLexicalSort	If the names should be sorted based on the string value of the name rather than the comparison index. This is slower when enabled.
+	 * @param	SortOrder		If the array should be sorted in ascending or descending order.
+	 */
+	UFUNCTION(BlueprintCallable, Category = "Utilities|Array|Sort", meta = (AdvancedDisplay = "bStableSort,bLexicalSort,SortOrder"))
+	static ENGINE_API void SortNameArray(UPARAM(Ref) TArray<FName>& TargetArray, bool bStableSort = false, bool bLexicalSort = true, EArraySortOrder SortOrder = EArraySortOrder::Ascending);
+
+	/**
+	 * Sorts an array of bytes.
+	 *
+	 * @param	TargetArray		The array to sort.
+	 * @param	bStableSort		If a stable sort should be used. This preserves the order of identical elements, but is slower.
+	 * @param	SortOrder		If the array should be sorted in ascending or descending order.
+	 */
+	UFUNCTION(BlueprintCallable, Category = "Utilities|Array|Sort", meta = (AdvancedDisplay = "bStableSort,SortOrder"))
+	static ENGINE_API void SortByteArray(UPARAM(Ref) TArray<uint8>& TargetArray, bool bStableSort = false, EArraySortOrder SortOrder = EArraySortOrder::Ascending);
+
+	/**
+	 * Sorts an array of integers.
+	 *
+	 * @param	TargetArray		The array to sort.
+	 * @param	bStableSort		If a stable sort should be used. This preserves the order of identical elements, but is slower.
+	 * @param	SortOrder		If the array should be sorted in ascending or descending order.
+	 */
+	UFUNCTION(BlueprintCallable, Category = "Utilities|Array|Sort", meta = (AdvancedDisplay = "bStableSort,SortOrder"))
+	static ENGINE_API void SortIntArray(UPARAM(Ref) TArray<int32>& TargetArray, bool bStableSort = false, EArraySortOrder SortOrder = EArraySortOrder::Ascending);
+
+	/**
+	 * Sorts an array of 64-bit integers.
+	 *
+	 * @param	TargetArray		The array to sort.
+	 * @param	bStableSort		If a stable sort should be used. This preserves the order of identical elements, but is slower.
+	 * @param	SortOrder		If the array should be sorted in ascending or descending order.
+	 */
+	UFUNCTION(BlueprintCallable, Category = "Utilities|Array|Sort", meta = (AdvancedDisplay = "bStableSort,SortOrder"))
+	static ENGINE_API void SortInt64Array(UPARAM(Ref) TArray<int64>& TargetArray, bool bStableSort = false, EArraySortOrder SortOrder = EArraySortOrder::Ascending);
+
+	/**
+	 * Sorts an array of doubles.
+	 *
+	 * @param	TargetArray		The array to sort.
+	 * @param	bStableSort		If a stable sort should be used. This preserves the order of identical elements, but is slower.
+	 * @param	SortOrder		If the array should be sorted in ascending or descending order.
+	 */
+	UFUNCTION(BlueprintCallable, Category = "Utilities|Array|Sort", meta = (AdvancedDisplay = "bStableSort,SortOrder"))
+	static ENGINE_API void SortFloatArray(UPARAM(Ref) TArray<double>& TargetArray, bool bStableSort = false, EArraySortOrder SortOrder = EArraySortOrder::Ascending);
+
 	// Native functions that will be called by the below custom thunk layers, which read off the property address, and call the appropriate native handler
 	static ENGINE_API int32 GenericArray_Add(void* TargetArray, const FArrayProperty* ArrayProp, const void* NewItem);
 	static ENGINE_API int32 GenericArray_AddUnique(void* TargetArray, const FArrayProperty* ArrayProp, const void* NewItem);
diff --git a/Engine/Source/Runtime/Engine/Private/KismetArrayLibrary.cpp b/Engine/Source/Runtime/Engine/Private/KismetArrayLibrary.cpp
--- a/Engine/Source/Runtime/Engine/Private/KismetArrayLibrary.cpp
+++ b/Engine/Source/Runtime/Engine/Private/KismetArrayLibrary.cpp
@@ -24,6 +24,41 @@ const FName RandomAccessToEmptyArrayWarning = FName("RandomAccessToEmptyArrayWar
 
 const FName UKismetArrayLibrary::ReachedMaximumContainerSizeWarning = FName("ReachedMaximumContainerSizeWarning");
 
+namespace UE::KismetArray::Private
+{
+	template<class T, typename InLessThanFunc = TLess<T>, typename InGreaterThanFunc = TGreater<T>>
+	void SortHelper(
+		TArray<T>& InArray,
+		bool bStableSort,
+		EArraySortOrder SortOrder,
+		InLessThanFunc&& LessThanFunc = {},
+		InGreaterThanFunc&& GreaterThanFunc = {})
+	{
+		if (!bStableSort)
+		{
+			if (SortOrder == EArraySortOrder::Ascending)
+			{
+				InArray.Sort(Forward<InLessThanFunc>(LessThanFunc));
+			}
+			else
+			{
+				InArray.Sort(Forward<InGreaterThanFunc>(GreaterThanFunc));
+			}
+		}
+		else
+		{
+			if (SortOrder == EArraySortOrder::Ascending)
+			{
+				InArray.StableSort(Forward<InLessThanFunc>(LessThanFunc));
+			}
+			else
+			{
+				InArray.StableSort(Forward<InGreaterThanFunc>(GreaterThanFunc));
+			}
+		}
+	}
+}
+
 UKismetArrayLibrary::UKismetArrayLibrary(const FObjectInitializer& ObjectInitializer)
 	: Super(ObjectInitializer)
 {
@@ -90,6 +125,46 @@ void UKismetArrayLibrary::FilterArray(const TArray<AActor*>& TargetArray, TSubcl
 	}
 }
 
+void UKismetArrayLibrary::SortStringArray(TArray<FString>& TargetArray, bool bStableSort, EArraySortOrder SortOrder)
+{
+	UE::KismetArray::Private::SortHelper(TargetArray, bStableSort, SortOrder);
+}
+
+void UKismetArrayLibrary::SortNameArray(TArray<FName>& TargetArray, bool bStableSort, bool bLexicalSort, EArraySortOrder SortOrder)
+{
+	const auto LessThan = [bLexicalSort](const FName& A, const FName& B)
+	{
+		return bLexicalSort ? FNameLexicalLess{}(A, B) : FNameFastLess{}(A, B);
+	};
+
+	const auto GreaterThan = [bLexicalSort](const FName& A, const FName& B)
+	{
+		return bLexicalSort ? FNameLexicalLess{}(B, A) : FNameFastLess{}(B, A);
+	};
+
+	UE::KismetArray::Private::SortHelper(TargetArray, bStableSort, SortOrder, LessThan, GreaterThan);
+}
+
+void UKismetArrayLibrary::SortByteArray(TArray<uint8>& TargetArray, bool bStableSort, EArraySortOrder SortOrder)
+{
+	UE::KismetArray::Private::SortHelper(TargetArray, bStableSort, SortOrder);
+}
+
+void UKismetArrayLibrary::SortIntArray(TArray<int32>& TargetArray, bool bStableSort, EArraySortOrder SortOrder)
+{
+	UE::KismetArray::Private::SortHelper(TargetArray, bStableSort, SortOrder);
+}
+
+void UKismetArrayLibrary::SortInt64Array(TArray<int64>& TargetArray, bool bStableSort, EArraySortOrder SortOrder)
+{
+	UE::KismetArray::Private::SortHelper(TargetArray, bStableSort, SortOrder);
+}
+
+void UKismetArrayLibrary::SortFloatArray(TArray<double>& TargetArray, bool bStableSort, EArraySortOrder SortOrder)
+{
+	UE::KismetArray::Private::SortHelper(TargetArray, bStableSort, SortOrder);
+}
+
 int32 UKismetArrayLibrary::GenericArray_Add(void* TargetArray, const FArrayProperty* ArrayProp, const void* NewItem)
 {
 	int32 NewIndex = INDEX_NONE;
diff --git a/Engine/Source/Runtime/Engine/Private/Tests/KismetArrayTests.cpp b/Engine/Source/Runtime/Engine/Private/Tests/KismetArrayTests.cpp
new file mode 100644
--- /dev/null
+++ b/Engine/Source/Runtime/Engine/Private/Tests/KismetArrayTests.cpp
@@ -0,0 +1,158 @@
+// Copyright Epic Games, Inc. All Rights Reserved.
+
+#if WITH_TESTS
+
+#include "Kismet/KismetArrayLibrary.h"
+#include "Tests/TestHarnessAdapter.h"
+
+namespace UE::Kismet::Test
+{
+	namespace Private
+	{
+		template<typename ElementType, typename InSortFunc, typename... InArgs>
+		void VerifySort(
+			const TArray<ElementType>& InArray,
+			const TArray<ElementType>& InExpected,
+			InSortFunc SortFunc,
+			InArgs... ExtraArgs)
+		{
+			const auto RunTest = [&](const TCHAR* SortType, bool bStableSort, EArraySortOrder SortOrder)
+			{
+				TArray<ElementType> ArrayCopy = InArray;
+				TArray<ElementType> ExpectedCopy = InExpected;
+				if (SortOrder == EArraySortOrder::Descending)
+				{
+					Algo::Reverse(ExpectedCopy);
+				}
+
+				Invoke(SortFunc, ArrayCopy, bStableSort, SortOrder, ExtraArgs...);
+
+				const auto ToStringHelper = [](const ElementType& Element) { return LexToString(Element); };
+
+				const FString ErrorMessage = FString::Printf(
+					TEXT("Sorted array doesn't match expected array for %s sort. Sorted: %s, Expected: %s."),
+					SortType,
+					*FString::JoinBy(ArrayCopy, TEXT(","), ToStringHelper),
+					*FString::JoinBy(ExpectedCopy, TEXT(","), ToStringHelper));
+
+				REQUIRE_MESSAGE(*ErrorMessage, ArrayCopy == ExpectedCopy);
+			};
+
+			// Test each combo of stable/sort order.
+			RunTest(TEXT("unstable-ascending"), false, EArraySortOrder::Ascending);
+			RunTest(TEXT("unstable-descending"), false, EArraySortOrder::Descending);
+			RunTest(TEXT("stable-ascending"),true, EArraySortOrder::Ascending);
+			RunTest(TEXT("stable-descending"),true, EArraySortOrder::Descending);
+		}
+	}
+
+	TEST_CASE(
+		"UE::KismetArray::SortStrings",
+		"[KismetArray][ClientContext][EditorContext][EngineFilter]")
+	{
+		{
+			Private::VerifySort(
+				TArray<FString>{ TEXT("Billy"), TEXT("Carol"), TEXT("Thomas"), TEXT("Riley"), TEXT("Adam") },
+				TArray<FString>{ TEXT("Adam"), TEXT("Billy"), TEXT("Carol"), TEXT("Riley"), TEXT("Thomas") },
+				&UKismetArrayLibrary::SortStringArray);
+		}
+
+		// Mixed casing
+		{
+			Private::VerifySort(
+				TArray<FString>{ TEXT("Billy"), TEXT("BILLY"), TEXT("billy"), TEXT("BiLlY") },
+				TArray<FString>{ TEXT("billy"), TEXT("Billy"), TEXT("BiLlY"), TEXT("BILLY") },
+				&UKismetArrayLibrary::SortStringArray);
+		}
+
+		// Stable sort
+		{
+			TArray<FString> InArray{ TEXT("Billy"), TEXT("Billy"), TEXT("Billy") };
+			TArray<const FString*> ElementOrder;
+			for (const FString& Str : InArray)
+			{
+				ElementOrder.Add(&Str);
+			}
+
+			UKismetArrayLibrary::SortStringArray(InArray, true);
+
+			for (int32 Index = 0; Index < InArray.Num(); ++Index)
+			{
+				if (ElementOrder[Index] != &InArray[Index])
+				{
+					FAIL_CHECK(TEXT("Identical elements didn't remain in original order after stable sort"));
+				}
+			}
+		}
+	}
+
+	TEST_CASE(
+		"UE::KismetArray::SortNames",
+		"[KismetArray][ClientContext][EditorContext][EngineFilter]")
+	{
+		const auto SortNamesRedirector = [](TArray<FName>& TargetArray, bool bStableSort, EArraySortOrder SortOrder, bool bLexicalSort)
+		{
+			UKismetArrayLibrary::SortNameArray(TargetArray, bStableSort, bLexicalSort, SortOrder);
+		};
+
+		// Test non-lexical sort
+		{
+			TArray<FName> InArray{ TEXT("Billy"), TEXT("Carol"), TEXT("Thomas"), TEXT("Riley"), TEXT("Adam") };
+			TArray<FName> Expected = InArray;
+			// Sorted by comparison index so can't easily hard-code expected array.
+			Expected.Sort(FNameFastLess{});
+
+			constexpr bool bLexicalSort = false;
+			Private::VerifySort(InArray, Expected, SortNamesRedirector, bLexicalSort);
+		}
+
+		// Test lexical sort
+		{
+			constexpr bool bLexicalSort = true;
+			Private::VerifySort(
+				TArray<FName>{ TEXT("Billy"), TEXT("Carol"), TEXT("Thomas"), TEXT("Riley"), TEXT("Adam") },
+				TArray<FName>{ TEXT("Adam"), TEXT("Billy"), TEXT("Carol"), TEXT("Riley"), TEXT("Thomas") },
+				SortNamesRedirector,
+				bLexicalSort);
+		}
+	}
+
+	TEST_CASE(
+		"UE::KismetArray::SortIntegerTypes",
+		"[KismetArray][ClientContext][EditorContext][EngineFilter]")
+	{
+		// byte array
+		{
+			Private::VerifySort(
+				TArray<uint8>{ 10, 1, 22, 244, 90, 3 },
+				TArray<uint8>{ 1, 3, 10, 22, 90, 244 },
+				&UKismetArrayLibrary::SortByteArray);
+		}
+
+		// int32
+		{
+			Private::VerifySort(
+				TArray<int32>{ 10, 1, 22, 244, 90, 3 },
+				TArray<int32>{ 1, 3, 10, 22, 90, 244 },
+				&UKismetArrayLibrary::SortIntArray);
+		}
+
+		// int64
+		{
+			Private::VerifySort(
+				TArray<int64>{ 10, 1, 22, 244, 90, 3 },
+				TArray<int64>{ 1, 3, 10, 22, 90, 244 },
+				&UKismetArrayLibrary::SortInt64Array);
+		}
+
+		// double
+		{
+			Private::VerifySort(
+				TArray<double>{ 10.0, 1.0, 22.0, 244.0, 90.0, 3.0 },
+				TArray<double>{ 1.0, 3.0, 10.0, 22.0, 90.0, 244.0 },
+				&UKismetArrayLibrary::SortFloatArray);
+		}
+	}
+
+} // UE::Kismet::Test
+#endif // WITH_TESTS
\ No newline at end of file
