Do not compile Raytraced input on a RayTracingQualitySwitch material node for platforms that do not support raytracing

diff --git a/Engine/Source/Runtime/Engine/Private/Materials/MaterialExpressions.cpp b/Engine/Source/Runtime/Engine/Private/Materials/MaterialExpressions.cpp
--- a/Engine/Source/Runtime/Engine/Private/Materials/MaterialExpressions.cpp
+++ b/Engine/Source/Runtime/Engine/Private/Materials/MaterialExpressions.cpp
@@ -18466,7 +18466,7 @@ int32 UMaterialExpressionRayTracingQualitySwitch::Compile(class FMaterialCompile
 	else
 	{
 		int32 Arg1 = Normal.Compile(Compiler);
-		int32 Arg2 = RayTraced.Compile(Compiler);
+		int32 Arg2 = FDataDrivenShaderPlatformInfo::GetSupportsRayTracing(Compiler->GetShaderPlatform()) ? RayTraced.Compile(Compiler) : INDEX_NONE;
 
 		//only when both of these are real expressions do the actual code.  otherwise various output pins will
 		//end up considered 'set' when really we just want a default.  This can cause us to force depth output when we don't want it for example.
