World Partition:
  - FWorldPartitionActorDescView is now a virtual interface that route calls to its inner FWorldPartitionActorDesc by default.
  - FStreamingGenerationActorDescView is now used during streaming generation to include mutations introduced by validation pass.
  - FActorDescViewMap is now a template, and streaming generation now uses custom FStreamingGenerationActorDescViewMap version using FStreamingGenerationActorDescView.

This is the groundwork for implementing actor descriptor mutators to allow game-specific changes to be done to actor descriptors.

diff --git a/Engine/Plugins/Editor/WorldPartitionHLODUtilities/Source/Private/WorldPartition/HLOD/Utilities/WorldPartitionHLODUtilities.cpp b/Engine/Plugins/Editor/WorldPartitionHLODUtilities/Source/Private/WorldPartition/HLOD/Utilities/WorldPartitionHLODUtilities.cpp
--- a/Engine/Plugins/Editor/WorldPartitionHLODUtilities/Source/Private/WorldPartition/HLOD/Utilities/WorldPartitionHLODUtilities.cpp
+++ b/Engine/Plugins/Editor/WorldPartitionHLODUtilities/Source/Private/WorldPartition/HLOD/Utilities/WorldPartitionHLODUtilities.cpp
@@ -101,7 +101,7 @@ void AddSourceActor(const FWorldPartitionActorDescView& ActorDescView, const ISt
 			false);
 
 		// Add its runtime references, recursively
-		const FActorDescViewMap* ActorDescViewMap = ActorInstance.ActorSetInstance->ContainerInstance->ActorDescViewMap;
+		const FStreamingGenerationActorDescViewMap* ActorDescViewMap = ActorInstance.ActorSetInstance->ContainerInstance->ActorDescViewMap;
 		for (const FGuid& ReferenceGuid : ActorDescView.GetReferences())
 		{
 			const FWorldPartitionActorDescView& RefActorDescView = ActorDescViewMap->FindByGuidChecked(ReferenceGuid);
diff --git a/Engine/Source/Editor/WorldPartitionEditor/Private/WorldPartition/SWorldPartitionEditorGrid2D.cpp b/Engine/Source/Editor/WorldPartitionEditor/Private/WorldPartition/SWorldPartitionEditorGrid2D.cpp
--- a/Engine/Source/Editor/WorldPartitionEditor/Private/WorldPartition/SWorldPartitionEditorGrid2D.cpp
+++ b/Engine/Source/Editor/WorldPartitionEditor/Private/WorldPartition/SWorldPartitionEditorGrid2D.cpp
@@ -220,7 +220,7 @@ public:
 		, bUseActor(bInUseActor)
 	{}
 
-	FBox GetEditorBounds() const
+	virtual FBox GetEditorBounds() const override
 	{
 		if (bUseActor)
 		{
@@ -233,12 +233,12 @@ public:
 		return ActorDesc->GetEditorBounds();
 	}
 
-	AActor* GetActor() const
+	virtual AActor* GetActor() const override
 	{
 		return bUseActor ? ActorDesc->GetActor(false) : nullptr;
 	}
 
-	FName GetActorLabel() const
+	virtual FName GetActorLabel() const override
 	{
 		FName ActorLabel = ActorDesc->GetActorLabel();
 		
diff --git a/Engine/Source/Runtime/Engine/Private/WorldPartition/DataLayer/DataLayerUtils.cpp b/Engine/Source/Runtime/Engine/Private/WorldPartition/DataLayer/DataLayerUtils.cpp
--- a/Engine/Source/Runtime/Engine/Private/WorldPartition/DataLayer/DataLayerUtils.cpp
+++ b/Engine/Source/Runtime/Engine/Private/WorldPartition/DataLayer/DataLayerUtils.cpp
@@ -80,7 +80,7 @@ TArray<FName> FDataLayerUtils::ResolvedDataLayerInstanceNames(const UDataLayerMa
 }
 
 // For performance reasons, this function assumes that InActorDesc's DataLayerInstanceNames was already resolved.
-bool FDataLayerUtils::ResolveRuntimeDataLayerInstanceNames(const UDataLayerManager* InDataLayerManager, const FWorldPartitionActorDescView& InActorDescView, const FActorDescViewMap& ActorDescViewMap, TArray<FName>& OutRuntimeDataLayerInstanceNames)
+bool FDataLayerUtils::ResolveRuntimeDataLayerInstanceNames(const UDataLayerManager* InDataLayerManager, const FWorldPartitionActorDescView& InActorDescView, const TArray<const FWorldPartitionActorDescView*>& InWorldDataLayerViews, TArray<FName>& OutRuntimeDataLayerInstanceNames)
 {
 	const TArray<FName>& ActorDescViewDataLayerInstanceNames = InActorDescView.GetDataLayerInstanceNames();
 
@@ -99,7 +99,7 @@ bool FDataLayerUtils::ResolveRuntimeDataLayerInstanceNames(const UDataLayerManag
 	}
 
 	// Fallback on FWorldDataLayersActorDesc
-	TArray<const FWorldDataLayersActorDesc*> WorldDataLayersActorDescs = FindWorldDataLayerActorDescs(ActorDescViewMap);
+	TArray<const FWorldDataLayersActorDesc*> WorldDataLayersActorDescs = FindWorldDataLayerActorDescs(InWorldDataLayerViews);
 	if (WorldDataLayersActorDescs.Num())
 	{
 		check(AreWorldDataLayersActorDescsSane(WorldDataLayersActorDescs));
@@ -120,11 +120,11 @@ bool FDataLayerUtils::ResolveRuntimeDataLayerInstanceNames(const UDataLayerManag
 	return false;
 }
 
-const FDataLayerInstanceDesc* FDataLayerUtils::GetDataLayerInstanceDescFromInstanceName(const TArray<const FWorldDataLayersActorDesc*>& InWorldDataLayersActorDescs, const FName& DataLayerInstanceName)
+const FDataLayerInstanceDesc* FDataLayerUtils::GetDataLayerInstanceDescFromInstanceName(const TArray<const FWorldDataLayersActorDesc*>& InWorldDataLayersActorDescs, const FName& InDataLayerInstanceName)
 {
 	for (const FWorldDataLayersActorDesc* WorldDataLayerActorDesc : InWorldDataLayersActorDescs)
 	{
-		if (const FDataLayerInstanceDesc* DataLayerInstanceDesc = WorldDataLayerActorDesc->GetDataLayerInstanceFromInstanceName(DataLayerInstanceName))
+		if (const FDataLayerInstanceDesc* DataLayerInstanceDesc = WorldDataLayerActorDesc->GetDataLayerInstanceFromInstanceName(InDataLayerInstanceName))
 		{
 			return DataLayerInstanceDesc;
 		}
@@ -133,11 +133,11 @@ const FDataLayerInstanceDesc* FDataLayerUtils::GetDataLayerInstanceDescFromInsta
 	return nullptr;
 }
 
-const FDataLayerInstanceDesc* FDataLayerUtils::GetDataLayerInstanceDescFromAssetPath(const TArray<const FWorldDataLayersActorDesc*>& InWorldDataLayersActorDescs, const FName& DataLayerAssetPath)
+const FDataLayerInstanceDesc* FDataLayerUtils::GetDataLayerInstanceDescFromAssetPath(const TArray<const FWorldDataLayersActorDesc*>& InWorldDataLayersActorDescs, const FName& InDataLayerAssetPath)
 {
 	for (const FWorldDataLayersActorDesc* WorldDataLayerActorDesc : InWorldDataLayersActorDescs)
 	{
-		if (const FDataLayerInstanceDesc* DataLayerInstanceDesc = WorldDataLayerActorDesc->GetDataLayerInstanceFromAssetPath(DataLayerAssetPath))
+		if (const FDataLayerInstanceDesc* DataLayerInstanceDesc = WorldDataLayerActorDesc->GetDataLayerInstanceFromAssetPath(InDataLayerAssetPath))
 		{
 			return DataLayerInstanceDesc;
 		}
@@ -146,11 +146,10 @@ const FDataLayerInstanceDesc* FDataLayerUtils::GetDataLayerInstanceDescFromAsset
 	return nullptr;
 }
 
-TArray<const FWorldDataLayersActorDesc*> FDataLayerUtils::FindWorldDataLayerActorDescs(const FActorDescViewMap& ActorDescViewMap)
+TArray<const FWorldDataLayersActorDesc*> FDataLayerUtils::FindWorldDataLayerActorDescs(const TArray<const FWorldPartitionActorDescView*>& InWorldDataLayerViews)
 {
 	TArray<const FWorldDataLayersActorDesc*> WorldDataLayersActorDescs;
-	TArray<const FWorldPartitionActorDescView*> WorldDataLayerViews = ActorDescViewMap.FindByExactNativeClass<AWorldDataLayers>();
-	Algo::TransformIf(WorldDataLayerViews, WorldDataLayersActorDescs, 
+	Algo::TransformIf(InWorldDataLayerViews, WorldDataLayersActorDescs, 
 		[](const FWorldPartitionActorDescView* WorldDataLayersActorDescView) { return ((FWorldDataLayersActorDesc*)WorldDataLayersActorDescView->GetActorDesc())->IsValid(); },
 		[](const FWorldPartitionActorDescView* WorldDataLayersActorDescView) { return (FWorldDataLayersActorDesc*)WorldDataLayersActorDescView->GetActorDesc(); });
 	return WorldDataLayersActorDescs;
diff --git a/Engine/Source/Runtime/Engine/Private/WorldPartition/RuntimeHashSet/WorldPartitionRuntimeHashSetHLODGeneration.cpp b/Engine/Source/Runtime/Engine/Private/WorldPartition/RuntimeHashSet/WorldPartitionRuntimeHashSetHLODGeneration.cpp
--- a/Engine/Source/Runtime/Engine/Private/WorldPartition/RuntimeHashSet/WorldPartitionRuntimeHashSetHLODGeneration.cpp
+++ b/Engine/Source/Runtime/Engine/Private/WorldPartition/RuntimeHashSet/WorldPartitionRuntimeHashSetHLODGeneration.cpp
@@ -55,7 +55,7 @@ public:
 
 	FBox WorldBounds;
 	FActorSetContainer ActorSetContainer;
-	FActorDescViewMap ActorDescViewMap;
+	FStreamingGenerationActorDescViewMap ActorDescViewMap;
 	FActorSetInstanceList ActorSetInstanceList;
 };
 
@@ -280,7 +280,7 @@ bool UWorldPartitionRuntimeHashSet::SetupHLODActors(const IStreamingGenerationCo
 				FWorldPartitionActorDesc* HLODActorDesc = WorldPartition->GetActorDesc(HLODActorGuid);
 				check(HLODActorDesc);
 
-				FWorldPartitionActorDescView* HLODActorDescView = HLODStreamingGenerationContext->ActorDescViewMap.Emplace(HLODActorDesc);
+				FStreamingGenerationActorDescView* HLODActorDescView = HLODStreamingGenerationContext->ActorDescViewMap.Emplace(HLODActorDesc);
 				HLODStreamingGenerationContext->WorldBounds += HLODActorDescView->GetRuntimeBounds();
 			
 				// Create actor set instances
@@ -296,8 +296,11 @@ bool UWorldPartitionRuntimeHashSet::SetupHLODActors(const IStreamingGenerationCo
 				ActorSetInstance.ContainerInstance = &HLODStreamingGenerationContext->ActorSetContainer;
 				ActorSetInstance.ActorSet = ActorSet;
 
+				TArray<const FWorldPartitionActorDescView*> WorldDataLayerViews;
+				Algo::Transform(MainActorSetContainer->ActorDescViewMap->FindByExactNativeClass<AWorldDataLayers>(), WorldDataLayerViews, [](const FStreamingGenerationActorDescView* ActorDescView) { return ActorDescView; });
+
 				TArray<FName> RuntimeDataLayerInstanceNames;
-				if (FDataLayerUtils::ResolveRuntimeDataLayerInstanceNames(DataLayerManager, *HLODActorDescView, *MainActorSetContainer->ActorDescViewMap, RuntimeDataLayerInstanceNames))
+				if (FDataLayerUtils::ResolveRuntimeDataLayerInstanceNames(DataLayerManager, *HLODActorDescView, WorldDataLayerViews, RuntimeDataLayerInstanceNames))
 				{
 					HLODActorDescView->SetRuntimeDataLayerInstanceNames(RuntimeDataLayerInstanceNames);
 					ActorSetInstance.DataLayers = DataLayerManager->GetRuntimeDataLayerInstances(RuntimeDataLayerInstanceNames);
diff --git a/Engine/Source/Runtime/Engine/Private/WorldPartition/RuntimeSpatialHash/RuntimeSpatialHashHLOD.cpp b/Engine/Source/Runtime/Engine/Private/WorldPartition/RuntimeSpatialHash/RuntimeSpatialHashHLOD.cpp
--- a/Engine/Source/Runtime/Engine/Private/WorldPartition/RuntimeSpatialHash/RuntimeSpatialHashHLOD.cpp
+++ b/Engine/Source/Runtime/Engine/Private/WorldPartition/RuntimeSpatialHash/RuntimeSpatialHashHLOD.cpp
@@ -274,8 +274,8 @@ static TMap<UHLODLayer*, int32> GatherHLODLayers(const IStreamingGenerationConte
 
 	StreamingGenerationContext->ForEachActorSetContainer([&HLODLayersLevel, WorldPartition](const IStreamingGenerationContext::FActorSetContainer& ActorSetContainer)
 	{
-		const FActorDescViewMap* ActorDescViewMap = ActorSetContainer.ActorDescViewMap;
-		ActorDescViewMap->ForEachActorDescView([&HLODLayersLevel, WorldPartition](const FWorldPartitionActorDescView& ActorDescView)
+		const FStreamingGenerationActorDescViewMap* ActorDescViewMap = ActorSetContainer.ActorDescViewMap;
+		ActorDescViewMap->ForEachActorDescView([&HLODLayersLevel, WorldPartition](const FStreamingGenerationActorDescView& ActorDescView)
 		{
 			if (!ActorDescView.GetActorNativeClass()->IsChildOf<AWorldPartitionHLOD>())
 			{
@@ -496,7 +496,7 @@ bool UWorldPartitionRuntimeSpatialHash::SetupHLODActors(const IStreamingGenerati
 
 	// Keep track of all valid HLOD actors, along with which runtime grid they live in
 	TMap<FName, TArray<FGuid>> GridsHLODActors;
-	TArray<TUniquePtr<FWorldPartitionActorDescView>> ActorDescViews;
+	TArray<TUniquePtr<FStreamingGenerationActorDescView>> ActorDescViews;
 
 	const UDataLayerManager* DataLayerManager = WorldPartition->GetDataLayerManager();
 
@@ -512,16 +512,20 @@ bool UWorldPartitionRuntimeSpatialHash::SetupHLODActors(const IStreamingGenerati
 			{
 				FWorldPartitionActorDesc* HLODActorDesc = WorldPartition->GetActorDesc(HLODActorGuid);
 				check(HLODActorDesc);
-				FWorldPartitionActorDescView* ActorDescView = ActorDescViews.Emplace_GetRef(MakeUnique<FWorldPartitionActorDescView>(HLODActorDesc)).Get();
+				FStreamingGenerationActorDescView* ActorDescView = ActorDescViews.Emplace_GetRef(MakeUnique<FStreamingGenerationActorDescView>(HLODActorDesc)).Get();
+				
+				TArray<const FWorldPartitionActorDescView*> WorldDataLayerViews;
+				Algo::Transform(MainActorSetContainer->ActorDescViewMap->FindByExactNativeClass<AWorldDataLayers>(), WorldDataLayerViews, [](const FStreamingGenerationActorDescView* ActorDescView) { return ActorDescView; });
+
 				TArray<FName> RuntimeDataLayerInstanceNames;
-				if (FDataLayerUtils::ResolveRuntimeDataLayerInstanceNames(DataLayerManager, *ActorDescView, *MainActorSetContainer->ActorDescViewMap, RuntimeDataLayerInstanceNames))
+				if (FDataLayerUtils::ResolveRuntimeDataLayerInstanceNames(DataLayerManager, *ActorDescView, WorldDataLayerViews, RuntimeDataLayerInstanceNames))
 				{
 					ActorDescView->SetRuntimeDataLayerInstanceNames(RuntimeDataLayerInstanceNames);
 				}
-				FActorDescViewMap* NonConstActorDescViewMap = const_cast<FActorDescViewMap*>(MainActorSetContainer->ActorDescViewMap);			
+				FStreamingGenerationActorDescViewMap* NonConstActorDescViewMap = const_cast<FStreamingGenerationActorDescViewMap*>(MainActorSetContainer->ActorDescViewMap);			
 				NonConstActorDescViewMap->Emplace(HLODActorGuid, *ActorDescView);
 
-				const FWorldPartitionActorDescView& HLODActorDescView = MainActorSetContainer->ActorDescViewMap->FindByGuidChecked(HLODActorGuid);
+				const FStreamingGenerationActorDescView& HLODActorDescView = MainActorSetContainer->ActorDescViewMap->FindByGuidChecked(HLODActorGuid);
 				GridsHLODActors.FindOrAdd(HLODActorDescView.GetRuntimeGrid()).Add(HLODActorGuid);
 			}
 		}
@@ -555,7 +559,7 @@ bool UWorldPartitionRuntimeSpatialHash::SetupHLODActors(const IStreamingGenerati
 
 		for (const FGuid& HLODActorGuid : GridsHLODActors[HLODGridName])
 		{
-			const FWorldPartitionActorDescView& HLODActorDescView = MainActorSetContainer->ActorDescViewMap->FindByGuidChecked(HLODActorGuid);
+			const FStreamingGenerationActorDescView& HLODActorDescView = MainActorSetContainer->ActorDescViewMap->FindByGuidChecked(HLODActorGuid);
 			
 			IStreamingGenerationContext::FActorSetInstance& NewHLODActorSetInstance = HLODActorSetInstances.Emplace_GetRef();
 			
diff --git a/Engine/Source/Runtime/Engine/Private/WorldPartition/WorldPartitionActorDescView.cpp b/Engine/Source/Runtime/Engine/Private/WorldPartition/WorldPartitionActorDescView.cpp
--- a/Engine/Source/Runtime/Engine/Private/WorldPartition/WorldPartitionActorDescView.cpp
+++ b/Engine/Source/Runtime/Engine/Private/WorldPartition/WorldPartitionActorDescView.cpp
@@ -3,348 +3,47 @@
 #include "WorldPartition/WorldPartitionActorDescView.h"
 
 #if WITH_EDITOR
-#include "WorldPartition/WorldPartitionLog.h"
 #include "WorldPartition/WorldPartitionActorDesc.h"
 
-FWorldPartitionActorDescView::FWorldPartitionActorDescView()
-	: FWorldPartitionActorDescView(nullptr)
-{}
+FWorldPartitionActorDescView::FWorldPartitionActorDescView() : FWorldPartitionActorDescView(nullptr) {}
+FWorldPartitionActorDescView::FWorldPartitionActorDescView(const FWorldPartitionActorDesc* InActorDesc)	: ActorDesc(InActorDesc) {}
+const FGuid& FWorldPartitionActorDescView::GetGuid() const { return ActorDesc->GetGuid(); }
+FTopLevelAssetPath FWorldPartitionActorDescView::GetBaseClass() const {	return ActorDesc->GetBaseClass(); }
+FTopLevelAssetPath FWorldPartitionActorDescView::GetNativeClass() const { return ActorDesc->GetNativeClass(); }
+UClass* FWorldPartitionActorDescView::GetActorNativeClass() const {	return ActorDesc->GetActorNativeClass(); }
+FName FWorldPartitionActorDescView::GetRuntimeGrid() const { return ActorDesc->GetRuntimeGrid(); }
+bool FWorldPartitionActorDescView::GetIsSpatiallyLoaded() const { return ActorDesc->GetIsSpatiallyLoaded(); }
+bool FWorldPartitionActorDescView::GetActorIsEditorOnly() const { return ActorDesc->GetActorIsEditorOnly(); }
+bool FWorldPartitionActorDescView::GetActorIsRuntimeOnly() const { return ActorDesc->GetActorIsRuntimeOnly(); }
+bool FWorldPartitionActorDescView::GetActorIsHLODRelevant() const {	return ActorDesc->GetActorIsHLODRelevant(); }
+FSoftObjectPath FWorldPartitionActorDescView::GetHLODLayer() const { return ActorDesc->GetHLODLayer(); }
+bool FWorldPartitionActorDescView::HasResolvedDataLayerInstanceNames() const { return ActorDesc->HasResolvedDataLayerInstanceNames(); }
+const TArray<FName>& FWorldPartitionActorDescView::GetDataLayerInstanceNames() const { return ActorDesc->GetDataLayerInstanceNames(); }
+const TArray<FName>& FWorldPartitionActorDescView::GetTags() const { return ActorDesc->GetTags(); }
+FName FWorldPartitionActorDescView::GetActorPackage() const { return ActorDesc->GetActorPackage(); }
+FSoftObjectPath FWorldPartitionActorDescView::GetActorSoftPath() const { return ActorDesc->GetActorSoftPath(); }
+FName FWorldPartitionActorDescView::GetActorLabel() const {	return ActorDesc->GetActorLabel(); }
+FBox FWorldPartitionActorDescView::GetEditorBounds() const { return ActorDesc->GetEditorBounds(); }
+FBox FWorldPartitionActorDescView::GetRuntimeBounds() const { return ActorDesc->GetRuntimeBounds(); }
+const TArray<FGuid>& FWorldPartitionActorDescView::GetReferences() const { return ActorDesc->GetReferences(); }
+FString FWorldPartitionActorDescView::ToString() const { return ActorDesc->ToString(); }
+const FGuid& FWorldPartitionActorDescView::GetParentActor() const {	return ActorDesc->GetParentActor(); }
+FName FWorldPartitionActorDescView::GetActorName() const { return ActorDesc->GetActorName(); }
+const FGuid& FWorldPartitionActorDescView::GetFolderGuid() const { return ActorDesc->GetFolderGuid(); }
+FGuid FWorldPartitionActorDescView::GetContentBundleGuid() const { return ActorDesc->GetContentBundleGuid(); }
+FName FWorldPartitionActorDescView::GetContainerPackage() const { return ActorDesc->GetContainerPackage(); }
+bool FWorldPartitionActorDescView::IsContainerInstance() const { return ActorDesc->IsContainerInstance(); }
+bool FWorldPartitionActorDescView::GetContainerInstance(FWorldPartitionActorDesc::FContainerInstance& OutContainerInstance) const {	return ActorDesc->GetContainerInstance(OutContainerInstance); }
+EWorldPartitionActorFilterType FWorldPartitionActorDescView::GetContainerFilterType() const { return ActorDesc->GetContainerFilterType(); }
+const FWorldPartitionActorFilter* FWorldPartitionActorDescView::GetContainerFilter() const { return ActorDesc->GetContainerFilter(); }
+void FWorldPartitionActorDescView::CheckForErrors(IStreamingGenerationErrorHandler* ErrorHandler) const { return ActorDesc->CheckForErrors(ErrorHandler); }
+FName FWorldPartitionActorDescView::GetActorLabelOrName() const { return ActorDesc->GetActorLabelOrName(); }
+AActor* FWorldPartitionActorDescView::GetActor() const { return ActorDesc->GetActor(); }
+bool FWorldPartitionActorDescView::IsEditorOnlyReference(const FGuid& ReferenceGuid) const { return ActorDesc->IsEditorOnlyReference(ReferenceGuid); }
+bool FWorldPartitionActorDescView::GetProperty(FName PropertyName, FName* PropertyValue) const { return ActorDesc->GetProperty(PropertyName, PropertyValue); }
+bool FWorldPartitionActorDescView::HasProperty(FName PropertyName) const { return ActorDesc->HasProperty(PropertyName); }
+
+const TArray<FName>& FWorldPartitionActorDescView::GetRuntimeDataLayerInstanceNames() const { static TArray<FName> EmptyDataLayers; return EmptyDataLayers; }
+const TArray<FGuid>& FWorldPartitionActorDescView::GetEditorReferences() const { static TArray<FGuid> EmptyEditorReferences; return EmptyEditorReferences; }
 
-FWorldPartitionActorDescView::FWorldPartitionActorDescView(const FWorldPartitionActorDesc* InActorDesc)
-	: ActorDesc(InActorDesc)
-	, ParentView(nullptr)
-	, bIsForcedNonSpatiallyLoaded(false)
-	, bIsForcedNoRuntimeGrid(false)
-	, bIsForcedNoDataLayers(false)	
-	, bIsForceNoHLODLayer(false)
-{}
-
-const FGuid& FWorldPartitionActorDescView::GetGuid() const
-{
-	return ActorDesc->GetGuid();
-}
-
-FTopLevelAssetPath FWorldPartitionActorDescView::GetBaseClass() const
-{
-	return ActorDesc->GetBaseClass();
-}
-
-FTopLevelAssetPath FWorldPartitionActorDescView::GetNativeClass() const
-{
-	return ActorDesc->GetNativeClass();
-}
-
-UClass* FWorldPartitionActorDescView::GetActorNativeClass() const
-{
-	return ActorDesc->GetActorNativeClass();
-}
-
-PRAGMA_DISABLE_DEPRECATION_WARNINGS
-FVector FWorldPartitionActorDescView::GetOrigin() const
-{
-	return ActorDesc->GetOrigin();
-}
-PRAGMA_ENABLE_DEPRECATION_WARNINGS
-
-FName FWorldPartitionActorDescView::GetRuntimeGrid() const
-{
-	if (bIsForcedNoRuntimeGrid)
-	{
-		return NAME_None;
-	}
-
-	if (ParentView)
-	{
-		return ParentView->GetRuntimeGrid();
-	}
-
-	return ActorDesc->GetRuntimeGrid();
-}
-
-bool FWorldPartitionActorDescView::GetActorIsEditorOnly() const
-{
-	return ActorDesc->GetActorIsEditorOnly();
-}
-
-bool FWorldPartitionActorDescView::GetActorIsRuntimeOnly() const
-{
-	return ActorDesc->GetActorIsRuntimeOnly();
-}
-
-bool FWorldPartitionActorDescView::GetIsSpatiallyLoaded() const
-{
-	if (bIsForcedNonSpatiallyLoaded)
-	{
-		return false;
-	}
-
-	bool bIsSpatiallyLoaded = ActorDesc->GetIsSpatiallyLoaded();
-	if (bIsSpatiallyLoaded && ParentView)
-	{
-		bIsSpatiallyLoaded = ParentView->GetIsSpatiallyLoaded();
-	}
-	
-	return bIsSpatiallyLoaded;
-}
-
-bool FWorldPartitionActorDescView::GetActorIsHLODRelevant() const
-{
-	return ActorDesc->GetActorIsHLODRelevant();
-}
-
-FSoftObjectPath FWorldPartitionActorDescView::GetHLODLayer() const
-{
-	if (bIsForceNoHLODLayer)
-	{
-		return FSoftObjectPath();
-	}
-	
-	if (RuntimedHLODLayer.IsSet())
-	{
-		return RuntimedHLODLayer.GetValue();
-	}
-
-	return ActorDesc->GetHLODLayer();
-}
-
-void FWorldPartitionActorDescView::SetDataLayerInstanceNames(const TArray<FName>& InDataLayerInstanceNames)
-{
-	check(!ActorDesc->HasResolvedDataLayerInstanceNames());
-	ResolvedDataLayerInstanceNames = InDataLayerInstanceNames;
-}
-
-const TArray<FName>& FWorldPartitionActorDescView::GetDataLayerInstanceNames() const
-{
-	if (bIsForcedNoDataLayers)
-	{
-		static TArray<FName> EmptyDataLayers;
-		return EmptyDataLayers;
-	}
-
-	if (ParentView)
-	{
-		return ParentView->GetDataLayerInstanceNames();
-	}
-	
-	if (ResolvedDataLayerInstanceNames.IsSet())
-	{
-		return ResolvedDataLayerInstanceNames.GetValue();
-	}
-
-	return ActorDesc->GetDataLayerInstanceNames();
-}
-
-const TArray<FName>& FWorldPartitionActorDescView::GetRuntimeDataLayerInstanceNames() const
-{
-	if (bIsForcedNoDataLayers || !ensure(RuntimeDataLayerInstanceNames.IsSet()))
-	{
-		static TArray<FName> EmptyDataLayers;
-		return EmptyDataLayers;
-	}
-
-	if (ParentView)
-	{
-		return ParentView->GetRuntimeDataLayerInstanceNames();
-	}
-
-	return RuntimeDataLayerInstanceNames.GetValue();
-}
-
-const TArray<FName>& FWorldPartitionActorDescView::GetTags() const
-{
-	return ActorDesc->GetTags();
-}
-
-FName FWorldPartitionActorDescView::GetActorPackage() const
-{
-	return ActorDesc->GetActorPackage();
-}
-
-FSoftObjectPath FWorldPartitionActorDescView::GetActorSoftPath() const
-{
-	return ActorDesc->GetActorSoftPath();
-}
-
-FName FWorldPartitionActorDescView::GetActorLabel() const
-{
-	return ActorDesc->GetActorLabel();
-}
-
-FName FWorldPartitionActorDescView::GetActorName() const
-{
-	return ActorDesc->GetActorName();
-}
-
-PRAGMA_DISABLE_DEPRECATION_WARNINGS
-FBox FWorldPartitionActorDescView::GetBounds() const
-{
-	return ActorDesc->GetBounds();
-}
-PRAGMA_ENABLE_DEPRECATION_WARNINGS
-
-FBox FWorldPartitionActorDescView::GetEditorBounds() const
-{
-	return ActorDesc->GetEditorBounds();
-}
-
-FBox FWorldPartitionActorDescView::GetRuntimeBounds() const
-{
-	return ActorDesc->GetRuntimeBounds();
-}
-
-const TArray<FGuid>& FWorldPartitionActorDescView::GetReferences() const
-{
-	return RuntimeReferences.IsSet() ? RuntimeReferences.GetValue() : ActorDesc->GetReferences();
-}
-
-const TArray<FGuid>& FWorldPartitionActorDescView::GetEditorReferences() const
-{
-	return EditorReferences;
-}
-
-FString FWorldPartitionActorDescView::ToString() const
-{
-	return ActorDesc->ToString();
-}
-
-const FGuid& FWorldPartitionActorDescView::GetParentActor() const
-{
-	return ActorDesc->GetParentActor();
-}
-
-const FGuid& FWorldPartitionActorDescView::GetFolderGuid() const
-{
-	return ActorDesc->GetFolderGuid();
-}
-
-FGuid FWorldPartitionActorDescView::GetContentBundleGuid() const
-{
-	return ActorDesc->GetContentBundleGuid();
-}
-
-bool FWorldPartitionActorDescView::IsContainerInstance() const
-{
-	return ActorDesc->IsContainerInstance();
-}
-
-EWorldPartitionActorFilterType FWorldPartitionActorDescView::GetContainerFilterType() const
-{
-	return ActorDesc->GetContainerFilterType();
-}
-
-FName FWorldPartitionActorDescView::GetContainerPackage() const
-{
-	return ActorDesc->GetContainerPackage();
-}
-
-bool FWorldPartitionActorDescView::GetContainerInstance(FWorldPartitionActorDesc::FContainerInstance& OutContainerInstance) const
-{
-	return ActorDesc->GetContainerInstance(OutContainerInstance);
-}
-
-const FWorldPartitionActorFilter* FWorldPartitionActorDescView::GetContainerFilter() const
-{
-	check(GetContainerFilterType() != EWorldPartitionActorFilterType::None);
-	return ActorDesc->GetContainerFilter();
-}
-
-void FWorldPartitionActorDescView::CheckForErrors(IStreamingGenerationErrorHandler* ErrorHandler) const
-{
-	ActorDesc->CheckForErrors(ErrorHandler);
-}
-
-FName FWorldPartitionActorDescView::GetActorLabelOrName() const
-{
-	return ActorDesc->GetActorLabelOrName();
-}
-
-bool FWorldPartitionActorDescView::ShouldValidateRuntimeGrid() const
-{
-PRAGMA_DISABLE_DEPRECATION_WARNINGS
-	return ActorDesc->ShouldValidateRuntimeGrid();
-PRAGMA_ENABLE_DEPRECATION_WARNINGS
-}
-
-void FWorldPartitionActorDescView::SetForcedNonSpatiallyLoaded()
-{
-	if (!bIsForcedNonSpatiallyLoaded)
-	{
-		bIsForcedNonSpatiallyLoaded = true;
-		UE_LOG(LogWorldPartition, Verbose, TEXT("Actor '%s' forced to be non-spatially loaded"), *GetActorLabelOrName().ToString());
-	}
-}
-
-void FWorldPartitionActorDescView::SetForcedNoRuntimeGrid()
-{
-	if (!bIsForcedNoRuntimeGrid)
-	{
-		bIsForcedNoRuntimeGrid = true;	
-		UE_LOG(LogWorldPartition, Verbose, TEXT("Actor '%s' runtime grid invalidated"), *GetActorLabelOrName().ToString());
-	}
-}
-
-void FWorldPartitionActorDescView::SetForcedNoDataLayers()
-{
-	if (!bIsForcedNoDataLayers)
-	{
-		bIsForcedNoDataLayers = true;
-		UE_LOG(LogWorldPartition, Verbose, TEXT("Actor '%s' data layers invalidated"), *GetActorLabelOrName().ToString());
-	}
-}
-
-void FWorldPartitionActorDescView::SetRuntimeDataLayerInstanceNames(const TArray<FName>& InRuntimeDataLayerInstanceNames)
-{
-	RuntimeDataLayerInstanceNames = InRuntimeDataLayerInstanceNames;
-}
-
-void FWorldPartitionActorDescView::SetRuntimeReferences(const TArray<FGuid>& InRuntimeReferences)
-{
-	RuntimeReferences = InRuntimeReferences;
-}
-
-void FWorldPartitionActorDescView::SetEditorReferences(const TArray<FGuid>& InEditorReferences)
-{
-	EditorReferences = InEditorReferences;
-}
-
-void FWorldPartitionActorDescView::SetForcedNoHLODLayer()
-{
-	if (!bIsForceNoHLODLayer)
-	{
-		bIsForceNoHLODLayer = true;
-		UE_LOG(LogWorldPartition, Verbose, TEXT("Actor '%s' HLOD layer invalidated"), *GetActorLabelOrName().ToString());
-	}
-}
-
-void FWorldPartitionActorDescView::SetRuntimeHLODLayer(const FSoftObjectPath& InHLODLayer)
-{
-	RuntimedHLODLayer = InHLODLayer;
-}
-
-AActor* FWorldPartitionActorDescView::GetActor() const
-{
-	return ActorDesc->GetActor();
-}
-
-bool FWorldPartitionActorDescView::IsEditorOnlyReference(const FGuid& ReferenceGuid) const
-{
-	return ActorDesc->IsEditorOnlyReference(ReferenceGuid);
-}
-
-bool FWorldPartitionActorDescView::GetProperty(FName PropertyName, FName* PropertyValue) const
-{
-	return ActorDesc->GetProperty(PropertyName, PropertyValue);
-}
-
-bool FWorldPartitionActorDescView::HasProperty(FName PropertyName) const
-{
-	return ActorDesc->HasProperty(PropertyName);
-}
-
-void FWorldPartitionActorDescView::SetParentView(const FWorldPartitionActorDescView* InParentView)
-{
-	check(!ParentView);
-	check(GetParentActor().IsValid());
-	ParentView = InParentView;
-}
 #endif
diff --git a/Engine/Source/Runtime/Engine/Private/WorldPartition/WorldPartitionRuntimeSpatialHash.cpp b/Engine/Source/Runtime/Engine/Private/WorldPartition/WorldPartitionRuntimeSpatialHash.cpp
--- a/Engine/Source/Runtime/Engine/Private/WorldPartition/WorldPartitionRuntimeSpatialHash.cpp
+++ b/Engine/Source/Runtime/Engine/Private/WorldPartition/WorldPartitionRuntimeSpatialHash.cpp
@@ -1211,7 +1211,8 @@ bool UWorldPartitionRuntimeSpatialHash::GenerateStreaming(UWorldPartitionStreami
 	TArray<FSpatialHashRuntimeGrid> AllGrids;
 	AllGrids.Append(Grids);
 
-	TArray<const FWorldPartitionActorDescView*> SpatialHashRuntimeGridInfos = StreamingGenerationContext->GetMainWorldContainer()->ActorDescViewMap->FindByExactNativeClass<ASpatialHashRuntimeGridInfo>();
+	TArray<const FWorldPartitionActorDescView*> SpatialHashRuntimeGridInfos;
+	Algo::Transform(StreamingGenerationContext->GetMainWorldContainer()->ActorDescViewMap->FindByExactNativeClass<ASpatialHashRuntimeGridInfo>(), SpatialHashRuntimeGridInfos, [](const FStreamingGenerationActorDescView* ActorDescView) { return ActorDescView; });
 	for (const FWorldPartitionActorDescView* SpatialHashRuntimeGridInfo : SpatialHashRuntimeGridInfos)
 	{
 		FWorldPartitionReference Ref(WorldPartition, SpatialHashRuntimeGridInfo->GetGuid());
diff --git a/Engine/Source/Runtime/Engine/Private/WorldPartition/WorldPartitionStreamingGeneration.cpp b/Engine/Source/Runtime/Engine/Private/WorldPartition/WorldPartitionStreamingGeneration.cpp
--- a/Engine/Source/Runtime/Engine/Private/WorldPartition/WorldPartitionStreamingGeneration.cpp
+++ b/Engine/Source/Runtime/Engine/Private/WorldPartition/WorldPartitionStreamingGeneration.cpp
@@ -76,41 +76,193 @@ private:
 	IStreamingGenerationErrorHandler* ErrorHandler;
 };
 
-FActorDescViewMap::FActorDescViewMap()
+UWorldPartition::FCheckForErrorsParams::FCheckForErrorsParams()
+	: ErrorHandler(nullptr)
+	, ActorDescContainerCollection(nullptr)
+	, bEnableStreaming(false)
 {}
 
-TArray<const FWorldPartitionActorDescView*> FActorDescViewMap::FindByExactNativeClass(UClass* InExactNativeClass) const
+FStreamingGenerationActorDescView::FStreamingGenerationActorDescView()
+	: FWorldPartitionActorDescView(nullptr)
+{}
+
+FStreamingGenerationActorDescView::FStreamingGenerationActorDescView(const FWorldPartitionActorDesc* InActorDesc)
+	: FWorldPartitionActorDescView(InActorDesc)
+	, ParentView(nullptr)
+	, bIsForcedNonSpatiallyLoaded(false)
+	, bIsForcedNoRuntimeGrid(false)
+	, bIsForcedNoDataLayers(false)	
+	, bIsForceNoHLODLayer(false)
+{}
+
+FName FStreamingGenerationActorDescView::GetRuntimeGrid() const
+{
+	if (bIsForcedNoRuntimeGrid)
+	{
+		return NAME_None;
+	}
+
+	if (ParentView)
+	{
+		return ParentView->GetRuntimeGrid();
+	}
+
+	return FWorldPartitionActorDescView::GetRuntimeGrid();
+}
+
+bool FStreamingGenerationActorDescView::GetIsSpatiallyLoaded() const
+{
+	if (bIsForcedNonSpatiallyLoaded)
+	{
+		return false;
+	}
+
+	bool bIsSpatiallyLoaded = FWorldPartitionActorDescView::GetIsSpatiallyLoaded();
+
+	if (bIsSpatiallyLoaded && ParentView)
+	{
+		bIsSpatiallyLoaded = ParentView->GetIsSpatiallyLoaded();
+	}
+
+	return bIsSpatiallyLoaded;
+}
+
+FSoftObjectPath FStreamingGenerationActorDescView::GetHLODLayer() const
 {
-	check(InExactNativeClass->IsNative());
-	const FName NativeClassName = InExactNativeClass->GetFName();
-	TArray<const FWorldPartitionActorDescView*> Result;
-	ActorDescViewsByClass.MultiFind(NativeClassName, Result);
-	return Result;
+	if (bIsForceNoHLODLayer)
+	{
+		return FSoftObjectPath();
+	}
+	
+	if (RuntimedHLODLayer.IsSet())
+	{
+		return RuntimedHLODLayer.GetValue();
+	}
+
+	return FWorldPartitionActorDescView::GetHLODLayer();
 }
 
-FWorldPartitionActorDescView* FActorDescViewMap::Emplace(const FGuid& InGuid, const FWorldPartitionActorDescView& InActorDescView)
+const TArray<FName>& FStreamingGenerationActorDescView::GetDataLayerInstanceNames() const
 {
-	FWorldPartitionActorDescView* NewActorDescView = ActorDescViewList.Emplace_GetRef(MakeUnique<FWorldPartitionActorDescView>(InActorDescView)).Get();
+	if (bIsForcedNoDataLayers)
+	{
+		static TArray<FName> EmptyDataLayers;
+		return EmptyDataLayers;
+	}
+
+	if (ParentView)
+	{
+		return ParentView->GetDataLayerInstanceNames();
+	}
 	
-	const UClass* NativeClass = NewActorDescView->GetActorNativeClass();
-	const FName NativeClassName = NativeClass->GetFName();
+	if (ResolvedDataLayerInstanceNames.IsSet())
+	{
+		return ResolvedDataLayerInstanceNames.GetValue();
+	}
 
-	ActorDescViewsByGuid.Emplace(InGuid, NewActorDescView);
-	ActorDescViewsByClass.Add(NativeClassName, NewActorDescView);
+	return FWorldPartitionActorDescView::GetDataLayerInstanceNames();
+}
 
-	return NewActorDescView;
+const TArray<FGuid>& FStreamingGenerationActorDescView::GetReferences() const
+{
+	return RuntimeReferences.IsSet() ? RuntimeReferences.GetValue() : FWorldPartitionActorDescView::GetReferences();
 }
 
-FWorldPartitionActorDescView* FActorDescViewMap::Emplace(const FWorldPartitionActorDesc* InActorDesc)
+const TArray<FGuid>& FStreamingGenerationActorDescView::GetEditorReferences() const
 {
-	return Emplace(InActorDesc->GetGuid(), FWorldPartitionActorDescView(InActorDesc));
+	return EditorReferences;
 }
 
-UWorldPartition::FCheckForErrorsParams::FCheckForErrorsParams()
-	: ErrorHandler(nullptr)
-	, ActorDescContainerCollection(nullptr)
-	, bEnableStreaming(false)
-{}
+bool FStreamingGenerationActorDescView::ShouldValidateRuntimeGrid() const
+{
+PRAGMA_DISABLE_DEPRECATION_WARNINGS
+	return ActorDesc->ShouldValidateRuntimeGrid();
+PRAGMA_ENABLE_DEPRECATION_WARNINGS
+}
+
+void FStreamingGenerationActorDescView::SetParentView(const FWorldPartitionActorDescView* InParentView)
+{
+	check(!ParentView);
+	check(GetParentActor().IsValid());
+	ParentView = InParentView;
+}
+
+void FStreamingGenerationActorDescView::SetDataLayerInstanceNames(const TArray<FName>& InDataLayerInstanceNames)
+{
+	check(!FWorldPartitionActorDescView::HasResolvedDataLayerInstanceNames());
+	ResolvedDataLayerInstanceNames = InDataLayerInstanceNames;
+}
+
+void FStreamingGenerationActorDescView::SetForcedNonSpatiallyLoaded()
+{
+	if (!bIsForcedNonSpatiallyLoaded)
+	{
+		bIsForcedNonSpatiallyLoaded = true;
+		UE_LOG(LogWorldPartition, Verbose, TEXT("Actor '%s' forced to be non-spatially loaded"), *GetActorLabelOrName().ToString());
+	}
+}
+
+void FStreamingGenerationActorDescView::SetForcedNoRuntimeGrid()
+{
+	if (!bIsForcedNoRuntimeGrid)
+	{
+		bIsForcedNoRuntimeGrid = true;	
+		UE_LOG(LogWorldPartition, Verbose, TEXT("Actor '%s' runtime grid invalidated"), *GetActorLabelOrName().ToString());
+	}
+}
+
+void FStreamingGenerationActorDescView::SetForcedNoDataLayers()
+{
+	if (!bIsForcedNoDataLayers)
+	{
+		bIsForcedNoDataLayers = true;
+		UE_LOG(LogWorldPartition, Verbose, TEXT("Actor '%s' data layers invalidated"), *GetActorLabelOrName().ToString());
+	}
+}
+
+void FStreamingGenerationActorDescView::SetRuntimeDataLayerInstanceNames(const TArray<FName>& InRuntimeDataLayerInstanceNames)
+{
+	RuntimeDataLayerInstanceNames = InRuntimeDataLayerInstanceNames;
+}
+
+void FStreamingGenerationActorDescView::SetRuntimeReferences(const TArray<FGuid>& InRuntimeReferences)
+{
+	RuntimeReferences = InRuntimeReferences;
+}
+
+void FStreamingGenerationActorDescView::SetEditorReferences(const TArray<FGuid>& InEditorReferences)
+{
+	EditorReferences = InEditorReferences;
+}
+
+void FStreamingGenerationActorDescView::SetForcedNoHLODLayer()
+{
+	if (!bIsForceNoHLODLayer)
+	{
+		bIsForceNoHLODLayer = true;
+		UE_LOG(LogWorldPartition, Verbose, TEXT("Actor '%s' HLOD layer invalidated"), *GetActorLabelOrName().ToString());
+	}
+}
+
+void FStreamingGenerationActorDescView::SetRuntimeHLODLayer(const FSoftObjectPath& InHLODLayer)
+{
+	RuntimedHLODLayer = InHLODLayer;
+}
+
+const TArray<FName>& FStreamingGenerationActorDescView::GetRuntimeDataLayerInstanceNames() const
+{
+	if (bIsForcedNoDataLayers || !RuntimeDataLayerInstanceNames.IsSet())
+	{
+		return FWorldPartitionActorDescView::GetRuntimeDataLayerInstanceNames();
+	}
+
+	if (ParentView)
+	{
+		return ParentView->GetRuntimeDataLayerInstanceNames();
+	}
+
+	return RuntimeDataLayerInstanceNames.GetValue();
+}
 
 class FWorldPartitionStreamingGenerator
 {
@@ -253,13 +405,13 @@ class FWorldPartitionStreamingGenerator
 		TSharedPtr<FStreamingGenerationActorDescCollection> ActorDescCollection;
 
 		/** The actor descriptor views for for this descriptor */
-		FActorDescViewMap ActorDescViewMap;
+		FStreamingGenerationActorDescViewMap ActorDescViewMap;
 
 		/** Set of editor-only actors that are not part of the actor descriptor views */
 		TSet<FGuid> EditorOnlyActorDescMap;
 
 		/** List of actor descriptor views that are containers (mainly level instances) */
-		TArray<FWorldPartitionActorDescView> ContainerCollectionInstanceViews;
+		TArray<FStreamingGenerationActorDescView> ContainerCollectionInstanceViews;
 
 		/** List of actor clusters for this descriptor */
 		TArray<TArray<FGuid>> Clusters;
@@ -338,7 +490,7 @@ class FWorldPartitionStreamingGenerator
 		TMap<FGuid, FSetElementId> PerInstanceData;
 	};
 
-	void ResolveRuntimeSpatiallyLoaded(FWorldPartitionActorDescView& ActorDescView)
+	void ResolveRuntimeSpatiallyLoaded(FStreamingGenerationActorDescView& ActorDescView)
 	{
 		if (!bEnableStreaming)
 		{
@@ -346,7 +498,7 @@ class FWorldPartitionStreamingGenerator
 		}
 	}
 
-	void ResolveRuntimeGrid(FWorldPartitionActorDescView& ActorDescView)
+	void ResolveRuntimeGrid(FStreamingGenerationActorDescView& ActorDescView)
 	{
 		if (!bEnableStreaming)
 		{
@@ -354,26 +506,29 @@ class FWorldPartitionStreamingGenerator
 		}
 	}
 
-	void ResolveRuntimeDataLayers(FWorldPartitionActorDescView& ActorDescView, const FActorDescViewMap& ActorDescViewMap)
+	void ResolveRuntimeDataLayers(FStreamingGenerationActorDescView& ActorDescView, const FStreamingGenerationActorDescViewMap& ActorDescViewMap)
 	{
+		TArray<const FWorldPartitionActorDescView*> WorldDataLayerViews;
+		Algo::Transform(ActorDescViewMap.FindByExactNativeClass<AWorldDataLayers>(), WorldDataLayerViews, [](const FStreamingGenerationActorDescView* ActorDescView) { return ActorDescView; });
+
 		// Resolve DataLayerInstanceNames of ActorDescView only when necessary (i.e. when container is a template)
 		if (!ActorDescView.GetActorDesc()->HasResolvedDataLayerInstanceNames())
 		{
 			// Build a WorldDataLayerActorDescs if DataLayerManager can't resolve Data Layers (i.e. when validating changelists and World is not loaded)
 			const bool bDataLayerManagerCanResolve = DataLayerManager && DataLayerManager->CanResolveDataLayers();
-			const TArray<const FWorldDataLayersActorDesc*> WorldDataLayerActorDescs = !bDataLayerManagerCanResolve ? FDataLayerUtils::FindWorldDataLayerActorDescs(ActorDescViewMap) : TArray<const FWorldDataLayersActorDesc*>();
+			const TArray<const FWorldDataLayersActorDesc*> WorldDataLayerActorDescs = !bDataLayerManagerCanResolve ? FDataLayerUtils::FindWorldDataLayerActorDescs(WorldDataLayerViews) : TArray<const FWorldDataLayersActorDesc*>();
 			const TArray<FName> DataLayerInstanceNames = FDataLayerUtils::ResolvedDataLayerInstanceNames(DataLayerManager, ActorDescView.GetActorDesc(), WorldDataLayerActorDescs);
 			ActorDescView.SetDataLayerInstanceNames(DataLayerInstanceNames);
 		}
 
 		TArray<FName> RuntimeDataLayerInstanceNames;
-		if (FDataLayerUtils::ResolveRuntimeDataLayerInstanceNames(DataLayerManager, ActorDescView, ActorDescViewMap, RuntimeDataLayerInstanceNames))
+		if (FDataLayerUtils::ResolveRuntimeDataLayerInstanceNames(DataLayerManager, ActorDescView, WorldDataLayerViews, RuntimeDataLayerInstanceNames))
 		{
 			ActorDescView.SetRuntimeDataLayerInstanceNames(RuntimeDataLayerInstanceNames);
 		}
 	}
 
-	void ResolveHLODLayer(FWorldPartitionActorDescView& ActorDescView, const FSoftObjectPath& DefaultHLODLayer)
+	void ResolveHLODLayer(FStreamingGenerationActorDescView& ActorDescView, const FSoftObjectPath& DefaultHLODLayer)
 	{
 		// Only assign the default layer to actors that don't have a valid HLOD layer set. HLOD actors will have their 
 		// parent HLOD layer set during HLOD generation.
@@ -383,7 +538,7 @@ class FWorldPartitionStreamingGenerator
 		}
 	}
 
-	void ResolveParentView(FWorldPartitionActorDescView& ActorDescView, const FActorDescViewMap& ActorDescViewMap)
+	void ResolveParentView(FStreamingGenerationActorDescView& ActorDescView, const FStreamingGenerationActorDescViewMap& ActorDescViewMap)
 	{
 		if (FGuid ParentGuid = ActorDescView.GetParentActor(); ParentGuid.IsValid())
 		{
@@ -391,7 +546,7 @@ class FWorldPartitionStreamingGenerator
 		}
 	}
 
-	void CreateActorDescViewMap(const FStreamingGenerationActorDescCollection& InActorDescCollection, FActorDescViewMap& OutActorDescViewMap, TSet<FGuid>& OutEditorOnlyActorDescMap, const FActorContainerID& InContainerID, TArray<FWorldPartitionActorDescView>& OutContainerInstances)
+	void CreateActorDescViewMap(const FStreamingGenerationActorDescCollection& InActorDescCollection, FStreamingGenerationActorDescViewMap& OutActorDescViewMap, TSet<FGuid>& OutEditorOnlyActorDescMap, const FActorContainerID& InContainerID, TArray<FStreamingGenerationActorDescView>& OutContainerInstances)
 	{
 		// Should we handle unsaved or newly created actors?
 		const bool bHandleUnsavedActors = ModifiedActorsDescList && InContainerID.IsMainContainer();
@@ -441,7 +596,7 @@ class FWorldPartitionStreamingGenerator
 		};
 
 		// Register the actor descriptor view
-		auto RegisterActorDescView = [this, &OutActorDescViewMap, &OutContainerInstances](const FGuid& ActorGuid, FWorldPartitionActorDescView& InActorDescView)
+		auto RegisterActorDescView = [this, &OutActorDescViewMap, &OutContainerInstances](const FGuid& ActorGuid, FStreamingGenerationActorDescView& InActorDescView)
 		{
 			if (InActorDescView.IsContainerInstance())
 			{
@@ -471,14 +626,14 @@ class FWorldPartitionStreamingGenerator
 					if (bHandleUnsavedActors && (bIsTempContainerPackage || Actor->GetPackage()->IsDirty()))
 					{
 						// Dirty, unsaved actor for PIE
-						FWorldPartitionActorDescView ModifiedActorDescView = GetModifiedActorDesc(Actor, InActorDescCollection);
+						FStreamingGenerationActorDescView ModifiedActorDescView = GetModifiedActorDesc(Actor, InActorDescCollection);
 						RegisterActorDescView(ActorDescIt->GetGuid(), ModifiedActorDescView);
 						continue;
 					}
 				}
 
 				// Non-dirty actor
-				FWorldPartitionActorDescView ActorDescView(*ActorDescIt);
+				FStreamingGenerationActorDescView ActorDescView(*ActorDescIt);
 				RegisterActorDescView(ActorDescIt->GetGuid(), ActorDescView);
 			}
 			else
@@ -496,7 +651,7 @@ class FWorldPartitionStreamingGenerator
 					&& (Actor->GetContentBundleGuid() == InActorDescCollection.GetContentBundleGuid())
 					&& !InActorDescCollection.GetActorDesc(Actor->GetActorGuid()))
 				{
-					FWorldPartitionActorDescView ModifiedActorDescView = GetModifiedActorDesc(Actor, InActorDescCollection);
+					FStreamingGenerationActorDescView ModifiedActorDescView = GetModifiedActorDesc(Actor, InActorDescCollection);
 					if (ShouldRegisterActorDesc(ModifiedActorDescView.GetActorDesc(), InContainerID))
 					{
 						RegisterActorDescView(Actor->GetActorGuid(), ModifiedActorDescView);
@@ -564,7 +719,7 @@ class FWorldPartitionStreamingGenerator
 
 				// Calculate Bounds of non-container actor descriptor views
 				check(!ContainerCollectionInstanceDescriptor.Bounds.IsValid);
-				ContainerCollectionDescriptor.ActorDescViewMap.ForEachActorDescView([&ContainerCollectionInstanceDescriptor](const FWorldPartitionActorDescView& ActorDescView)
+				ContainerCollectionDescriptor.ActorDescViewMap.ForEachActorDescView([&ContainerCollectionInstanceDescriptor](const FStreamingGenerationActorDescView& ActorDescView)
 				{
 					if (ActorDescView.GetIsSpatiallyLoaded())
 					{
@@ -577,7 +732,7 @@ class FWorldPartitionStreamingGenerator
 			}
 
 			// Parse actor containers
-			TArray<FWorldPartitionActorDescView> ContainerCollectionInstanceViews = ContainerCollectionDescriptor.ContainerCollectionInstanceViews;
+			TArray<FStreamingGenerationActorDescView> ContainerCollectionInstanceViews = ContainerCollectionDescriptor.ContainerCollectionInstanceViews;
 			for (const FWorldPartitionActorDescView& ContainerCollectionInstanceView : ContainerCollectionInstanceViews)
 			{
 				FWorldPartitionActorDesc::FContainerInstance SubContainerInstance;
@@ -633,7 +788,7 @@ class FWorldPartitionStreamingGenerator
 
 		// Apply per-instance data
 		ContainerCollectionInstanceDescriptor.PerInstanceData.Reserve(ContainerCollectionDescriptor.ActorDescViewMap.Num());
-		ContainerCollectionDescriptor.ActorDescViewMap.ForEachActorDescView([&ContainerCollectionInstanceDescriptor, &InheritParentContainerPerInstanceData](FWorldPartitionActorDescView& ActorDescView)
+		ContainerCollectionDescriptor.ActorDescViewMap.ForEachActorDescView([&ContainerCollectionInstanceDescriptor, &InheritParentContainerPerInstanceData](FStreamingGenerationActorDescView& ActorDescView)
 		{
 			const FContainerCollectionInstanceDescriptor::FPerInstanceData PerInstanceData = InheritParentContainerPerInstanceData(ContainerCollectionInstanceDescriptor.InstanceData, ActorDescView);
 
@@ -691,7 +846,7 @@ class FWorldPartitionStreamingGenerator
 	 */
 	void ResolveContainerDescriptor(FContainerCollectionDescriptor& ContainerCollectionDescriptor)
 	{
-		auto ResolveActorDescView = [this, &ContainerCollectionDescriptor](FWorldPartitionActorDescView& ActorDescView)
+		auto ResolveActorDescView = [this, &ContainerCollectionDescriptor](FStreamingGenerationActorDescView& ActorDescView)
 		{
 			ResolveRuntimeSpatiallyLoaded(ActorDescView);
 			ResolveRuntimeGrid(ActorDescView);
@@ -700,12 +855,12 @@ class FWorldPartitionStreamingGenerator
 			ResolveParentView(ActorDescView, ContainerCollectionDescriptor.ActorDescViewMap);
 		};
 
-		ContainerCollectionDescriptor.ActorDescViewMap.ForEachActorDescView([this, &ResolveActorDescView](FWorldPartitionActorDescView& ActorDescView)
+		ContainerCollectionDescriptor.ActorDescViewMap.ForEachActorDescView([this, &ResolveActorDescView](FStreamingGenerationActorDescView& ActorDescView)
 		{
 			ResolveActorDescView(ActorDescView);
 		});
 
-		for (FWorldPartitionActorDescView& ContainerCollectionInstanceView : ContainerCollectionDescriptor.ContainerCollectionInstanceViews)
+		for (FStreamingGenerationActorDescView& ContainerCollectionInstanceView : ContainerCollectionDescriptor.ContainerCollectionInstanceViews)
 		{
 			ResolveActorDescView(ContainerCollectionInstanceView);
 		}
@@ -748,7 +903,7 @@ class FWorldPartitionStreamingGenerator
 
 			for (const FGuid& LevelScriptReferenceActorGuid : LevelScriptReferences)
 			{
-				if (FWorldPartitionActorDescView* ActorDescView = ContainerCollectionDescriptor.ActorDescViewMap.FindByGuid(LevelScriptReferenceActorGuid))
+				if (FStreamingGenerationActorDescView* ActorDescView = (FStreamingGenerationActorDescView*)ContainerCollectionDescriptor.ActorDescViewMap.FindByGuid(LevelScriptReferenceActorGuid))
 				{
 					if (ActorDescView->GetIsSpatiallyLoaded())
 					{
@@ -766,7 +921,7 @@ class FWorldPartitionStreamingGenerator
 		}
 
 		// Route standard CheckForErrors calls which should not modify actor descriptors in any ways
-		ContainerCollectionDescriptor.ActorDescViewMap.ForEachActorDescView([this](FWorldPartitionActorDescView& ActorDescView)
+		ContainerCollectionDescriptor.ActorDescViewMap.ForEachActorDescView([this](FStreamingGenerationActorDescView& ActorDescView)
 		{
 			ActorDescView.CheckForErrors(ErrorHandler);
 		});
@@ -795,7 +950,7 @@ class FWorldPartitionStreamingGenerator
 
 			NbErrorsDetected = 0;
 
-			ContainerCollectionDescriptor.ActorDescViewMap.ForEachActorDescView([this, &ContainerCollectionDescriptor, &NbErrorsDetected, PassType](FWorldPartitionActorDescView& ActorDescView)
+			ContainerCollectionDescriptor.ActorDescViewMap.ForEachActorDescView([this, &ContainerCollectionDescriptor, &NbErrorsDetected, PassType](FStreamingGenerationActorDescView& ActorDescView)
 			{
 				// Validate grid placement
 				auto IsReferenceGridPlacementValid = [](const FWorldPartitionActorDescView& RefererActorDescView, const FWorldPartitionActorDescView& ReferenceActorDescView)
@@ -829,9 +984,9 @@ class FWorldPartitionStreamingGenerator
 				struct FActorReferenceInfo
 				{
 					FGuid ActorGuid;
-					FWorldPartitionActorDescView* ActorDesc;
+					FStreamingGenerationActorDescView* ActorDescView;
 					FGuid ReferenceGuid;
-					FWorldPartitionActorDescView* ReferenceActorDesc;
+					FStreamingGenerationActorDescView* ReferenceActorDescView;
 				};
 
 				TArray<FActorReferenceInfo> References;
@@ -842,7 +997,7 @@ class FWorldPartitionStreamingGenerator
 					if (ReferenceGuid != ActorDescView.GetParentActor()) // References to the parent are inversed in their handling 
 					{
 						// Filter out parent back references
-						FWorldPartitionActorDescView* ReferenceActorDesc = ContainerCollectionDescriptor.ActorDescViewMap.FindByGuid(ReferenceGuid);
+						FStreamingGenerationActorDescView* ReferenceActorDesc = (FStreamingGenerationActorDescView*)ContainerCollectionDescriptor.ActorDescViewMap.FindByGuid(ReferenceGuid);
 						if (!ReferenceActorDesc || (ReferenceActorDesc->GetParentActor() != ActorDescView.GetGuid()))
 						{
 							References.Emplace(FActorReferenceInfo { ActorDescView.GetGuid(), &ActorDescView, ReferenceGuid, ReferenceActorDesc });
@@ -854,11 +1009,11 @@ class FWorldPartitionStreamingGenerator
 				// parent to be refering to us, not the child to be referering the parent.
 				{
 					FGuid ParentGuid = ActorDescView.GetParentActor();
-					FWorldPartitionActorDescView* TopParentDescView = nullptr;
+					FStreamingGenerationActorDescView* TopParentDescView = nullptr;
 
 					while (ParentGuid.IsValid())
 					{
-						FWorldPartitionActorDescView* ParentDescView = ContainerCollectionDescriptor.ActorDescViewMap.FindByGuid(ParentGuid);
+						FStreamingGenerationActorDescView* ParentDescView = (FStreamingGenerationActorDescView*)ContainerCollectionDescriptor.ActorDescViewMap.FindByGuid(ParentGuid);
 					
 						if (ParentDescView)
 						{
@@ -885,8 +1040,8 @@ class FWorldPartitionStreamingGenerator
 
 				for (FActorReferenceInfo& Info : References)
 				{
-					FWorldPartitionActorDescView* RefererActorDescView = Info.ActorDesc;
-					FWorldPartitionActorDescView* ReferenceActorDescView = Info.ReferenceActorDesc;
+					FStreamingGenerationActorDescView* RefererActorDescView = Info.ActorDescView;
+					FStreamingGenerationActorDescView* ReferenceActorDescView = Info.ReferenceActorDescView;
 
 					if (ReferenceActorDescView)
 					{
@@ -969,7 +1124,7 @@ class FWorldPartitionStreamingGenerator
 		}
 
 		// Split runtime and editor references
-		ContainerCollectionDescriptor.ActorDescViewMap.ForEachActorDescView([this, &ContainerCollectionDescriptor](FWorldPartitionActorDescView& ActorDescView)
+		ContainerCollectionDescriptor.ActorDescViewMap.ForEachActorDescView([this, &ContainerCollectionDescriptor](FStreamingGenerationActorDescView& ActorDescView)
 		{
 			TArray<FGuid> RuntimeReferences;
 			TArray<FGuid> EditorReferences;
@@ -979,7 +1134,7 @@ class FWorldPartitionStreamingGenerator
 
 			for (const FGuid& ReferenceGuid : ActorDescView.GetReferences())
 			{
-				if (FWorldPartitionActorDescView* ReferenceActorDescView = ContainerCollectionDescriptor.ActorDescViewMap.FindByGuid(ReferenceGuid))
+				if (FStreamingGenerationActorDescView* ReferenceActorDescView = (FStreamingGenerationActorDescView*)ContainerCollectionDescriptor.ActorDescViewMap.FindByGuid(ReferenceGuid))
 				{
 					// The actor reference is not editor-only, but we are referencing it through an editor-only property
 					if (ActorDescView.IsEditorOnlyReference(ReferenceGuid))
@@ -1032,7 +1187,7 @@ class FWorldPartitionStreamingGenerator
 
 			NbErrorsDetected = 0;
 
-			ContainerCollectionDescriptor.ActorDescViewMap.ForEachActorDescView([this, &ContainerCollectionDescriptor, &ContainerCollectionInstanceDescriptor, &NbErrorsDetected, PassType](FWorldPartitionActorDescView& ActorDescView)
+			ContainerCollectionDescriptor.ActorDescViewMap.ForEachActorDescView([this, &ContainerCollectionDescriptor, &ContainerCollectionInstanceDescriptor, &NbErrorsDetected, PassType](FStreamingGenerationActorDescView& ActorDescView)
 			{
 				FContainerCollectionInstanceDescriptor::FPerInstanceData& PerInstanceData = ContainerCollectionInstanceDescriptor.GetInstanceData(ActorDescView.GetGuid());
 
@@ -1102,7 +1257,7 @@ PRAGMA_ENABLE_DEPRECATION_WARNINGS
 	{
 		// Build clusters for this container - at this point, all actors references should be in the same data layers, grid, etc because of actor descriptors validation.
 		TArray<TPair<FGuid, TArray<FGuid>>> ActorsWithRefs;
-		ContainerCollectionDescriptor.ActorDescViewMap.ForEachActorDescView([&ActorsWithRefs](const FWorldPartitionActorDescView& ActorDescView) { ActorsWithRefs.Emplace(ActorDescView.GetGuid(), ActorDescView.GetReferences()); });
+		ContainerCollectionDescriptor.ActorDescViewMap.ForEachActorDescView([&ActorsWithRefs](const FStreamingGenerationActorDescView& ActorDescView) { ActorsWithRefs.Emplace(ActorDescView.GetGuid(), ActorDescView.GetReferences()); });
 		ContainerCollectionDescriptor.Clusters = GenerateObjectsClusters(ActorsWithRefs);
 	}
 
@@ -1229,7 +1384,7 @@ public:
 			{
 				UE_SCOPED_INDENT_LOG_ARCHIVE(Ar.PrintfIndent(TEXT("ActorDescs:")));
 
-				TMap<FGuid, FWorldPartitionActorDescView*> SortedActorDescViewMap = ContainerCollectionDescriptor.ActorDescViewMap.ActorDescViewsByGuid;
+				TMap<FGuid, FStreamingGenerationActorDescView*> SortedActorDescViewMap = ContainerCollectionDescriptor.ActorDescViewMap.ActorDescViewsByGuid;
 				SortedActorDescViewMap.KeySort([](const FGuid& GuidA, const FGuid& GuidB) { return GuidA < GuidB; });
 
 				for (auto& [ActorGuid, ActorDescView] : SortedActorDescViewMap)
diff --git a/Engine/Source/Runtime/Engine/Public/WorldPartition/DataLayer/DataLayerUtils.h b/Engine/Source/Runtime/Engine/Public/WorldPartition/DataLayer/DataLayerUtils.h
--- a/Engine/Source/Runtime/Engine/Public/WorldPartition/DataLayer/DataLayerUtils.h
+++ b/Engine/Source/Runtime/Engine/Public/WorldPartition/DataLayer/DataLayerUtils.h
@@ -29,13 +29,13 @@ public:
 
 	static ENGINE_API TArray<FName> ResolvedDataLayerInstanceNames(const UDataLayerManager* InDataLayerManager, const FWorldPartitionActorDesc* InActorDesc, const TArray<const FWorldDataLayersActorDesc*>& InWorldDataLayersActorDescs = TArray<const FWorldDataLayersActorDesc*>());
 	
-	static ENGINE_API bool ResolveRuntimeDataLayerInstanceNames(const UDataLayerManager* InDataLayerManager, const FWorldPartitionActorDescView& InActorDescView, const FActorDescViewMap& ActorDescViewMap, TArray<FName>& OutRuntimeDataLayerInstanceNames);
+	static ENGINE_API bool ResolveRuntimeDataLayerInstanceNames(const UDataLayerManager* InDataLayerManager, const FWorldPartitionActorDescView& InActorDescView, const TArray<const FWorldPartitionActorDescView*>& InWorldDataLayerViews, TArray<FName>& OutRuntimeDataLayerInstanceNames);
 
-	static ENGINE_API const FDataLayerInstanceDesc* GetDataLayerInstanceDescFromInstanceName(const TArray<const FWorldDataLayersActorDesc*>& InWorldDataLayersActorDescs, const FName& DataLayerInstanceName);
+	static ENGINE_API const FDataLayerInstanceDesc* GetDataLayerInstanceDescFromInstanceName(const TArray<const FWorldDataLayersActorDesc*>& InWorldDataLayersActorDescs, const FName& InDataLayerInstanceName);
 
-	static ENGINE_API const FDataLayerInstanceDesc* GetDataLayerInstanceDescFromAssetPath(const TArray<const FWorldDataLayersActorDesc*>& InWorldDataLayersActorDescs, const FName& DataLayerAssetPath);
+	static ENGINE_API const FDataLayerInstanceDesc* GetDataLayerInstanceDescFromAssetPath(const TArray<const FWorldDataLayersActorDesc*>& InWorldDataLayersActorDescs, const FName& InDataLayerAssetPath);
 
-	static ENGINE_API TArray<const FWorldDataLayersActorDesc*> FindWorldDataLayerActorDescs(const FActorDescViewMap& ActorDescViewMap);
+	static ENGINE_API TArray<const FWorldDataLayersActorDesc*> FindWorldDataLayerActorDescs(const TArray<const FWorldPartitionActorDescView*>& InWorldDataLayerViews);
 
 	static ENGINE_API bool AreWorldDataLayersActorDescsSane(const TArray<const FWorldDataLayersActorDesc*>& InWorldDataLayersActorDescs);
 
diff --git a/Engine/Source/Runtime/Engine/Public/WorldPartition/WorldPartitionActorDescView.h b/Engine/Source/Runtime/Engine/Public/WorldPartition/WorldPartitionActorDescView.h
--- a/Engine/Source/Runtime/Engine/Public/WorldPartition/WorldPartitionActorDescView.h
+++ b/Engine/Source/Runtime/Engine/Public/WorldPartition/WorldPartitionActorDescView.h
@@ -16,110 +16,60 @@ struct FWorldPartitionActorFilter;
 enum class EWorldPartitionActorFilterType : uint8;
 
 /**
- * A view on top of an actor desc, used to cache information that can be (potentially) different than the actor desc
- * itself due to streaming generation logic, etc.
+ * A view on top of an actor descriptor, used to store information that can be potentially different than the actor
+ * descriptor itself due to streaming generation logic, etc.
  */
 class FWorldPartitionActorDescView
 {
-	friend class UWorldPartitionRuntimeHash;
-
 public:
 	ENGINE_API FWorldPartitionActorDescView();
 	ENGINE_API FWorldPartitionActorDescView(const FWorldPartitionActorDesc* InActorDesc);
-
-	ENGINE_API const FGuid& GetGuid() const;
-	ENGINE_API FTopLevelAssetPath GetBaseClass() const;
-	ENGINE_API FTopLevelAssetPath GetNativeClass() const;
-	ENGINE_API UClass* GetActorNativeClass() const;
-	UE_DEPRECATED(5.2, "GetOrigin is deprecated.")
-	ENGINE_API FVector GetOrigin() const;
-	ENGINE_API FName GetRuntimeGrid() const;
-	ENGINE_API bool GetIsSpatiallyLoaded() const;
-	ENGINE_API bool GetActorIsEditorOnly() const;
-	ENGINE_API bool GetActorIsRuntimeOnly() const;
-	ENGINE_API bool GetActorIsHLODRelevant() const;
-	ENGINE_API FSoftObjectPath GetHLODLayer() const;
-	ENGINE_API const TArray<FName>& GetDataLayerInstanceNames() const;
-	ENGINE_API const TArray<FName>& GetRuntimeDataLayerInstanceNames() const;
-	ENGINE_API const TArray<FName>& GetTags() const;
-	ENGINE_API FName GetActorPackage() const;
-	
-	ENGINE_API FSoftObjectPath GetActorSoftPath() const;
-	ENGINE_API FName GetActorLabel() const;
-
-	UE_DEPRECATED(5.2, "GetBounds is deprecated, GetEditorBounds or GetRuntimeBounds should be used instead.")
-	ENGINE_API FBox GetBounds() const;
-
-	ENGINE_API FBox GetEditorBounds() const;
-	ENGINE_API FBox GetRuntimeBounds() const;
-
-	ENGINE_API const TArray<FGuid>& GetReferences() const;
-	ENGINE_API const TArray<FGuid>& GetEditorReferences() const;
-	ENGINE_API FString ToString() const;
-	ENGINE_API const FGuid& GetParentActor() const;
-	ENGINE_API FName GetActorName() const;
-	ENGINE_API const FGuid& GetFolderGuid() const;
-
-	ENGINE_API FGuid GetContentBundleGuid() const;
-	ENGINE_API FName GetContainerPackage() const;
-	ENGINE_API bool IsContainerInstance() const;
-	ENGINE_API bool GetContainerInstance(FWorldPartitionActorDesc::FContainerInstance& OutContainerInstance) const;
-	ENGINE_API EWorldPartitionActorFilterType GetContainerFilterType() const;
-	ENGINE_API const FWorldPartitionActorFilter* GetContainerFilter() const;
 	
-	UE_DEPRECATED(5.3, "GetLevelPackage is deprecated use GetContainerPackage instead.")
-	FName GetLevelPackage() const { return GetContainerPackage(); }
-
-	ENGINE_API void CheckForErrors(IStreamingGenerationErrorHandler* ErrorHandler) const;
-
-	ENGINE_API FName GetActorLabelOrName() const;
-
-	UE_DEPRECATED(5.2, "ShouldValidateRuntimeGrid is deprecated and should not be used.")
-	ENGINE_API bool ShouldValidateRuntimeGrid() const;
-
-	ENGINE_API void SetForcedNonSpatiallyLoaded();
-	ENGINE_API void SetForcedNoRuntimeGrid();
-	ENGINE_API void SetForcedNoDataLayers();
-	ENGINE_API void SetRuntimeDataLayerInstanceNames(const TArray<FName>& InRuntimeDataLayerInstanceNames);
-	ENGINE_API void SetRuntimeReferences(const TArray<FGuid>& InRuntimeReferences);
-	ENGINE_API void SetEditorReferences(const TArray<FGuid>& InEditorReferences);
-	ENGINE_API void SetDataLayerInstanceNames(const TArray<FName>& InDataLayerInstanceNames);
-	
-	ENGINE_API void SetForcedNoHLODLayer();
-	ENGINE_API void SetRuntimeHLODLayer(const FSoftObjectPath& InHLODLayer);
-
-	ENGINE_API AActor* GetActor() const;
-
-	bool operator==(const FWorldPartitionActorDescView& Other) const
-	{
-		return GetGuid() == Other.GetGuid();
-	}
-
-	friend uint32 GetTypeHash(const FWorldPartitionActorDescView& Key)
-	{
-		return GetTypeHash(Key.GetGuid());
-	}
-
+	virtual ~FWorldPartitionActorDescView() {}
+
+	/** Interface on top of FWorldPartitionActorDesc */
+	ENGINE_API virtual const FGuid& GetGuid() const;
+	ENGINE_API virtual FTopLevelAssetPath GetBaseClass() const;
+	ENGINE_API virtual FTopLevelAssetPath GetNativeClass() const;
+	ENGINE_API virtual UClass* GetActorNativeClass() const;
+	ENGINE_API virtual FName GetRuntimeGrid() const;
+	ENGINE_API virtual bool GetIsSpatiallyLoaded() const;
+	ENGINE_API virtual bool GetActorIsEditorOnly() const;
+	ENGINE_API virtual bool GetActorIsRuntimeOnly() const;
+	ENGINE_API virtual bool GetActorIsHLODRelevant() const;
+	ENGINE_API virtual FSoftObjectPath GetHLODLayer() const;
+	ENGINE_API virtual bool HasResolvedDataLayerInstanceNames() const;
+	ENGINE_API virtual const TArray<FName>& GetDataLayerInstanceNames() const;
+	ENGINE_API virtual const TArray<FName>& GetTags() const;
+	ENGINE_API virtual FName GetActorPackage() const;	
+	ENGINE_API virtual FSoftObjectPath GetActorSoftPath() const;
+	ENGINE_API virtual FName GetActorLabel() const;
+	ENGINE_API virtual FBox GetEditorBounds() const;
+	ENGINE_API virtual FBox GetRuntimeBounds() const;
+	ENGINE_API virtual const TArray<FGuid>& GetReferences() const;
+	ENGINE_API virtual const TArray<FGuid>& GetEditorReferences() const;
+	ENGINE_API virtual FString ToString() const;
+	ENGINE_API virtual const FGuid& GetParentActor() const;
+	ENGINE_API virtual FName GetActorName() const;
+	ENGINE_API virtual const FGuid& GetFolderGuid() const;
+	ENGINE_API virtual FGuid GetContentBundleGuid() const;
+	ENGINE_API virtual FName GetContainerPackage() const;
+	ENGINE_API virtual bool IsContainerInstance() const;
+	ENGINE_API virtual bool GetContainerInstance(FWorldPartitionActorDesc::FContainerInstance& OutContainerInstance) const;
+	ENGINE_API virtual EWorldPartitionActorFilterType GetContainerFilterType() const;
+	ENGINE_API virtual const FWorldPartitionActorFilter* GetContainerFilter() const;	
+	ENGINE_API virtual void CheckForErrors(IStreamingGenerationErrorHandler* ErrorHandler) const;
+	ENGINE_API virtual FName GetActorLabelOrName() const;
+	ENGINE_API virtual AActor* GetActor() const;
+	ENGINE_API virtual bool IsEditorOnlyReference(const FGuid& ReferenceGuid) const;
+	ENGINE_API virtual bool GetProperty(FName PropertyName, FName* PropertyValue) const;
+	ENGINE_API virtual bool HasProperty(FName PropertyName) const;
+
+	/** Helper functions */
 	const FWorldPartitionActorDesc* GetActorDesc() const { return ActorDesc; }
-
-	ENGINE_API bool IsEditorOnlyReference(const FGuid& ReferenceGuid) const;
-
-	ENGINE_API bool GetProperty(FName PropertyName, FName* PropertyValue) const;
-	ENGINE_API bool HasProperty(FName PropertyName) const;
-
-	ENGINE_API void SetParentView(const FWorldPartitionActorDescView* InParentView);
+	ENGINE_API virtual const TArray<FName>& GetRuntimeDataLayerInstanceNames() const;
 
 protected:
 	const FWorldPartitionActorDesc* ActorDesc;
-	const FWorldPartitionActorDescView* ParentView;
-	bool bIsForcedNonSpatiallyLoaded;
-	bool bIsForcedNoRuntimeGrid;
-	bool bIsForcedNoDataLayers;
-	bool bIsForceNoHLODLayer;
-	TOptional<TArray<FName>> ResolvedDataLayerInstanceNames;
-	TOptional<TArray<FName>> RuntimeDataLayerInstanceNames;
-	TOptional<TArray<FGuid>> RuntimeReferences;
-	TOptional<FSoftObjectPath> RuntimedHLODLayer;
-	TArray<FGuid> EditorReferences;
 };
 #endif
diff --git a/Engine/Source/Runtime/Engine/Public/WorldPartition/WorldPartitionActorDescViewMap.h b/Engine/Source/Runtime/Engine/Public/WorldPartition/WorldPartitionActorDescViewMap.h
new file mode 100644
--- /dev/null
+++ b/Engine/Source/Runtime/Engine/Public/WorldPartition/WorldPartitionActorDescViewMap.h
@@ -0,0 +1,128 @@
+// Copyright Epic Games, Inc. All Rights Reserved.
+#pragma once
+
+#if WITH_EDITOR
+#include "CoreMinimal.h"
+#include "OverrideVoidReturnInvoker.h"
+#include "WorldPartition/WorldPartitionActorDescView.h"
+
+template <class Type>
+class TActorDescViewMap
+{
+	friend class FWorldPartitionStreamingGenerator;
+
+private:
+	template <class Func>
+	void ForEachActorDescView(Func InFunc)
+	{
+		TOverrideVoidReturnInvoker Invoker(true, InFunc);
+
+		for (TUniquePtr<Type>& ActorDescView : ActorDescViewList)
+		{
+			if (!Invoker(*((Type*)ActorDescView.Get())))
+			{
+				return;
+			}
+		}
+	}
+
+	Type* FindByGuid(const FGuid& InGuid)
+	{
+		if (Type** ActorDescViewPtr = ActorDescViewsByGuid.Find(InGuid))
+		{
+			return *ActorDescViewPtr;
+		}
+		return nullptr;
+	}
+
+	Type& FindByGuidChecked(const FGuid& InGuid)
+	{
+		return *ActorDescViewsByGuid.FindChecked(InGuid);
+	}
+
+public:
+	TActorDescViewMap()
+	{}
+
+	// Non-copyable but movable
+	TActorDescViewMap(const TActorDescViewMap&) = delete;
+	TActorDescViewMap(TActorDescViewMap&&) = default;
+	TActorDescViewMap& operator=(const TActorDescViewMap&) = delete;
+	TActorDescViewMap& operator=(TActorDescViewMap&&) = default;
+
+	Type* Emplace(const FGuid& InActorGuid, const Type& InActorDescView)
+	{
+		Type* NewActorDescView = (Type*)ActorDescViewList.Emplace_GetRef(MakeUnique<Type>(InActorDescView)).Get();
+	
+		const UClass* NativeClass = NewActorDescView->GetActorNativeClass();
+		const FName NativeClassName = NativeClass->GetFName();
+
+		ActorDescViewsByGuid.Emplace(InActorGuid, NewActorDescView);
+		ActorDescViewsByClass.Add(NativeClassName, NewActorDescView);
+
+		return NewActorDescView;
+	}
+
+	Type* Emplace(const FWorldPartitionActorDesc* InActorDesc)
+	{
+		return Emplace(InActorDesc->GetGuid(), Type(InActorDesc));
+	}
+
+	FORCEINLINE int32 Num() const
+	{
+		return ActorDescViewList.Num();
+	}
+
+	template <class Func>
+	void ForEachActorDescView(Func InFunc) const
+	{
+		TOverrideVoidReturnInvoker Invoker(true, InFunc);
+
+		for (const TUniquePtr<Type>& ActorDescView : ActorDescViewList)
+		{
+			if (!Invoker(*ActorDescView))
+			{
+				return;
+			}
+		}
+	}
+
+	const Type* FindByGuid(const FGuid& InGuid) const
+	{
+		if (const Type* const* ActorDescViewPtr = ActorDescViewsByGuid.Find(InGuid))
+		{
+			return *ActorDescViewPtr;
+		}
+		return nullptr;
+	}
+
+	const Type& FindByGuidChecked(const FGuid& InGuid) const
+	{
+		return *ActorDescViewsByGuid.FindChecked(InGuid);
+	}
+
+	template <class ClassType>
+	TArray<const Type*> FindByExactNativeClass() const
+	{
+		return FindByExactNativeClass(ClassType::StaticClass());
+	}
+
+	TArray<const Type*> FindByExactNativeClass(UClass* InExactNativeClass) const
+	{
+		check(InExactNativeClass->IsNative());
+		const FName NativeClassName = InExactNativeClass->GetFName();
+		TArray<const Type*> Result;
+		ActorDescViewsByClass.MultiFind(NativeClassName, Result);
+		return Result;
+	}
+
+	const TMap<FGuid, Type*>& GetActorDescViewsByGuid() const { return ActorDescViewsByGuid; }
+
+protected:
+	TArray<TUniquePtr<Type>> ActorDescViewList;
+
+	TMap<FGuid, Type*> ActorDescViewsByGuid;
+	TMultiMap<FName, const Type*> ActorDescViewsByClass;
+};
+
+#endif
diff --git a/Engine/Source/Runtime/Engine/Public/WorldPartition/WorldPartitionStreamingGeneration.h b/Engine/Source/Runtime/Engine/Public/WorldPartition/WorldPartitionStreamingGeneration.h
--- a/Engine/Source/Runtime/Engine/Public/WorldPartition/WorldPartitionStreamingGeneration.h
+++ b/Engine/Source/Runtime/Engine/Public/WorldPartition/WorldPartitionStreamingGeneration.h
@@ -3,104 +3,8 @@
 
 #if WITH_EDITOR
 #include "CoreMinimal.h"
-#include "OverrideVoidReturnInvoker.h"
 #include "WorldPartition/ActorDescContainerCollection.h"
-#include "WorldPartition/WorldPartitionActorDescView.h"
-
-class FActorDescViewMap
-{
-	friend class FWorldPartitionStreamingGenerator;
-
-private:
-	template <class Func>
-	void ForEachActorDescView(Func InFunc)
-	{
-		TOverrideVoidReturnInvoker Invoker(true, InFunc);
-
-		for (TUniquePtr<FWorldPartitionActorDescView>& ActorDescView : ActorDescViewList)
-		{
-			if (!Invoker(*ActorDescView))
-			{
-				return;
-			}
-		}
-	}
-
-	FWorldPartitionActorDescView* FindByGuid(const FGuid& InGuid)
-	{
-		if (FWorldPartitionActorDescView** ActorDescViewPtr = ActorDescViewsByGuid.Find(InGuid))
-		{
-			return *ActorDescViewPtr;
-		}
-		return nullptr;
-	}
-
-	FWorldPartitionActorDescView& FindByGuidChecked(const FGuid& InGuid)
-	{
-		return *ActorDescViewsByGuid.FindChecked(InGuid);
-	}
-
-public:
-	ENGINE_API FActorDescViewMap();
-
-	// Non-copyable but movable
-	FActorDescViewMap(const FActorDescViewMap&) = delete;
-	FActorDescViewMap(FActorDescViewMap&&) = default;
-	FActorDescViewMap& operator=(const FActorDescViewMap&) = delete;
-	FActorDescViewMap& operator=(FActorDescViewMap&&) = default;
-
-	ENGINE_API FWorldPartitionActorDescView* Emplace(const FGuid& InActorGuid, const FWorldPartitionActorDescView& InActorDescView);
-	ENGINE_API FWorldPartitionActorDescView* Emplace(const FWorldPartitionActorDesc* InActorDesc);
-
-	FORCEINLINE int32 Num() const
-	{
-		return ActorDescViewList.Num();
-	}
-
-	template <class Func>
-	void ForEachActorDescView(Func InFunc) const
-	{
-		TOverrideVoidReturnInvoker Invoker(true, InFunc);
-
-		for (const TUniquePtr<FWorldPartitionActorDescView>& ActorDescView : ActorDescViewList)
-		{
-			if (!Invoker(*ActorDescView))
-			{
-				return;
-			}
-		}
-	}
-
-	const FWorldPartitionActorDescView* FindByGuid(const FGuid& InGuid) const
-	{
-		if (const FWorldPartitionActorDescView* const* ActorDescViewPtr = ActorDescViewsByGuid.Find(InGuid))
-		{
-			return *ActorDescViewPtr;
-		}
-		return nullptr;
-	}
-
-	const FWorldPartitionActorDescView& FindByGuidChecked(const FGuid& InGuid) const
-	{
-		return *ActorDescViewsByGuid.FindChecked(InGuid);
-	}
-
-	template <class ClassType>
-	TArray<const FWorldPartitionActorDescView*> FindByExactNativeClass() const
-	{
-		return FindByExactNativeClass(ClassType::StaticClass());
-	}
-
-	ENGINE_API TArray<const FWorldPartitionActorDescView*> FindByExactNativeClass(UClass* InExactNativeClass) const;
-
-	const TMap<FGuid, FWorldPartitionActorDescView*>& GetActorDescViewsByGuid() const { return ActorDescViewsByGuid; }
-
-protected:
-	TArray<TUniquePtr<FWorldPartitionActorDescView>> ActorDescViewList;
-
-	TMap<FGuid, FWorldPartitionActorDescView*> ActorDescViewsByGuid;
-	TMultiMap<FName, const FWorldPartitionActorDescView*> ActorDescViewsByClass;
-};
+#include "WorldPartition/WorldPartitionActorDescViewMap.h"
 
 class FStreamingGenerationActorDescCollection : public TActorDescContainerCollection<TObjectPtr<const UActorDescContainer>>
 {
@@ -128,4 +32,58 @@ private:
 	static constexpr int ExternalDataLayerContainerStartIdx = MainContainerIdx + 1;
 };
 
+class FStreamingGenerationActorDescView : public FWorldPartitionActorDescView
+{
+public:
+	ENGINE_API FStreamingGenerationActorDescView();
+	ENGINE_API FStreamingGenerationActorDescView(const FWorldPartitionActorDesc* InActorDesc);
+
+	//~ Begin FWorldPartitionActorDescView interface
+	ENGINE_API virtual FName GetRuntimeGrid() const override;
+	ENGINE_API virtual bool GetIsSpatiallyLoaded() const override;
+	ENGINE_API virtual FSoftObjectPath GetHLODLayer() const override;
+	ENGINE_API virtual const TArray<FName>& GetDataLayerInstanceNames() const override;
+	ENGINE_API virtual const TArray<FGuid>& GetReferences() const override;
+	ENGINE_API virtual const TArray<FGuid>& GetEditorReferences() const override;
+	//~ End FWorldPartitionActorDescView interface
+
+	ENGINE_API bool ShouldValidateRuntimeGrid() const;
+	ENGINE_API void SetParentView(const FWorldPartitionActorDescView* InParentView);
+	ENGINE_API void SetDataLayerInstanceNames(const TArray<FName>& InDataLayerInstanceNames);
+	ENGINE_API void SetForcedNonSpatiallyLoaded();
+	ENGINE_API void SetForcedNoRuntimeGrid();
+	ENGINE_API void SetForcedNoDataLayers();
+	ENGINE_API void SetRuntimeDataLayerInstanceNames(const TArray<FName>& InRuntimeDataLayerInstanceNames);
+	ENGINE_API void SetRuntimeReferences(const TArray<FGuid>& InRuntimeReferences);
+	ENGINE_API void SetEditorReferences(const TArray<FGuid>& InEditorReferences);		
+	ENGINE_API void SetForcedNoHLODLayer();
+	ENGINE_API void SetRuntimeHLODLayer(const FSoftObjectPath& InHLODLayer);
+
+	ENGINE_API const TArray<FName>& GetRuntimeDataLayerInstanceNames() const;
+
+	bool operator==(const FWorldPartitionActorDescView& Other) const
+	{
+		return GetGuid() == Other.GetGuid();
+	}
+
+	friend uint32 GetTypeHash(const FWorldPartitionActorDescView& Key)
+	{
+		return GetTypeHash(Key.GetGuid());
+	}
+
+protected:
+	const FWorldPartitionActorDescView* ParentView;
+	bool bIsForcedNonSpatiallyLoaded;
+	bool bIsForcedNoRuntimeGrid;
+	bool bIsForcedNoDataLayers;
+	bool bIsForceNoHLODLayer;
+	TOptional<TArray<FName>> ResolvedDataLayerInstanceNames;
+	TOptional<TArray<FName>> RuntimeDataLayerInstanceNames;
+	TOptional<TArray<FGuid>> RuntimeReferences;
+	TOptional<FSoftObjectPath> RuntimedHLODLayer;
+	TArray<FGuid> EditorReferences;
+};
+
+class FStreamingGenerationActorDescViewMap : public TActorDescViewMap<FStreamingGenerationActorDescView> {};
+
 #endif // WITH_EDITOR
diff --git a/Engine/Source/Runtime/Engine/Public/WorldPartition/WorldPartitionStreamingGenerationContext.h b/Engine/Source/Runtime/Engine/Public/WorldPartition/WorldPartitionStreamingGenerationContext.h
--- a/Engine/Source/Runtime/Engine/Public/WorldPartition/WorldPartitionStreamingGenerationContext.h
+++ b/Engine/Source/Runtime/Engine/Public/WorldPartition/WorldPartitionStreamingGenerationContext.h
@@ -5,7 +5,7 @@
 #include "OverrideVoidReturnInvoker.h"
 #include "WorldPartition/WorldPartitionActorContainerID.h"
 
-class FActorDescViewMap;
+class FStreamingGenerationActorDescViewMap;
 class FWorldPartitionStreamingGenerator;
 class FStreamingGenerationActorDescCollection;
 class FWorldPartitionActorDescView;
@@ -42,7 +42,7 @@ public:
 		FActorSetContainer(const FActorSetContainer&) = delete;
 		FActorSetContainer& operator=(const FActorSetContainer&) = delete;
 
-		const FActorDescViewMap* ActorDescViewMap;
+		const FStreamingGenerationActorDescViewMap* ActorDescViewMap;
 		const FStreamingGenerationActorDescCollection* ActorDescCollection; // Only used by UWorldPartitionRuntimeSpatialHash::SetupHLODActors
 		TArray<TUniquePtr<FActorSet>> ActorSets;
 	};
