[NetPhysics] Allow a NetworkPhysicsComponent to be implemented on AActors instead of only on APawns.

diff --git a/Engine/Source/Runtime/Engine/Private/PhysicsEngine/NetworkPhysicsComponent.cpp b/Engine/Source/Runtime/Engine/Private/PhysicsEngine/NetworkPhysicsComponent.cpp
--- a/Engine/Source/Runtime/Engine/Private/PhysicsEngine/NetworkPhysicsComponent.cpp
+++ b/Engine/Source/Runtime/Engine/Private/PhysicsEngine/NetworkPhysicsComponent.cpp
@@ -176,7 +176,7 @@ int32 FNetworkPhysicsCallback::TriggerRewindIfNeeded_Internal(int32 LatestStepCo
 			const int32 ReplicationFrame = RewindData->GetResimFrame();
 
 #if DEBUG_NETWORK_PHYSICS || DEBUG_REWIND_DATA
-			UE_LOG(LogTemp, Log, TEXT("CLIENT | PT | TriggerRewindIfNeeded_Internal | Replication Frame = %d"), ReplicationFrame);
+			UE_LOG(LogChaos, Log, TEXT("CLIENT | PT | TriggerRewindIfNeeded_Internal | Replication Frame = %d"), ReplicationFrame);
 #endif
 			ResimFrame = (ResimFrame == INDEX_NONE) ? ReplicationFrame : (ReplicationFrame == INDEX_NONE) ? ResimFrame : FMath::Min(ReplicationFrame, ResimFrame);
 			RewindData->SetResimFrame(INDEX_NONE);
@@ -186,7 +186,7 @@ int32 FNetworkPhysicsCallback::TriggerRewindIfNeeded_Internal(int32 LatestStepCo
 		{
 			const int32 ValidFrame = RewindData->FindValidResimFrame(ResimFrame);
 #if DEBUG_NETWORK_PHYSICS || DEBUG_REWIND_DATA
-			UE_LOG(LogTemp, Log, TEXT("CLIENT | PT | TriggerRewindIfNeeded_Internal | Resim Frame = %d | Valid Frame = %d"), ResimFrame, ValidFrame);
+			UE_LOG(LogChaos, Log, TEXT("CLIENT | PT | TriggerRewindIfNeeded_Internal | Resim Frame = %d | Valid Frame = %d"), ResimFrame, ValidFrame);
 #endif
 			ResimFrame = ValidFrame;
 		}
@@ -211,7 +211,7 @@ void FNetworkPhysicsCallback::UpdateClientPlayer_External(int32 PhysicsStep)
 
 		if (InputCmdCVars::TimeDilationEnabled > 0)
 		{
-			UE_LOG(LogPhysics, Warning, TEXT("p.net.TimeDilationEnabled is set to true, this CVar is deprecated in UE5.4 and does not affect Time Dilation. Time Dilation is automatically used via the PlayerController if Physics Prediction is enabled in Project Settings. It's also recommended to disable the legacy flow that handled physics frame offset and this time dilation by setting: p.net.CmdOffsetEnabled = 0"));
+			UE_LOG(LogChaos, Warning, TEXT("p.net.TimeDilationEnabled is set to true, this CVar is deprecated in UE5.4 and does not affect Time Dilation. Time Dilation is automatically used via the PlayerController if Physics Prediction is enabled in Project Settings. It's also recommended to disable the legacy flow that handled physics frame offset and this time dilation by setting: p.net.CmdOffsetEnabled = 0"));
 		}
 	}
 }
@@ -242,7 +242,7 @@ void FNetworkPhysicsCallback::UpdateServerPlayer_External(int32 PhysicsStep)
 			// Check Overflow
 			if (NumBufferedInputCmds > InputCmdCVars::MaxBufferedCmds)
 			{
-				UE_LOG(LogPhysics, Warning, TEXT("[Remote.Input] overflow %d %d -> %d"), InputBuffer.HeadFrame(), FrameInfo.LastProcessedInputFrame, NumBufferedInputCmds);
+				UE_LOG(LogChaos, Warning, TEXT("[Remote.Input] overflow %d %d -> %d"), InputBuffer.HeadFrame(), FrameInfo.LastProcessedInputFrame, NumBufferedInputCmds);
 				FrameInfo.LastProcessedInputFrame = InputBuffer.HeadFrame() - InputCmdCVars::MaxBufferedCmds + 1;
 			}
 			// Check fault - we are waiting for Cmds to reach TargetNumBufferedCmds before continuing
@@ -419,7 +419,7 @@ void UNetworkPhysicsComponent::BeginPlay()
 				}
 				else
 				{
-					UE_LOG(LogPhysics, Warning, TEXT("A NetworkPhysicsComponent is trying to set up but 'Project Settings -> Physics -> Physics Prediction' is not enabled. The component might not work as intended."));
+					UE_LOG(LogChaos, Warning, TEXT("A NetworkPhysicsComponent is trying to set up but 'Project Settings -> Physics -> Physics Prediction' is not enabled. The component might not work as intended."));
 				}
 			}
 		}
@@ -464,15 +464,15 @@ void UNetworkPhysicsComponent::AsyncPhysicsTickComponent(float DeltaTime, float
 
 	Super ::AsyncPhysicsTickComponent(DeltaTime, SimTime);
 #if DEBUG_NETWORK_PHYSICS
-	if(HasServerWorld() && !HasLocalController() && InputsHistory)
+	if(HasServerWorld() && !IsLocallyControlled() && InputsHistory)
 	{
 		TArray<int32> LocalFrames, ServerFrames, InputFrames;
 		InputsHistory->DebugDatas(*ReplicatedInputs.History, LocalFrames, ServerFrames, InputFrames);
 
-		UE_LOG(LogTemp, Log, TEXT("SERVER | PT | AsyncPhysicsTickComponent | Receiving %d inputs from CLIENT | Component = %s"), LocalFrames.Num(), *GetFullName());
+		UE_LOG(LogChaos, Log, TEXT("SERVER | PT | AsyncPhysicsTickComponent | Receiving %d inputs from CLIENT | Component = %s"), LocalFrames.Num(), *GetFullName());
 		for (int32 FrameIndex = 0; FrameIndex < LocalFrames.Num(); ++FrameIndex)
 		{
-			UE_LOG(LogTemp, Log, TEXT("		Debugging replicated inputs at local frame = %d | server frame = %d | Component = %s"),
+			UE_LOG(LogChaos, Log, TEXT("		Debugging replicated inputs at local frame = %d | server frame = %d | Component = %s"),
 				LocalFrames[FrameIndex], ServerFrames[FrameIndex], *GetFullName());
 		}
 	}
@@ -501,16 +501,16 @@ void UNetworkPhysicsComponent::AsyncPhysicsTickComponent(float DeltaTime, float
 
 void UNetworkPhysicsComponent::SendLocalInputsDatas()
 {
-	const APlayerController* PlayerController = GetPlayerController();
-	if (!PlayerController)
+	if (IsLocallyControlled() && InputsHistory)
 	{
-		return;
-	}
+		const APlayerController* PlayerController = GetPlayerController();
+		if (!PlayerController)
+		{
+			PlayerController = GetWorld()->GetFirstPlayerController();
+		}
 
-	if (PlayerController->IsLocalController() && InputsHistory)
-	{
 		// We just check that the local client to server offset is valid before doing something
-		if (HasServerWorld() || PlayerController->GetNetworkPhysicsTickOffsetAssigned())
+		if (PlayerController && (HasServerWorld() || PlayerController->GetNetworkPhysicsTickOffsetAssigned()))
 		{
 			const int32 LocalOffset = HasServerWorld() ? 0 : PlayerController->GetNetworkPhysicsTickOffset();
 			const int32 NextIndex = (InputsIndex + 1) % (InputsRedundancy + 1);
@@ -526,10 +526,10 @@ void UNetworkPhysicsComponent::SendLocalInputsDatas()
 				TArray<int32> LocalFrames, ServerFrames, InputFrames;
 				InputsHistory->DebugDatas(*ReplicatedInputs.History, LocalFrames, ServerFrames, InputFrames);
 
-				UE_LOG(LogTemp, Log, TEXT("CLIENT | GT | SendLocalInputsDatas | Sending %d inputs from CLIENT | Component = %s"), LocalFrames.Num(), *GetFullName());
+				UE_LOG(LogChaos, Log, TEXT("CLIENT | GT | SendLocalInputsDatas | Sending %d inputs from CLIENT | Component = %s"), LocalFrames.Num(), *GetFullName());
 				for (int32 FrameIndex = 0; FrameIndex < LocalFrames.Num(); ++FrameIndex)
 				{
-					UE_LOG(LogTemp, Log, TEXT("		Debugging local inputs at local frame = %d | server frame = %d | Current local frame = %d | Current server frame = %d"),
+					UE_LOG(LogChaos, Log, TEXT("		Debugging local inputs at local frame = %d | server frame = %d | Current local frame = %d | Current server frame = %d"),
 						LocalFrames[FrameIndex], ServerFrames[FrameIndex], Timestamp.LocalFrame, Timestamp.ServerFrame);
 				}
 #endif
@@ -556,7 +556,7 @@ void UNetworkPhysicsComponent::SendLocalStatesDatas()
 /* Deprecated 5.4 */
 void UNetworkPhysicsComponent::CorrectServerToLocalOffset(const int32 LocalToServerOffset)
 {
-	if (HasLocalController() && !HasServerWorld() && StatesHistory)
+	if (IsLocallyControlled() && !HasServerWorld() && StatesHistory)
 	{
 		TArray<int32> LocalFrames, ServerFrames, InputFrames;
 		StatesHistory->DebugDatas(*ReplicatedStates.History, LocalFrames, ServerFrames, InputFrames);
@@ -564,7 +564,7 @@ void UNetworkPhysicsComponent::CorrectServerToLocalOffset(const int32 LocalToSer
 		for (int32 FrameIndex = 0; FrameIndex < LocalFrames.Num(); ++FrameIndex)
 		{
 #if DEBUG_NETWORK_PHYSICS || DEBUG_REWIND_DATA
-			UE_LOG(LogTemp, Log, TEXT("CLIENT | GT | CorrectServerToLocalOffset | Server frame = %d | Client Frame = %d"), ServerFrames[FrameIndex], InputFrames[FrameIndex]);
+			UE_LOG(LogChaos, Log, TEXT("CLIENT | GT | CorrectServerToLocalOffset | Server frame = %d | Client Frame = %d"), ServerFrames[FrameIndex], InputFrames[FrameIndex]);
 #endif
 			ServerToLocalOffset = FMath::Min(ServerToLocalOffset, ServerFrames[FrameIndex] - InputFrames[FrameIndex]);
 		}
@@ -573,22 +573,22 @@ void UNetworkPhysicsComponent::CorrectServerToLocalOffset(const int32 LocalToSer
 		PRAGMA_ENABLE_DEPRECATION_WARNINGS
 
 #if DEBUG_NETWORK_PHYSICS || DEBUG_REWIND_DATA
-		UE_LOG(LogTemp, Log, TEXT("CLIENT | GT | CorrectServerToLocalOffset | Server to local offset = %d | Local to server offset = %d"), ServerToLocalOffset, LocalToServerOffset);
+		UE_LOG(LogChaos, Log, TEXT("CLIENT | GT | CorrectServerToLocalOffset | Server to local offset = %d | Local to server offset = %d"), ServerToLocalOffset, LocalToServerOffset);
 #endif
 	}
 }
 
 void UNetworkPhysicsComponent::OnRep_SetReplicatedStates()
 {
-	// The replicated states should only be used on the client since the server already have authoritative local ones
-	if (!HasServerWorld() && StatesHistory)
+	APlayerController* PlayerController = GetPlayerController();
+	if (!PlayerController)
 	{
-		APlayerController* PlayerController = GetPlayerController();
-		if (!PlayerController)
-		{
-			PlayerController = GetWorld()->GetFirstPlayerController();
-		}
+		PlayerController = GetWorld()->GetFirstPlayerController();
+	}
 
+	// The replicated states should only be used on the client since the server already have authoritative local ones
+	if (PlayerController && !HasServerWorld() && StatesHistory)
+	{
 		const int32 LocalOffset = PlayerController->GetNetworkPhysicsTickOffset();
 
 		/* Deprecated 5.4 */
@@ -614,14 +614,14 @@ void UNetworkPhysicsComponent::OnRep_SetReplicatedStates()
 						TArray<int32> LocalFrames, ServerFrames, InputFrames;
 						StatesHistory->DebugDatas(*ReceivedStates, LocalFrames, ServerFrames, InputFrames);
 
-						UE_LOG(LogTemp, Log, TEXT("CLIENT | PT | OnRep_SetReplicatedStates | Receiving %d states from SERVER | Local offset = %d | Component = %s "), LocalFrames.Num(), LocalOffset, *GetFullName());
+						UE_LOG(LogChaos, Log, TEXT("CLIENT | PT | OnRep_SetReplicatedStates | Receiving %d states from SERVER | Local offset = %d | Component = %s "), LocalFrames.Num(), LocalOffset, *GetFullName());
 						for (int32 FrameIndex = 0; FrameIndex < LocalFrames.Num(); ++FrameIndex)
 						{
-							UE_LOG(LogTemp, Log, TEXT("		Recording replicated states at local frame = %d | server frame = %d | life time = %d | Component = %s"), ServerFrames[FrameIndex] - LocalOffset, ServerFrames[FrameIndex], InputFrames[FrameIndex], *GetFullName());
+							UE_LOG(LogChaos, Log, TEXT("		Recording replicated states at local frame = %d | server frame = %d | life time = %d | Component = %s"), ServerFrames[FrameIndex] - LocalOffset, ServerFrames[FrameIndex], InputFrames[FrameIndex], *GetFullName());
 
-							if (HasLocalController() && (InputFrames[FrameIndex] != (ServerFrames[FrameIndex] - LocalOffset)))
+							if (IsLocallyControlled() && (InputFrames[FrameIndex] != (ServerFrames[FrameIndex] - LocalOffset)))
 							{
-								UE_LOG(LogTemp, Log, TEXT("		Bad local frame compared to input frame!!!"));
+								UE_LOG(LogChaos, Log, TEXT("		Bad local frame compared to input frame!!!"));
 							}
 						}
 					}
@@ -634,15 +634,15 @@ void UNetworkPhysicsComponent::OnRep_SetReplicatedStates()
 
 void UNetworkPhysicsComponent::OnRep_SetReplicatedInputs()
 {
-	// For local controller we should already have correct replicated inputs
-	if (!HasLocalController() && !HasServerWorld() && InputsHistory)
+	APlayerController* PlayerController = GetPlayerController();
+	if (!PlayerController)
 	{
-		APlayerController* PlayerController = GetPlayerController();
-		if(!PlayerController)
-		{
-			PlayerController = GetWorld()->GetFirstPlayerController();
-		}
+		PlayerController = GetWorld()->GetFirstPlayerController();
+	}
 
+	// For local controller we should already have correct replicated inputs
+	if (PlayerController && !IsLocallyControlled() && !HasServerWorld() && InputsHistory)
+	{
 		const int32 LocalOffset = PlayerController->GetNetworkPhysicsTickOffset();
 
 		// Record the received inputs from the server into the history for future use
@@ -660,10 +660,10 @@ void UNetworkPhysicsComponent::OnRep_SetReplicatedInputs()
 						TArray<int32> LocalFrames, ServerFrames, InputFrames;
 						InputsHistory->DebugDatas(*ReceivedInputs, LocalFrames, ServerFrames, InputFrames);
 
-						UE_LOG(LogTemp, Log, TEXT("CLIENT | PT | OnRep_SetReplicatedInputs | Receiving %d inputs from SERVER | Local offset = %d | Component = %s"), LocalFrames.Num(), LocalOffset, *GetFullName());
+						UE_LOG(LogChaos, Log, TEXT("CLIENT | PT | OnRep_SetReplicatedInputs | Receiving %d inputs from SERVER | Local offset = %d | Component = %s"), LocalFrames.Num(), LocalOffset, *GetFullName());
 						for (int32 FrameIndex = 0; FrameIndex < LocalFrames.Num(); ++FrameIndex)
 						{
-							UE_LOG(LogTemp, Log, TEXT("		Recording replicated inputs at local frame = %d | server frame = %d | Component = %s"), ServerFrames[FrameIndex] - LocalOffset, ServerFrames[FrameIndex], *GetFullName());
+							UE_LOG(LogChaos, Log, TEXT("		Recording replicated inputs at local frame = %d | server frame = %d | Component = %s"), ServerFrames[FrameIndex] - LocalOffset, ServerFrames[FrameIndex], *GetFullName());
 						}
 					}
 #endif
@@ -675,14 +675,11 @@ void UNetworkPhysicsComponent::OnRep_SetReplicatedInputs()
 
 void UNetworkPhysicsComponent::ServerReceiveInputsDatas_Implementation(const FNetworkPhysicsRewindDataInputProxy& ClientInputs)
 {
-	if(InputsHistory)
+	if (InputsHistory)
 	{ 
 		// We could probably skip that test since the server RPC is on server
 		ensure(HasServerWorld());
 
-		// We could probably skip that test since the server RPC is on server
-		ensure(!HasLocalController());
-
 		// Record the received inputs from the client into the history for future use
 		ReplicatedInputs.History = ClientInputs.History->Clone();
 
@@ -704,10 +701,10 @@ void UNetworkPhysicsComponent::ServerReceiveInputsDatas_Implementation(const FNe
 						const int32 CurrentFrame = PhysScene->GetSolver()->GetCurrentFrame();
 
 						const int32 EvalOffset = CurrentFrame - InputFrames[InputFrames.Num()-1] + 4;
-						UE_LOG(LogTemp, Log, TEXT("SERVER | PT | ServerReceiveInputsDatas | Receiving %d inputs from CLIENT | Inputs frame = %d | Server frame = %d | Eval Offset = %d | Component = %s"), LocalFrames.Num(), InputFrames[InputFrames.Num() - 1], CurrentFrame, EvalOffset, *GetFullName());
+						UE_LOG(LogChaos, Log, TEXT("SERVER | PT | ServerReceiveInputsDatas | Receiving %d inputs from CLIENT | Inputs frame = %d | Server frame = %d | Eval Offset = %d | Component = %s"), LocalFrames.Num(), InputFrames[InputFrames.Num() - 1], CurrentFrame, EvalOffset, *GetFullName());
 						for (int32 FrameIndex = 0; FrameIndex < LocalFrames.Num(); ++FrameIndex)
 						{
-							UE_LOG(LogTemp, Log, TEXT("		Recording replicated inputs at local frame = %d | server frame = %d | Solver offset = %d | Component = %s"), 
+							UE_LOG(LogChaos, Log, TEXT("		Recording replicated inputs at local frame = %d | server frame = %d | Solver offset = %d | Component = %s"), 
 								LocalFrames[FrameIndex], ServerFrames[FrameIndex], ServerFrames[FrameIndex] - LocalFrames[FrameIndex], *GetFullName());
 						}
 					}
@@ -723,11 +720,11 @@ void UNetworkPhysicsComponent::OnPreProcessInputsInternal(const int32 PhysicsSte
 #if DEBUG_NETWORK_PHYSICS
 	if (HasServerWorld())
 	{
-		UE_LOG(LogTemp, Log, TEXT("SERVER | PT | OnPreProcessInputsInternal | At Frame %d | Component = %s"), PhysicsStep, *GetFullName());
+		UE_LOG(LogChaos, Log, TEXT("SERVER | PT | OnPreProcessInputsInternal | At Frame %d | Component = %s"), PhysicsStep, *GetFullName());
 	}
 	else
 	{
-		UE_LOG(LogTemp, Log, TEXT("CLIENT | PT | OnPreProcessInputsInternal | At Frame %d | Component = %s"), PhysicsStep, *GetFullName());
+		UE_LOG(LogChaos, Log, TEXT("CLIENT | PT | OnPreProcessInputsInternal | At Frame %d | Component = %s"), PhysicsStep, *GetFullName());
 	}
 #endif
 
@@ -746,12 +743,12 @@ void UNetworkPhysicsComponent::OnPreProcessInputsInternal(const int32 PhysicsSte
 		}
 
 		// for the inputs client local ones are ground truth otherwise use the replicated ones coming from the server
-		if (!HasLocalController() || bIsSolverResim)
+		if (!IsLocallyControlled() || bIsSolverResim)
 		{
 			FNetworkPhysicsDatas* PhysicsDatas = InputsDatas.Get();
 			PhysicsDatas->LocalFrame = PhysicsStep;
 	#if DEBUG_NETWORK_PHYSICS
-			UE_LOG(LogTemp, Log, TEXT("		Extracting history inputs at frame %d | Component = %s"), PhysicsStep, *GetFullName());
+			UE_LOG(LogChaos, Log, TEXT("		Extracting history inputs at frame %d | Component = %s"), PhysicsStep, *GetFullName());
 	#endif
 			if (InputsHistory->ExtractDatas(PhysicsStep, bIsSolverReset, PhysicsDatas))
 			{ 
@@ -764,7 +761,7 @@ void UNetworkPhysicsComponent::OnPreProcessInputsInternal(const int32 PhysicsSte
 			FNetworkPhysicsDatas* PhysicsDatas = StatesDatas.Get();
 			PhysicsDatas->LocalFrame = PhysicsStep;
 	#if DEBUG_NETWORK_PHYSICS
-			UE_LOG(LogTemp, Log, TEXT("		Extracting history states at frame %d | Component = %s"), PhysicsStep, *GetFullName());
+			UE_LOG(LogChaos, Log, TEXT("		Extracting history states at frame %d | Component = %s"), PhysicsStep, *GetFullName());
 	#endif
 			if (StatesHistory->ExtractDatas(PhysicsStep, bIsSolverReset, PhysicsDatas, true))
 			{
@@ -781,11 +778,11 @@ void UNetworkPhysicsComponent::OnPostProcessInputsInternal(const int32 PhysicsSt
 #if DEBUG_NETWORK_PHYSICS
 	if (HasServerWorld())
 	{
-		UE_LOG(LogTemp, Log, TEXT("SERVER | PT | OnPostProcessInputsInternal | At Frame %d | Component = %s"), PhysicsStep, *GetFullName());
+		UE_LOG(LogChaos, Log, TEXT("SERVER | PT | OnPostProcessInputsInternal | At Frame %d | Component = %s"), PhysicsStep, *GetFullName());
 	}
 	else
 	{
-		UE_LOG(LogTemp, Log, TEXT("CLIENT | PT | OnPostProcessInputsInternal | At Frame %d | Component = %s"), PhysicsStep, *GetFullName());
+		UE_LOG(LogChaos, Log, TEXT("CLIENT | PT | OnPostProcessInputsInternal | At Frame %d | Component = %s"), PhysicsStep, *GetFullName());
 	}
 #endif
 
@@ -800,12 +797,18 @@ void UNetworkPhysicsComponent::OnPostProcessInputsInternal(const int32 PhysicsSt
 			}
 		}
 
+		APlayerController* PlayerController = GetPlayerController();
+		if (!PlayerController)
+		{
+			PlayerController = GetWorld()->GetFirstPlayerController();
+		}
+
 		// for the inputs client local ones are ground truth otherwise use the replicated ones coming from the server
-		if (HasLocalController() && !bIsSolverResim && (InputsDatas != nullptr))
+		if (PlayerController && IsLocallyControlled() && !bIsSolverResim && (InputsDatas != nullptr))
 		{
 			FNetworkPhysicsDatas* PhysicsDatas = InputsDatas.Get();
 			PhysicsDatas->LocalFrame = PhysicsStep;
-			PhysicsDatas->ServerFrame = HasServerWorld() ? PhysicsStep : PhysicsStep + GetPlayerController()->GetNetworkPhysicsTickOffset();
+			PhysicsDatas->ServerFrame = HasServerWorld() ? PhysicsStep : PhysicsStep + PlayerController->GetNetworkPhysicsTickOffset();
 			PhysicsDatas->InputFrame = PhysicsStep;
 
 			PhysicsDatas->BuildDatas(ActorComponent);
@@ -814,7 +817,7 @@ void UNetworkPhysicsComponent::OnPostProcessInputsInternal(const int32 PhysicsSt
 			InputsHistory->RecordDatas(PhysicsStep, PhysicsDatas);
 
 #if DEBUG_NETWORK_PHYSICS
-			UE_LOG(LogTemp, Log, TEXT("		Recording local inputs at frame %d | Component = %s"), PhysicsDatas->LocalFrame, *GetFullName());
+			UE_LOG(LogChaos, Log, TEXT("		Recording local inputs at frame %d | Component = %s"), PhysicsDatas->LocalFrame, *GetFullName());
 #endif
 		}
 
@@ -841,7 +844,7 @@ void UNetworkPhysicsComponent::OnPostProcessInputsInternal(const int32 PhysicsSt
 			StatesHistory->RecordDatas(PhysicsStep, PhysicsDatas);
 
 #if DEBUG_NETWORK_PHYSICS
-			UE_LOG(LogTemp, Log, TEXT("		Recording local states at frame %d | from input frame = %d | Component = %s"), PhysicsDatas->LocalFrame, PhysicsDatas->InputFrame, *GetFullName());
+			UE_LOG(LogChaos, Log, TEXT("		Recording local states at frame %d | from input frame = %d | Component = %s"), PhysicsDatas->LocalFrame, PhysicsDatas->InputFrame, *GetFullName());
 #endif
 		}
 	}
@@ -861,6 +864,20 @@ bool UNetworkPhysicsComponent::HasLocalController() const
 	return false;
 }
 
+bool UNetworkPhysicsComponent::IsLocallyControlled() const
+{
+	if (bIsLocallyPossessed && !GetWorld()->IsNetMode(NM_DedicatedServer))
+	{
+		return true;
+	}
+
+	if (APlayerController* PlayerController = GetPlayerController())
+	{
+		return PlayerController->IsLocalController();
+	}
+	return false;
+}
+
 APlayerController* UNetworkPhysicsComponent::GetPlayerController() const
 {
 	if (APlayerController* PC = Cast<APlayerController>(GetOwner()))
diff --git a/Engine/Source/Runtime/Engine/Public/Physics/NetworkPhysicsComponent.h b/Engine/Source/Runtime/Engine/Public/Physics/NetworkPhysicsComponent.h
--- a/Engine/Source/Runtime/Engine/Public/Physics/NetworkPhysicsComponent.h
+++ b/Engine/Source/Runtime/Engine/Public/Physics/NetworkPhysicsComponent.h
@@ -88,7 +88,7 @@ struct TNetRewindHistory : public Chaos::TDatasRewindHistory<DatasType>
 		
 		if (Super::NumFrames > GetMaxArraySize())
 		{
-			UE_LOG(LogPhysics, Warning, TEXT("TNetRewindHistory: serialized array of size %d exceeds maximum size %d."), Super::NumFrames, GetMaxArraySize());
+			UE_LOG(LogChaos, Warning, TEXT("TNetRewindHistory: serialized array of size %d exceeds maximum size %d."), Super::NumFrames, GetMaxArraySize());
 			Ar.SetError();
 			return;
 		}
@@ -146,7 +146,7 @@ private :
 		}
 		else
 		{
-			UE_LOG(LogPhysics, Error, TEXT("TNetRewindHistory::NetSerializeDatas called on data struct %s without a native NetSerialize"), *ScriptStruct->GetName());
+			UE_LOG(LogChaos, Error, TEXT("TNetRewindHistory::NetSerializeDatas called on data struct %s without a native NetSerialize"), *ScriptStruct->GetName());
 
 			// Not working for now since the packagemap could be null
 			// UNetConnection* Connection = CastChecked<UPackageMapClient>(PackageMap)->GetConnection();
@@ -437,7 +437,21 @@ public:
 	ENGINE_API bool HasServerWorld() const;
 
 	// Check if the player controller exists and is local
+	UE_DEPRECATED(5.4, "Deprecated, use IsLocallyControlled() which takes both local player controlled and local possession into account.")
 	ENGINE_API bool HasLocalController() const;
+	
+	// Check if this is controlled locally through possession or an existing local player controller
+	ENGINE_API bool IsLocallyControlled() const;
+
+	/** Mark this as controlled through local possession rather than controlled as a pawn
+	* Set if NetworkPhysicsComponent is implemented on an AActor instead of APawn and it's currently being fed inputs from the local player */
+	ENGINE_API void SetIsLocallyPossessed(bool bPossess)
+	{
+		bIsLocallyPossessed = bPossess;
+	}
+
+	/** Check if this is controlled locally through possession. It's recommended to use HasLocalController() when checking if this is locally controlled. */
+	ENGINE_API const bool GetIsLocallyPossessed() const { return bIsLocallyPossessed; }
 
 protected : 
 
@@ -498,6 +512,9 @@ private:
 
 	// Actor component that will be used to fill the histories
 	TObjectPtr<UActorComponent> ActorComponent;
+
+	// Locally possessed makes this component act as if it's a locally controlled pawn.
+	bool bIsLocallyPossessed = false;
 };
 
 template<typename PhysicsTraits>
