Fix a crash when using USkeletalMesh::LockPropertiesUntil. The FEvent ownership is now handle by the LockPropertiesUntil to make sure the FEvent is not release right after the trigger which can crash the FEvent::Wait if the event is destroy before the Wait in the thread is notify by the FEvent::Trigger.

The API change and the USkeletalMesh::LockPropertiesUntil now return a FEvent instead of receiving one in parameter.

#rnx

diff --git a/Engine/Plugins/Interchange/Runtime/Source/Import/Private/Mesh/InterchangeSkeletalMeshFactory.cpp b/Engine/Plugins/Interchange/Runtime/Source/Import/Private/Mesh/InterchangeSkeletalMeshFactory.cpp
--- a/Engine/Plugins/Interchange/Runtime/Source/Import/Private/Mesh/InterchangeSkeletalMeshFactory.cpp
+++ b/Engine/Plugins/Interchange/Runtime/Source/Import/Private/Mesh/InterchangeSkeletalMeshFactory.cpp
@@ -990,8 +990,7 @@ UInterchangeFactoryBase::FImportAssetResult UInterchangeSkeletalMeshFactory::Beg
 	//Lock the skeletalmesh properties if the skeletal mesh already exist (re-import)
 	if (ExistingAsset)
 	{
-		SkeletalMeshLockPropertiesEvent = FPlatformProcess::GetSynchEventFromPool();
-		SkeletalMesh->LockPropertiesUntil(SkeletalMeshLockPropertiesEvent);
+		SkeletalMeshLockPropertiesEvent = SkeletalMesh->LockPropertiesUntil();
 	}
 
 	ImportAssetResult.ImportedObject = SkeletalMesh;
@@ -1795,7 +1794,6 @@ void UInterchangeSkeletalMeshFactory::Cancel()
 	if (SkeletalMeshLockPropertiesEvent)
 	{
 		SkeletalMeshLockPropertiesEvent->Trigger();
-		FPlatformProcess::ReturnSynchEventToPool(SkeletalMeshLockPropertiesEvent);
 		SkeletalMeshLockPropertiesEvent = nullptr;
 	}
 }
@@ -1817,7 +1815,6 @@ void UInterchangeSkeletalMeshFactory::SetupObject_GameThread(const FSetupObjectP
 		if (SkeletalMeshLockPropertiesEvent)
 		{
 			SkeletalMeshLockPropertiesEvent->Trigger();
-			FPlatformProcess::ReturnSynchEventToPool(SkeletalMeshLockPropertiesEvent);
 			SkeletalMeshLockPropertiesEvent = nullptr;
 		}
 
diff --git a/Engine/Source/Runtime/Engine/Classes/Engine/SkeletalMesh.h b/Engine/Source/Runtime/Engine/Classes/Engine/SkeletalMesh.h
--- a/Engine/Source/Runtime/Engine/Classes/Engine/SkeletalMesh.h
+++ b/Engine/Source/Runtime/Engine/Classes/Engine/SkeletalMesh.h
@@ -2373,10 +2373,11 @@ public:
 
 	/**
 	* Lock the skeletalmesh properties until the caller trigger the Event parameter.
+	* The returned Event memory is all handle by the LockPropertiesUntil function, do not  return the event to the pool
 	* 
-	* @Param Event - When the caller will trigger the event the skeletal mesh properties will be unlocked
+	* @Return - Return the Event caller must trigger to unlock the skeletal mesh properties
 	*/
-	ENGINE_API void LockPropertiesUntil(FEvent* Event);
+	ENGINE_API FEvent* LockPropertiesUntil();
 
 	ENGINE_API virtual void PostEditChangeProperty(FPropertyChangedEvent& PropertyChangedEvent) override;
 
diff --git a/Engine/Source/Runtime/Engine/Classes/Engine/SkinnedAssetAsyncCompileUtils.h b/Engine/Source/Runtime/Engine/Classes/Engine/SkinnedAssetAsyncCompileUtils.h
--- a/Engine/Source/Runtime/Engine/Classes/Engine/SkinnedAssetAsyncCompileUtils.h
+++ b/Engine/Source/Runtime/Engine/Classes/Engine/SkinnedAssetAsyncCompileUtils.h
@@ -41,12 +41,11 @@ public:
 class FSkinnedAsyncTaskContext : public FSkinnedAssetCompilationContext
 {
 public:
-	FSkinnedAsyncTaskContext(TFunctionRef<void()> InAsyncTaskFunction)
-		: AsyncTaskFunction(InAsyncTaskFunction)
+	FSkinnedAsyncTaskContext(FEvent* InEvent)
+		: Event(InEvent)
 	{
 	}
-
-	TFunctionRef<void()> AsyncTaskFunction;
+	FEvent* Event = nullptr;
 };
 
 #if WITH_EDITOR
diff --git a/Engine/Source/Runtime/Engine/Private/SkeletalMesh.cpp b/Engine/Source/Runtime/Engine/Private/SkeletalMesh.cpp
--- a/Engine/Source/Runtime/Engine/Private/SkeletalMesh.cpp
+++ b/Engine/Source/Runtime/Engine/Private/SkeletalMesh.cpp
@@ -1515,23 +1515,21 @@ FScopedSkeletalMeshRenderData::FScopedSkeletalMeshRenderData(USkeletalMesh* InMe
 	{
 		// Lock the skeletalmesh properties since we call USkeletalMesh::Cache() function (through GetPlatformSkeletalMeshRenderData -> CachePlatform -> Cache) 
 		// and which could be called by other threads at the same time
-		Lock = FPlatformProcess::GetSynchEventFromPool();
-		Mesh->LockPropertiesUntil(Lock);
+		Lock = Mesh->LockPropertiesUntil();
 	}
 }
 
 FScopedSkeletalMeshRenderData::~FScopedSkeletalMeshRenderData()
 {
-	if (Mesh)
+	if (Lock)
 	{
-		check(Lock);
-
+		check(Mesh);
 		Lock->Trigger();
-		FPlatformProcess::ReturnSynchEventToPool(Lock);
-		Data = nullptr;
-		Mesh = nullptr;
-		Lock = nullptr;
 	}
+
+	Data = nullptr;
+	Mesh = nullptr;
+	Lock = nullptr;
 }
 
 const FSkeletalMeshRenderData* FScopedSkeletalMeshRenderData::GetData() const
@@ -1992,21 +1990,20 @@ void USkeletalMesh::FinishBuildInternal(FSkinnedAssetBuildContext& Context)
 	PostMeshCached.Broadcast(this);
 }
 
-void USkeletalMesh::LockPropertiesUntil(FEvent* Event)
+FEvent* USkeletalMesh::LockPropertiesUntil()
 {
 	TRACE_CPUPROFILER_EVENT_SCOPE(USkeletalMesh::Import);
+	
+	FEvent* Event = FPlatformProcess::GetSynchEventFromPool();
+	check(Event);
+
 	if (IsCompiling())
 	{
 		FSkinnedAssetCompilingManager::Get().FinishCompilation({ this });
 	}
 
-	auto AsyncTaskFunction = [Event]()
-	{
-		Event->Wait();
-	};
-
 	//Use the async task compile to lock the properties
-	FSkinnedAsyncTaskContext Context(AsyncTaskFunction);
+	FSkinnedAsyncTaskContext Context(Event);
 	BeginAsyncTaskInternal(Context);
 	PrepareForAsyncCompilation();
 	FQueuedThreadPool* SkeletalMeshThreadPool = FSkinnedAssetCompilingManager::Get().GetThreadPool();
@@ -2015,6 +2012,7 @@ void USkeletalMesh::LockPropertiesUntil(FEvent* Event)
 	AsyncTask = MakeUnique<FSkinnedAssetAsyncBuildTask>(this, MoveTemp(Context));
 	AsyncTask->StartBackgroundTask(SkeletalMeshThreadPool, BasePriority, EQueuedWorkFlags::DoNotRunInsideBusyWait);
 	FSkinnedAssetCompilingManager::Get().AddSkinnedAssets({ this });
+	return Event;
 }
 
 void USkeletalMesh::BeginAsyncTaskInternal(FSkinnedAsyncTaskContext& Context)
@@ -2028,7 +2026,11 @@ void USkeletalMesh::BeginAsyncTaskInternal(FSkinnedAsyncTaskContext& Context)
 
 void USkeletalMesh::ExecuteAsyncTaskInternal(FSkinnedAsyncTaskContext& Context)
 {
-	Context.AsyncTaskFunction();
+	if (ensure(Context.Event))
+	{
+		Context.Event->Wait();
+		FPlatformProcess::ReturnSynchEventToPool(Context.Event);
+	}
 }
 
 void USkeletalMesh::FinishAsyncTaskInternal(FSkinnedAsyncTaskContext& Context)
