Disable PathTracing and RayTracingDebug showflags when Ray Tracing is disabled.
- prevents asserts and corrupted image when ray tracing is dynamically disabled while using those visualizations.

diff --git a/Engine/Source/Runtime/Engine/Private/ShowFlags.cpp b/Engine/Source/Runtime/Engine/Private/ShowFlags.cpp
--- a/Engine/Source/Runtime/Engine/Private/ShowFlags.cpp
+++ b/Engine/Source/Runtime/Engine/Private/ShowFlags.cpp
@@ -5,6 +5,7 @@
 =============================================================================*/
 
 #include "ShowFlags.h"
+#include "RenderUtils.h"
 #include "Engine/EngineBaseTypes.h"
 #include "HAL/IConsoleManager.h"
 #include "Misc/ScopeRWLock.h"
@@ -474,6 +475,12 @@ void EngineShowFlagOverride(EShowFlagInitMode ShowFlagInitMode, EViewModeIndex V
 		}
 	}
 
+	if (!IsRayTracingEnabled())
+	{
+		EngineShowFlags.SetPathTracing(false);
+		EngineShowFlags.SetRayTracingDebug(false);
+	}
+
 	// Some view modes want some features off or on (no state)
 	{
 		if (ViewModeIndex == VMI_BrushWireframe ||
@@ -580,14 +587,8 @@ void EngineShowFlagOverride(EShowFlagInitMode ShowFlagInitMode, EViewModeIndex V
 			EngineShowFlags.SetDecals(false); // Decals require the use of FDebugPSInLean.
 		}
 
-		if (ViewModeIndex == VMI_PathTracing)
-		{
-			EngineShowFlags.SetPathTracing(true);
-		}
-
-		if (ViewModeIndex == VMI_RayTracingDebug)
+		if (IsRayTracingEnabled() && ViewModeIndex == VMI_RayTracingDebug)
 		{
-			EngineShowFlags.SetRayTracingDebug(true);
 			EngineShowFlags.SetVisualizeHDR(false);
 			EngineShowFlags.SetVisualizeLocalExposure(false);
 			EngineShowFlags.SetVisualizeMotionBlur(false);
