World Partition: don't report missing actor references when checking the map for errors as it can happen in legit scenarios involving weak actor references to deleted actors, for instance.
The LinkerLoad code will report these errors when loading the actors and null out references, without failing editor loads or cooks, so still report on changelist validation though to let the user resave the actors.

diff --git a/Engine/Source/Runtime/Engine/Private/WorldPartition/ErrorHandling/WorldPartitionStreamingGenerationTokenizedMessageErrorHandler.cpp b/Engine/Source/Runtime/Engine/Private/WorldPartition/ErrorHandling/WorldPartitionStreamingGenerationTokenizedMessageErrorHandler.cpp
--- a/Engine/Source/Runtime/Engine/Private/WorldPartition/ErrorHandling/WorldPartitionStreamingGenerationTokenizedMessageErrorHandler.cpp
+++ b/Engine/Source/Runtime/Engine/Private/WorldPartition/ErrorHandling/WorldPartitionStreamingGenerationTokenizedMessageErrorHandler.cpp
@@ -25,23 +25,21 @@ void ITokenizedMessageErrorHandler::OnInvalidRuntimeGrid(const FWorldPartitionAc
 
 void ITokenizedMessageErrorHandler::OnInvalidReference(const FWorldPartitionActorDescView& ActorDescView, const FGuid& ReferenceGuid, FWorldPartitionActorDescView* ReferenceActorDescView)
 {
-	TSharedRef<FTokenizedMessage> Message = FTokenizedMessage::Create(EMessageSeverity::Error);
-	Message->AddToken(FTextToken::Create(LOCTEXT("TokenMessage_WorldPartition_Actor", "Actor")))
-		->AddToken(FActorToken::Create(ActorDescView.GetActorSoftPath().ToString(), ActorDescView.GetGuid(), FText::FromString(GetActorName(ActorDescView))))
-		->AddToken(FTextToken::Create(LOCTEXT("TokenMessage_WorldPartition_HaveMissingRefsTo", "has an invalid reference to")));
-
+	// Don't report invalid references to non-existing actors as it can happen in valid scenarios, as the linker code will silently null out references when loading actors with missing
+	// references.
 	if (ReferenceActorDescView)
 	{
+		TSharedRef<FTokenizedMessage> Message = FTokenizedMessage::Create(EMessageSeverity::Error);
+		Message->AddToken(FTextToken::Create(LOCTEXT("TokenMessage_WorldPartition_Actor", "Actor")))
+			->AddToken(FActorToken::Create(ActorDescView.GetActorSoftPath().ToString(), ActorDescView.GetGuid(), FText::FromString(GetActorName(ActorDescView))))
+			->AddToken(FTextToken::Create(LOCTEXT("TokenMessage_WorldPartition_HaveMissingRefsTo", "has an invalid reference to")));
+
 		Message->AddToken(FActorToken::Create(ReferenceActorDescView->GetActorSoftPath().ToString(), ReferenceActorDescView->GetGuid(), FText::FromString(GetActorName(*ReferenceActorDescView))));
-	}
-	else
-	{
-		Message->AddToken(FTextToken::Create(FText::FromString(ReferenceGuid.ToString())));
-	}
 
-	Message->AddToken(FMapErrorToken::Create(TEXT("WorldPartition_MissingActorReference_CheckForErrors")));
+		Message->AddToken(FMapErrorToken::Create(TEXT("WorldPartition_MissingActorReference_CheckForErrors")));
 
-	HandleTokenizedMessage(MoveTemp(Message));
+		HandleTokenizedMessage(MoveTemp(Message));
+	}
 }
 
 void ITokenizedMessageErrorHandler::OnInvalidReferenceGridPlacement(const FWorldPartitionActorDescView& ActorDescView, const FWorldPartitionActorDescView& ReferenceActorDescView)
