Allow querying whether a client has initialized a level via the ULevel

UNetConnection::ClientHasInitializedLevelFor is now deprecated and UNetConnection::ClientHasInitializedLevel should be used/overridden instead.

diff --git a/Engine/Plugins/NetcodeUnitTest/NetcodeUnitTest/Source/NetcodeUnitTest/Private/NUTActor.cpp b/Engine/Plugins/NetcodeUnitTest/NetcodeUnitTest/Source/NetcodeUnitTest/Private/NUTActor.cpp
--- a/Engine/Plugins/NetcodeUnitTest/NetcodeUnitTest/Source/NetcodeUnitTest/Private/NUTActor.cpp
+++ b/Engine/Plugins/NetcodeUnitTest/NetcodeUnitTest/Source/NetcodeUnitTest/Private/NUTActor.cpp
@@ -627,7 +627,7 @@ void ANUTActor::ServerClientPing_Implementation()
 		{
 			// Based on UNetDriver::IsLevelInitializeForActor
 			bNotLoaded = !(CurConn->GetClientWorldPackageName() == CurWorld->GetOutermost()->GetFName() &&
-										CurConn->ClientHasInitializedLevelFor(this));
+										CurConn->ClientHasInitializedLevel(GetLevel()));
 
 			// Also trigger if there is no PlayerController yet set for the connection
 			if (CurConn->OwningActor == nullptr)
diff --git a/Engine/Source/Runtime/Engine/Classes/Engine/DemoNetConnection.h b/Engine/Source/Runtime/Engine/Classes/Engine/DemoNetConnection.h
--- a/Engine/Source/Runtime/Engine/Classes/Engine/DemoNetConnection.h
+++ b/Engine/Source/Runtime/Engine/Classes/Engine/DemoNetConnection.h
@@ -33,7 +33,7 @@ public:
 	ENGINE_API virtual void FlushNet( bool bIgnoreSimulation = false ) override;
 	ENGINE_API virtual void HandleClientPlayer( APlayerController* PC, class UNetConnection* NetConnection ) override;
 	ENGINE_API virtual TSharedPtr<const FInternetAddr> GetRemoteAddr() override;
-	ENGINE_API virtual bool ClientHasInitializedLevelFor( const AActor* TestActor ) const override;
+	ENGINE_API virtual bool ClientHasInitializedLevel( const ULevel* TestLevel ) const override;
 	ENGINE_API virtual TSharedPtr<FObjectReplicator> CreateReplicatorForNewActorChannel(UObject* Object);
 	virtual FString RemoteAddressToString() override { return TEXT("Demo"); }
 
diff --git a/Engine/Source/Runtime/Engine/Classes/Engine/NetConnection.h b/Engine/Source/Runtime/Engine/Classes/Engine/NetConnection.h
--- a/Engine/Source/Runtime/Engine/Classes/Engine/NetConnection.h
+++ b/Engine/Source/Runtime/Engine/Classes/Engine/NetConnection.h
@@ -1176,10 +1176,17 @@ public:
 	/** Pops the LastStart bits off of the send buffer, used for merging bunches */
 	void PopLastStart();
 
+	/**
+	 * returns whether the client has initialized the given level
+	 * @return true if the client has initialized the given level, false otherwise
+	 */
+	ENGINE_API virtual bool ClientHasInitializedLevel(const ULevel* TestLevel) const;
+
 	/** 
 	 * returns whether the client has initialized the level required for the given object
 	 * @return true if the client has initialized the level the object is in or the object is not in a level, false otherwise
 	 */
+	UE_DEPRECATED(5.4, "ClientHasInitializedLevelFor is deprecated. Use ClientHasInitializedLevel and pass the actor's level (Actor->GetLevel()) instead.")
 	ENGINE_API virtual bool ClientHasInitializedLevelFor(const AActor* TestActor) const;
 
 	/**
@@ -1944,7 +1951,7 @@ public:
 	virtual void LowLevelSend(void* Data, int32 CountBits, FOutPacketTraits& Traits) override { }
 	ENGINE_API void HandleClientPlayer( APlayerController* PC, UNetConnection* NetConnection ) override;
 	virtual FString LowLevelGetRemoteAddress(bool bAppendPort=false) override { return FString(); }
-	virtual bool ClientHasInitializedLevelFor(const AActor* TestActor) const { return true; }
+	virtual bool ClientHasInitializedLevel(const ULevel* TestLevel) const override { return true; }
 
 	virtual void DestroyOwningActor() override { /* Don't destroy the OwningActor since we follow a real PlayerController*/ }
 
diff --git a/Engine/Source/Runtime/Engine/Private/DemoNetDriver.cpp b/Engine/Source/Runtime/Engine/Private/DemoNetDriver.cpp
--- a/Engine/Source/Runtime/Engine/Private/DemoNetDriver.cpp
+++ b/Engine/Source/Runtime/Engine/Private/DemoNetDriver.cpp
@@ -4971,12 +4971,12 @@ TSharedPtr<const FInternetAddr> UDemoNetConnection::GetRemoteAddr()
 	return FInternetAddrDemo::DemoInternetAddr;
 }
 
-bool UDemoNetConnection::ClientHasInitializedLevelFor(const AActor* TestActor) const
+bool UDemoNetConnection::ClientHasInitializedLevel(const ULevel* TestLevel) const
 {
 	// We save all currently streamed levels into the demo stream so we can force the demo playback client
 	// to stay in sync with the recording server
 	// This may need to be tweaked or re-evaluated when we start recording demos on the client
-	return (GetDriver()->GetDemoFrameNum() > 2 || Super::ClientHasInitializedLevelFor(TestActor));
+	return (GetDriver()->GetDemoFrameNum() > 2 || Super::ClientHasInitializedLevel(TestLevel));
 }
 
 TSharedPtr<FObjectReplicator> UDemoNetConnection::CreateReplicatorForNewActorChannel(UObject* Object)
diff --git a/Engine/Source/Runtime/Engine/Private/NetConnection.cpp b/Engine/Source/Runtime/Engine/Private/NetConnection.cpp
--- a/Engine/Source/Runtime/Engine/Private/NetConnection.cpp
+++ b/Engine/Source/Runtime/Engine/Private/NetConnection.cpp
@@ -1518,7 +1518,7 @@ FNetLevelVisibilityTransactionId UNetConnection::UpdateLevelStreamStatusChangedT
 	return TransactionId;
 }
 
-bool UNetConnection::ClientHasInitializedLevelFor(const AActor* TestActor) const
+bool UNetConnection::ClientHasInitializedLevel(const ULevel* TestLevel) const
 {
 	checkSlow(Driver);
 	checkSlow(Driver->IsServer());
@@ -1526,18 +1526,23 @@ bool UNetConnection::ClientHasInitializedLevelFor(const AActor* TestActor) const
 	// This function is called a lot, basically for every replicated actor every time it replicates, on every client connection
 	// Each client connection has a different visibility state (what levels are currently loaded for them).
 
-	const FName PackageName = TestActor->GetLevel()->GetPackage()->GetFName();
+	const FName PackageName = TestLevel->GetPackage()->GetFName();
 
 	if (const bool* bIsVisible = ClientVisibleActorOuters.Find(PackageName))
 	{
 		return *bIsVisible;
 	}
 
-	// The actor's outer was not in the acceleration map so we perform the "legacy" function and 
+	// The level was not in the acceleration map so we perform the "legacy" function and 
 	// cache the result so that we don't do this every time:
 	return UpdateCachedLevelVisibility(PackageName);
 }
 
+bool UNetConnection::ClientHasInitializedLevelFor(const AActor* TestActor) const
+{
+	return ClientHasInitializedLevel(TestActor->GetLevel());
+}
+
 bool UNetConnection::UpdateCachedLevelVisibility(const FName& PackageName) const
 {
 	bool bVisible = false;
@@ -5101,7 +5106,7 @@ void UNetConnection::FlushDormancyForObject(AActor* DormantActor, UObject* Repli
 		
 		// Init using the object's current state if the client has this actor's level loaded. If the level
 		// is unloaded, we need to use the default state since the client has no current state.
-		if (bFlushDormancyUseDefaultStateForUnloadedLevels && !ClientHasInitializedLevelFor(DormantActor))
+		if (bFlushDormancyUseDefaultStateForUnloadedLevels && !ClientHasInitializedLevel(DormantActor->GetLevel()))
 		{
 			bUseDefaultState = true;
 		}
diff --git a/Engine/Source/Runtime/Engine/Private/NetDriver.cpp b/Engine/Source/Runtime/Engine/Private/NetDriver.cpp
--- a/Engine/Source/Runtime/Engine/Private/NetDriver.cpp
+++ b/Engine/Source/Runtime/Engine/Private/NetDriver.cpp
@@ -2013,7 +2013,7 @@ bool UNetDriver::IsLevelInitializedForActor(const AActor* InActor, const UNetCon
 #endif
 
 	// we can't create channels while the client is in the wrong world
-	const bool bCorrectWorld = GetWorldPackage() != nullptr && (InConnection->GetClientWorldPackageName() == GetWorldPackage()->GetFName()) && InConnection->ClientHasInitializedLevelFor(InActor);
+	const bool bCorrectWorld = GetWorldPackage() != nullptr && (InConnection->GetClientWorldPackageName() == GetWorldPackage()->GetFName()) && InConnection->ClientHasInitializedLevel(InActor->GetLevel());
 	
 	// exception: Special case for PlayerControllers as they are required for the client to travel to the new world correctly			
 	bool bIsConnectionPC = false;
diff --git a/Engine/Source/Runtime/Engine/Private/ReplayHelper.cpp b/Engine/Source/Runtime/Engine/Private/ReplayHelper.cpp
--- a/Engine/Source/Runtime/Engine/Private/ReplayHelper.cpp
+++ b/Engine/Source/Runtime/Engine/Private/ReplayHelper.cpp
@@ -1786,7 +1786,7 @@ bool FReplayHelper::ReplicateActor(AActor* Actor, UNetConnection* Connection, bo
 		const bool bShouldHaveChannel =
 			Actor->bRelevantForNetworkReplays &&
 			!Actor->GetTearOff() &&
-			(!Actor->IsNetStartupActor() || Connection->ClientHasInitializedLevelFor(Actor));
+			(!Actor->IsNetStartupActor() || Connection->ClientHasInitializedLevel(Actor->GetLevel()));
 
 		UActorChannel* Channel = Connection->FindActorChannelRef(Actor);
 
diff --git a/Engine/Source/Runtime/Engine/Private/ReplayNetConnection.cpp b/Engine/Source/Runtime/Engine/Private/ReplayNetConnection.cpp
--- a/Engine/Source/Runtime/Engine/Private/ReplayNetConnection.cpp
+++ b/Engine/Source/Runtime/Engine/Private/ReplayNetConnection.cpp
@@ -216,9 +216,9 @@ TSharedPtr<const FInternetAddr> UReplayNetConnection::GetRemoteAddr()
 	return FInternetAddrDemo::DemoInternetAddr;
 }
 
-bool UReplayNetConnection::ClientHasInitializedLevelFor(const AActor* TestActor) const
+bool UReplayNetConnection::ClientHasInitializedLevel(const ULevel* TestLevel) const
 {
-	return (DemoFrameNum > 2 || Super::ClientHasInitializedLevelFor(TestActor));
+	return (DemoFrameNum > 2 || Super::ClientHasInitializedLevel(TestLevel));
 }
 
 void UReplayNetConnection::AddEvent(const FString& Group, const FString& Meta, const TArray<uint8>& Data)
diff --git a/Engine/Source/Runtime/Engine/Public/ReplayNetConnection.h b/Engine/Source/Runtime/Engine/Public/ReplayNetConnection.h
--- a/Engine/Source/Runtime/Engine/Public/ReplayNetConnection.h
+++ b/Engine/Source/Runtime/Engine/Public/ReplayNetConnection.h
@@ -26,7 +26,7 @@ public:
 	ENGINE_API virtual int32 IsNetReady(bool Saturate) override;
 	ENGINE_API virtual void HandleClientPlayer(APlayerController* PC, UNetConnection* NetConnection) override;
 	ENGINE_API virtual TSharedPtr<const FInternetAddr> GetRemoteAddr() override;
-	ENGINE_API virtual bool ClientHasInitializedLevelFor(const AActor* TestActor) const override;
+	ENGINE_API virtual bool ClientHasInitializedLevel(const ULevel* TestLevel) const override;
 	virtual FString RemoteAddressToString() override { return TEXT("Replay"); }
 	ENGINE_API virtual void CleanUp() override;
 	ENGINE_API virtual void Serialize(FArchive& Ar) override;
