World Partition: don't validate runtime grids for runtime spatial hash autogenerated HLODs.

diff --git a/Engine/Source/Runtime/Engine/Private/WorldPartition/WorldPartitionRuntimeSpatialHash.cpp b/Engine/Source/Runtime/Engine/Private/WorldPartition/WorldPartitionRuntimeSpatialHash.cpp
--- a/Engine/Source/Runtime/Engine/Private/WorldPartition/WorldPartitionRuntimeSpatialHash.cpp
+++ b/Engine/Source/Runtime/Engine/Private/WorldPartition/WorldPartitionRuntimeSpatialHash.cpp
@@ -1562,20 +1562,7 @@ void UWorldPartitionRuntimeSpatialHash::FlushStreaming()
 
 bool UWorldPartitionRuntimeSpatialHash::IsValidGrid(FName GridName) const
 {
-	if (GridName.IsNone())
-	{
-		return true;
-	}
-
-	for (const FSpatialHashRuntimeGrid& Grid : Grids)
-	{
-		if (Grid.GridName == GridName)
-		{
-			return true;
-		}
-	}
-
-	return false;
+	return true;
 }
 
 #endif //WITH_EDITOR
diff --git a/Engine/Source/Runtime/Engine/Private/WorldPartition/WorldPartitionStreamingGeneration.cpp b/Engine/Source/Runtime/Engine/Private/WorldPartition/WorldPartitionStreamingGeneration.cpp
--- a/Engine/Source/Runtime/Engine/Private/WorldPartition/WorldPartitionStreamingGeneration.cpp
+++ b/Engine/Source/Runtime/Engine/Private/WorldPartition/WorldPartitionStreamingGeneration.cpp
@@ -173,13 +173,6 @@ const TArray<FGuid>& FStreamingGenerationActorDescView::GetEditorReferences() co
 	return EditorReferences;
 }
 
-bool FStreamingGenerationActorDescView::ShouldValidateRuntimeGrid() const
-{
-PRAGMA_DISABLE_DEPRECATION_WARNINGS
-	return ActorDesc->ShouldValidateRuntimeGrid();
-PRAGMA_ENABLE_DEPRECATION_WARNINGS
-}
-
 void FStreamingGenerationActorDescView::SetParentView(const FWorldPartitionActorDescView* InParentView)
 {
 	check(!ParentView);
@@ -1191,37 +1184,32 @@ class FWorldPartitionStreamingGenerator
 			{
 				FContainerCollectionInstanceDescriptor::FPerInstanceData& PerInstanceData = ContainerCollectionInstanceDescriptor.GetInstanceData(ActorDescView.GetGuid());
 
-PRAGMA_DISABLE_DEPRECATION_WARNINGS
-				if (ActorDescView.ShouldValidateRuntimeGrid())
-PRAGMA_ENABLE_DEPRECATION_WARNINGS
+				if (!IsValidGrid(PerInstanceData.RuntimeGrid))
 				{
-					if (!IsValidGrid(PerInstanceData.RuntimeGrid))
+					if (PassType == EPassType::ErrorReporting)
 					{
-						if (PassType == EPassType::ErrorReporting)
-						{
-							ErrorHandler->OnInvalidRuntimeGrid(ActorDescView, PerInstanceData.RuntimeGrid);
-						}
-						else
-						{
-							PerInstanceData.RuntimeGrid = NAME_None;
-						}
-
-						NbErrorsDetected++;
+						ErrorHandler->OnInvalidRuntimeGrid(ActorDescView, PerInstanceData.RuntimeGrid);
 					}
-					
-					if (ActorDescView.GetHLODLayer().IsValid() && !IsValidHLODLayer(PerInstanceData.RuntimeGrid, ActorDescView.GetHLODLayer()))
+					else
 					{
-						if (PassType == EPassType::ErrorReporting)
-						{
-							ErrorHandler->OnInvalidHLODLayer(ActorDescView);
-						}
-						else
-						{
-							ActorDescView.SetForcedNoHLODLayer();
-						}
+						PerInstanceData.RuntimeGrid = NAME_None;
+					}
 
-						NbErrorsDetected++;
+					NbErrorsDetected++;
+				}
+					
+				if (ActorDescView.GetHLODLayer().IsValid() && !IsValidHLODLayer(PerInstanceData.RuntimeGrid, ActorDescView.GetHLODLayer()))
+				{
+					if (PassType == EPassType::ErrorReporting)
+					{
+						ErrorHandler->OnInvalidHLODLayer(ActorDescView);
+					}
+					else
+					{
+						ActorDescView.SetForcedNoHLODLayer();
 					}
+
+					NbErrorsDetected++;
 				}
 
 				if (TSet<FGuid>* FilteredActors = ContainerFilteredActors.Find(ContainerCollectionInstanceDescriptor.ID))
diff --git a/Engine/Source/Runtime/Engine/Public/WorldPartition/HLOD/HLODActorDesc.h b/Engine/Source/Runtime/Engine/Public/WorldPartition/HLOD/HLODActorDesc.h
--- a/Engine/Source/Runtime/Engine/Public/WorldPartition/HLOD/HLODActorDesc.h
+++ b/Engine/Source/Runtime/Engine/Public/WorldPartition/HLOD/HLODActorDesc.h
@@ -35,7 +35,6 @@ public:
 
 	//~ Begin FWorldPartitionActorDesc Interface.
 	virtual bool IsRuntimeRelevant(const FActorContainerID& InContainerID) const override { return !bIsForcedNonSpatiallyLoaded; }
-	virtual bool ShouldValidateRuntimeGrid() const override { return false; }
 	virtual FBox GetEditorBounds() const override { return EditorBounds; }
 	//~ End FWorldPartitionActorDesc Interface.
 
diff --git a/Engine/Source/Runtime/Engine/Public/WorldPartition/WorldPartitionActorDesc.h b/Engine/Source/Runtime/Engine/Public/WorldPartition/WorldPartitionActorDesc.h
--- a/Engine/Source/Runtime/Engine/Public/WorldPartition/WorldPartitionActorDesc.h
+++ b/Engine/Source/Runtime/Engine/Public/WorldPartition/WorldPartitionActorDesc.h
@@ -159,9 +159,6 @@ ENGINE_API PRAGMA_ENABLE_DEPRECATION_WARNINGS
 	ENGINE_API virtual bool IsEditorRelevant() const;
 	ENGINE_API virtual void CheckForErrors(IStreamingGenerationErrorHandler* ErrorHandler) const;
 
-	UE_DEPRECATED(5.2, "ShouldValidateRuntimeGrid is deprecated and should not be used.")
-	virtual bool ShouldValidateRuntimeGrid() const { return true; }
-
 	UE_DEPRECATED(5.3, "GetLevelPackage is deprecated use GetContainerPackage instead.")
 	virtual FName GetLevelPackage() const { return GetContainerPackage(); }
 
diff --git a/Engine/Source/Runtime/Engine/Public/WorldPartition/WorldPartitionStreamingGeneration.h b/Engine/Source/Runtime/Engine/Public/WorldPartition/WorldPartitionStreamingGeneration.h
--- a/Engine/Source/Runtime/Engine/Public/WorldPartition/WorldPartitionStreamingGeneration.h
+++ b/Engine/Source/Runtime/Engine/Public/WorldPartition/WorldPartitionStreamingGeneration.h
@@ -47,7 +47,6 @@ public:
 	ENGINE_API virtual const TArray<FGuid>& GetEditorReferences() const override;
 	//~ End FWorldPartitionActorDescView interface
 
-	ENGINE_API bool ShouldValidateRuntimeGrid() const;
 	ENGINE_API void SetParentView(const FWorldPartitionActorDescView* InParentView);
 	ENGINE_API void SetDataLayerInstanceNames(const TArray<FName>& InDataLayerInstanceNames);
 	ENGINE_API void SetForcedNonSpatiallyLoaded();
