World Partition HashSet Conversions:
  - Fix crash when converting a map without any actors in it (when creating a new map from the empty template).
  - Better logging when converting maps.

diff --git a/Engine/Source/Runtime/Engine/Private/WorldPartition/RuntimeHashSet/WorldPartitionRuntimeHashSetConversions.cpp b/Engine/Source/Runtime/Engine/Private/WorldPartition/RuntimeHashSet/WorldPartitionRuntimeHashSetConversions.cpp
--- a/Engine/Source/Runtime/Engine/Private/WorldPartition/RuntimeHashSet/WorldPartitionRuntimeHashSetConversions.cpp
+++ b/Engine/Source/Runtime/Engine/Private/WorldPartition/RuntimeHashSet/WorldPartitionRuntimeHashSetConversions.cpp
@@ -29,6 +29,11 @@ UWorldPartitionRuntimeHashSet* UWorldPartitionRuntimeHashSet::CreateFrom(const U
 		TMap<FName, TMap<FName, TSet<const UHLODLayer*>>> GridHLODLayersMap;
 		const UHLODLayer* DefaultHLODLayer = WorldPartition->GetDefaultHLODLayer();
 
+		for (const FSpatialHashRuntimeGrid& Grid : SpatialHash->Grids)
+		{
+			GridHLODLayersMap.Add(Grid.GridName);
+		}
+
 		StreamingGenerationContext->ForEachActorSetInstance([&GridHLODLayersMap, SpatialHash, DefaultHLODLayer](const IStreamingGenerationContext::FActorSetInstance& ActorSetInstance)
 		{
 			ActorSetInstance.ForEachActor([&GridHLODLayersMap, SpatialHash, DefaultHLODLayer, &ActorSetInstance](const FGuid& ActorGuid)
diff --git a/Engine/Source/Runtime/Engine/Private/WorldPartition/WorldPartitionRuntimeHash.cpp b/Engine/Source/Runtime/Engine/Private/WorldPartition/WorldPartitionRuntimeHash.cpp
--- a/Engine/Source/Runtime/Engine/Private/WorldPartition/WorldPartitionRuntimeHash.cpp
+++ b/Engine/Source/Runtime/Engine/Private/WorldPartition/WorldPartitionRuntimeHash.cpp
@@ -513,27 +513,29 @@ UWorldPartitionRuntimeHash* UWorldPartitionRuntimeHash::ConvertWorldPartitionHas
 	check(InDstClass->IsChildOf<UWorldPartitionRuntimeHash>());
 	check(!InDstClass->HasAnyClassFlags(CLASS_Abstract));
 
-	UWorldPartitionRuntimeHash* NewHash = nullptr;
-
+	// Look for a registered converter
 	const UClass* CurrentSrcClass = InSrcHash->GetClass();
 	while(CurrentSrcClass != UWorldPartitionRuntimeHash::StaticClass())
 	{
 		if (FRuntimeHashConvertFunc* Converter = WorldPartitionRuntimeHashConverters.Find({ CurrentSrcClass, InDstClass }))
 		{
-			NewHash = (*Converter)(InSrcHash);
-			break;
+			if (UWorldPartitionRuntimeHash* NewHash = (*Converter)(InSrcHash))
+			{
+				UE_LOG(LogWorldPartition, Log, TEXT("Converted '%s' runtime hash class from '%s' to '%s'."), *InSrcHash->GetPackage()->GetName(), *InSrcHash->GetClass()->GetName(), *InDstClass->GetName());
+				return NewHash;
+			}
+			else
+			{
+				UE_LOG(LogWorldPartition, Warning, TEXT("Failed to convert '%s' runtime hash class from '%s' to '%s'."), *InSrcHash->GetPackage()->GetName(), *InSrcHash->GetClass()->GetName(), *InDstClass->GetName());
+			}
 		}
 		CurrentSrcClass = CurrentSrcClass->GetSuperClass();
 	}
 
-	if (!NewHash)
-	{
-		NewHash = NewObject<UWorldPartitionRuntimeHash>(InSrcHash->GetOuter(), InDstClass, NAME_None, RF_Transactional);
-		NewHash->SetDefaultValues();
-	}
-
-	UE_LOG(LogWorldPartition, Log, TEXT("Converted '%s' runtime hash class from '%s' to '%s'"), *InSrcHash->GetPackage()->GetName(), *InSrcHash->GetClass()->GetName(), *InDstClass->GetName());
-
+	// No converter found, create a new hash of the target type with default values
+	UE_LOG(LogWorldPartition, Log, TEXT("No converter found to convert '%s' runtime hash class from '%s' to '%s', creating new with default values."), *InSrcHash->GetPackage()->GetName(), *InSrcHash->GetClass()->GetName(), *InDstClass->GetName());
+	UWorldPartitionRuntimeHash* NewHash = NewObject<UWorldPartitionRuntimeHash>(InSrcHash->GetOuter(), InDstClass, NAME_None, RF_Transactional);
+	NewHash->SetDefaultValues();
 	return NewHash;
 }
 #endif
