World Partition HashSet: fix non-spatially loaded actors never loading when in a data layer or content bundle.

diff --git a/Engine/Source/Runtime/Engine/Private/WorldPartition/RuntimeHashSet/WorldPartitionRuntimeHashSetStreamingGeneration.cpp b/Engine/Source/Runtime/Engine/Private/WorldPartition/RuntimeHashSet/WorldPartitionRuntimeHashSetStreamingGeneration.cpp
--- a/Engine/Source/Runtime/Engine/Private/WorldPartition/RuntimeHashSet/WorldPartitionRuntimeHashSetStreamingGeneration.cpp
+++ b/Engine/Source/Runtime/Engine/Private/WorldPartition/RuntimeHashSet/WorldPartitionRuntimeHashSetStreamingGeneration.cpp
@@ -224,13 +224,13 @@ bool UWorldPartitionRuntimeHashSet::GenerateStreaming(UWorldPartitionStreamingPo
 				StreamingData.Name = CellDescInstance.SourcePartition->Name;
 				StreamingData.LoadingRange = CellDescInstance.SourcePartition->LoadingRange;
 
-				if (RuntimeCell->IsAlwaysLoaded())
+				if (CellDescInstance.bIsSpatiallyLoaded)
 				{
-					StreamingData.NonStreamingCells.Add(RuntimeCell);
+					StreamingData.StreamingCells.Add(RuntimeCell);
 				}
 				else
 				{
-					StreamingData.StreamingCells.Add(RuntimeCell);
+					StreamingData.NonStreamingCells.Add(RuntimeCell);
 				}
 			}
 		}
