Added new Blueprint functions for converting multidimensional array indices.

Specifically, this adds conversion functions for 1D -> 2D, 1D -> 3D, 2D -> 1D, and 3D -> 1D indices.

For example, in a 3x2 array, converting the 1D index 4 to a 2D dimensional index would yield (1,1). Tests were also added, which demonstrate usage and expected output.

#virtualized

diff --git a/Engine/Source/Runtime/Engine/Classes/Kismet/KismetMathLibrary.h b/Engine/Source/Runtime/Engine/Classes/Kismet/KismetMathLibrary.h
--- a/Engine/Source/Runtime/Engine/Classes/Kismet/KismetMathLibrary.h
+++ b/Engine/Source/Runtime/Engine/Classes/Kismet/KismetMathLibrary.h
@@ -4546,6 +4546,60 @@ class UKismetMathLibrary : public UBlueprintFunctionLibrary
 	UFUNCTION(BlueprintPure, Category="Math|Smoothing", meta=(DisplayName="Dynamic Weighted Moving Average Rotator"))
 	static ENGINE_API FRotator DynamicWeightedMovingAverage_FRotator(FRotator CurrentSample, FRotator PreviousSample, float MaxDistance, float MinWeight, float MaxWeight);
 	
+	//
+	// Multidimensional index conversions
+	// 
+
+	/**
+	 *
+	 * Maps a 1D array index to a 2D array index.
+	 * 
+	 * @param Index1D - The 1D array index
+	 * @param XSize - X dimension of the 2D array
+	 * 
+	 * @return The equivalent 2D index of the array
+	 */
+	UFUNCTION(BlueprintPure, Category="Math|Conversions|Indices", meta=(DisplayName="Convert a 1D Index to a 2D Index"))
+	static FIntPoint Convert1DTo2D(int32 Index1D, int32 XSize);
+
+	/**
+	 *
+	 * Maps a 1D array index to a 3D array index.
+	 *
+	 * @param Index1D - The 1D array index
+	 * @param XSize - X dimension of the 3D array
+	 * @param YSize - Y dimension of the 3D array
+	 *
+	 * @return The equivalent 3D index of the array
+	 */
+	UFUNCTION(BlueprintPure, Category="Math|Conversions|Indices", meta=(DisplayName="Convert a 1D Index to a 3D Index"))
+	static FIntVector Convert1DTo3D(int32 Index1D, int32 XSize, int32 YSize);
+
+	/**
+	 *
+	 * Maps a 2D array index to a 1D array index.
+	 *
+	 * @param Index2D - The 2D array index
+	 * @param XSize - X dimension of the 2D array
+	 *
+	 * @return The equivalent 1D index of the array
+	 */
+	UFUNCTION(BlueprintPure, Category="Math|Conversions|Indices", meta=(DisplayName="Convert a 2D Index to a 1D Index"))
+	static int32 Convert2DTo1D(const FIntPoint& Index2D, int32 XSize);
+
+	/**
+	 *
+	 * Maps a 3D array index to a 1D array index.
+	 *
+	 * @param Index3D - The 3D array index
+	 * @param XSize - X dimension of the 3D array
+	 * @param YSize - Y dimension of the 3D array
+	 *
+	 * @return The equivalent 1D index of the array
+	 */
+	UFUNCTION(BlueprintPure, Category="Math|Conversions|Indices", meta=(DisplayName="Convert a 3D Index to a 1D Index"))
+	static int32 Convert3DTo1D(const FIntVector& Index3D, int32 XSize, int32 YSize);
+
 	// NetQuantized vector make/breaks
 	UFUNCTION(BlueprintPure, Category = "Math|Vector", meta = (NativeMakeFunc))
 	static FVector_NetQuantize MakeVector_NetQuantize(double X, double Y, double Z) { return FVector_NetQuantize(X, Y, Z); }
diff --git a/Engine/Source/Runtime/Engine/Private/KismetMathLibrary.cpp b/Engine/Source/Runtime/Engine/Private/KismetMathLibrary.cpp
--- a/Engine/Source/Runtime/Engine/Private/KismetMathLibrary.cpp
+++ b/Engine/Source/Runtime/Engine/Private/KismetMathLibrary.cpp
@@ -45,6 +45,7 @@ const FName DivideByZeroWarning = FName("DivideByZeroWarning");
 const FName NegativeSqrtWarning = FName("NegativeSqrtWarning");
 const FName ZeroLengthProjectionWarning = FName("ZeroLengthProjectionWarning");
 const FName InvalidDateWarning = FName("InvalidDateWarning");
+const FName InvalidIndexConversionParameterWarning = FName("InvalidIndexConversionParameterWarning");
 
 UKismetMathLibrary::UKismetMathLibrary(const FObjectInitializer& ObjectInitializer)
 	: Super(ObjectInitializer)
@@ -73,6 +74,12 @@ UKismetMathLibrary::UKismetMathLibrary(const FObjectInitializer& ObjectInitializ
 			LOCTEXT("InvalidDateWarning", "Invalid date warning")
 		)
 	);
+	FBlueprintSupport::RegisterBlueprintWarning(
+		FBlueprintWarningDeclaration(
+			InvalidIndexConversionParameterWarning,
+			LOCTEXT("InvalidIndexConversionParameterWarning", "Invalid index conversion parameter warning")
+		)
+	);
 }
 
 void UKismetMathLibrary::ReportError_Divide_ByteByte()
@@ -1580,6 +1587,94 @@ FRotator UKismetMathLibrary::DynamicWeightedMovingAverage_FRotator(FRotator Curr
 	return OutRotator;
 }
 
+FIntPoint UKismetMathLibrary::Convert1DTo2D(int32 Index1D, int32 XSize)
+{
+	if (Index1D < 0)
+	{
+		FFrame::KismetExecutionMessage(TEXT("Index1D must be non-negative: Convert1DTo2D"), ELogVerbosity::Warning, InvalidIndexConversionParameterWarning);
+		return FIntPoint::ZeroValue;
+	}
+
+	if (XSize <= 0)
+	{
+		FFrame::KismetExecutionMessage(TEXT("XSize must be positive: Convert1DTo2D"), ELogVerbosity::Warning, InvalidIndexConversionParameterWarning);
+		return FIntPoint::ZeroValue;
+	}
+
+	int32 X = Index1D % XSize;
+	int32 Y = Index1D / XSize;
+
+	return FIntPoint{ X,Y };
+}
+
+FIntVector UKismetMathLibrary::Convert1DTo3D(int32 Index1D, int32 XSize, int32 YSize)
+{
+	if (Index1D < 0)
+	{
+		FFrame::KismetExecutionMessage(TEXT("Index1D must be non-negative: Convert1DTo3D"), ELogVerbosity::Warning, InvalidIndexConversionParameterWarning);
+		return FIntVector::ZeroValue;
+	}
+
+	if (XSize <= 0 || YSize <= 0)
+	{
+		FFrame::KismetExecutionMessage(TEXT("XSize and YSize must be positive: Convert1DTo3D"), ELogVerbosity::Warning, InvalidIndexConversionParameterWarning);
+		return FIntVector::ZeroValue;
+	}
+
+	int32 X = Index1D % XSize;
+	int32 Y = (Index1D / XSize) % YSize;
+	int32 Z = (Index1D / (XSize * YSize));
+
+	return FIntVector{ X, Y, Z };
+}
+
+int32 UKismetMathLibrary::Convert2DTo1D(const FIntPoint& Index2D, int32 XSize)
+{
+	const bool bInvalidBounds =
+		Index2D.X < 0 || 
+		Index2D.X >= XSize || 
+		Index2D.Y < 0
+	;
+
+	if (bInvalidBounds)
+	{
+		FFrame::KismetExecutionMessage(TEXT("Index2D out of bounds: Convert2DTo1D"), ELogVerbosity::Warning, InvalidIndexConversionParameterWarning);
+		return 0;
+	}
+
+	if (XSize <= 0)
+	{
+		FFrame::KismetExecutionMessage(TEXT("XSize must be positive: Convert2DTo1D"), ELogVerbosity::Warning, InvalidIndexConversionParameterWarning);
+		return 0;
+	}
+
+	return Index2D.X + (Index2D.Y * XSize);
+}
+
+int32 UKismetMathLibrary::Convert3DTo1D(const FIntVector& Index3D, int32 XSize, int32 YSize)
+{
+	const bool bInvalidBounds =
+		Index3D.X < 0 || 
+		Index3D.X >= XSize || 
+		Index3D.Y < 0 || 
+		Index3D.Y >= YSize || 
+		Index3D.Z < 0
+	;
+
+	if (bInvalidBounds)
+	{
+		FFrame::KismetExecutionMessage(TEXT("Index3D out of bounds: Convert3DTo1D"), ELogVerbosity::Warning, InvalidIndexConversionParameterWarning);
+		return 0;
+	}
+
+	if (XSize <= 0 || YSize <= 0)
+	{
+		FFrame::KismetExecutionMessage(TEXT("XSize and YSize must be positive: Convert3DTo1D"), ELogVerbosity::Warning, InvalidIndexConversionParameterWarning);
+		return 0;
+	}
+
+	return Index3D.X + (Index3D.Y * XSize) + (Index3D.Z * XSize * YSize);
+}
 
 #undef LOCTEXT_NAMESPACE
 
