[Backout] - CL30049175
Original CL Desc
-----------------------------------------------------------------
World Partition HashSet: fix non-spatially loaded actors never loading when in a data layer or content bundle.

diff --git a/Engine/Source/Runtime/Engine/Private/WorldPartition/RuntimeHashSet/WorldPartitionRuntimeHashSetStreamingGeneration.cpp b/Engine/Source/Runtime/Engine/Private/WorldPartition/RuntimeHashSet/WorldPartitionRuntimeHashSetStreamingGeneration.cpp
--- a/Engine/Source/Runtime/Engine/Private/WorldPartition/RuntimeHashSet/WorldPartitionRuntimeHashSetStreamingGeneration.cpp
+++ b/Engine/Source/Runtime/Engine/Private/WorldPartition/RuntimeHashSet/WorldPartitionRuntimeHashSetStreamingGeneration.cpp
@@ -224,13 +224,13 @@ bool UWorldPartitionRuntimeHashSet::GenerateStreaming(UWorldPartitionStreamingPo
 				StreamingData.Name = CellDescInstance.SourcePartition->Name;
 				StreamingData.LoadingRange = CellDescInstance.SourcePartition->LoadingRange;
 
-				if (CellDescInstance.bIsSpatiallyLoaded)
+				if (RuntimeCell->IsAlwaysLoaded())
 				{
-					StreamingData.StreamingCells.Add(RuntimeCell);
+					StreamingData.NonStreamingCells.Add(RuntimeCell);
 				}
 				else
 				{
-					StreamingData.NonStreamingCells.Add(RuntimeCell);
+					StreamingData.StreamingCells.Add(RuntimeCell);
 				}
 			}
 		}
