World Partition Streaming Generation: only log actors that are going to be moved into the persistent level when generating the state dump log, not all actors in the persistent level.

diff --git a/Engine/Source/Runtime/Engine/Private/WorldPartition/RuntimeHashSet/WorldPartitionRuntimeHashSetStreamingGeneration.cpp b/Engine/Source/Runtime/Engine/Private/WorldPartition/RuntimeHashSet/WorldPartitionRuntimeHashSetStreamingGeneration.cpp
--- a/Engine/Source/Runtime/Engine/Private/WorldPartition/RuntimeHashSet/WorldPartitionRuntimeHashSetStreamingGeneration.cpp
+++ b/Engine/Source/Runtime/Engine/Private/WorldPartition/RuntimeHashSet/WorldPartitionRuntimeHashSetStreamingGeneration.cpp
@@ -258,7 +258,7 @@ void UWorldPartitionRuntimeHashSet::DumpStateLog(FHierarchicalLogArchive& Ar) co
 	Ar.Printf(TEXT("----------------------------------------------------------------------------------------------------------------"));
 
 	TArray<const UWorldPartitionRuntimeCell*> StreamingCells;
-	ForEachStreamingCells([&StreamingCells](const UWorldPartitionRuntimeCell* StreamingCell) { StreamingCells.Add(StreamingCell); return true; });
+	ForEachStreamingCells([&StreamingCells](const UWorldPartitionRuntimeCell* StreamingCell) { if (!StreamingCell->IsAlwaysLoaded()) { StreamingCells.Add(StreamingCell); } return true; });
 				
 	StreamingCells.Sort([this](const UWorldPartitionRuntimeCell& A, const UWorldPartitionRuntimeCell& B) { return A.GetFName().LexicalLess(B.GetFName()); });
 
diff --git a/Engine/Source/Runtime/Engine/Private/WorldPartition/WorldPartitionRuntimeHash.cpp b/Engine/Source/Runtime/Engine/Private/WorldPartition/WorldPartitionRuntimeHash.cpp
--- a/Engine/Source/Runtime/Engine/Private/WorldPartition/WorldPartitionRuntimeHash.cpp
+++ b/Engine/Source/Runtime/Engine/Private/WorldPartition/WorldPartitionRuntimeHash.cpp
@@ -372,34 +372,55 @@ bool UWorldPartitionRuntimeHash::PopulateGeneratorPackageForCook(const TArray<FW
 
 void UWorldPartitionRuntimeHash::DumpStateLog(FHierarchicalLogArchive& Ar) const
 {
-	if (!GIsAutomationTesting)
+	Ar.Printf(TEXT("----------------------------------------------------------------------------------------------------------------"));
+	Ar.Printf(TEXT("%s - Persistent Level"), *GetWorld()->GetName());
+	Ar.Printf(TEXT("----------------------------------------------------------------------------------------------------------------"));
 	{
-		Ar.Printf(TEXT("----------------------------------------------------------------------------------------------------------------"));
-		Ar.Printf(TEXT("%s - Persistent Level"), *GetWorld()->GetName());
-		Ar.Printf(TEXT("----------------------------------------------------------------------------------------------------------------"));
-		{
-			FHierarchicalLogArchive::FIndentScope CellIndentScope = Ar.PrintfIndent(TEXT("Content of %s Persistent Level"), *GetWorld()->GetName());
+		FHierarchicalLogArchive::FIndentScope CellIndentScope = Ar.PrintfIndent(TEXT("Content of %s Persistent Level"), *GetWorld()->GetName());
+
+		TArray<TPair<FString, FString>> Actors;
 
-			TArray<const AActor*> Actors;
-			for (const AActor* Actor : GetWorld()->PersistentLevel->Actors)
+		if (!IsRunningCookCommandlet())
+		{
+			for (const FAlwaysLoadedActorForPIE& AlwaysLoadedActor : AlwaysLoadedActorsForPIE)
 			{
-				if (Actor)
+				if (AlwaysLoadedActor.Actor.IsValid())
 				{
-					Actors.Add(Actor);
+					Actors.Add(
+					{
+						FString::Printf(TEXT("Actor Path: %s"), *AlwaysLoadedActor.Actor->GetPathName()), 
+						FString::Printf(TEXT("Actor Package: %s"), *AlwaysLoadedActor.Actor->GetPackage()->GetName())
+					});
 				}
 			}
-
-			Actors.Sort([this](const AActor& A, const AActor& B) { return A.GetFName().LexicalLess(B.GetFName()); });
-
-			Ar.Printf(TEXT("Always loaded Actor Count: %d "), Actors.Num());
-			for (const AActor* Actor : Actors)
+		}
+		else
+		{
+			for (UWorldPartitionRuntimeCell* Cell : GetAlwaysLoadedCells())
 			{
-				Ar.Printf(TEXT("Actor Path: %s"), *Actor->GetPathName());
-				Ar.Printf(TEXT("Actor Package: %s"), *Actor->GetPackage()->GetName());
+				const UWorldPartitionRuntimeLevelStreamingCell* RuntimeCell = CastChecked<UWorldPartitionRuntimeLevelStreamingCell>(Cell);
+
+				for (const FWorldPartitionRuntimeCellObjectMapping& Package : RuntimeCell->GetPackages())
+				{
+					Actors.Add(
+					{
+						FString::Printf(TEXT("Actor Path: %s"), *Package.Path.ToString()), 
+						FString::Printf(TEXT("Actor Package: %s"), *Package.Package.ToString())
+					});
+				}
 			}
 		}
-		Ar.Printf(TEXT(""));
+
+		Actors.Sort();
+
+		Ar.Printf(TEXT("Always loaded Actor Count: %d "), Actors.Num());
+		for (const TPair<FString, FString>& Actor : Actors)
+		{
+			Ar.Print(*Actor.Key);
+			Ar.Print(*Actor.Value);
+		}
 	}
+	Ar.Printf(TEXT(""));
 }
 
 void UWorldPartitionRuntimeHash::ForceExternalActorLevelReference(bool bForceExternalActorLevelReferenceForPIE)
diff --git a/Engine/Source/Runtime/Engine/Private/WorldPartition/WorldPartitionRuntimeSpatialHash.cpp b/Engine/Source/Runtime/Engine/Private/WorldPartition/WorldPartitionRuntimeSpatialHash.cpp
--- a/Engine/Source/Runtime/Engine/Private/WorldPartition/WorldPartitionRuntimeSpatialHash.cpp
+++ b/Engine/Source/Runtime/Engine/Private/WorldPartition/WorldPartitionRuntimeSpatialHash.cpp
@@ -1341,7 +1341,10 @@ void UWorldPartitionRuntimeSpatialHash::DumpStateLog(FHierarchicalLogArchive& Ar
 				{
 					for (const TObjectPtr<UWorldPartitionRuntimeCell>& Cell : LayerCell.GridCells)
 					{
-						Cells.Add(Cell);
+						if (!Cell->IsAlwaysLoaded())
+						{
+							Cells.Add(Cell);
+						}
 					}
 				}
 
