TextureRenderTarget clean up some hard-coding of specific format support; make it more generic
mostly comments clarifying things

diff --git a/Engine/Plugins/Media/ImgMedia/Source/ImgMediaEditor/Private/Widgets/SImgMediaProcessEXR.cpp b/Engine/Plugins/Media/ImgMedia/Source/ImgMediaEditor/Private/Widgets/SImgMediaProcessEXR.cpp
--- a/Engine/Plugins/Media/ImgMedia/Source/ImgMediaEditor/Private/Widgets/SImgMediaProcessEXR.cpp
+++ b/Engine/Plugins/Media/ImgMedia/Source/ImgMediaEditor/Private/Widgets/SImgMediaProcessEXR.cpp
@@ -426,7 +426,7 @@ void SImgMediaProcessEXR::ProcessImageCustomRawData(TArray64<uint8>& RawData,
 	int32 TileWidth = InTileWidth;
 	int32 TileHeight = InTileHeight;
 	int32 BytesPerPixel = RawData.Num() / (Width * Height);
-	int32 BytesPerPixelPerChannel = 16 / 8;
+	int32 BytesPerPixelPerChannel = 16 / 8; // <- 2 bytes per channel, assumes RGBA16F
 	int32 NumChannels = BytesPerPixel / BytesPerPixelPerChannel;
 	int32 DestNumChannels = NumChannels;
 
@@ -931,10 +931,15 @@ void SImgMediaProcessEXR::HandleProcessing()
 					DrawTextureToRenderTarget();
 
 					// Process this render.
-					TArray64<uint8> RawData;
-					bool bReadSuccess = FImageUtils::GetRawData(RenderTarget, RawData);
+					FImage Image;
+					bool bReadSuccess = FImageUtils::GetRenderTargetImage(RenderTarget, Image);
+
 					if (bReadSuccess)
 					{
+						Image.ChangeFormat(ERawImageFormat::RGBA16F,EGammaSpace::Linear);
+						// "Image" is not used in following code, just Image.RawData as uint8
+						//	this file is hard-coded to assume 16F pixels
+
 						int32 Width = RenderTarget->GetSurfaceWidth();
 						int32 Height = RenderTarget->GetSurfaceHeight();
 						int32 InTileWidth = Options->bEnableTiling ? Options->TileSizeX : 0;
@@ -946,7 +951,7 @@ void SImgMediaProcessEXR::HandleProcessing()
 						FString FileName = FString::Printf(TEXT("image%05d.exr"), CurrentFrameIndex);
 						FString Name = FPaths::Combine(OutPath, FileName);
 
-						Async(EAsyncExecution::Thread, [this, RawData = MoveTemp(RawData), Width, Height,
+						Async(EAsyncExecution::Thread, [this, RawData = MoveTemp(Image.RawData), Width, Height,
 							InTileWidth, InTileHeight, TileBorder, bEnableMips, Name]() mutable
 						{
 							ProcessImageCustomRawData(RawData, Width, Height,
diff --git a/Engine/Source/Editor/UnrealEd/Private/Factories/EditorFactories.cpp b/Engine/Source/Editor/UnrealEd/Private/Factories/EditorFactories.cpp
--- a/Engine/Source/Editor/UnrealEd/Private/Factories/EditorFactories.cpp
+++ b/Engine/Source/Editor/UnrealEd/Private/Factories/EditorFactories.cpp
@@ -5063,6 +5063,16 @@ URenderTargetExporterPNG::URenderTargetExporterPNG(const FObjectInitializer& Obj
 	FormatDescription.Add(TEXT("PNG"));
 }
 
+static int GetBitsPerComponent(EPixelFormat Format)
+{
+	if ( Format == PF_A2B10G10R10 ) return 10; // doesn't handle heterogenous bit counts well
+
+	const FPixelFormatInfo & Info = GPixelFormats[Format];
+	// rounds down
+	int BitsPerComponent = ( Info.BlockBytes * 8 ) / ( Info.BlockSizeX * Info.BlockSizeY * Info.BlockSizeZ * Info.NumComponents );
+	return BitsPerComponent;
+}
+
 bool URenderTargetExporterPNG::SupportsObject(UObject* Object) const
 {
 	if (Super::SupportsObject(Object))
@@ -5073,7 +5083,7 @@ bool URenderTargetExporterPNG::SupportsObject(UObject* Object) const
 		{
 			EPixelFormat PixelFormat = TexRT2D->GetFormat();
 
-			return PixelFormat == PF_B8G8R8A8;
+			return GetBitsPerComponent(PixelFormat) <= 8;
 		}
 	}
 	return false;
@@ -5102,8 +5112,8 @@ bool URenderTargetExporterEXR::SupportsObject(UObject* Object) const
 		if (TexRT2D)
 		{
 			EPixelFormat PixelFormat = TexRT2D->GetFormat();
-
-			return PixelFormat == PF_FloatRGBA;
+			
+			return GetBitsPerComponent(PixelFormat) > 8;
 		}
 	}
 	return false;
diff --git a/Engine/Source/Runtime/Engine/Classes/Engine/TextureRenderTarget2D.h b/Engine/Source/Runtime/Engine/Classes/Engine/TextureRenderTarget2D.h
--- a/Engine/Source/Runtime/Engine/Classes/Engine/TextureRenderTarget2D.h
+++ b/Engine/Source/Runtime/Engine/Classes/Engine/TextureRenderTarget2D.h
@@ -246,6 +246,10 @@ PRAGMA_ENABLE_DEPRECATION_WARNINGS
 
 	bool IsSRGB() const
 	{
+		// ?? note: UTextureRenderTarget::TargetGamma is ignored here
+
+		// ?? note: GetDisplayGamma forces linear for some float formats, but this doesn't
+
 		if (OverrideFormat == PF_Unknown)
 		{
 			return RenderTargetFormat == RTF_RGBA8_SRGB;
diff --git a/Engine/Source/Runtime/Engine/Private/ImageUtils.cpp b/Engine/Source/Runtime/Engine/Private/ImageUtils.cpp
--- a/Engine/Source/Runtime/Engine/Private/ImageUtils.cpp
+++ b/Engine/Source/Runtime/Engine/Private/ImageUtils.cpp
@@ -314,83 +314,114 @@ bool FImageUtils::ExportTextureSourceToDDS(TArray64<uint8> & OutData, UTexture *
 #endif
 }
 
-/**
- * Returns data containing the pixmap of the passed in rendertarget.
- * @param TexRT - The 2D rendertarget from which to read pixmap data.
- * @param RawData - an array to be filled with pixel data.
- * @return true if RawData has been successfully filled.
- */
 bool FImageUtils::GetRawData(UTextureRenderTarget2D* TexRT, TArray64<uint8>& RawData)
 {
-	FRenderTarget* RenderTarget = TexRT->GameThread_GetRenderTargetResource();
-	EPixelFormat Format = TexRT->GetFormat();
+	// DEPRECATED , use GetRenderTargetImage
 
-	int32 ImageBytes = CalculateImageBytes(TexRT->SizeX, TexRT->SizeY, 0, Format);
-	RawData.AddUninitialized(ImageBytes);
-	bool bReadSuccess = false;
-	switch (Format)
-	{
-	case PF_FloatRGBA:
-		{
-		TArray<FFloat16Color> FloatColors;
-		bReadSuccess = RenderTarget->ReadFloat16Pixels(FloatColors);
-		FMemory::Memcpy(RawData.GetData(), FloatColors.GetData(), ImageBytes);
-		}
-		break;
-	case PF_B8G8R8A8:
-		bReadSuccess = RenderTarget->ReadPixelsPtr((FColor*)RawData.GetData());
-		break;
-	default:
-		// bReadSuccess == false
-		UE_LOG(LogImageUtils,Warning,TEXT("RenderTarget GetRawData PixelFormat no supported : %s") , *(StaticEnum<EPixelFormat>()->GetDisplayNameTextByValue(Format).ToString()) );
-		break;
-	}
-	if (bReadSuccess == false)
+	RawData.Empty();
+
+	FImage Image;
+	if ( ! GetRenderTargetImage(TexRT,Image) )
 	{
-		RawData.Empty();
+		return false;
 	}
-	return bReadSuccess;
+
+	RawData = MoveTemp(Image.RawData);
+	return true;
+}
+
+static int GetBitsPerComponent(EPixelFormat Format)
+{
+	if ( Format == PF_A2B10G10R10 ) return 10; // doesn't handle heterogenous bit counts well
+
+	const FPixelFormatInfo & Info = GPixelFormats[Format];
+	// rounds down
+	int BitsPerComponent = ( Info.BlockBytes * 8 ) / ( Info.BlockSizeX * Info.BlockSizeY * Info.BlockSizeZ * Info.NumComponents );
+	return BitsPerComponent;
 }
 
 bool FImageUtils::GetRenderTargetImage(UTextureRenderTarget2D* TexRT, FImage & Image)
 {
 	Image = FImage();
 	
-	// see ETextureRenderTargetFormat
-	// for allowed formats
+	FRenderTarget* RenderTarget = TexRT->GameThread_GetRenderTargetResource();
+	EPixelFormat RTFormat = TexRT->GetFormat();
 
-	TArray64<uint8> RawData;
-	if ( ! GetRawData(TexRT,RawData) )
+	bool Use8Bit;
+	
+	if ( GetBitsPerComponent(RTFormat) <= 8 )
 	{
-		return false;
+		Use8Bit = true;
+	}
+	else if ( IsFloatFormat(RTFormat) || IsHDR(RTFormat) )
+	{
+		Use8Bit = false;
+	}
+	else if ( IsDepthOrStencilFormat(RTFormat) )
+	{
+		Use8Bit = false;
+	}
+	else
+	{
+		// eg. 16-bit integer
+		Use8Bit = false;
 	}
 	
-	// @todo Oodle : in theory the RenderTarget knows its gamma
-	//   but from what I've seen it's usually wrong?
-	// TexRT->GetDisplayGamma() or TexRT->SRGB or TexRT->IsSRGB() , OMG
-
-	Image.RawData = MoveTemp(RawData);
-	Image.SizeX = TexRT->SizeX;
-	Image.SizeY = TexRT->SizeY;
-	Image.NumSlices = 1;
-	
-	EPixelFormat PixelFormat = TexRT->GetFormat();
-	switch(PixelFormat)
-	{
-	case PF_FloatRGBA:
-		Image.Format = ERawImageFormat::RGBA16F;
-		Image.GammaSpace = EGammaSpace::Linear;
-		break;	
-	case PF_B8G8R8A8:
-		Image.Format = ERawImageFormat::BGRA8;
-		Image.GammaSpace = EGammaSpace::sRGB;
-		break;
-	default:
-		// should not get here because GetRawData would have returned false
-		check(0);
-		return false;
+	if ( RTFormat == PF_FloatRGBA )
+	{
+		// ReadFloat16Pixels does no conversions
+		//	must be used only exactly with FloatRGBA type
+
+		Image.Init(TexRT->SizeX,TexRT->SizeY,ERawImageFormat::RGBA16F,EGammaSpace::Linear);
+		
+		TArray<FFloat16Color> Colors;
+		if ( ! RenderTarget->ReadFloat16Pixels(Colors) )
+		{
+			return false;
+		}
+
+		check( Image.GetImageSizeBytes() == Colors.Num() * sizeof(Colors[0]) );
+		memcpy( &Image.RawData[0], &Colors[0], Image.GetImageSizeBytes() );
+	}
+	else if ( Use8Bit )
+	{
+		// ?? not clear TexRT->IsSRGB is right , see other notes on various issues there
+		//	mainly we are trying to catch the check for whether the _SRGB or non _SRGB BGRA8 format as chosen
+		EGammaSpace GammaSpace = TexRT->IsSRGB() ? EGammaSpace::sRGB : EGammaSpace::Linear;
+		
+		Image.Init(TexRT->SizeX,TexRT->SizeY,ERawImageFormat::BGRA8,GammaSpace);
+		
+		FReadSurfaceDataFlags InFlags(RCM_MinMax, CubeFace_MAX);
+
+		InFlags.SetLinearToGamma( GammaSpace == EGammaSpace::sRGB );
+
+		TArray<FColor> Colors;
+		if ( ! RenderTarget->ReadPixels(Colors,InFlags) )
+		{
+			return false;
+		}
+
+		check( Image.GetImageSizeBytes() == Colors.Num() * sizeof(Colors[0]) );
+		memcpy( &Image.RawData[0], &Colors[0], Image.GetImageSizeBytes() );
 	}
+	else // use F32
+	{
+		Image.Init(TexRT->SizeX,TexRT->SizeY,ERawImageFormat::RGBA32F,EGammaSpace::Linear);
+		
+		FReadSurfaceDataFlags InFlags(RCM_MinMax, CubeFace_MAX);
+
+		InFlags.SetLinearToGamma( false );
+
+		TArray<FLinearColor> Colors;
+		if ( ! RenderTarget->ReadLinearColorPixels(Colors,InFlags) )
+		{
+			return false;
+		}
 
+		check( Image.GetImageSizeBytes() == Colors.Num() * sizeof(Colors[0]) );
+		memcpy( &Image.RawData[0], &Colors[0], Image.GetImageSizeBytes() );
+	}	
+	
 	return true;
 }
 
@@ -686,6 +717,9 @@ UTexture2D* FImageUtils::CreateTexture2D(int32 SrcWidth, int32 SrcHeight, const
 
 void FImageUtils::CropAndScaleImage( int32 SrcWidth, int32 SrcHeight, int32 DesiredWidth, int32 DesiredHeight, const TArray<FColor> &SrcData, TArray<FColor> &DstData  )
 {
+	// DEPRECATED , uses the bad ImageResize
+	//	this is used by the old Thumbnail code, not the new paths
+
 	// Get the aspect ratio, and calculate the dimension of the image to crop
 	float DesiredAspectRatio = (float)DesiredWidth/(float)DesiredHeight;
 
@@ -952,122 +986,13 @@ UVolumeTexture* FImageUtils::CreateCheckerboardVolumeTexture(FColor ColorOne, FC
 HDR file format helper.
 ------------------------------------------------------------------------------*/
 // DEPRECATED
+// only used for cube maps for GenerateLongLatUnwrap
 // do not use HDR for exports, it is very lossy, use EXR instead for float output
 // if you need HDR, do not use this export code
 // instead use HdrImageWrapper via FImageUtils::CompressImage
 class FHDRExportHelper
 {
 public:
-	/**
-	* Writes HDR format image to an FArchive
-	* @param TexRT - A 2D source render target to read from.
-	* @param Ar - Archive object to write HDR data to.
-	* @return true on successful export.
-	*/
-	bool ExportHDR(UTextureRenderTarget2D* TexRT, FArchive& Ar)
-	{
-		check(TexRT != nullptr);
-		FRenderTarget* RenderTarget = TexRT->GameThread_GetRenderTargetResource();
-		Size = RenderTarget->GetSizeXY();
-		Format = TexRT->GetFormat();
-
-		TArray64<uint8> RawData;
-		bool bReadSuccess = FImageUtils::GetRawData(TexRT, RawData);
-		if (bReadSuccess)
-		{
-			WriteHDRImage(RawData, Ar);
-			return true;
-		}
-		return false;
-	}
-
-	/**
-	* Writes HDR format image to an FArchive
-	* @param TexRT - A 2D source render target to read from.
-	* @param Ar - Archive object to write HDR data to.
-	* @return true on successful export.
-	*/
-	bool ExportHDR(UTexture2D* Texture, FArchive& Ar)
-	{
-		check(Texture != nullptr);
-		bool bReadSuccess = true;
-		TArray64<uint8> RawData;
-
-#if WITH_EDITORONLY_DATA
-		Size = FIntPoint(Texture->Source.GetSizeX(), Texture->Source.GetSizeY());
-
-		bReadSuccess = Texture->Source.GetMipData(RawData, 0);
-		const ETextureSourceFormat NewFormat = Texture->Source.GetFormat();
-
-		// DEPRECATED
-		// not a general purpose HDR exporter
-		// can't write F32 or BGRE8
-		// do not use this except for longlat cubemaps
-		if (NewFormat == TSF_BGRA8)
-		{
-			Format = PF_B8G8R8A8;
-		}
-		else if (NewFormat == TSF_RGBA16F)
-		{
-			Format = PF_FloatRGBA;
-		}
-		else
-		{
-			bReadSuccess = false;			
-			FMessageLog("ImageUtils").Warning(LOCTEXT("ExportHDRUnsupportedSourceTextureFormat", "Unsupported source texture format provided."));
-		}
-#else
-		TArray<uint8*> RawData2;
-		Size = Texture->GetImportedSize();
-		RawData2.AddZeroed(Texture->GetNumMips());
-		// this is PlatformData GetMipData , not Source :
-		Texture->GetMipData(0, (void**)RawData2.GetData());
-		const EPixelFormat NewFormat = Texture->GetPixelFormat();
-
-		if (Texture->GetPlatformData()->Mips.Num() == 0)
-		{
-			bReadSuccess = false;
-			FMessageLog("ImageUtils").Warning(FText::Format(LOCTEXT("ExportHDRFailedToReadMipData", "Failed to read Mip Data in: '{0}'"), FText::FromString(Texture->GetName())));
-		}
-
-		if (NewFormat == PF_B8G8R8A8)
-		{
-			Format = PF_B8G8R8A8;
-		}
-		else if (NewFormat == PF_FloatRGBA)
-		{
-			Format = PF_FloatRGBA;
-		}
-		else
-		{
-			bReadSuccess = false;
-			FMessageLog("ImageUtils").Warning(LOCTEXT("ExportHDRUnsupportedTextureFormat", "Unsupported texture format provided."));
-		}
-
-		//Put first mip data into usable array
-		if (bReadSuccess)
-		{
-			const uint32 TotalSize = Texture->GetPlatformData()->Mips[0].BulkData.GetBulkDataSize();
-			RawData.AddZeroed(TotalSize);
-			FMemory::Memcpy(RawData.GetData(), RawData2[0], TotalSize);
-		}
-
-		//Deallocate the mip data
-		for (auto MipData : RawData2)
-		{
-			FMemory::Free(MipData);
-		}
-
-#endif // WITH_EDITORONLY_DATA
-
-		if (bReadSuccess)
-		{
-			WriteHDRImage(RawData, Ar);
-			return true;
-		}
-
-		return false;
-	}
 
 	/**
 	* Writes HDR format image to an FArchive
@@ -1250,8 +1175,6 @@ private:
 
 bool FImageUtils::ExportRenderTarget2DAsHDR(UTextureRenderTarget2D* TexRT, FArchive& Ar)
 {
-	UE_LOG(LogImageUtils, Warning, TEXT("HDR file format is very lossy, use EXR or PNG instead.") );
-
 	FImage Image;
 	if ( ! GetRenderTargetImage(TexRT,Image) )
 	{
@@ -1307,6 +1230,27 @@ bool FImageUtils::ExportRenderTarget2DAsEXR(UTextureRenderTarget2D* TexRT, FArch
 	return true;
 }
 
+bool FImageUtils::ExportTexture2DAsHDR(UTexture2D* Tex, FArchive& Ar)
+{
+	UE_LOG(LogImageUtils, Warning, TEXT("HDR file format is very lossy, use EXR or PNG instead.") );
+	
+	FImage Image;
+	if ( ! GetTexture2DSourceImage(Tex,Image) )
+	{
+		return false;
+	}
+
+	TArray64<uint8> CompressedData;
+	if ( ! CompressImage(CompressedData,TEXT("HDR"),Image) )
+	{
+		return false;
+	}
+
+	Ar.Serialize((void*)CompressedData.GetData(), CompressedData.GetAllocatedSize());
+
+	return true;
+}
+
 // if Texture source is available, get it as an FImage
 bool FImageUtils::GetTexture2DSourceImage(UTexture2D* Texture, FImage & OutImage)
 {
@@ -1338,28 +1282,13 @@ bool FImageUtils::GetTexture2DSourceImage(UTexture2D* Texture, FImage & OutImage
 #endif
 }
 
-bool FImageUtils::ExportTexture2DAsHDR(UTexture2D* TexRT, FArchive& Ar)
-{
-	UE_LOG(LogImageUtils, Warning, TEXT("HDR file format is very lossy, use EXR or PNG instead.") );
-
-	// could use GetTexture2DSourceImage
-	// then CompressImage
-	// for arbitrary format exports
-	// but just leave it alone for now
-
-	// this is only used by UKismetRenderingLibrary::ExportTexture2D
-	//	so we should take filename as input to auto-detect format
-
-	FHDRExportHelper Exporter;
-	return Exporter.ExportHDR(TexRT, Ar);
-}
-
 UTexture2D* FImageUtils::ImportFileAsTexture2D(const FString& Filename)
 {
 	UTexture2D* NewTexture = nullptr;
 	TArray64<uint8> Buffer;
 	if (FFileHelper::LoadFileToArray(Buffer, *Filename))
 	{
+		// note this make a Transient / PlatformData only Texture (no TextureSource)
 		NewTexture = FImageUtils::ImportBufferAsTexture2D(Buffer);
 
 		if(!NewTexture)
@@ -1384,6 +1313,7 @@ UTexture2D* FImageUtils::ImportBufferAsTexture2D(TArrayView64<const uint8> Buffe
 		return nullptr;
 	}
 
+	// note this make a Transient / PlatformData only Texture (no TextureSource)
 	return CreateTexture2DFromImage(Image);
 }
 
@@ -1406,6 +1336,7 @@ UTexture2D* FImageUtils::CreateTexture2DFromImage(const FImageView & Image)
 	int64 MipDataSize = NewTexture->GetPlatformData()->Mips[0].BulkData.GetBulkDataSize();
 
 	FImageView MipImage(MipData,Image.SizeX,Image.SizeY,1,PixelFormatRawFormat,Image.GammaSpace);
+	check( MipImage.GetImageSizeBytes() <= MipDataSize ); // is it exactly == ?
 
 	// copy into texture and convert if necessary :
 	FImageCore::CopyImage(Image,MipImage);
diff --git a/Engine/Source/Runtime/Engine/Private/TextureRenderTarget.cpp b/Engine/Source/Runtime/Engine/Private/TextureRenderTarget.cpp
--- a/Engine/Source/Runtime/Engine/Private/TextureRenderTarget.cpp
+++ b/Engine/Source/Runtime/Engine/Private/TextureRenderTarget.cpp
@@ -13,6 +13,7 @@
 #include "RenderUtils.h"
 #include "TextureCompiler.h"
 #include "Logging/MessageLog.h"
+#include "Hash/xxhash.h"
 
 #include UE_INLINE_GENERATED_CPP_BY_NAME(TextureRenderTarget)
 
@@ -23,8 +24,9 @@ UTextureRenderTarget
 UTextureRenderTarget::UTextureRenderTarget(const FObjectInitializer& ObjectInitializer)
 	: Super(ObjectInitializer)
 {
+	TargetGamma = 0.f; // 0 means inherit from resource which is 2.2
 	NeverStream = true;
-	SRGB = true;
+	SRGB = true; // <- odd, not usually what you want; often replaced with IsSRGB()
 	LODGroup = TEXTUREGROUP_RenderTarget;	
 	bNeedsTwoCopies = false;
 	bCanCreateUAV = false;
@@ -93,6 +95,7 @@ ETextureSourceFormat UTextureRenderTarget::ValidateTextureFormatForConversionToT
 
 	// Return what ETextureSourceFormat corresponds to this EPixelFormat (must match the conversion capabilities of UTextureRenderTarget::UpdateTexture) : 
 	ETextureSourceFormat TextureFormat = TSF_Invalid;
+	// @todo Oodle : just support all formats and remove this
 	switch (InFormat)
 	{
 	case PF_B8G8R8A8:
@@ -170,7 +173,7 @@ bool UTextureRenderTarget::UpdateTexture(UTexture* InTexture, EConstructTextureF
 		return false;
 	}
 
-	TextureCompressionSettings CompressionSettingsForTexture = (PixelFormat == EPixelFormat::PF_FloatRGBA) ? TC_HDR : TC_Default;
+	TextureCompressionSettings CompressionSettingsForTexture = IsHDR(PixelFormat) ? TC_HDR : TC_Default;
 
 	const int32 NumMips = 1; // There's only support for mip 0 ATM
 
@@ -183,7 +186,7 @@ bool UTextureRenderTarget::UpdateTexture(UTexture* InTexture, EConstructTextureF
 	bTextureChanging |= InTexture->Source.GetFormat() != TextureSourceFormat;
 	bTextureChanging |= InTexture->CompressionSettings != CompressionSettingsForTexture;
 
-	uint32 OldDataHash = 0;
+	FXxHash64 OldDataHash;
 
 	// Hashing the content is only really useful if none of the other texture settings have changed.
 	const bool bHashContent = !bTextureChanging;
@@ -191,7 +194,7 @@ bool UTextureRenderTarget::UpdateTexture(UTexture* InTexture, EConstructTextureF
 	{
 		const uint8* OldData = InTexture->Source.LockMipReadOnly(0);
 		const int32 OldDataSize = InTexture->Source.CalcMipSize(0);
-		OldDataHash = FCrc::MemCrc32(OldData, OldDataSize);
+		OldDataHash = FXxHash64::HashBuffer(OldData, OldDataSize);
 		InTexture->Source.UnlockMip(0);
 	}
 
@@ -207,6 +210,8 @@ bool UTextureRenderTarget::UpdateTexture(UTexture* InTexture, EConstructTextureF
 	TArray<uint8> NewData;
 	NewData.SetNumUninitialized(DestinationSize);
 
+	// @todo Oodle : make this more generic like ImageUtils::GetRenderTargetImage
+
 	// Read the data surface by surface (i.e. for a 2D array or a volume: slice by slice, for a cubemap: face by face and for a cubemap array: face by face, slice by slice
 	for (int32 SurfaceIndex = 0; SurfaceIndex < NumSurfaces; ++SurfaceIndex)
 	{
@@ -256,7 +261,10 @@ bool UTextureRenderTarget::UpdateTexture(UTexture* InTexture, EConstructTextureF
 		{
 			TArray<FFloat16Color> NewDataFloat16Color;
 			check(NumBytesPerPixel == sizeof(FFloat16Color));
+
+			// ReadFloat16Pixels is only allowed if PF is exactly FloatRGBA
 			RenderTarget->ReadFloat16Pixels(NewDataFloat16Color, ReadSurfaceDataFlags);
+
 			check(NewDataFloat16Color.Num() == NumPixelsPerSurface);
 			const int32 NumBytes = NewDataFloat16Color.Num() * NewDataFloat16Color.GetTypeSize();
 			check(NumBytes == SurfaceByteData.Num());
@@ -273,18 +281,22 @@ bool UTextureRenderTarget::UpdateTexture(UTexture* InTexture, EConstructTextureF
 			}
 			else if (InFlags & CTF_RemapAlphaAsMasked)
 			{
+				// ?? unclear this path makes any sense on Float16
+				FFloat16 F16Zero; F16Zero.SetZero();
+				FFloat16 F16One; F16One.SetOne();
+
 				// if the target was rendered with a masked texture, then the depth will probably have been written instead of 0/255 for the
 				// alpha, and the depth when unwritten will be 255, so remap 255 to 0 (masked out area) and anything else as 255 (written to area)
 				for (int32 PixelIndex = 0; PixelIndex < NumPixelsPerSurface; ++PixelIndex)
 				{
-					SurfaceFloat16ColorData[PixelIndex].A = (SurfaceFloat16ColorData[PixelIndex].A == 255) ? 0 : 255;
+					SurfaceFloat16ColorData[PixelIndex].A = (SurfaceFloat16ColorData[PixelIndex].A >= 1.f) ? F16Zero : F16One;
 				}
 			}
 			else if (InFlags & CTF_ForceOpaque)
 			{
 				for (int32 PixelIndex = 0; PixelIndex < NumPixelsPerSurface; ++PixelIndex)
 				{
-					SurfaceFloat16ColorData[PixelIndex].A = 255;
+					SurfaceFloat16ColorData[PixelIndex].A.SetOne();
 				}
 			}
 		}
@@ -308,7 +320,7 @@ bool UTextureRenderTarget::UpdateTexture(UTexture* InTexture, EConstructTextureF
 
 	if (bHashContent)
 	{
-		uint32 NewDataHash = FCrc::MemCrc32(NewData.GetData(), NewData.Num());
+		FXxHash64 NewDataHash = FXxHash64::HashBuffer(NewData.GetData(), NewData.Num());
 		bTextureChanging = OldDataHash != NewDataHash;
 	}
 
@@ -405,11 +417,20 @@ UTexture* UTextureRenderTarget::ConstructTexture(UObject* InOuter, const FString
  */
 bool FTextureRenderTargetResource::IsSupportedFormat( EPixelFormat Format )
 {
+	// this should at least support all the formats in GetPixelFormatFromRenderTargetFormat
+
 	switch( Format )
 	{
+	case PF_G8:
+	case PF_R8G8:
 	case PF_B8G8R8A8:
 	case PF_R8G8B8A8:
 	case PF_A16B16G16R16:
+	case PF_R16F:
+	case PF_G16R16F:
+	case PF_R32_FLOAT:
+	case PF_G32R32F:
+	case PF_A32B32G32R32F:
 	case PF_FloatRGB:
 	case PF_FloatRGBA: // for exporting materials to .obj/.mtl
 	case PF_FloatR11G11B10: //Pixel inspector for Reading HDR Color
@@ -437,6 +458,7 @@ const FTextureRHIRef& FTextureRenderTargetResource::GetShaderResourceTexture() c
 */
 float FTextureRenderTargetResource::GetDisplayGamma() const
 {
+	// when we say we want a 2.2 gamma, what we actually mean is that we want SRGB conversion in most cases
 	return 2.2f;  
 }
 
diff --git a/Engine/Source/Runtime/Engine/Private/TextureRenderTarget2D.cpp b/Engine/Source/Runtime/Engine/Private/TextureRenderTarget2D.cpp
--- a/Engine/Source/Runtime/Engine/Private/TextureRenderTarget2D.cpp
+++ b/Engine/Source/Runtime/Engine/Private/TextureRenderTarget2D.cpp
@@ -43,10 +43,15 @@ UTextureRenderTarget2D::UTextureRenderTarget2D(const FObjectInitializer& ObjectI
 	NumMips = 0;
 	ClearColor = FLinearColor(0.0f, 0.0f, 0.0f, 1.0f);
 	OverrideFormat = PF_Unknown;
-	bForceLinearGamma = true;
+	bForceLinearGamma = true; // <<-- if you set RTF_RGBA8_SRGB, this is turned off
 	MipsSamplerFilter = Filter;
 	MipsAddressU = TA_Clamp;
 	MipsAddressV = TA_Clamp;
+
+	// note UTextureRenderTarget::UTextureRenderTarget set SRGB = true
+	// later SRGB may be set = IsSRGB();
+
+	// see also UTextureRenderTarget::TargetGamma
 }
 
 FTextureResource* UTextureRenderTarget2D::CreateResource()
@@ -127,6 +132,10 @@ void UTextureRenderTarget2D::InitAutoFormat(uint32 InSizeX, uint32 InSizeY)
 {
 	check(InSizeX > 0 && InSizeY > 0);
 
+	// ?? missing ?
+	//OverrideFormat = PF_Unknown;
+	//bForceLinearGamma = true;
+
 	// set required size
 	SizeX = InSizeX;
 	SizeY = InSizeY;
@@ -358,6 +367,7 @@ bool UTextureRenderTarget2D::CanConvertToTexture(ETextureSourceFormat& OutTextur
 {
 	const EPixelFormat LocalFormat = GetFormat();
 	// These are the formats currently available for conversion to texture for UTextureRenderTarget2D : 
+	// @todo Oodle : get rid of this, support all
 	const ETextureSourceFormat TextureSourceFormat = ValidateTextureFormatForConversionToTextureInternal(LocalFormat, { PF_G8, PF_R8G8, PF_B8G8R8A8, PF_FloatRGBA }, OutErrorMessage);
 	if (TextureSourceFormat == TSF_Invalid)
 	{
@@ -642,9 +652,12 @@ float FTextureRenderTarget2DResource::GetDisplayGamma() const
 	{
 		return Owner->TargetGamma;
 	}
+	// ?? special casing just two of the float PixelFormats to force 1.0 gamma here is inconsistent
+	//		(there are lots of other float formats)
 	if (Format == PF_FloatRGB || Format == PF_FloatRGBA || Owner->bForceLinearGamma )
 	{
 		return 1.0f;
 	}
+	// ignores Owner->IsSRGB() ?
 	return FTextureRenderTargetResource::GetDisplayGamma();
 }
diff --git a/Engine/Source/Runtime/Engine/Private/TextureRenderTargetCube.cpp b/Engine/Source/Runtime/Engine/Private/TextureRenderTargetCube.cpp
--- a/Engine/Source/Runtime/Engine/Private/TextureRenderTargetCube.cpp
+++ b/Engine/Source/Runtime/Engine/Private/TextureRenderTargetCube.cpp
@@ -54,6 +54,9 @@ void UTextureRenderTargetCube::Init(uint32 InSizeX, EPixelFormat InFormat)
 
 void UTextureRenderTargetCube::InitAutoFormat(uint32 InSizeX)
 {
+	// ?? looks like this is missing :
+	//OverrideFormat = PF_Unknown;
+
 	check(InSizeX > 0);
 	check(!(InSizeX % GPixelFormats[GetFormat()].BlockSizeX));
 	check(FTextureRenderTargetResource::IsSupportedFormat(GetFormat()));
diff --git a/Engine/Source/Runtime/Engine/Private/TextureRenderTargetVolume.cpp b/Engine/Source/Runtime/Engine/Private/TextureRenderTargetVolume.cpp
--- a/Engine/Source/Runtime/Engine/Private/TextureRenderTargetVolume.cpp
+++ b/Engine/Source/Runtime/Engine/Private/TextureRenderTargetVolume.cpp
@@ -344,6 +344,17 @@ bool FTextureRenderTargetVolumeResource::ReadFloat16Pixels(TArray<FFloat16Color>
 	{
 		InSrcRect = FIntRect(0, 0, GetSizeXY().X, GetSizeXY().Y);
 	}
+	
+	/**
+
+	deep inside RHI this will crash unless the format is one of these :
+	no checks are done in outer code
+
+	bool bIsRGBAFmt = TextureDesc.Format == GPixelFormats[PF_FloatRGBA].PlatformFormat;
+	bool bIsR16FFmt = TextureDesc.Format == GPixelFormats[PF_R16F].PlatformFormat;
+	bool bIsR32FFmt = TextureDesc.Format == GPixelFormats[PF_R32_FLOAT].PlatformFormat;
+	check(bIsRGBAFmt || bIsR16FFmt || bIsR32FFmt);
+	**/
 
 	OutImageData.Reset();
 
diff --git a/Engine/Source/Runtime/Engine/Private/UnrealClient.cpp b/Engine/Source/Runtime/Engine/Private/UnrealClient.cpp
--- a/Engine/Source/Runtime/Engine/Private/UnrealClient.cpp
+++ b/Engine/Source/Runtime/Engine/Private/UnrealClient.cpp
@@ -91,6 +91,10 @@ PRAGMA_ENABLE_DEPRECATION_WARNINGS
 
 bool FRenderTarget::ReadFloat16Pixels(TArray<FFloat16Color>& OutImageData, FReadSurfaceDataFlags InFlags, FIntRect InSrcRect)
 {
+	// if the RenderTarget is not EXACTLY PF_FloatRGBA , this will check down in the RHI
+	// (eg. PF_FloatRGB will fail)
+	check( GetShaderResourceTexture()->GetDesc().Format == PF_FloatRGBA );
+
 	if (InSrcRect == FIntRect(0, 0, 0, 0))
 	{
 		InSrcRect = FIntRect(0, 0, GetSizeXY().X, GetSizeXY().Y);
diff --git a/Engine/Source/Runtime/Engine/Public/ImageUtils.h b/Engine/Source/Runtime/Engine/Public/ImageUtils.h
--- a/Engine/Source/Runtime/Engine/Public/ImageUtils.h
+++ b/Engine/Source/Runtime/Engine/Public/ImageUtils.h
@@ -232,6 +232,9 @@ public:
 	 * @param InParams		Params about how to set up the texture.
 	 * @return				Returns a pointer to the constructed 2D texture object.
 	 *
+	 * this fills a TextureSource , and will then Build a Platform texture from that
+	 * can be used WITH_EDITOR only
+	 * contrast to CreateTexture2DFromImage
 	 */
 	ENGINE_API static UTexture2D* CreateTexture2D(int32 SrcWidth, int32 SrcHeight, const TArray<FColor> &SrcData, UObject* Outer, const FString& Name, const EObjectFlags &Flags, const FCreateTexture2DParameters& InParams);
 	
@@ -240,6 +243,9 @@ public:
 	 * 
 	 * @param Image			Image that will be copied into a Texture
 	 * @return				Returns a pointer to the constructed 2D texture object.
+	 *
+	 * NOTE: this makes a Transient texture with the Image in the PlatformData
+	 * this is different than making a Texture via TextureSource.Init(Image)
 	 */	
 	ENGINE_API static UTexture2D* CreateTexture2DFromImage(const FImageView & Image);
 
@@ -253,6 +259,7 @@ public:
 	 * @param SrcData			Raw image array.
 	 * @param DstData			compressed image array.
 	 *
+	 * DEPRECATED, avoid this, uses the bad image resize
 	 */
 	ENGINE_API static void CropAndScaleImage( int32 SrcWidth, int32 SrcHeight, int32 DesiredWidth, int32 DesiredHeight, const TArray<FColor> &SrcData, TArray<FColor> &DstData  );
 
@@ -385,11 +392,17 @@ public:
 
 	/**
 	 * Imports a texture file from disk and creates Texture2D from it
+	 *
+	 * note this make a Transient / PlatformData only Texture (no TextureSource)
+	 *
 	 */
 	ENGINE_API static UTexture2D* ImportFileAsTexture2D(const FString& Filename);
 
 	/**
 	 * Imports a texture a buffer and creates Texture2D from it
+	 *
+	 * note this make a Transient / PlatformData only Texture (no TextureSource)
+	 *
 	 */
 	ENGINE_API static UTexture2D* ImportBufferAsTexture2D(TArrayView64<const uint8> Buffer);
 	ENGINE_API static UTexture2D* ImportBufferAsTexture2D(const TArray<uint8>& Buffer);
@@ -416,8 +429,7 @@ public:
 	*/
 	ENGINE_API static bool ExportTextureCubeAsHDR(UTextureCube* TexRT, FArchive& Ar);
 
-	// Should be removed from public API.
-	//  Move out of header.
+	UE_DEPRECATED(5.5, "Use GetRenderTargetImage")
 	ENGINE_API static bool GetRawData(UTextureRenderTarget2D* TexRT, TArray64<uint8>& RawData);
 	
 	/**
diff --git a/Engine/Source/Runtime/Engine/Public/UnrealClient.h b/Engine/Source/Runtime/Engine/Public/UnrealClient.h
--- a/Engine/Source/Runtime/Engine/Public/UnrealClient.h
+++ b/Engine/Source/Runtime/Engine/Public/UnrealClient.h
@@ -93,15 +93,19 @@ public:
 	* @param InFlags - Additional information about how to to read the surface data (cube face, slice index, etc.)
 	* @param InSrcRect - InSrcRect not specified means the whole rect
 	* @return True if the read succeeded.
+	*
+	* This will convert whatever the pixel format is to FColor
 	*/
 	ENGINE_API virtual bool ReadPixels(TArray<FColor>& OutImageData, FReadSurfaceDataFlags InFlags = FReadSurfaceDataFlags(RCM_UNorm, CubeFace_MAX), FIntRect InSrcRect = FIntRect(0, 0, 0, 0));
 
 	/**
 	* Reads the render target's displayed pixels into a preallocated color buffer.
 	* @param OutImageBytes - RGBA8 values will be stored in this buffer.  Buffer must be preallocated with the correct size!
-	* * @param InFlags - Additional information about how to to read the surface data (cube face, slice index, etc.)
+	* @param InFlags - Additional information about how to to read the surface data (cube face, slice index, etc.)
 	* @param InSrcRect - InSrcRect not specified means the whole rect
 	* @return True if the read succeeded.
+	*
+	* Ptr variant of this API just does an extra memcpy; prefer the TArray variant
 	*/
 	ENGINE_API bool ReadPixelsPtr(FColor* OutImageBytes, FReadSurfaceDataFlags InFlags = FReadSurfaceDataFlags(RCM_UNorm, CubeFace_MAX), FIntRect InSrcRect = FIntRect(0, 0, 0, 0));
 
@@ -119,24 +123,19 @@ public:
 	 * @param InFlags - Additional information about how to to read the surface data (cube face, slice index, etc.)
 	 * @param InSrcRect - InSrcRect not specified means the whole rect
 	 * @return True if the read succeeded.
+	 *
+	 * this only works if surface is PF_FloatRGBA exactly ; it does not convert
 	 */
 	ENGINE_API virtual bool ReadFloat16Pixels(TArray<FFloat16Color>& OutImageData, FReadSurfaceDataFlags InFlags = FReadSurfaceDataFlags(RCM_UNorm, CubeFace_MAX), FIntRect InSrcRect = FIntRect(0, 0, 0, 0));
 
-	/**
-	* Reads the render target's displayed pixels into a preallocated color buffer.
-	* @param OutImageBytes - RGBA16F values will be stored in this buffer.  Buffer must be preallocated with the correct size!
-	* @param InFlags - Additional information about how to to read the surface data (cube face, slice index, etc.)
-	* @param InSrcRect - InSrcRect not specified means the whole rect
-	* @return True if the read succeeded.
-	*/
-	ENGINE_API bool ReadFloat16PixelsPtr(FFloat16Color* OutImageBytes, FReadSurfaceDataFlags InFlags = FReadSurfaceDataFlags(RCM_UNorm, CubeFace_MAX), FIntRect InSrcRect = FIntRect(0, 0, 0, 0));
-
 	/**
 	 * Reads the render target's displayed pixels into the given color buffer.
 	 * @param OutImageData - Linear color array to store the value
 	 * @param InFlags - Additional information about how to to read the surface data (cube face, slice index, etc.)
 	 * @param InSrcRect - InSrcRect not specified means the whole rect
 	 * @return True if the read succeeded.
+	 *
+	 * This will convert whatever the pixel format is to FLinearColor
 	 */
 	ENGINE_API virtual bool ReadLinearColorPixels(TArray<FLinearColor>& OutImageData, FReadSurfaceDataFlags InFlags = FReadSurfaceDataFlags(RCM_MinMax, CubeFace_MAX), FIntRect InSrcRect = FIntRect(0, 0, 0, 0));
 
@@ -146,6 +145,8 @@ public:
 	 * @param InFlags - Additional information about how to to read the surface data (cube face, slice index, etc.)
 	 * @param InSrcRect - InSrcRect not specified means the whole rect
 	 * @return True if the read succeeded.
+	 *
+	 * Ptr variant of this API just does an extra memcpy; prefer the TArray variant
 	 */
 	ENGINE_API bool ReadLinearColorPixelsPtr(FLinearColor* OutImageBytes, FReadSurfaceDataFlags InFlags = FReadSurfaceDataFlags(RCM_MinMax, CubeFace_MAX), FIntRect InSrcRect = FIntRect(0, 0, 0, 0));
 
