TextureRenderTarget clean up Create/Update not supporting all formats
use generic routines so all-format support is automatic

diff --git a/Engine/Source/Runtime/Core/Public/PixelFormat.h b/Engine/Source/Runtime/Core/Public/PixelFormat.h
--- a/Engine/Source/Runtime/Core/Public/PixelFormat.h
+++ b/Engine/Source/Runtime/Core/Public/PixelFormat.h
@@ -23,8 +23,8 @@ enum EPixelFormat : uint8
 	PF_DXT3                 =6,
 	PF_DXT5                 =7,
 	PF_UYVY                 =8,
-	PF_FloatRGB             =9, // 16F
-	PF_FloatRGBA            =10, // 16F
+	PF_FloatRGB             =9,  // FloatRGB == PF_FloatR11G11B10 , NOT 16F usually, but varies
+	PF_FloatRGBA            =10, // RGBA16F
 	PF_DepthStencil         =11,
 	PF_ShadowDepth          =12,
 	PF_R32_FLOAT            =13,
@@ -303,6 +303,8 @@ static inline bool IsBlockCompressedFormat(EPixelFormat Format)
 	return false;
 }
 
+// unclear what IsHDR is supposed to mean
+//	see also IsFloatFormat
 FORCEINLINE bool IsHDR(EPixelFormat PixelFormat)
 {
 	return PixelFormat == PF_FloatRGBA || PixelFormat == PF_BC6H || PixelFormat == PF_R16F || PixelFormat == PF_R32_FLOAT || PixelFormat == PF_A32B32G32R32F
@@ -334,6 +336,7 @@ FORCEINLINE bool IsInteger(EPixelFormat PixelFormat)
 	return false;
 }
 
+// IsFloatFormat does not include all floating point formats; see also IsHDR
 static bool IsFloatFormat(EPixelFormat Format)
 {
 	switch (Format)
diff --git a/Engine/Source/Runtime/D3D12RHI/Private/D3D12RenderTarget.cpp b/Engine/Source/Runtime/D3D12RHI/Private/D3D12RenderTarget.cpp
--- a/Engine/Source/Runtime/D3D12RHI/Private/D3D12RenderTarget.cpp
+++ b/Engine/Source/Runtime/D3D12RHI/Private/D3D12RenderTarget.cpp
@@ -592,9 +592,11 @@ static void ConvertDXGIToFColor(DXGI_FORMAT Format, uint32 Width, uint32 Height,
 
 static void ConvertRAWSurfaceDataToFLinearColor(EPixelFormat Format, uint32 Width, uint32 Height, uint8 *In, uint32 SrcPitch, FLinearColor* Out, FReadSurfaceDataFlags InFlags)
 {
-	if (Format == PF_R16F || Format == PF_R16F_FILTER)
+	// @todo Oodle : need to factor this out of here to somewhere generic and share it with all the RHI's
+	// also get the UNorm/MinMax stuff out of the inner loops here, do that as a post-process
+
+	if (Format == PF_R16_UINT )
 	{
-		// e.g. shadow maps
 		for (uint32 Y = 0; Y < Height; Y++)
 		{
 			uint16* SrcPtr = (uint16*)(In + Y * SrcPitch);
@@ -605,6 +607,22 @@ static void ConvertRAWSurfaceDataToFLinearColor(EPixelFormat Format, uint32 Widt
 				uint16 Value16 = *SrcPtr;
 				float Value = Value16 * (1.f/0xffff);
 
+				*DestPtr = FLinearColor(Value, Value, Value); // G (gray) not R
+				++SrcPtr;
+				++DestPtr;
+			}
+		}
+	}
+	else if (Format == PF_R16F || Format == PF_R16F_FILTER)
+	{
+		for (uint32 Y = 0; Y < Height; Y++)
+		{
+			FFloat16 * SrcPtr = (FFloat16*)(In + Y * SrcPitch);
+			FLinearColor* DestPtr = Out + Y * Width;
+
+			for (uint32 X = 0; X < Width; X++)
+			{
+				float Value = SrcPtr->GetFloat();
 				*DestPtr = FLinearColor(Value, Value, Value);
 				++SrcPtr;
 				++DestPtr;
@@ -868,6 +886,55 @@ static void ConvertRAWSurfaceDataToFLinearColor(EPixelFormat Format, uint32 Widt
 			}
 		}
 	}
+	else if (Format == PF_G16R16F)
+	{
+		// Read the data out of the buffer, converting it to FLinearColor.
+		for (uint32 Y = 0; Y < Height; Y++)
+		{
+			FFloat16 * SrcPtr = (FFloat16*)(In + Y * SrcPitch);
+			FLinearColor* DestPtr = Out + Y * Width;
+			for (uint32 X = 0; X < Width; X++)
+			{
+				*DestPtr = FLinearColor( SrcPtr[0].GetFloat(), SrcPtr[1].GetFloat(), 0.f,1.f);
+				SrcPtr += 2;
+				++DestPtr;
+			}
+		}
+	}
+	else if (Format == PF_G32R32F)
+	{
+		// not doing MinMax/Unorm remap here
+	
+		// Read the data out of the buffer, converting it to FLinearColor.
+		for (uint32 Y = 0; Y < Height; Y++)
+		{
+			float * SrcPtr = (float *)(In + Y * SrcPitch);
+			FLinearColor* DestPtr = Out + Y * Width;
+			for (uint32 X = 0; X < Width; X++)
+			{
+				*DestPtr = FLinearColor( SrcPtr[0], SrcPtr[1], 0.f, 1.f );
+				SrcPtr += 2;
+				++DestPtr;
+			}
+		}
+	}
+	else if (Format == PF_R32_FLOAT)
+	{
+		// not doing MinMax/Unorm remap here
+	
+		// Read the data out of the buffer, converting it to FLinearColor.
+		for (uint32 Y = 0; Y < Height; Y++)
+		{
+			float * SrcPtr = (float *)(In + Y * SrcPitch);
+			FLinearColor* DestPtr = Out + Y * Width;
+			for (uint32 X = 0; X < Width; X++)
+			{
+				*DestPtr = FLinearColor( SrcPtr[0], 0.f, 0.f, 1.f );
+				++SrcPtr;
+				++DestPtr;
+			}
+		}
+	}
 	else
 	{
 		// not supported yet
diff --git a/Engine/Source/Runtime/Engine/Classes/Engine/TextureRenderTarget.h b/Engine/Source/Runtime/Engine/Classes/Engine/TextureRenderTarget.h
--- a/Engine/Source/Runtime/Engine/Classes/Engine/TextureRenderTarget.h
+++ b/Engine/Source/Runtime/Engine/Classes/Engine/TextureRenderTarget.h
@@ -73,7 +73,13 @@ class UTextureRenderTarget : public UTexture
 	 */
 	ENGINE_API virtual TSubclassOf<UTexture> GetTextureUClass() const PURE_VIRTUAL(UTextureRenderTarget, return nullptr;);
 
+	// get the variant of ReadPixels call that should be used
+	//	either BGRA8,RGBA32F,or RGBA16F
+	//	for FColor,FLinearColor, or Float16
+	static ERawImageFormat::Type GetReadPixelsFormat(EPixelFormat PF,bool bIsVolume);
+
 #if WITH_EDITOR
+
 	/**
 	 * Utility for updating an existing UTexture from a TextureRenderTarget
 	 * @param InTexture - Texture which will contain the content of this render target after the call.
diff --git a/Engine/Source/Runtime/Engine/Classes/Engine/TextureRenderTargetCube.h b/Engine/Source/Runtime/Engine/Classes/Engine/TextureRenderTargetCube.h
--- a/Engine/Source/Runtime/Engine/Classes/Engine/TextureRenderTargetCube.h
+++ b/Engine/Source/Runtime/Engine/Classes/Engine/TextureRenderTargetCube.h
@@ -17,7 +17,7 @@ struct FPropertyChangedEvent;
  * for rendering as well as rendered as a regular cube texture resource.
  *
  */
-UCLASS(hidecategories=Object, hidecategories=Texture, MinimalAPI)
+UCLASS(hidecategories=Object, hidecategories=Texture, hidecategories=Compression, hidecategories=Adjustments, hidecategories=Compositing, MinimalAPI)
 class UTextureRenderTargetCube : public UTextureRenderTarget
 {
 	GENERATED_UCLASS_BODY()
@@ -30,12 +30,13 @@ class UTextureRenderTargetCube : public UTextureRenderTarget
 	UPROPERTY()
 	FLinearColor ClearColor;
 
-	/** The format of the texture data.											*/
-	/** Normally the format is derived from bHDR, this allows code to set the format explicitly. */
+	/** The format of the texture data.											
+	* Normally the format is derived from bHDR, this allows code to set the format explicitly. */
 	UPROPERTY()
 	TEnumAsByte<enum EPixelFormat> OverrideFormat;
 
-	/** Whether to support storing HDR values, which requires more memory. */
+	/** If OverrideFormat is not set, bHDR chooses the format of the RT.
+	With bHDR on it is RGBA16F , off is BGRA8 */
 	UPROPERTY(EditAnywhere, BlueprintReadOnly, Category=TextureRenderTargetCube, AssetRegistrySearchable)
 	uint8 bHDR:1;
 
diff --git a/Engine/Source/Runtime/Engine/Classes/Engine/TextureRenderTargetVolume.h b/Engine/Source/Runtime/Engine/Classes/Engine/TextureRenderTargetVolume.h
--- a/Engine/Source/Runtime/Engine/Classes/Engine/TextureRenderTargetVolume.h
+++ b/Engine/Source/Runtime/Engine/Classes/Engine/TextureRenderTargetVolume.h
@@ -17,7 +17,7 @@ struct FPropertyChangedEvent;
  * for rendering as well as rendered as a regular Volume texture resource.
  *
  */
-UCLASS(hidecategories=Object, hidecategories=Texture, MinimalAPI)
+UCLASS(hidecategories=Object, hidecategories=Texture, hidecategories=Compression, hidecategories=Adjustments, hidecategories=Compositing, MinimalAPI)
 class UTextureRenderTargetVolume : public UTextureRenderTarget
 {
 	GENERATED_UCLASS_BODY()
@@ -37,13 +37,14 @@ class UTextureRenderTargetVolume : public UTextureRenderTarget
 	/** the color the texture is cleared to */
 	UPROPERTY()
 	FLinearColor ClearColor;
-
-	/** The format of the texture data.											*/
-	/** Normally the format is derived from bHDR, this allows code to set the format explicitly. */
+	
+	/** The format of the texture data.											
+	* Normally the format is derived from bHDR, this allows code to set the format explicitly. */
 	UPROPERTY()
 	TEnumAsByte<enum EPixelFormat> OverrideFormat;
 
-	/** Whether to support storing HDR values, which requires more memory. */
+	/** If OverrideFormat is not set, bHDR chooses the format of the RT.
+	With bHDR on it is RGBA16F , off is BGRA8 */
 	UPROPERTY(EditAnywhere, BlueprintReadOnly, Category=TextureRenderTargetVolume, AssetRegistrySearchable)
 	uint8 bHDR:1;
 
diff --git a/Engine/Source/Runtime/Engine/Private/ImageUtils.cpp b/Engine/Source/Runtime/Engine/Private/ImageUtils.cpp
--- a/Engine/Source/Runtime/Engine/Private/ImageUtils.cpp
+++ b/Engine/Source/Runtime/Engine/Private/ImageUtils.cpp
@@ -330,16 +330,6 @@ bool FImageUtils::GetRawData(UTextureRenderTarget2D* TexRT, TArray64<uint8>& Raw
 	return true;
 }
 
-static int GetBitsPerComponent(EPixelFormat Format)
-{
-	if ( Format == PF_A2B10G10R10 ) return 10; // doesn't handle heterogenous bit counts well
-
-	const FPixelFormatInfo & Info = GPixelFormats[Format];
-	// rounds down
-	int BitsPerComponent = ( Info.BlockBytes * 8 ) / ( Info.BlockSizeX * Info.BlockSizeY * Info.BlockSizeZ * Info.NumComponents );
-	return BitsPerComponent;
-}
-
 bool FImageUtils::GetRenderTargetImage(UTextureRenderTarget2D* TexRT, FImage & Image)
 {
 	Image = FImage();
@@ -347,27 +337,9 @@ bool FImageUtils::GetRenderTargetImage(UTextureRenderTarget2D* TexRT, FImage & I
 	FRenderTarget* RenderTarget = TexRT->GameThread_GetRenderTargetResource();
 	EPixelFormat RTFormat = TexRT->GetFormat();
 
-	bool Use8Bit;
-	
-	if ( GetBitsPerComponent(RTFormat) <= 8 )
-	{
-		Use8Bit = true;
-	}
-	else if ( IsFloatFormat(RTFormat) || IsHDR(RTFormat) )
-	{
-		Use8Bit = false;
-	}
-	else if ( IsDepthOrStencilFormat(RTFormat) )
-	{
-		Use8Bit = false;
-	}
-	else
-	{
-		// eg. 16-bit integer
-		Use8Bit = false;
-	}
-	
-	if ( RTFormat == PF_FloatRGBA )
+	ERawImageFormat::Type ReadFormat = UTextureRenderTarget::GetReadPixelsFormat(RTFormat,false);
+		
+	if ( ReadFormat == ERawImageFormat::RGBA16F )
 	{
 		// ReadFloat16Pixels does no conversions
 		//	must be used only exactly with FloatRGBA type
@@ -383,7 +355,7 @@ bool FImageUtils::GetRenderTargetImage(UTextureRenderTarget2D* TexRT, FImage & I
 		check( Image.GetImageSizeBytes() == Colors.Num() * sizeof(Colors[0]) );
 		memcpy( &Image.RawData[0], &Colors[0], Image.GetImageSizeBytes() );
 	}
-	else if ( Use8Bit )
+	else if ( ReadFormat == ERawImageFormat::BGRA8 )
 	{
 		// ?? not clear TexRT->IsSRGB is right , see other notes on various issues there
 		//	mainly we are trying to catch the check for whether the _SRGB or non _SRGB BGRA8 format as chosen
@@ -404,7 +376,7 @@ bool FImageUtils::GetRenderTargetImage(UTextureRenderTarget2D* TexRT, FImage & I
 		check( Image.GetImageSizeBytes() == Colors.Num() * sizeof(Colors[0]) );
 		memcpy( &Image.RawData[0], &Colors[0], Image.GetImageSizeBytes() );
 	}
-	else // use F32
+	else if ( ReadFormat == ERawImageFormat::RGBA32F )
 	{
 		Image.Init(TexRT->SizeX,TexRT->SizeY,ERawImageFormat::RGBA32F,EGammaSpace::Linear);
 		
@@ -420,7 +392,11 @@ bool FImageUtils::GetRenderTargetImage(UTextureRenderTarget2D* TexRT, FImage & I
 
 		check( Image.GetImageSizeBytes() == Colors.Num() * sizeof(Colors[0]) );
 		memcpy( &Image.RawData[0], &Colors[0], Image.GetImageSizeBytes() );
-	}	
+	}
+	else
+	{
+		check(0); // unexpected ReadFormat
+	}
 	
 	return true;
 }
diff --git a/Engine/Source/Runtime/Engine/Private/TextureRenderTarget.cpp b/Engine/Source/Runtime/Engine/Private/TextureRenderTarget.cpp
--- a/Engine/Source/Runtime/Engine/Private/TextureRenderTarget.cpp
+++ b/Engine/Source/Runtime/Engine/Private/TextureRenderTarget.cpp
@@ -14,6 +14,7 @@
 #include "TextureCompiler.h"
 #include "Logging/MessageLog.h"
 #include "Hash/xxhash.h"
+#include "ImageCoreUtils.h"
 
 #include UE_INLINE_GENERATED_CPP_BY_NAME(TextureRenderTarget)
 
@@ -77,7 +78,8 @@ bool UTextureRenderTarget::CanConvertToTexture(FText* OutErrorMessage) const
 
 ETextureSourceFormat UTextureRenderTarget::ValidateTextureFormatForConversionToTextureInternal(EPixelFormat InFormat, const TArrayView<const EPixelFormat>& InCompatibleFormats, FText* OutErrorMessage) const
 {
-	if (!InCompatibleFormats.Contains(InFormat))
+	// InCompatibleFormats can be empty, meaning anything works
+	if ( InCompatibleFormats.Num() != 0 && !InCompatibleFormats.Contains(InFormat))
 	{
 		if (OutErrorMessage != nullptr)
 		{
@@ -94,34 +96,50 @@ ETextureSourceFormat UTextureRenderTarget::ValidateTextureFormatForConversionToT
 	}
 
 	// Return what ETextureSourceFormat corresponds to this EPixelFormat (must match the conversion capabilities of UTextureRenderTarget::UpdateTexture) : 
-	ETextureSourceFormat TextureFormat = TSF_Invalid;
-	// @todo Oodle : just support all formats and remove this
-	switch (InFormat)
-	{
-	case PF_B8G8R8A8:
-	case PF_R8G8: // Will produce 2 useless channels as it will be expanded to a TSF_BGRA8 texture source but there's no 2 channel ETextureSourceFormat
-		TextureFormat = TSF_BGRA8;
-		break;
-	case PF_R32_FLOAT: // Narrowed down to 16bits float (if available by the FRenderTarget::ReadFloat16Pixels implementation and corresponding RHI implementation)
-	case PF_FloatRGBA:
-		TextureFormat = TSF_RGBA16F;
-		break;
-	case PF_G8:
-		TextureFormat = TSF_G8;
-		break;
-	default:
+	ERawImageFormat::Type RawFormat = FImageCoreUtils::GetRawImageFormatForPixelFormat(InFormat);
+	ETextureSourceFormat TextureFormat = FImageCoreUtils::ConvertToTextureSourceFormat(RawFormat);
+
+	return TextureFormat;
+}
+
+// there are three different RHI APIs to read pixels from render targets
+// they correspond to FColor, FLinearColor, and FFloat16
+// choose which of the three to use for this RT
+ERawImageFormat::Type UTextureRenderTarget::GetReadPixelsFormat(EPixelFormat PF,bool bIsVolume)
+{
+	if ( bIsVolume )
 	{
-		checkf(false, TEXT("Only some select texture formats are currently supported by UTextureRenderTarget::UpdateTexture. Please update this switch if you add conversion code."));
+		// volumes are different, must always use 16F path
+		return ERawImageFormat::RGBA16F;
 	}
+	else if ( PF == PF_FloatRGBA )
+	{
+		// for non-volumes, only exact match 16F is supported
+		return ERawImageFormat::RGBA16F;
 	}
+	else
+	{
+		// either 8-bit FColor or 32F FLinearColor path
 
-	return TextureFormat;
+		ERawImageFormat::Type RF = FImageCoreUtils::GetRawImageFormatForPixelFormat(PF);
+
+		if ( RF == ERawImageFormat::BGRA8 || RF == ERawImageFormat::G8 )
+		{
+			return ERawImageFormat::BGRA8; // use FColor
+		}
+		else
+		{
+			return ERawImageFormat::RGBA32F; // use FLinearColor
+		}
+	}
 }
 
 #if WITH_EDITOR
 
 bool UTextureRenderTarget::UpdateTexture(UTexture* InTexture, EConstructTextureFlags InFlags, const TArray<uint8>* InAlphaOverride, FOnTextureChangingDelegate InOnTextureChangingDelegate, FText* OutErrorMessage)
 {
+	// InTexture will be filled from the RT
+	// InTexture must already be set up as the right type of texture
 
 	if (InTexture == nullptr)
 	{
@@ -146,12 +164,13 @@ bool UTextureRenderTarget::UpdateTexture(UTexture* InTexture, EConstructTextureF
 		return false;
 	}
 
-	ETextureSourceFormat TextureSourceFormat = TSF_Invalid;
+	ETextureSourceFormat OutTextureSourceFormat = TSF_Invalid;
 	EPixelFormat PixelFormat = PF_Unknown; 
-	if (!CanConvertToTexture(TextureSourceFormat, PixelFormat, OutErrorMessage))
+	if (!CanConvertToTexture(OutTextureSourceFormat, PixelFormat, OutErrorMessage))
 	{
 		return false;
 	}
+	check( OutTextureSourceFormat != TSF_Invalid );
 
 	// Note : use UTexture's GetTextureClass() here, as URenderTarget's GetTextureClass() unconveniently returns ETextureClass::RenderTarget :
 	ETextureClass TextureClass = InTexture->GetTextureClass();
@@ -183,7 +202,7 @@ bool UTextureRenderTarget::UpdateTexture(UTexture* InTexture, EConstructTextureF
 	bTextureChanging |= InTexture->Source.GetVolumeSizeZ() != SizeZ;
 	// Note : for FTextureSource, NumSlices means NumSurfaces (e.g. 6 * GetNumSlices() for a cube map array) : 
 	bTextureChanging |= InTexture->Source.GetNumSlices() != NumSurfaces;
-	bTextureChanging |= InTexture->Source.GetFormat() != TextureSourceFormat;
+	bTextureChanging |= InTexture->Source.GetFormat() != OutTextureSourceFormat;
 	bTextureChanging |= InTexture->CompressionSettings != CompressionSettingsForTexture;
 
 	FXxHash64 OldDataHash;
@@ -199,10 +218,27 @@ bool UTextureRenderTarget::UpdateTexture(UTexture* InTexture, EConstructTextureF
 	}
 
 	// Temp storage that will be used for the texture init.
+	ERawImageFormat::Type ReadFormat = GetReadPixelsFormat(PixelFormat,bIsVolume);
+
+	check( ReadFormat == ERawImageFormat::BGRA8 ||
+		ReadFormat == ERawImageFormat::RGBA32F ||
+		ReadFormat == ERawImageFormat::RGBA16F );
+
+	EGammaSpace ReadGammaSpace = EGammaSpace::Linear;
+	bool bSRGB = RenderTarget->GetDisplayGamma() > 1.5f; // @@ no way to call IsSRGB() , no uniform query on RenderTarget
+	if ( ERawImageFormat::GetFormatNeedsGammaSpace(ReadFormat) && bSRGB )
+	{
+		ReadGammaSpace = EGammaSpace::sRGB;
+	}
+
+	FImage ReadImage;
+	ReadImage.Init(SizeX,SizeY,NumSurfaces,ReadFormat,ReadGammaSpace);
+
 	const int32 NumPixelsPerSurface = SizeX * SizeY;
-	const int32 NumBytesPerPixel = GTextureSourceFormats[TextureSourceFormat].BytesPerPixel;
-	const int32 DestinationSurfaceMipSize = NumPixelsPerSurface * NumBytesPerPixel;
-	const int32 DestinationSize = NumPixelsPerSurface * NumSurfaces * NumBytesPerPixel;
+	const int32 NumBytesPerPixel = ReadImage.GetBytesPerPixel();
+	const int64 DestinationSurfaceMipSize = ReadImage.GetSliceSizeBytes();
+	const int64 DestinationSize = ReadImage.GetImageSizeBytes();
+
 	check(NumBytesPerPixel > 0);
 	checkf((InAlphaOverride == nullptr) || (NumPixelsPerSurface == InAlphaOverride->Num()), TEXT("If InAlphaOverride is specified, it must have the same number of pixels as a single slice (currently, render target has %d pixels per slice, AlphaOverride: %d). It is also expected to be 1-byte encoded (1 element per pixel)."),
 		NumPixelsPerSurface, InAlphaOverride->Num());
@@ -210,28 +246,28 @@ bool UTextureRenderTarget::UpdateTexture(UTexture* InTexture, EConstructTextureF
 	TArray<uint8> NewData;
 	NewData.SetNumUninitialized(DestinationSize);
 
-	// @todo Oodle : make this more generic like ImageUtils::GetRenderTargetImage
-
 	// Read the data surface by surface (i.e. for a 2D array or a volume: slice by slice, for a cubemap: face by face and for a cubemap array: face by face, slice by slice
 	for (int32 SurfaceIndex = 0; SurfaceIndex < NumSurfaces; ++SurfaceIndex)
 	{
-		TArrayView<uint8> SurfaceByteData = MakeArrayView(NewData.GetData() + SurfaceIndex * DestinationSurfaceMipSize, DestinationSurfaceMipSize);
 		FReadSurfaceDataFlags ReadSurfaceDataFlags(RCM_UNorm, bIsCube ? static_cast<ECubeFace>(SurfaceIndex % 6) : CubeFace_MAX);
 		ReadSurfaceDataFlags.SetArrayIndex(bIsCube ? SurfaceIndex / 6 : SurfaceIndex);
-		switch (TextureSourceFormat)
+
+		void * ReadIntoSlice = ReadImage.GetPixelPointer(0,0,SurfaceIndex);
+
+		switch (ReadFormat)
 		{
-		case TSF_BGRA8:
+		case ERawImageFormat::BGRA8: // FColor
 		{
 			TArray<FColor> NewDataColor;
 			RenderTarget->ReadPixels(NewDataColor, ReadSurfaceDataFlags);
 			check(NewDataColor.Num() == NumPixelsPerSurface);
 
 			const int32 NumBytes = NewDataColor.Num() * NewDataColor.GetTypeSize();
-			check(NumBytes == SurfaceByteData.Num());
-			FMemory::Memcpy(SurfaceByteData.GetData(), NewDataColor.GetData(), NumBytes);
+			check(NumBytes == DestinationSurfaceMipSize);
+			FMemory::Memcpy(ReadIntoSlice, NewDataColor.GetData(), NumBytes);
 
 			// override the alpha if desired
-			TArrayView<FColor> SurfaceColorData = MakeArrayView(reinterpret_cast<FColor*>(SurfaceByteData.GetData()), NumPixelsPerSurface);
+			TArrayView<FColor> SurfaceColorData = MakeArrayView(reinterpret_cast<FColor*>(ReadIntoSlice), NumPixelsPerSurface);
 			if (InAlphaOverride != nullptr)
 			{
 				for (int32 PixelIndex = 0; PixelIndex < NumPixelsPerSurface; ++PixelIndex)
@@ -257,7 +293,46 @@ bool UTextureRenderTarget::UpdateTexture(UTexture* InTexture, EConstructTextureF
 			}
 		}
 		break;
-		case TSF_RGBA16F:
+		
+		case ERawImageFormat::RGBA32F: // FLinearColor
+		{
+			TArray<FLinearColor> NewDataColor;
+			RenderTarget->ReadLinearColorPixels(NewDataColor, ReadSurfaceDataFlags);
+			check(NewDataColor.Num() == NumPixelsPerSurface);
+
+			const int32 NumBytes = NewDataColor.Num() * NewDataColor.GetTypeSize();
+			check(NumBytes == DestinationSurfaceMipSize);
+			FMemory::Memcpy(ReadIntoSlice, NewDataColor.GetData(), NumBytes);
+
+			// override the alpha if desired
+			TArrayView<FLinearColor> SurfaceColorData = MakeArrayView(reinterpret_cast<FLinearColor*>(ReadIntoSlice), NumPixelsPerSurface);
+			if (InAlphaOverride != nullptr)
+			{
+				for (int32 PixelIndex = 0; PixelIndex < NumPixelsPerSurface; ++PixelIndex)
+				{
+					SurfaceColorData[PixelIndex].A = (*InAlphaOverride)[PixelIndex] * (1.f/255.f);
+				}
+			}
+			else if (InFlags & CTF_RemapAlphaAsMasked)
+			{
+				// if the target was rendered with a masked texture, then the depth will probably have been written instead of 0/255 for the
+				// alpha, and the depth when unwritten will be 255, so remap 255 to 0 (masked out area) and anything else as 255 (written to area)
+				for (int32 PixelIndex = 0; PixelIndex < NumPixelsPerSurface; ++PixelIndex)
+				{
+					SurfaceColorData[PixelIndex].A = (SurfaceColorData[PixelIndex].A >= 1.f) ? 0.f : 1.f;
+				}
+			}
+			else if (InFlags & CTF_ForceOpaque)
+			{
+				for (int32 PixelIndex = 0; PixelIndex < NumPixelsPerSurface; ++PixelIndex)
+				{
+					SurfaceColorData[PixelIndex].A = 1.f;
+				}
+			}
+		}
+		break;
+
+		case ERawImageFormat::RGBA16F:
 		{
 			TArray<FFloat16Color> NewDataFloat16Color;
 			check(NumBytesPerPixel == sizeof(FFloat16Color));
@@ -267,16 +342,16 @@ bool UTextureRenderTarget::UpdateTexture(UTexture* InTexture, EConstructTextureF
 
 			check(NewDataFloat16Color.Num() == NumPixelsPerSurface);
 			const int32 NumBytes = NewDataFloat16Color.Num() * NewDataFloat16Color.GetTypeSize();
-			check(NumBytes == SurfaceByteData.Num());
-			FMemory::Memcpy(SurfaceByteData.GetData(), NewDataFloat16Color.GetData(), NumBytes);
+			check(NumBytes == DestinationSurfaceMipSize);
+			FMemory::Memcpy(ReadIntoSlice, NewDataFloat16Color.GetData(), NumBytes);
 
 			// override the alpha if desired
-			TArrayView<FFloat16Color> SurfaceFloat16ColorData = MakeArrayView(reinterpret_cast<FFloat16Color*>(SurfaceByteData.GetData()), NumPixelsPerSurface);
+			TArrayView<FFloat16Color> SurfaceFloat16ColorData = MakeArrayView(reinterpret_cast<FFloat16Color*>(ReadIntoSlice), NumPixelsPerSurface);
 			if (InAlphaOverride != nullptr)
 			{
 				for (int32 PixelIndex = 0; PixelIndex < NumPixelsPerSurface; ++PixelIndex)
 				{
-					SurfaceFloat16ColorData[PixelIndex].A = (*InAlphaOverride)[PixelIndex];
+					SurfaceFloat16ColorData[PixelIndex].A = (*InAlphaOverride)[PixelIndex] * (1.f/255.f);
 				}
 			}
 			else if (InFlags & CTF_RemapAlphaAsMasked)
@@ -301,26 +376,30 @@ bool UTextureRenderTarget::UpdateTexture(UTexture* InTexture, EConstructTextureF
 			}
 		}
 		break;
-		case TSF_G8:
-		{
-			
-			TArray<FColor> NewDataColor;
-			RenderTarget->ReadPixels(NewDataColor, ReadSurfaceDataFlags);
-			check(NewDataColor.Num() == NumPixelsPerSurface);
-			for (int32 PixelIndex = 0; PixelIndex < NumPixelsPerSurface; PixelIndex++)
-			{
-				SurfaceByteData[PixelIndex] = NewDataColor[PixelIndex].R;
-			}
-		}
-		break;
+
 		default:
-			checkf(false, TEXT("Unsupported format, it needs to be added to CanConvertToTexture and UpdateTexture"));
+			checkf(false, TEXT("Unexpected ReadFormat"));
+		}
+	}
+
+	// ReadFormat is one of the three primary color types
+	// OutTextureSourceFormat is the closest TSF to our PF
+	if ( OutTextureSourceFormat != FImageCoreUtils::ConvertToTextureSourceFormat(ReadFormat) )
+	{
+		// convert image to OutTextureSourceFormat
+		ReadFormat = FImageCoreUtils::ConvertToRawImageFormat(OutTextureSourceFormat);
+		
+		ReadGammaSpace = EGammaSpace::Linear;
+		if ( ERawImageFormat::GetFormatNeedsGammaSpace(ReadFormat) && bSRGB )
+		{
+			ReadGammaSpace = EGammaSpace::sRGB;
 		}
+		ReadImage.ChangeFormat(ReadFormat,ReadGammaSpace);
 	}
 
 	if (bHashContent)
 	{
-		FXxHash64 NewDataHash = FXxHash64::HashBuffer(NewData.GetData(), NewData.Num());
+		FXxHash64 NewDataHash = FXxHash64::HashBuffer(ReadImage.RawData.GetData(), ReadImage.RawData.Num());
 		bTextureChanging = OldDataHash != NewDataHash;
 	}
 
@@ -332,8 +411,9 @@ bool UTextureRenderTarget::UpdateTexture(UTexture* InTexture, EConstructTextureF
 		FTextureCompilingManager::Get().FinishCompilation({ InTexture });
 
 		// init to the same size as the render target
-		InTexture->Source.Init(SizeX, SizeY, NumSurfaces, NumMips, TextureSourceFormat, NewData.GetData());
+		InTexture->Source.Init(ReadImage);
 		InTexture->CompressionSettings = CompressionSettingsForTexture;
+		InTexture->SRGB = ReadImage.GammaSpace == EGammaSpace::sRGB;
 	}
 
 	return true;
diff --git a/Engine/Source/Runtime/Engine/Private/TextureRenderTarget2D.cpp b/Engine/Source/Runtime/Engine/Private/TextureRenderTarget2D.cpp
--- a/Engine/Source/Runtime/Engine/Private/TextureRenderTarget2D.cpp
+++ b/Engine/Source/Runtime/Engine/Private/TextureRenderTarget2D.cpp
@@ -366,12 +366,12 @@ TSubclassOf<UTexture> UTextureRenderTarget2D::GetTextureUClass() const
 bool UTextureRenderTarget2D::CanConvertToTexture(ETextureSourceFormat& OutTextureSourceFormat, EPixelFormat& OutPixelFormat, FText* OutErrorMessage) const
 {
 	const EPixelFormat LocalFormat = GetFormat();
-	// These are the formats currently available for conversion to texture for UTextureRenderTarget2D : 
-	// @todo Oodle : get rid of this, support all
-	const ETextureSourceFormat TextureSourceFormat = ValidateTextureFormatForConversionToTextureInternal(LocalFormat, { PF_G8, PF_R8G8, PF_B8G8R8A8, PF_FloatRGBA }, OutErrorMessage);
+
+	// empty array means all formats supported
+	const ETextureSourceFormat TextureSourceFormat = ValidateTextureFormatForConversionToTextureInternal(LocalFormat, { }, OutErrorMessage);
 	if (TextureSourceFormat == TSF_Invalid)
 	{
-		return TSF_Invalid;
+		return false;
 	}
 
 	if ((SizeX <= 0) || (SizeY <= 0))
@@ -424,7 +424,7 @@ FTextureRenderTarget2DResource::FTextureRenderTarget2DResource(const class UText
 	,	TargetSizeX(Owner->SizeX)
 	,	TargetSizeY(Owner->SizeY)
 {
-	
+	// note: Resource has a bSRGB field which is not set or checked in the RenderTarget code
 }
 
 PRAGMA_DISABLE_DEPRECATION_WARNINGS
@@ -648,16 +648,19 @@ FIntPoint FTextureRenderTarget2DResource::GetSizeXY() const
 */
 float FTextureRenderTarget2DResource::GetDisplayGamma() const
 {
+	// if TargetGamma is set (not zero), it overrides everything else
 	if (Owner->TargetGamma > UE_KINDA_SMALL_NUMBER * 10.0f)
 	{
 		return Owner->TargetGamma;
 	}
+
 	// ?? special casing just two of the float PixelFormats to force 1.0 gamma here is inconsistent
 	//		(there are lots of other float formats)
+	// ignores Owner->IsSRGB() ? it's similar but not quite the same
 	if (Format == PF_FloatRGB || Format == PF_FloatRGBA || Owner->bForceLinearGamma )
 	{
 		return 1.0f;
 	}
-	// ignores Owner->IsSRGB() ?
-	return FTextureRenderTargetResource::GetDisplayGamma();
+
+	return FTextureRenderTargetResource::GetDisplayGamma(); // hard-coded 2.2 , actually means SRGB
 }
diff --git a/Engine/Source/Runtime/Engine/Private/TextureRenderTargetCube.cpp b/Engine/Source/Runtime/Engine/Private/TextureRenderTargetCube.cpp
--- a/Engine/Source/Runtime/Engine/Private/TextureRenderTargetCube.cpp
+++ b/Engine/Source/Runtime/Engine/Private/TextureRenderTargetCube.cpp
@@ -168,8 +168,8 @@ TSubclassOf<UTexture> UTextureRenderTargetCube::GetTextureUClass() const
 bool UTextureRenderTargetCube::CanConvertToTexture(ETextureSourceFormat& OutTextureSourceFormat, EPixelFormat& OutPixelFormat, FText* OutErrorMessage) const
 {
 	const EPixelFormat LocalFormat = GetFormat();
-	// These are the formats currently available for conversion to texture for UTextureRenderTargetCube : 
-	const ETextureSourceFormat TextureSourceFormat = ValidateTextureFormatForConversionToTextureInternal(GetFormat(), { PF_B8G8R8A8, PF_FloatRGBA }, OutErrorMessage);
+
+	const ETextureSourceFormat TextureSourceFormat = ValidateTextureFormatForConversionToTextureInternal(LocalFormat, { }, OutErrorMessage);
 	if (TextureSourceFormat == TSF_Invalid)
 	{
 		return false;
@@ -224,6 +224,8 @@ void FTextureRenderTargetCubeResource::InitRHI(FRHICommandListBase& RHICmdList)
 
 	if(Owner->SizeX > 0)
 	{
+		// compare with IsSRGB() call on RenderTarget2D
+		//	also "bool SRGB" should be set = bIsSRGB, but isn't
 		bool bIsSRGB = true;
 		// if render target gamma used was 1.0 then disable SRGB for the static texture
 		if(FMath::Abs(GetDisplayGamma() - 1.0f) < UE_KINDA_SMALL_NUMBER)
@@ -371,6 +373,8 @@ FIntPoint FTextureRenderTargetCubeResource::GetSizeXY() const
 
 float FTextureRenderTargetCubeResource::GetDisplayGamma() const
 {
+	// code dupe of RenderTarget2D
+
 	if(Owner->TargetGamma > UE_KINDA_SMALL_NUMBER * 10.0f)
 	{
 		return Owner->TargetGamma;
diff --git a/Engine/Source/Runtime/Engine/Private/TextureRenderTargetVolume.cpp b/Engine/Source/Runtime/Engine/Private/TextureRenderTargetVolume.cpp
--- a/Engine/Source/Runtime/Engine/Private/TextureRenderTargetVolume.cpp
+++ b/Engine/Source/Runtime/Engine/Private/TextureRenderTargetVolume.cpp
@@ -147,15 +147,19 @@ TSubclassOf<UTexture> UTextureRenderTargetVolume::GetTextureUClass() const
 bool UTextureRenderTargetVolume::CanConvertToTexture(ETextureSourceFormat& OutTextureSourceFormat, EPixelFormat& OutPixelFormat, FText* OutErrorMessage) const
 {
 	const EPixelFormat LocalFormat = GetFormat();
-	// These are the formats currently available for conversion to texture for UTextureRenderTargetVolume : 
-	// PF_R32_FLOAT and PF_FloatRGBA are supported by FTextureRenderTargetVolumeResource::ReadFloat16Pixels
-	// TODO [jonathan.bard] : PF_R16F is also supported by by FTextureRenderTargetVolumeResource::ReadFloat16Pixels so it could be added to the list
-	ETextureSourceFormat TextureSourceFormat = ValidateTextureFormatForConversionToTextureInternal(GetFormat(), { PF_R32_FLOAT, PF_FloatRGBA }, OutErrorMessage);
+
+	// Formats limited by support in Read3DSurfaceFloatData (see FTextureRenderTargetVolumeResource::ReadFloat16Pixels) :
+	ETextureSourceFormat TextureSourceFormat = ValidateTextureFormatForConversionToTextureInternal(LocalFormat, 
+		{ PF_R32_FLOAT, PF_FloatRGBA, PF_R16F }, OutErrorMessage);
 	if (TextureSourceFormat == TSF_Invalid)
 	{
 		return false;
 	}
 
+	// this is not actually required, TextureSourceFormat is a free choice
+	// volumes will always use the 16F read path, that comes from GetReadPixelsFormat()
+	TextureSourceFormat = TSF_RGBA16F;
+
 	if ((SizeX <= 0) || (SizeY <= 0) || (SizeZ <= 0))
 	{
 		if (OutErrorMessage != nullptr)
@@ -207,8 +211,8 @@ void FTextureRenderTargetVolumeResource::InitRHI(FRHICommandListBase& RHICmdList
 
 	if((Owner->SizeX > 0) && (Owner->SizeY > 0) && (Owner->SizeZ > 0))
 	{
+		// ?? Volume's version of IsSRGB() here :
 		bool bIsSRGB = true;
-
 		// if render target gamma used was 1.0 then disable SRGB for the static texture
 		if(FMath::Abs(GetDisplayGamma() - 1.0f) < UE_KINDA_SMALL_NUMBER)
 		{
@@ -320,11 +324,14 @@ FIntPoint FTextureRenderTargetVolumeResource::GetSizeXY() const
 
 float FTextureRenderTargetVolumeResource::GetDisplayGamma() const
 {
+	// code dupe ; move this up to the top level, it's duped everywhere
+
 	if(Owner->TargetGamma > UE_KINDA_SMALL_NUMBER * 10.0f)
 	{
 		return Owner->TargetGamma;
 	}
 	EPixelFormat Format = Owner->GetFormat();
+	// ?? hard-coding a few formats but not others, likely wrong
 	if(Format == PF_R32_FLOAT || Format == PF_FloatRGBA || Owner->bForceLinearGamma)
 	{
 		return 1.0f;
@@ -345,16 +352,23 @@ bool FTextureRenderTargetVolumeResource::ReadFloat16Pixels(TArray<FFloat16Color>
 		InSrcRect = FIntRect(0, 0, GetSizeXY().X, GetSizeXY().Y);
 	}
 	
-	/**
+	check( Owner->GetFormat() == GetRenderTargetTexture()->GetFormat() );
+	EPixelFormat PF = GetRenderTargetTexture()->GetFormat();
 
-	deep inside RHI this will crash unless the format is one of these :
-	no checks are done in outer code
+	if ( PF != PF_FloatRGBA && PF != PF_R16F && PF != PF_R32_FLOAT )
+	{
+		/**
 
-	bool bIsRGBAFmt = TextureDesc.Format == GPixelFormats[PF_FloatRGBA].PlatformFormat;
-	bool bIsR16FFmt = TextureDesc.Format == GPixelFormats[PF_R16F].PlatformFormat;
-	bool bIsR32FFmt = TextureDesc.Format == GPixelFormats[PF_R32_FLOAT].PlatformFormat;
-	check(bIsRGBAFmt || bIsR16FFmt || bIsR32FFmt);
-	**/
+		deep inside RHI this will crash unless the format is one of these :
+
+		bool bIsRGBAFmt = TextureDesc.Format == GPixelFormats[PF_FloatRGBA].PlatformFormat;
+		bool bIsR16FFmt = TextureDesc.Format == GPixelFormats[PF_R16F].PlatformFormat;
+		bool bIsR32FFmt = TextureDesc.Format == GPixelFormats[PF_R32_FLOAT].PlatformFormat;
+		check(bIsRGBAFmt || bIsR16FFmt || bIsR32FFmt);
+		**/
+
+		return false;
+	}
 
 	OutImageData.Reset();
 
diff --git a/Engine/Source/Runtime/ImageCore/Private/ImageCoreUtils.cpp b/Engine/Source/Runtime/ImageCore/Private/ImageCoreUtils.cpp
--- a/Engine/Source/Runtime/ImageCore/Private/ImageCoreUtils.cpp
+++ b/Engine/Source/Runtime/ImageCore/Private/ImageCoreUtils.cpp
@@ -199,4 +199,126 @@ IMAGECORE_API ETextureSourceFormat FImageCoreUtils::GetCommonSourceFormat(ETextu
 	default:
 		return TSF_RGBA32F;
 	}
-}
\ No newline at end of file
+}
+
+IMAGECORE_API ERawImageFormat::Type FImageCoreUtils::GetRawImageFormatForPixelFormat(EPixelFormat PF)
+{
+	switch(PF)
+	{
+	case PF_Unknown:
+	case PF_DepthStencil:
+	case PF_ShadowDepth:
+	case PF_D24                  :
+		return ERawImageFormat::Invalid;
+
+	case PF_A32B32G32R32F:
+	case PF_G32R32F              :
+	case PF_R32G32B32A32_UINT	:
+	case PF_R32G32B32_UINT		:
+	case PF_R32G32B32_SINT		:
+	case PF_R32G32B32F			:
+	case PF_R32G32_UINT         :
+		return ERawImageFormat::RGBA32F;
+
+	case PF_B8G8R8A8:
+	case PF_R5G6B5_UNORM         :
+	case PF_R8G8B8A8             :
+	case PF_A8R8G8B8			:
+	case PF_R8G8                :
+	case PF_XGXR8				:
+	case PF_R8G8B8A8_UINT		:
+	case PF_R8G8B8A8_SNORM		:
+	case PF_B5G5R5A1_UNORM       :
+	case PF_G16R16_SNORM		:
+	case PF_R8G8_UINT			:
+		return ERawImageFormat::BGRA8;
+		
+	case PF_G8:
+	case PF_L8	:
+	case PF_R8_UINT				:
+	case PF_R8		            :
+		return ERawImageFormat::G8;
+
+	case PF_G16:
+		return ERawImageFormat::G16;
+
+	case PF_FloatRGB:
+	case PF_FloatRGBA:
+	case PF_G16R16F              :
+	case PF_G16R16F_FILTER       :
+	case PF_FloatR11G11B10       :
+	case PF_R9G9B9EXP5			:
+		return ERawImageFormat::RGBA16F;
+		
+	case PF_R16F                 :
+	case PF_R16F_FILTER          :
+		return ERawImageFormat::R16F;
+
+	case PF_R32_FLOAT:
+	case PF_R32_UINT             :
+	case PF_R32_SINT             :
+	case PF_R8_SINT				:
+	case PF_R64_UINT			:
+		return ERawImageFormat::R32F;
+
+	case PF_G16R16:
+	case PF_A2B10G10R10          :
+	case PF_A16B16G16R16         :
+	case PF_R16_UINT             :
+	case PF_R16_SINT             :
+	case PF_R16G16B16A16_UINT    :
+	case PF_R16G16B16A16_SINT    :
+	case PF_R16G16_UINT			:
+	case PF_R16G16B16A16_UNORM	:
+	case PF_R16G16B16A16_SNORM	:
+		return ERawImageFormat::RGBA16;
+		
+	case PF_DXT1:
+	case PF_DXT3:
+	case PF_DXT5:
+	case PF_UYVY:
+	case PF_BC5                  :
+	case PF_V8U8                 :
+	case PF_A1                   :
+	case PF_A8                   :
+	case PF_PVRTC2               :
+	case PF_PVRTC4               :
+	case PF_BC4					:
+	case PF_ATC_RGB				:
+	case PF_ATC_RGBA_E			:
+	case PF_ATC_RGBA_I			:
+	case PF_X24_G8				:
+	case PF_ETC1				:
+	case PF_ETC2_RGB			:
+	case PF_ETC2_RGBA			:
+	case PF_ASTC_4x4             :
+	case PF_ASTC_6x6             :
+	case PF_ASTC_8x8             :
+	case PF_ASTC_10x10           :
+	case PF_ASTC_12x12           :
+	case PF_BC6H				:
+	case PF_BC7					:
+	case PF_PLATFORM_HDR_0		:
+	case PF_PLATFORM_HDR_1		:
+	case PF_PLATFORM_HDR_2		:
+	case PF_NV12				:
+	case PF_ETC2_R11_EAC		:
+	case PF_ETC2_RG11_EAC		:
+	case PF_ASTC_4x4_HDR         :
+	case PF_ASTC_6x6_HDR         :
+	case PF_ASTC_8x8_HDR         :
+	case PF_ASTC_10x10_HDR       :
+	case PF_ASTC_12x12_HDR       :
+	case PF_P010				:
+	case PF_ASTC_4x4_NORM_RG	:
+	case PF_ASTC_6x6_NORM_RG	:
+	case PF_ASTC_8x8_NORM_RG	:
+	case PF_ASTC_10x10_NORM_RG	:
+	case PF_ASTC_12x12_NORM_RG	:
+		return ERawImageFormat::Invalid;
+
+	default:
+		check(0); // unexpected value
+		return ERawImageFormat::Invalid;
+	}
+};
diff --git a/Engine/Source/Runtime/ImageCore/Public/ImageCoreUtils.h b/Engine/Source/Runtime/ImageCore/Public/ImageCoreUtils.h
--- a/Engine/Source/Runtime/ImageCore/Public/ImageCoreUtils.h
+++ b/Engine/Source/Runtime/ImageCore/Public/ImageCoreUtils.h
@@ -54,5 +54,10 @@ namespace FImageCoreUtils
 	 * @param Format2 - Second source format.
 	 */
 	IMAGECORE_API ETextureSourceFormat GetCommonSourceFormat(ETextureSourceFormat Format1, ETextureSourceFormat Format2);
-
+	
+	// return an ERawImageFormat for a PixelFormat
+	// these are not equal and you cannot blit between them
+	//	it's just the closest reasonable thing that can hold that PixelFormat
+	//	often there is no good match
+	IMAGECORE_API ERawImageFormat::Type GetRawImageFormatForPixelFormat(EPixelFormat PF);
 };
diff --git a/Engine/Source/Runtime/Windows/D3D11RHI/Private/D3D11RenderTarget.cpp b/Engine/Source/Runtime/Windows/D3D11RHI/Private/D3D11RenderTarget.cpp
--- a/Engine/Source/Runtime/Windows/D3D11RHI/Private/D3D11RenderTarget.cpp
+++ b/Engine/Source/Runtime/Windows/D3D11RHI/Private/D3D11RenderTarget.cpp
@@ -760,6 +760,9 @@ void FD3D11DynamicRHI::RHIReadSurfaceFloatData(FRHITexture* TextureRHI,FIntRect
 
 static void ConvertRAWSurfaceDataToFLinearColor(EPixelFormat Format, uint32 Width, uint32 Height, uint8 *In, uint32 SrcPitch, FLinearColor* Out, FReadSurfaceDataFlags InFlags)
 {
+	// @@ missing formats I added to D3D12RenderTarget
+	//	factor out and share code!
+
 	bool bLinearToGamma = InFlags.GetLinearToGamma();
 	if (Format == PF_R16F || Format == PF_R16F_FILTER)
 	{
