Move Select Text/Name helper functions into kismetmathlibrary

#tests Preflight with cook

#virtualized

diff --git a/Engine/Source/Runtime/Engine/Classes/Kismet/KismetMathLibrary.h b/Engine/Source/Runtime/Engine/Classes/Kismet/KismetMathLibrary.h
--- a/Engine/Source/Runtime/Engine/Classes/Kismet/KismetMathLibrary.h
+++ b/Engine/Source/Runtime/Engine/Classes/Kismet/KismetMathLibrary.h
@@ -3834,6 +3834,14 @@ class UKismetMathLibrary : public UBlueprintFunctionLibrary
 	UFUNCTION(BlueprintPure, Category="Utilities|String")
 	static ENGINE_API FString SelectString(const FString& A, const FString& B, bool bPickA);
 
+	/** If bPickA is true, A is returned, otherwise B is */
+	UFUNCTION(BlueprintPure, Category="Utilities|String")
+	static ENGINE_API FText SelectText(const FText A, const FText B, bool bPickA);
+
+	/** If bPickA is true, A is returned, otherwise B is */
+	UFUNCTION(BlueprintPure, Category="Utilities|String")
+	static ENGINE_API FName SelectName(const FName A, const FName B, bool bPickA);
+
 	/** If bPickA is true, A is returned, otherwise B is */
 	UFUNCTION(BlueprintPure, Category="Math|Integer")
 	static ENGINE_API int32 SelectInt(int32 A, int32 B, bool bPickA);
diff --git a/Engine/Source/Runtime/Engine/Classes/Kismet/KismetMathLibrary.inl b/Engine/Source/Runtime/Engine/Classes/Kismet/KismetMathLibrary.inl
--- a/Engine/Source/Runtime/Engine/Classes/Kismet/KismetMathLibrary.inl
+++ b/Engine/Source/Runtime/Engine/Classes/Kismet/KismetMathLibrary.inl
@@ -3827,13 +3827,24 @@ FString UKismetMathLibrary::ToHex_LinearColor(FLinearColor InColor)
 	return InColor.ToFColor(true).ToHex();
 }
 
-
 KISMET_MATH_FORCEINLINE
 FString UKismetMathLibrary::SelectString(const FString& A, const FString& B, bool bSelectA)
 {
 	return bSelectA ? A : B;
 }
 
+KISMET_MATH_FORCEINLINE
+FText UKismetMathLibrary::SelectText(const FText A, const FText B, bool bSelectA)
+{
+	return bSelectA ? A : B;
+}
+
+KISMET_MATH_FORCEINLINE
+FName UKismetMathLibrary::SelectName(const FName A, const FName B, bool bSelectA)
+{
+	return bSelectA ? A : B;
+}
+
 KISMET_MATH_FORCEINLINE
 int32 UKismetMathLibrary::SelectInt(int32 A, int32 B, bool bSelectA)
 {
