[NetPhysics] Fix replication target alignment when receiving states with ServerFrame numbers after having received with ServerFrame 0 previously.'
The calculated SendInterval could be huge which would align the target incorrectly making the replicated object desync for a moment.

diff --git a/Engine/Source/Runtime/Engine/Private/PhysicsEngine/PhysicsReplication.cpp b/Engine/Source/Runtime/Engine/Private/PhysicsEngine/PhysicsReplication.cpp
--- a/Engine/Source/Runtime/Engine/Private/PhysicsEngine/PhysicsReplication.cpp
+++ b/Engine/Source/Runtime/Engine/Private/PhysicsEngine/PhysicsReplication.cpp
@@ -875,7 +875,7 @@ void FPhysicsReplicationAsync::UpdateAsyncTarget(const FPhysicsRepAsyncInputData
 	{
 		const int32 PrevTickCount = Target->TickCount;
 		const int32 PrevReceiveInterval = Target->ReceiveInterval;
-		const int32 SendInterval = Input.ServerFrame - Target->ServerFrame;
+		const int32 SendInterval = (Target->ServerFrame <= 0) ? 0 : Input.ServerFrame - Target->ServerFrame;
 		const int32 AdjustedAverageReceiveInterval = FMath::CeilToInt(Target->AverageReceiveInterval) * PhysicsReplicationCVars::PredictiveInterpolationCVars::TargetTickAlignmentClampMultiplier;
 
 		Target->PrevServerFrame = Target->bWaiting ? Input.ServerFrame : Target->ServerFrame;
