Bring back CL30049175 with fix.

-----------------------------------------------------------------
[Backout] - CL30049175
Original CL Desc
-----------------------------------------------------------------
World Partition HashSet: fix non-spatially loaded actors never loading when in a data layer or content bundle.

diff --git a/Engine/Source/Runtime/Engine/Private/WorldPartition/RuntimeHashSet/WorldPartitionRuntimeHashSet.cpp b/Engine/Source/Runtime/Engine/Private/WorldPartition/RuntimeHashSet/WorldPartitionRuntimeHashSet.cpp
--- a/Engine/Source/Runtime/Engine/Private/WorldPartition/RuntimeHashSet/WorldPartitionRuntimeHashSet.cpp
+++ b/Engine/Source/Runtime/Engine/Private/WorldPartition/RuntimeHashSet/WorldPartitionRuntimeHashSet.cpp
@@ -188,16 +188,6 @@ bool UWorldPartitionRuntimeHashSet::IsValidHLODLayer(FName GridName, const FSoft
 	return false;
 }
 
-TArray<UWorldPartitionRuntimeCell*> UWorldPartitionRuntimeHashSet::GetAlwaysLoadedCells() const
-{
-	TArray<UWorldPartitionRuntimeCell*> AlwaysLoadedCells;
-	ForEachStreamingData([&AlwaysLoadedCells](const FRuntimePartitionStreamingData& StreamingData)
-	{
-		AlwaysLoadedCells.Append(StreamingData.NonStreamingCells);
-	});
-	return AlwaysLoadedCells;
-}
-
 bool UWorldPartitionRuntimeHashSet::ParseGridName(FName GridName, TArray<FName>& MainPartitionTokens, TArray<FName>& HLODPartitionTokens)
 {
 	// If the grid name is none, it directly maps to the main partition
diff --git a/Engine/Source/Runtime/Engine/Private/WorldPartition/RuntimeHashSet/WorldPartitionRuntimeHashSetStreamingGeneration.cpp b/Engine/Source/Runtime/Engine/Private/WorldPartition/RuntimeHashSet/WorldPartitionRuntimeHashSetStreamingGeneration.cpp
--- a/Engine/Source/Runtime/Engine/Private/WorldPartition/RuntimeHashSet/WorldPartitionRuntimeHashSetStreamingGeneration.cpp
+++ b/Engine/Source/Runtime/Engine/Private/WorldPartition/RuntimeHashSet/WorldPartitionRuntimeHashSetStreamingGeneration.cpp
@@ -224,13 +224,13 @@ bool UWorldPartitionRuntimeHashSet::GenerateStreaming(UWorldPartitionStreamingPo
 				StreamingData.Name = CellDescInstance.SourcePartition->Name;
 				StreamingData.LoadingRange = CellDescInstance.SourcePartition->LoadingRange;
 
-				if (RuntimeCell->IsAlwaysLoaded())
+				if (CellDescInstance.bIsSpatiallyLoaded)
 				{
-					StreamingData.NonStreamingCells.Add(RuntimeCell);
+					StreamingData.StreamingCells.Add(RuntimeCell);
 				}
 				else
 				{
-					StreamingData.StreamingCells.Add(RuntimeCell);
+					StreamingData.NonStreamingCells.Add(RuntimeCell);
 				}
 			}
 		}
diff --git a/Engine/Source/Runtime/Engine/Public/WorldPartition/RuntimeHashSet/WorldPartitionRuntimeHashSet.h b/Engine/Source/Runtime/Engine/Public/WorldPartition/RuntimeHashSet/WorldPartitionRuntimeHashSet.h
--- a/Engine/Source/Runtime/Engine/Public/WorldPartition/RuntimeHashSet/WorldPartitionRuntimeHashSet.h
+++ b/Engine/Source/Runtime/Engine/Public/WorldPartition/RuntimeHashSet/WorldPartitionRuntimeHashSet.h
@@ -140,7 +140,6 @@ public:
 	ENGINE_API virtual void FlushStreaming() override;
 	ENGINE_API virtual bool IsValidGrid(FName GridName) const;
 	ENGINE_API virtual bool IsValidHLODLayer(FName GridName, const FSoftObjectPath& HLODLayerPath) const;
-	ENGINE_API virtual TArray<UWorldPartitionRuntimeCell*> GetAlwaysLoadedCells() const override;
 	ENGINE_API virtual void DumpStateLog(FHierarchicalLogArchive& Ar) const override;
 
 	// Helpers
