LFAtlas - added al lthe cvars into the project setting window

diff --git a/Engine/Source/Runtime/Engine/Classes/Engine/RendererSettings.h b/Engine/Source/Runtime/Engine/Classes/Engine/RendererSettings.h
--- a/Engine/Source/Runtime/Engine/Classes/Engine/RendererSettings.h
+++ b/Engine/Source/Runtime/Engine/Classes/Engine/RendererSettings.h
@@ -893,6 +893,22 @@ class URendererSettings : public UDeveloperSettings
 		ConfigRestartRequired = true))
 		TEnumAsByte<ELightFunctionAtlasPixelFormat::Type> LightFunctionAtlasPixelFormat;
 
+	/**
+	"Enable support for light function on volumetric fog, when the light function atlas is enabled."
+	*/
+	UPROPERTY(config, EditAnywhere, Category = LightFunctionAtlas, meta = (
+		ConsoleVariable = "r.VolumetricFog.LightFunction", DisplayName = "Volumetric Fog Uses Light Function Atlas.",
+		ToolTip = "Enable support for light function on volumetric fog, when the light function atlas is enabled."))
+		uint32 bVolumetricFogUsesLightFunctionAtlas : 1;
+
+	/**
+	"Enable support for light function on deferred lighting (multi-pass and clustered), when the light function atlas is enabled."
+	*/
+	UPROPERTY(config, EditAnywhere, Category = LightFunctionAtlas, meta = (
+		ConsoleVariable = "r.Deferred.UsesLightFunctionAtlas", DisplayName = "Deferred Lighting Uses Light Function Atlas.",
+		ToolTip = "Enable support for light function on deferred lighting (multi-pass and clustered), when the light function atlas is enabled."))
+		uint32 bDeferredLightingUsesLightFunctionAtlas : 1;
+
 	/**
 	"Enable support for light function on Single Layer Water when the light function atlas is enabled."
 	*/
