Fix CPU skinning crash on selecting Character>Mesh>SelectedBoneWeight option in editor
- CacheVertices is already called from rendering thread, calling ENQUEUE_RENDER_COMMAND within itself can cause render commands to be executed after FSkeletalMeshObjectLOD::ReleaseResources function releases the resources, so remove the ENQUEUE_RENDER_COMMAND calls.
- Remove CacheVertices function from the base FSkeletalMeshObject because it is only used by cpu skinning

diff --git a/Engine/Source/Runtime/Engine/Private/SkeletalRenderCPUSkin.cpp b/Engine/Source/Runtime/Engine/Private/SkeletalRenderCPUSkin.cpp
--- a/Engine/Source/Runtime/Engine/Private/SkeletalRenderCPUSkin.cpp
+++ b/Engine/Source/Runtime/Engine/Private/SkeletalRenderCPUSkin.cpp
@@ -196,7 +196,7 @@ void FSkeletalMeshObjectCPUSkin::UpdateDynamicData_RenderThread(FRHICommandList&
 	check(DynamicData);
 
 	// update vertices using the new data
-	CacheVertices(DynamicData->LODIndex,true);
+	CacheVertices(DynamicData->LODIndex, true, RHICmdList);
 }
 
 #define SKIN_LOD_VERTICES(VertexType, NumUVs) \
@@ -222,9 +222,10 @@ void FSkeletalMeshObjectCPUSkin::UpdateDynamicData_RenderThread(FRHICommandList&
 }\
 	
 
-void FSkeletalMeshObjectCPUSkin::CacheVertices(int32 LODIndex, bool bForce) const
+void FSkeletalMeshObjectCPUSkin::CacheVertices(int32 LODIndex, bool bForce, FRHICommandList& RHICmdList) const
 {
 	SCOPE_CYCLE_COUNTER( STAT_CPUSkinUpdateRTTime);
+	check(IsInParallelRenderingThread());
 
 	// Source skel mesh and static lod model
 	FSkeletalMeshLODRenderData& LOD = SkeletalMeshRenderData->LODRenderData[LODIndex];
@@ -297,25 +298,18 @@ void FSkeletalMeshObjectCPUSkin::CacheVertices(int32 LODIndex, bool bForce) cons
 			}
 		}
 
-		BeginUpdateResourceRHI(&MeshLOD.PositionVertexBuffer, &UE::RenderCommandPipe::SkeletalMesh);
-		BeginUpdateResourceRHI(&MeshLOD.StaticMeshVertexBuffer, &UE::RenderCommandPipe::SkeletalMesh);
+		MeshLOD.PositionVertexBuffer.UpdateRHI(RHICmdList);
+		MeshLOD.StaticMeshVertexBuffer.UpdateRHI(RHICmdList);
 
-		const FSkeletalMeshObjectCPUSkin::FSkeletalMeshObjectLOD* MeshLODptr = &MeshLOD;
-		FLocalVertexFactory* VertexFactoryPtr = &MeshLOD.VertexFactory;
-		ENQUEUE_RENDER_COMMAND(UpdateSkeletalMeshCPUSkinVertexFactory)(UE::RenderCommandPipe::SkeletalMesh,
-			[VertexFactoryPtr, MeshLODptr](FRHICommandList& RHICmdList)
-		{
-			FLocalVertexFactory::FDataType Data;
-
-			MeshLODptr->PositionVertexBuffer.BindPositionVertexBuffer(VertexFactoryPtr, Data);
-			MeshLODptr->StaticMeshVertexBuffer.BindTangentVertexBuffer(VertexFactoryPtr, Data);
-			MeshLODptr->StaticMeshVertexBuffer.BindPackedTexCoordVertexBuffer(VertexFactoryPtr, Data, MAX_TEXCOORDS);
-			MeshLODptr->StaticMeshVertexBuffer.BindLightMapVertexBuffer(VertexFactoryPtr, Data, 0);
-			MeshLODptr->MeshObjectColorBuffer->BindColorVertexBuffer(VertexFactoryPtr, Data);
+		FLocalVertexFactory::FDataType Data;
+		MeshLOD.PositionVertexBuffer.BindPositionVertexBuffer(&MeshLOD.VertexFactory, Data);
+		MeshLOD.StaticMeshVertexBuffer.BindTangentVertexBuffer(&MeshLOD.VertexFactory, Data);
+		MeshLOD.StaticMeshVertexBuffer.BindPackedTexCoordVertexBuffer(&MeshLOD.VertexFactory, Data, MAX_TEXCOORDS);
+		MeshLOD.StaticMeshVertexBuffer.BindLightMapVertexBuffer(&MeshLOD.VertexFactory, Data, 0);
+		MeshLOD.MeshObjectColorBuffer->BindColorVertexBuffer(&MeshLOD.VertexFactory, Data);
 
-			VertexFactoryPtr->SetData(RHICmdList, Data);
-			VertexFactoryPtr->InitResource(RHICmdList);
-		});
+		MeshLOD.VertexFactory.SetData(RHICmdList, Data);
+		MeshLOD.VertexFactory.InitResource(RHICmdList);
 	}
 }
 
diff --git a/Engine/Source/Runtime/Engine/Private/SkeletalRenderCPUSkin.h b/Engine/Source/Runtime/Engine/Private/SkeletalRenderCPUSkin.h
--- a/Engine/Source/Runtime/Engine/Private/SkeletalRenderCPUSkin.h
+++ b/Engine/Source/Runtime/Engine/Private/SkeletalRenderCPUSkin.h
@@ -101,12 +101,19 @@ public:
 	ENGINE_API virtual void Update(int32 LODIndex,USkinnedMeshComponent* InMeshComponent,const FMorphTargetWeightMap& InActiveMorphTargets, const TArray<float>& MorphTargetsWeights, EPreviousBoneTransformUpdateMode PreviousBoneTransformUpdateMode, const FExternalMorphWeightData& InExternalMorphWeightData) override;
 	ENGINE_API void UpdateDynamicData_RenderThread(FRHICommandList& RHICmdList, FDynamicSkelMeshObjectDataCPUSkin* InDynamicData, uint64 FrameNumberToPrepare, uint32 RevisionNumber);
 	ENGINE_API virtual void EnableOverlayRendering(bool bEnabled, const TArray<int32>* InBonesOfInterest, const TArray<UMorphTarget*>* InMorphTargetOfInterest) override;
-	ENGINE_API virtual void CacheVertices(int32 LODIndex, bool bForce) const override;
 	virtual bool IsCPUSkinned() const override { return true; }
 	ENGINE_API virtual const FVertexFactory* GetSkinVertexFactory(const FSceneView* View, int32 LODIndex, int32 ChunkIdx, ESkinVertexFactoryMode VFMode = ESkinVertexFactoryMode::Default) const override;
 	ENGINE_API virtual TArray<FTransform>* GetComponentSpaceTransforms() const override;
 	ENGINE_API virtual const TArray<FMatrix44f>& GetReferenceToLocalMatrices() const override;
 
+	/**
+	 * Re-skin cached vertices for an LOD and update the vertex buffer. Note that this
+	 * function is called from the render thread!
+	 * @param	LODIndex - index to LODs
+	 * @param	bForce - force update even if LOD index hasn't changed
+	 */
+	ENGINE_API void CacheVertices(int32 LODIndex, bool bForce, FRHICommandList& RHICmdList) const;
+
 	virtual int32 GetLOD() const override
 	{
 		if(DynamicData)
diff --git a/Engine/Source/Runtime/Engine/Private/SkeletalRenderGPUSkin.h b/Engine/Source/Runtime/Engine/Private/SkeletalRenderGPUSkin.h
--- a/Engine/Source/Runtime/Engine/Private/SkeletalRenderGPUSkin.h
+++ b/Engine/Source/Runtime/Engine/Private/SkeletalRenderGPUSkin.h
@@ -356,7 +356,6 @@ public:
 	ENGINE_API virtual void PreGDMECallback(FRHICommandList& RHICmdList, FGPUSkinCache* GPUSkinCache, uint32 FrameNumber) override;
 	ENGINE_API virtual const FVertexFactory* GetSkinVertexFactory(const FSceneView* View, int32 LODIndex,int32 ChunkIdx, ESkinVertexFactoryMode VFMode = ESkinVertexFactoryMode::Default) const override;
 	ENGINE_API virtual const FSkinBatchVertexFactoryUserData* GetVertexFactoryUserData(const int32 LODIndex, int32 ChunkIdx, ESkinVertexFactoryMode VFMode) const override;
-	virtual void CacheVertices(int32 LODIndex, bool bForce) const override {}
 	virtual bool IsCPUSkinned() const override { return false; }
 	ENGINE_API virtual TArray<FTransform>* GetComponentSpaceTransforms() const override;
 	ENGINE_API virtual const TArray<FMatrix44f>& GetReferenceToLocalMatrices() const override;
diff --git a/Engine/Source/Runtime/Engine/Private/SkeletalRenderStatic.h b/Engine/Source/Runtime/Engine/Private/SkeletalRenderStatic.h
--- a/Engine/Source/Runtime/Engine/Private/SkeletalRenderStatic.h
+++ b/Engine/Source/Runtime/Engine/Private/SkeletalRenderStatic.h
@@ -30,7 +30,6 @@ public:
 	virtual void Update(int32 LODIndex,USkinnedMeshComponent* InMeshComponent,const FMorphTargetWeightMap& InActiveMorphTargets, const TArray<float>& MorphTargetWeights, EPreviousBoneTransformUpdateMode PreviousBoneTransformUpdateMode, const FExternalMorphWeightData& InExternalMorphWeightData) override {};
 	//virtual void UpdateRecomputeTangent(int32 MaterialIndex, int32 LODIndex, bool bRecomputeTangent) override {};
 	virtual void EnableOverlayRendering(bool bEnabled, const TArray<int32>* InBonesOfInterest, const TArray<UMorphTarget*>* InMorphTargetOfInterest) override {};
-	virtual void CacheVertices(int32 LODIndex, bool bForce) const override {};
 	virtual bool IsCPUSkinned() const override { return true; }
 	ENGINE_API virtual const FVertexFactory* GetSkinVertexFactory(const FSceneView* View, int32 LODIndex, int32 ChunkIdx, ESkinVertexFactoryMode VFMode = ESkinVertexFactoryMode::Default) const override;
 	ENGINE_API virtual TArray<FTransform>* GetComponentSpaceTransforms() const override;
diff --git a/Engine/Source/Runtime/Engine/Public/SkeletalRenderPublic.h b/Engine/Source/Runtime/Engine/Public/SkeletalRenderPublic.h
--- a/Engine/Source/Runtime/Engine/Public/SkeletalRenderPublic.h
+++ b/Engine/Source/Runtime/Engine/Public/SkeletalRenderPublic.h
@@ -113,14 +113,6 @@ public:
 	 */
 	virtual const FSkinBatchVertexFactoryUserData* GetVertexFactoryUserData(const int32 LODIndex, int32 ChunkIdx, ESkinVertexFactoryMode VFMode) const { return nullptr; }
 
-	/**
-	 * Re-skin cached vertices for an LOD and update the vertex buffer. Note that this
-	 * function is called from the render thread!
-	 * @param	LODIndex - index to LODs
-	 * @param	bForce - force update even if LOD index hasn't changed
-	 */
-	virtual void CacheVertices(int32 LODIndex, bool bForce) const = 0;
-
 	/**
 	 * Returns true if this mesh performs skinning on the CPU.
 	 */
