LFAtlas - fix for deferred light pixel shader to compil when substrate is enabled.

diff --git a/Engine/Source/Runtime/Engine/Private/ShaderCompiler/ShaderCompiler.cpp b/Engine/Source/Runtime/Engine/Private/ShaderCompiler/ShaderCompiler.cpp
--- a/Engine/Source/Runtime/Engine/Private/ShaderCompiler/ShaderCompiler.cpp
+++ b/Engine/Source/Runtime/Engine/Private/ShaderCompiler/ShaderCompiler.cpp
@@ -9084,8 +9084,14 @@ void VerifyGlobalShaders(EShaderPlatform Platform, const ITargetPlatform* Target
 			}
 		}
 
+		int32 PermutationCountLimit = 832;	// Nanite culling as of today (2022-01-11) can go up to 832 permutations
+		if (Substrate::IsSubstrateEnabled())
+		{
+			// SUBSTRATE_TODO reduce the number of permutation of FDeferredLightPS.
+			PermutationCountLimit = 1304;	// FDeferredLightPS as of today (2023-12-04)
+		}
 		ensureMsgf(
-			PermutationCountToCompile <= 832,	// Nanite culling as of today (2022-01-11) can go up to 832 permutations
+			PermutationCountToCompile <= PermutationCountLimit,
 			TEXT("Global shader %s has %i permutations: probably more than it needs."),
 			GlobalShaderType->GetName(), PermutationCountToCompile);
 
