Fixed world partition streaming generation not properly handling dirty actors of a newly created level based on a template partition map.

diff --git a/Engine/Source/Runtime/Engine/Private/WorldPartition/ActorDescContainer.cpp b/Engine/Source/Runtime/Engine/Private/WorldPartition/ActorDescContainer.cpp
--- a/Engine/Source/Runtime/Engine/Private/WorldPartition/ActorDescContainer.cpp
+++ b/Engine/Source/Runtime/Engine/Private/WorldPartition/ActorDescContainer.cpp
@@ -199,7 +199,10 @@ bool UActorDescContainer::IsActorDescHandled(const AActor* Actor) const
 	if (Actor->GetContentBundleGuid() == GetContentBundleGuid())
 	{
 		const FString ActorPackageName = Actor->GetPackage()->GetName();
-		const FString ExternalActorPath = GetExternalActorPath() / TEXT("");
+		// Resolve ExternalActorPath using the actual package name instead of ContainerPackageName to properly handle 
+		// unsaved levels created with an existing template partitioned map.
+		// However, always rely on ContainerPackageName for content bundles to resolve the external actor path (no content bundles for template maps).
+		const FString ExternalActorPath = (GetContentBundleGuid().IsValid() ? GetExternalActorPath() : ULevel::GetExternalActorsPath(GetPackage()->GetName())) / TEXT("");
 		return ActorPackageName.StartsWith(ExternalActorPath);
 	}
 	return false;
diff --git a/Engine/Source/Runtime/Engine/Private/WorldPartition/WorldPartitionStreamingGeneration.cpp b/Engine/Source/Runtime/Engine/Private/WorldPartition/WorldPartitionStreamingGeneration.cpp
--- a/Engine/Source/Runtime/Engine/Private/WorldPartition/WorldPartitionStreamingGeneration.cpp
+++ b/Engine/Source/Runtime/Engine/Private/WorldPartition/WorldPartitionStreamingGeneration.cpp
@@ -560,13 +560,6 @@ class FWorldPartitionStreamingGenerator
 			FWorldPartitionActorDesc* ModifiedActorDesc = ModifiedActorsDescList->AddActor(InActor);
 
 			const UActorDescContainer* HandlingContainer = InActorDescCollection.FindHandlingContainer(InActor);
-			if (HandlingContainer == nullptr)
-			{
-				// @todo_ow : UActorDescContainer::IsActorDescHandled(AActor*) does not handle the case where the InActor is an unsaved new non-plugin actor
-				// and the level is in memory. Fix UActorDescContainer::IsActorDescHandled so it handle those actors
-				HandlingContainer = InActorDescCollection.GetMainActorDescContainer();
-			}
-
 			check(HandlingContainer != nullptr);
 
 			// Pretend that this actor descriptor belongs to the original container, even if it's not present. It's essentially a proxy
