PR #11174: Add LowerHemisphereColor to fast path for updating skylight scene proxy

#11174

diff --git a/Engine/Source/Runtime/Engine/Private/Components/SkyLightComponent.cpp b/Engine/Source/Runtime/Engine/Private/Components/SkyLightComponent.cpp
--- a/Engine/Source/Runtime/Engine/Private/Components/SkyLightComponent.cpp
+++ b/Engine/Source/Runtime/Engine/Private/Components/SkyLightComponent.cpp
@@ -461,12 +461,14 @@ void USkyLightComponent::UpdateLimitedRenderingStateFast()
 		FLinearColor InLightColor = FLinearColor(LightColor) * Intensity;
 		float InIndirectLightingIntensity = IndirectLightingIntensity;
 		float InVolumetricScatteringIntensity = VolumetricScatteringIntensity;
+		FLinearColor InLowerHemisphereColor = LowerHemisphereColor;
 		ENQUEUE_RENDER_COMMAND(FFastUpdateSkyLightCommand)(
-			[LightSceneProxy, InLightColor, InIndirectLightingIntensity, InVolumetricScatteringIntensity] (FRHICommandListBase&)
+			[LightSceneProxy, InLightColor, InIndirectLightingIntensity, InVolumetricScatteringIntensity, InLowerHemisphereColor] (FRHICommandListBase&)
 			{
 				LightSceneProxy->SetLightColor(InLightColor);
 				LightSceneProxy->IndirectLightingIntensity = InIndirectLightingIntensity;
 				LightSceneProxy->VolumetricScatteringIntensity = InVolumetricScatteringIntensity;
+				LightSceneProxy->LowerHemisphereColor = InLowerHemisphereColor;
 			});
 	}
 }
@@ -1044,7 +1046,7 @@ void USkyLightComponent::SetLowerHemisphereColor(const FLinearColor& InLowerHemi
 		&& LowerHemisphereColor != InLowerHemisphereColor)
 	{
 		LowerHemisphereColor = InLowerHemisphereColor;
-		MarkRenderStateDirty();
+		UpdateLimitedRenderingStateFast();
 	}
 }
 
