Refactor how static ray tracing geometries are cached in FPrimitiveSceneInfo.
- don't cache array in FPrimitiveSceneProxy to then move to FPrimitiveSceneInfo.
- simply build array during FPrimitiveSceneInfo construction as required.

diff --git a/Engine/Source/Runtime/Engine/Private/PrimitiveSceneProxy.cpp b/Engine/Source/Runtime/Engine/Private/PrimitiveSceneProxy.cpp
--- a/Engine/Source/Runtime/Engine/Private/PrimitiveSceneProxy.cpp
+++ b/Engine/Source/Runtime/Engine/Private/PrimitiveSceneProxy.cpp
@@ -1791,7 +1791,7 @@ ERayTracingPrimitiveFlags FPrimitiveSceneProxy::GetCachedRayTracingInstance(FRay
 
 	if (IsRayTracingStaticRelevant())
 	{
-		if (PrimitiveSceneInfo->GetRayTracingGeometryNum() == 0
+		if (PrimitiveSceneInfo->GetStaticRayTracingGeometryNum() == 0
 			|| PrimitiveSceneInfo->StaticMeshes.IsEmpty())
 		{
 			return ERayTracingPrimitiveFlags::Excluded;
diff --git a/Engine/Source/Runtime/Engine/Private/SkeletalMesh.cpp b/Engine/Source/Runtime/Engine/Private/SkeletalMesh.cpp
--- a/Engine/Source/Runtime/Engine/Private/SkeletalMesh.cpp
+++ b/Engine/Source/Runtime/Engine/Private/SkeletalMesh.cpp
@@ -5864,20 +5864,6 @@ FSkeletalMeshSceneProxy::FSkeletalMeshSceneProxy(const USkinnedMeshComponent* Co
 		EnableGPUSceneSupportFlags();
 
 	}
-#if RHI_RAYTRACING
-	if (IsRayTracingAllowed())
-	{
-		if (bRenderStatic)
-		{
-			RayTracingGeometries.AddDefaulted(SkeletalMeshRenderData->LODRenderData.Num());
-			for (int32 LODIndex = 0; LODIndex < SkeletalMeshRenderData->LODRenderData.Num(); LODIndex++)
-			{
-				ensure(SkeletalMeshRenderData->LODRenderData[LODIndex].NumReferencingStaticSkeletalMeshObjects > 0);
-				RayTracingGeometries[LODIndex] = &SkeletalMeshRenderData->LODRenderData[LODIndex].StaticRayTracingGeometry;
-			}
-		}
-	}
-#endif
 
 	if (IsAllowingApproximateOcclusionQueries())
 	{		
@@ -6492,6 +6478,24 @@ bool FSkeletalMeshSceneProxy::HasRayTracingRepresentation() const
 	return bRenderStatic;
 }
 
+TArray<FRayTracingGeometry*> FSkeletalMeshSceneProxy::GetStaticRayTracingGeometries() const
+{
+	if (IsRayTracingAllowed() && bRenderStatic)
+	{
+		TArray<FRayTracingGeometry*> RayTracingGeometries;
+		RayTracingGeometries.AddDefaulted(SkeletalMeshRenderData->LODRenderData.Num());
+		for (int32 LODIndex = 0; LODIndex < SkeletalMeshRenderData->LODRenderData.Num(); LODIndex++)
+		{
+			ensure(SkeletalMeshRenderData->LODRenderData[LODIndex].NumReferencingStaticSkeletalMeshObjects > 0);
+			RayTracingGeometries[LODIndex] = &SkeletalMeshRenderData->LODRenderData[LODIndex].StaticRayTracingGeometry;
+		}
+
+		return MoveTemp(RayTracingGeometries);
+	}
+
+	return {};
+}
+
 void FSkeletalMeshSceneProxy::GetDynamicRayTracingInstances(FRayTracingMaterialGatheringContext & Context, TArray<struct FRayTracingInstance>& OutRayTracingInstances)
 {
 	if (!CVarRayTracingSkeletalMeshes.GetValueOnRenderThread()
diff --git a/Engine/Source/Runtime/Engine/Private/StaticMeshRender.cpp b/Engine/Source/Runtime/Engine/Private/StaticMeshRender.cpp
--- a/Engine/Source/Runtime/Engine/Private/StaticMeshRender.cpp
+++ b/Engine/Source/Runtime/Engine/Private/StaticMeshRender.cpp
@@ -315,13 +315,7 @@ FStaticMeshSceneProxy::FStaticMeshSceneProxy(const FStaticMeshSceneProxyDesc& In
 	bNeedsDynamicRayTracingGeometries = false;
 	
 	if (IsRayTracingAllowed() && bSupportRayTracing)
-	{
-		RayTracingGeometries.AddDefaulted(RenderData->LODResources.Num());
-		for (int32 LODIndex = 0; LODIndex < RenderData->LODResources.Num(); LODIndex++)
-		{
-			RayTracingGeometries[LODIndex] = &RenderData->LODResources[LODIndex].RayTracingGeometry;
-		}
-		
+	{		
 		const bool bWantsRayTracingWPO = MaterialRelevance.bUsesWorldPositionOffset && InProxyDesc.bEvaluateWorldPositionOffsetInRayTracing;
 
 		// r.RayTracing.Geometry.StaticMeshes.WPO is handled in the following way:
@@ -1872,6 +1866,23 @@ bool FStaticMeshSceneProxy::HasRayTracingRepresentation() const
 	return bSupportRayTracing;
 }
 
+TArray<FRayTracingGeometry*> FStaticMeshSceneProxy::GetStaticRayTracingGeometries() const
+{
+	if (IsRayTracingAllowed() && bSupportRayTracing)
+	{
+		TArray<FRayTracingGeometry*> RayTracingGeometries;
+		RayTracingGeometries.AddDefaulted(RenderData->LODResources.Num());
+		for (int32 LODIndex = 0; LODIndex < RenderData->LODResources.Num(); LODIndex++)
+		{
+			RayTracingGeometries[LODIndex] = &RenderData->LODResources[LODIndex].RayTracingGeometry;
+		}
+
+		return MoveTemp(RayTracingGeometries);
+	}
+
+	return {};
+}
+
 void FStaticMeshSceneProxy::GetDynamicRayTracingInstances(FRayTracingMaterialGatheringContext& Context, TArray<FRayTracingInstance>& OutRayTracingInstances )
 {
 	if (DynamicRayTracingGeometries.IsEmpty() || CVarRayTracingStaticMeshes.GetValueOnRenderThread() == 0)
diff --git a/Engine/Source/Runtime/Engine/Public/PrimitiveSceneProxy.h b/Engine/Source/Runtime/Engine/Public/PrimitiveSceneProxy.h
--- a/Engine/Source/Runtime/Engine/Public/PrimitiveSceneProxy.h
+++ b/Engine/Source/Runtime/Engine/Public/PrimitiveSceneProxy.h
@@ -337,10 +337,7 @@ public:
 	/** Gathers dynamic ray tracing instances from this proxy. */
 	virtual void GetDynamicRayTracingInstances(struct FRayTracingMaterialGatheringContext& Context, TArray<struct FRayTracingInstance>& OutRayTracingInstances) {}
 
-	TArray<FRayTracingGeometry*>&& MoveRayTracingGeometries()
-	{
-		return static_cast<TArray<FRayTracingGeometry*>&&>(RayTracingGeometries);
-	}
+	virtual TArray<FRayTracingGeometry*> GetStaticRayTracingGeometries() const { return {}; }
 
 	/** 
 	 * Gathers static ray tracing primitives from this proxy. 
@@ -1098,10 +1095,6 @@ protected:
 		OwnerName = InOwnerName;
 	}
 
-#if RHI_RAYTRACING
-	TArray<FRayTracingGeometry*> RayTracingGeometries;
-#endif
-
 	void SetForceHidden(bool bForceHiddenIn) {bForceHidden = bForceHiddenIn;}
 
 	/** 
diff --git a/Engine/Source/Runtime/Engine/Public/SkeletalMeshSceneProxy.h b/Engine/Source/Runtime/Engine/Public/SkeletalMeshSceneProxy.h
--- a/Engine/Source/Runtime/Engine/Public/SkeletalMeshSceneProxy.h
+++ b/Engine/Source/Runtime/Engine/Public/SkeletalMeshSceneProxy.h
@@ -41,6 +41,8 @@ public:
 		return bRenderStatic;
 	}
 
+	ENGINE_API virtual TArray<FRayTracingGeometry*> GetStaticRayTracingGeometries() const override;
+
 	ENGINE_API virtual void GetDynamicRayTracingInstances(struct FRayTracingMaterialGatheringContext& Context, TArray<struct FRayTracingInstance>& OutRayTracingInstances) override;
 #endif // RHI_RAYTRACING
 
diff --git a/Engine/Source/Runtime/Engine/Public/StaticMeshSceneProxy.h b/Engine/Source/Runtime/Engine/Public/StaticMeshSceneProxy.h
--- a/Engine/Source/Runtime/Engine/Public/StaticMeshSceneProxy.h
+++ b/Engine/Source/Runtime/Engine/Public/StaticMeshSceneProxy.h
@@ -135,6 +135,8 @@ public:
 	ENGINE_API virtual const FCardRepresentationData* GetMeshCardRepresentation() const override;
 
 #if RHI_RAYTRACING
+	ENGINE_API virtual TArray<FRayTracingGeometry*> GetStaticRayTracingGeometries() const override;
+
 	ENGINE_API virtual void GetDynamicRayTracingInstances(FRayTracingMaterialGatheringContext& Context, TArray<FRayTracingInstance>& OutRayTracingInstances) override;
 	ENGINE_API virtual bool HasRayTracingRepresentation() const override;
 	virtual bool IsRayTracingRelevant() const override { return true; }
diff --git a/Engine/Source/Runtime/Renderer/Private/PrimitiveSceneInfo.cpp b/Engine/Source/Runtime/Renderer/Private/PrimitiveSceneInfo.cpp
--- a/Engine/Source/Runtime/Renderer/Private/PrimitiveSceneInfo.cpp
+++ b/Engine/Source/Runtime/Renderer/Private/PrimitiveSceneInfo.cpp
@@ -350,7 +350,8 @@ FPrimitiveSceneInfo::FPrimitiveSceneInfo(const FPrimitiveSceneInfoAdapter& InAda
 	FMemory::Memzero(CachedReflectionCaptureProxies);
 
 #if RHI_RAYTRACING
-	RayTracingGeometries = InAdapter.SceneProxy->MoveRayTracingGeometries();
+	// Cache static ray tracing geometries in SceneInfo to avoid having to access SceneProxy later
+	StaticRayTracingGeometries = InAdapter.SceneProxy->GetStaticRayTracingGeometries();
 	CachedRayTracingGeometry = nullptr;
 #endif
 
@@ -396,17 +397,17 @@ bool FPrimitiveSceneInfo::IsCachedRayTracingGeometryValid() const
 
 FRHIRayTracingGeometry* FPrimitiveSceneInfo::GetStaticRayTracingGeometryInstance(int LodLevel) const
 {
-	if (RayTracingGeometries.Num() > LodLevel)
+	if (StaticRayTracingGeometries.Num() > LodLevel)
 	{
 		// TODO: Select different LOD, when build is still pending for this LOD?
-		if (RayTracingGeometries[LodLevel]->HasPendingBuildRequest())
+		if (StaticRayTracingGeometries[LodLevel]->HasPendingBuildRequest())
 		{
-			RayTracingGeometries[LodLevel]->BoostBuildPriority();
+			StaticRayTracingGeometries[LodLevel]->BoostBuildPriority();
 			return nullptr;
 		}
-		else if (RayTracingGeometries[LodLevel]->IsValid())
+		else if (StaticRayTracingGeometries[LodLevel]->IsValid())
 		{
-			return RayTracingGeometries[LodLevel]->RayTracingGeometryRHI;
+			return StaticRayTracingGeometries[LodLevel]->RayTracingGeometryRHI;
 		}
 		else
 		{
diff --git a/Engine/Source/Runtime/Renderer/Public/PrimitiveSceneInfo.h b/Engine/Source/Runtime/Renderer/Public/PrimitiveSceneInfo.h
--- a/Engine/Source/Runtime/Renderer/Public/PrimitiveSceneInfo.h
+++ b/Engine/Source/Runtime/Renderer/Public/PrimitiveSceneInfo.h
@@ -554,7 +554,7 @@ public:
 
 	RENDERER_API FRHIRayTracingGeometry* GetStaticRayTracingGeometryInstance(int LodLevel) const;
 
-	int GetRayTracingGeometryNum() const { return RayTracingGeometries.Num(); }
+	int GetStaticRayTracingGeometryNum() const { return StaticRayTracingGeometries.Num(); }
 #endif
 
 	/** Return primitive fullname (for debugging only). */
@@ -750,10 +750,10 @@ private:
 	void RemoveCachedNaniteMaterialBins();
 
 #if RHI_RAYTRACING
-	TArray<FRayTracingGeometry*> RayTracingGeometries;
+	TArray<FRayTracingGeometry*> StaticRayTracingGeometries;
 
 	// Cache pointer to FRayTracingGeometry used by cached ray tracing instance
-	// since primitives using ERayTracingPrimitiveFlags::CacheInstances don't fill the RayTracingGeometries array above
+	// since primitives using ERayTracingPrimitiveFlags::CacheInstances don't fill the StaticRayTracingGeometries array above
 	const FRayTracingGeometry* CachedRayTracingGeometry;
 
 	/** Creates cached ray tracing representations for all meshes. */
