UnrealClient RenderTarget disable check that was causing problems due to incorrect use of RHI texture

#rnx

diff --git a/Engine/Source/Runtime/Engine/Private/UnrealClient.cpp b/Engine/Source/Runtime/Engine/Private/UnrealClient.cpp
--- a/Engine/Source/Runtime/Engine/Private/UnrealClient.cpp
+++ b/Engine/Source/Runtime/Engine/Private/UnrealClient.cpp
@@ -93,7 +93,8 @@ bool FRenderTarget::ReadFloat16Pixels(TArray<FFloat16Color>& OutImageData, FRead
 {
 	// if the RenderTarget is not EXACTLY PF_FloatRGBA , this will check down in the RHI
 	// (eg. PF_FloatRGB will fail)
-	check( GetShaderResourceTexture()->GetDesc().Format == PF_FloatRGBA );
+	// this check is correct, but you can't use GetShaderResourceTexture() except from render thread (could be a race)
+	//check( GetShaderResourceTexture()->GetDesc().Format == PF_FloatRGBA );
 
 	if (InSrcRect == FIntRect(0, 0, 0, 0))
 	{
