Removed SIMD copy of bone data as it's incorrect. The data is vector3 and needs to be 4 floats.

diff --git a/Engine/Source/Runtime/Engine/Private/GPUSkinCache.cpp b/Engine/Source/Runtime/Engine/Private/GPUSkinCache.cpp
--- a/Engine/Source/Runtime/Engine/Private/GPUSkinCache.cpp
+++ b/Engine/Source/Runtime/Engine/Private/GPUSkinCache.cpp
@@ -1814,24 +1814,11 @@ bool FGPUSkinCache::ProcessEntry(
 					ClothBufferIndexMapping.MappingOffset;  // Otherwise fallback to a 0 ClothLODBias to prevent from reading pass the buffer (but still raytrace broken shadows/reflections/etc.)
 			}
 
-#if PLATFORM_ENABLE_VECTORINTRINSICS
-			const int32 LastIndex = SimData->Positions.Num() - 1;
-			// Utilize SIMD to populate positions and normals data
-			for (int32 Index = 0; Index < LastIndex; Index++)
-			{
-				VectorStoreAligned(VectorLoadAligned((const float*)(SimData->Positions.GetData() + Index)), (float*)(Data + Index * 2));
-				VectorStoreAligned(VectorLoadAligned((const float*)(SimData->Normals.GetData() + Index)), (float*)(Data + Index * 2 + 1));
-			}
-			// Manually fetch last element 
-			*((Data + LastIndex * 2)) = SimData->Positions[LastIndex];
-			*((Data + LastIndex * 2 + 1)) = SimData->Normals[LastIndex];
-#else
 			for (int32 Index = 0; Index < SimData->Positions.Num(); Index++)
 			{
 				*(Data + Index * 2) = SimData->Positions[Index];
 				*(Data + Index * 2 + 1) = SimData->Normals[Index];
 			}
-#endif
 
 	        FResourceArrayInterface* ResourceArray = VertexAndNormalData.GetResourceArray();
 	        check(ResourceArray->GetResourceDataSize() > 0);
