Update DisplayName for SortIntArray and SortInt64Array to match user-facing BP type names.

#rnx

diff --git a/Engine/Source/Runtime/Engine/Classes/Kismet/KismetArrayLibrary.h b/Engine/Source/Runtime/Engine/Classes/Kismet/KismetArrayLibrary.h
--- a/Engine/Source/Runtime/Engine/Classes/Kismet/KismetArrayLibrary.h
+++ b/Engine/Source/Runtime/Engine/Classes/Kismet/KismetArrayLibrary.h
@@ -311,7 +311,7 @@ class UKismetArrayLibrary : public UBlueprintFunctionLibrary
 	 * @param	bStableSort		If a stable sort should be used. This preserves the order of identical elements, but is slower.
 	 * @param	SortOrder		If the array should be sorted in ascending or descending order.
 	 */
-	UFUNCTION(BlueprintCallable, Category = "Utilities|Array|Sort", meta = (AdvancedDisplay = "bStableSort,SortOrder"))
+	UFUNCTION(BlueprintCallable, Category = "Utilities|Array|Sort", meta = (DisplayName = "Sort Integer Array", AdvancedDisplay = "bStableSort,SortOrder"))
 	static ENGINE_API void SortIntArray(UPARAM(Ref) TArray<int32>& TargetArray, bool bStableSort = false, EArraySortOrder SortOrder = EArraySortOrder::Ascending);
 
 	/**
@@ -321,7 +321,7 @@ class UKismetArrayLibrary : public UBlueprintFunctionLibrary
 	 * @param	bStableSort		If a stable sort should be used. This preserves the order of identical elements, but is slower.
 	 * @param	SortOrder		If the array should be sorted in ascending or descending order.
 	 */
-	UFUNCTION(BlueprintCallable, Category = "Utilities|Array|Sort", meta = (AdvancedDisplay = "bStableSort,SortOrder"))
+	UFUNCTION(BlueprintCallable, Category = "Utilities|Array|Sort", meta = (DisplayName = "Sort Integer64 Array", AdvancedDisplay = "bStableSort,SortOrder"))
 	static ENGINE_API void SortInt64Array(UPARAM(Ref) TArray<int64>& TargetArray, bool bStableSort = false, EArraySortOrder SortOrder = EArraySortOrder::Ascending);
 
 	/**
