Move revision number to dispatch data to avoid iterating over whole dispatch batch collection everytime we add new one

Implementation by [REDACTED_PERSON]

diff --git a/Engine/Source/Runtime/Engine/Private/GPUSkinCache.cpp b/Engine/Source/Runtime/Engine/Private/GPUSkinCache.cpp
--- a/Engine/Source/Runtime/Engine/Private/GPUSkinCache.cpp
+++ b/Engine/Source/Runtime/Engine/Private/GPUSkinCache.cpp
@@ -344,7 +344,7 @@ public:
 		// triangle index buffer (input for the RecomputeSkinTangents, might need special index buffer unique to position and normal, not considering UV/vertex color)
 		uint32 IndexBufferOffsetValue = 0;
 		uint32 NumTriangles = 0;
-
+		uint32 RevisionNumber = 0;
 		FGPUSkinCache::FSkinCacheRWBuffer* TangentBuffer = nullptr;
 		FGPUSkinCache::FSkinCacheRWBuffer* IntermediateTangentBuffer = nullptr;
 		FGPUSkinCache::FSkinCacheRWBuffer* IntermediateAccumulatedTangentBuffer = nullptr;
@@ -1319,10 +1319,11 @@ void FGPUSkinCache::DoDispatch(FRHICommandList& RHICmdList)
 		for (int32 i = 0; i < BatchCount; ++i)
 		{
 			FDispatchEntry& DispatchItem = BatchDispatches[i];
-			PrepareUpdateSkinning(DispatchItem.SkinCacheEntry, DispatchItem.Section, DispatchItem.RevisionNumber, &BuffersToTransitionForSkinning);
+			PrepareUpdateSkinning(DispatchItem.SkinCacheEntry, DispatchItem.Section, DispatchItem.SkinCacheEntry->DispatchData[DispatchItem.Section].RevisionNumber, &BuffersToTransitionForSkinning);
 
 			// Clear the flag that this is queued for dispatch.
 			DispatchItem.SkinCacheEntry->bQueuedForDispatch = false;
+			DispatchItem.SkinCacheEntry->DispatchData[DispatchItem.Section].RevisionNumber = 0;
 		}
 
 		{
@@ -1844,19 +1845,18 @@ bool FGPUSkinCache::ProcessEntry(
 		InOutEntry->bQueuedForDispatch = true;
 
 		bool bFoundEntry = false;
-		for (FDispatchEntry& Entry : BatchDispatches)
+
+		if (InOutEntry->DispatchData[Section].RevisionNumber != 0)
 		{
 			// Check if the combo of skin cache entry and section index already exists, if so use the entry and update to latest revision number.
-			if (Entry.SkinCacheEntry == InOutEntry && Entry.Section == Section && Entry.RevisionNumber < RevisionNumber)
-			{
-				Entry.RevisionNumber = RevisionNumber;
-				bFoundEntry = true;
-				break;
-			}
+			InOutEntry->DispatchData[Section].RevisionNumber = FMath::Max(InOutEntry->DispatchData[Section].RevisionNumber, RevisionNumber);
+			bFoundEntry = true;
 		}
+
 		if (!bFoundEntry)
 		{
-			BatchDispatches.Add({ InOutEntry, RevisionNumber, uint32(Section) });
+			InOutEntry->DispatchData[Section].RevisionNumber = RevisionNumber;
+			BatchDispatches.Add({ InOutEntry, uint32(Section) });
 		}
 	}
 	else
diff --git a/Engine/Source/Runtime/Engine/Public/GPUSkinCache.h b/Engine/Source/Runtime/Engine/Public/GPUSkinCache.h
--- a/Engine/Source/Runtime/Engine/Public/GPUSkinCache.h
+++ b/Engine/Source/Runtime/Engine/Public/GPUSkinCache.h
@@ -115,7 +115,6 @@ public:
 	struct FDispatchEntry
 	{
 		FGPUSkinCacheEntry* SkinCacheEntry = nullptr;
-		uint32 RevisionNumber = 0;
 		uint32 Section = 0;	
 	};
 
