Fixed incorrect logic for deferring skeletal mesh updates to GDME when parallel GDME is enabled.

diff --git a/Engine/Source/Runtime/Engine/Private/SkeletalRenderGPUSkin.cpp b/Engine/Source/Runtime/Engine/Private/SkeletalRenderGPUSkin.cpp
--- a/Engine/Source/Runtime/Engine/Private/SkeletalRenderGPUSkin.cpp
+++ b/Engine/Source/Runtime/Engine/Private/SkeletalRenderGPUSkin.cpp
@@ -476,7 +476,7 @@ static TAutoConsoleVariable<int32> CVarDeferSkeletalDynamicDataUpdateUntilGDME(
 
 FORCEINLINE bool IsDeferredSkeletalDynamicDataUpdateEnabled()
 {
-	return CVarDeferSkeletalDynamicDataUpdateUntilGDME.GetValueOnRenderThread() > 0 && IsParallelGatherDynamicMeshElementsEnabled();
+	return CVarDeferSkeletalDynamicDataUpdateUntilGDME.GetValueOnRenderThread() > 0 && !IsParallelGatherDynamicMeshElementsEnabled();
 }
 
 void FSkeletalMeshObjectGPUSkin::UpdateDynamicData_RenderThread(FGPUSkinCache* GPUSkinCache, FRHICommandList& RHICmdList, FDynamicSkelMeshObjectDataGPUSkin* InDynamicData, FSceneInterface* Scene, uint64 FrameNumberToPrepare, uint32 RevisionNumber)
