Partially fix audio disconituities when seeking streaming audio during playback.

- Reset PCM decode state when seeking, so we don't write out stale audio data on the next read.
- Fix "Out of Order" operations. `WriteFromDecodedPCM` is the point at which decoded audio data is written to the output buffer. Move this _after_ we perform the requested "Seek" so we don't write out "stale" audio data. Additionally, when checking if we're "storing the end of the file", only return "bLooped" when we actually run out of audio data to write out. "bLooped" is the flag that's returned by the "StreamComprressedData" method, indicating that the "end of the file" has been reach, or has otherwise "looped" if "bIsLooping" was set to true. (It's confusing that it does double duty like this).

#tests PIE

diff --git a/Engine/Source/Runtime/Engine/Private/AudioDecompress.cpp b/Engine/Source/Runtime/Engine/Private/AudioDecompress.cpp
--- a/Engine/Source/Runtime/Engine/Private/AudioDecompress.cpp
+++ b/Engine/Source/Runtime/Engine/Private/AudioDecompress.cpp
@@ -263,11 +263,6 @@ bool IStreamedCompressedInfo::StreamCompressedData(uint8* Destination, bool bLoo
 		bStoringEndOfFile ? TEXT("YES") : TEXT("NO")
 	);
 
-	// Write out any PCM data that was decoded during the last request
-	uint32 RawPCMOffset = WriteFromDecodedPCM(Destination, BufferSize);
-	// immediately update the OutNumBytesStreamed to reflect what we just wrote
-	OutNumBytesStreamed = RawPCMOffset;
-
 	// If we have a pending next chunk from seeking, move to it now.
 	if (StreamSeekBlockIndex != INDEX_NONE)
 	{
@@ -285,7 +280,7 @@ bool IStreamedCompressedInfo::StreamCompressedData(uint8* Destination, bool bLoo
 		if (NewlySeekedChunk == nullptr)
 		{
 			// After a seek we're likely to need to wait a bit for the chunk to get in to memory.
-			ZeroBuffer(Destination + RawPCMOffset, BufferSize - RawPCMOffset);
+			ZeroBuffer(Destination, BufferSize);
 			return false;
 		}	
 
@@ -343,7 +338,7 @@ bool IStreamedCompressedInfo::StreamCompressedData(uint8* Destination, bool bLoo
 		StreamSeekToAudioFrames = INDEX_NONE;
 		StreamSeekBlockOffset = INDEX_NONE;
 	}
-
+	
 	// If next chunk wasn't loaded when last one finished reading, try to get it again now
 	if (SrcBufferData == NULL)
 	{
@@ -405,20 +400,25 @@ bool IStreamedCompressedInfo::StreamCompressedData(uint8* Destination, bool bLoo
 					CurrentChunkIndex, *StreamingSoundWave->GetFName().ToString(), ++PrintChunkFailMessageCount, EnumToString(Behavior), TimeSoFar);
 			}
 			
-			ZeroBuffer(Destination + RawPCMOffset, BufferSize - RawPCMOffset);
+			ZeroBuffer(Destination, BufferSize);
 			return false;
 		}
 	}
 
 	bool bLooped = false;
+	
+	// Write out any PCM data that was decoded during the last request
+	uint32 RawPCMOffset = WriteFromDecodedPCM(Destination, BufferSize);
+	// immediately update the OutNumBytesStreamed to reflect what we just wrote
+	OutNumBytesStreamed = RawPCMOffset;
 
-	if (bStoringEndOfFile && LastPCMByteSize > 0)
+	// if we were storing the end of the file and just now read to the end (LastPCMByteSize == 0), that means we looped
+	if (bStoringEndOfFile && LastPCMByteSize == 0)
 	{
-		// delayed returning looped because we hadn't read the entire buffer
 		bLooped = true;
 		bStoringEndOfFile = false;
 	}
-
+	
 	while (RawPCMOffset < BufferSize)
 	{
 		// Decompress the next compression frame of audio (many samples) into the PCM buffer
