Fix crash on startup on low-end DX11 GPUs due to TexCreate_UAV being enabled on gbuffer targets.

This flag was turned on if r.Nanite.AllowComputeMaterials was 1, without checking if Nanite is supported on the current platform. Old GPUs can't make UAVs for BGRA targets, so GBuffer creation failed.

#rnx

#localization none
#tests boot

diff --git a/Engine/Source/Runtime/Engine/Private/SceneTexturesConfig.cpp b/Engine/Source/Runtime/Engine/Private/SceneTexturesConfig.cpp
--- a/Engine/Source/Runtime/Engine/Private/SceneTexturesConfig.cpp
+++ b/Engine/Source/Runtime/Engine/Private/SceneTexturesConfig.cpp
@@ -292,12 +292,12 @@ void FSceneTexturesConfig::Init(const FSceneTexturesConfigInitSettings& InitSett
 				GBufferParams[Layout] = (Layout == GBL_Default) ? DefaultParams : FShaderCompileUtilities::FetchGBufferParamsRuntime(ShaderPlatform, (EGBufferLayout)Layout);
 				const FGBufferInfo GBufferInfo = FetchFullGBufferInfo(GBufferParams[Layout]);
 
-				BindingCache.Bindings[Layout].GBufferA = FindGBufferBindingByName(GBufferInfo, TEXT("GBufferA"));
-				BindingCache.Bindings[Layout].GBufferB = FindGBufferBindingByName(GBufferInfo, TEXT("GBufferB"));
-				BindingCache.Bindings[Layout].GBufferC = FindGBufferBindingByName(GBufferInfo, TEXT("GBufferC"));
-				BindingCache.Bindings[Layout].GBufferD = FindGBufferBindingByName(GBufferInfo, TEXT("GBufferD"));
-				BindingCache.Bindings[Layout].GBufferE = FindGBufferBindingByName(GBufferInfo, TEXT("GBufferE"));
-				BindingCache.Bindings[Layout].GBufferVelocity = FindGBufferBindingByName(GBufferInfo, TEXT("Velocity"));
+				BindingCache.Bindings[Layout].GBufferA = FindGBufferBindingByName(GBufferInfo, TEXT("GBufferA"), ShaderPlatform);
+				BindingCache.Bindings[Layout].GBufferB = FindGBufferBindingByName(GBufferInfo, TEXT("GBufferB"), ShaderPlatform);
+				BindingCache.Bindings[Layout].GBufferC = FindGBufferBindingByName(GBufferInfo, TEXT("GBufferC"), ShaderPlatform);
+				BindingCache.Bindings[Layout].GBufferD = FindGBufferBindingByName(GBufferInfo, TEXT("GBufferD"), ShaderPlatform);
+				BindingCache.Bindings[Layout].GBufferE = FindGBufferBindingByName(GBufferInfo, TEXT("GBufferE"), ShaderPlatform);
+				BindingCache.Bindings[Layout].GBufferVelocity = FindGBufferBindingByName(GBufferInfo, TEXT("Velocity"), ShaderPlatform);
 
 				// Remove DisableDCC flag for velocity. Only Nanite fast tile clear sets this flag currently
 				// but we want to exclude velocity because it usually doesn't have many written pixels
diff --git a/Engine/Source/Runtime/RenderCore/Private/GBufferInfo.cpp b/Engine/Source/Runtime/RenderCore/Private/GBufferInfo.cpp
--- a/Engine/Source/Runtime/RenderCore/Private/GBufferInfo.cpp
+++ b/Engine/Source/Runtime/RenderCore/Private/GBufferInfo.cpp
@@ -130,7 +130,7 @@ int32 FindGBufferTargetByName(const FGBufferInfo& GBufferInfo, const FString& Na
 	return -1;
 }
 
-FGBufferBinding FindGBufferBindingByName(const FGBufferInfo& GBufferInfo, const FString& Name)
+FGBufferBinding FindGBufferBindingByName(const FGBufferInfo& GBufferInfo, const FString& Name, EShaderPlatform ShaderPlatform)
 {
 	const int32 Index = FindGBufferTargetByName(GBufferInfo, Name);
 
@@ -183,7 +183,7 @@ FGBufferBinding FindGBufferBindingByName(const FGBufferInfo& GBufferInfo, const
 			Binding.Flags |= TexCreate_SRGB;
 		}
 
-		if (NaniteComputeMaterialsSupported())
+		if (DoesPlatformSupportNanite(ShaderPlatform, true) && NaniteComputeMaterialsSupported())
 		{
 			Binding.Flags |= TexCreate_UAV;
 
@@ -594,4 +594,3 @@ FGBufferInfo RENDERCORE_API FetchFullGBufferInfo(const FGBufferParams& Params)
 		return FetchLegacyGBufferInfo(Params);
 	}
 }
-
diff --git a/Engine/Source/Runtime/RenderCore/Public/GBufferInfo.h b/Engine/Source/Runtime/RenderCore/Public/GBufferInfo.h
--- a/Engine/Source/Runtime/RenderCore/Public/GBufferInfo.h
+++ b/Engine/Source/Runtime/RenderCore/Public/GBufferInfo.h
@@ -321,7 +321,13 @@ struct FGBufferParams
 
 int32 RENDERCORE_API FindGBufferTargetByName(const FGBufferInfo& GBufferInfo, const FString& Name);
 
-FGBufferBinding RENDERCORE_API FindGBufferBindingByName(const FGBufferInfo& GBufferInfo, const FString& Name);
+FGBufferBinding RENDERCORE_API FindGBufferBindingByName(const FGBufferInfo& GBufferInfo, const FString& Name, EShaderPlatform ShaderPlatform);
+
+UE_DEPRECATED(5.4, "Please use the overload which takes a shader platform parameter")
+inline FGBufferBinding FindGBufferBindingByName(const FGBufferInfo& GBufferInfo, const FString& Name)
+{
+	return FindGBufferBindingByName(GBufferInfo, Name, GMaxRHIShaderPlatform);
+}
 
 FGBufferInfo RENDERCORE_API FetchFullGBufferInfo(const FGBufferParams& Params);
 
