Added a new visibilty-based tick option that allows ticking montages and refresh bones when any montage is being played

PR #10703

diff --git a/Engine/Source/Runtime/Engine/Classes/Components/SkinnedMeshComponent.h b/Engine/Source/Runtime/Engine/Classes/Components/SkinnedMeshComponent.h
--- a/Engine/Source/Runtime/Engine/Classes/Components/SkinnedMeshComponent.h
+++ b/Engine/Source/Runtime/Engine/Classes/Components/SkinnedMeshComponent.h
@@ -78,6 +78,12 @@ enum class EVisibilityBasedAnimTickOption : uint8
 	AlwaysTickPoseAndRefreshBones,
 	/** Always Tick, but Refresh BoneTransforms only when rendered. */
 	AlwaysTickPose,
+	/**
+		When rendered Tick Pose and Refresh Bone Transforms,
+		otherwise, just update montages alongside BoneTransforms.
+		(AnimBP graph will not be updated).
+	*/
+	OnlyTickMontagesAndRefreshBonesWhenPlayingMontages,
 	/**
 		When rendered Tick Pose and Refresh Bone Transforms,
 		otherwise, just update montages and skip everything else.
diff --git a/Engine/Source/Runtime/Engine/Private/Components/SkeletalMeshComponent.cpp b/Engine/Source/Runtime/Engine/Private/Components/SkeletalMeshComponent.cpp
--- a/Engine/Source/Runtime/Engine/Private/Components/SkeletalMeshComponent.cpp
+++ b/Engine/Source/Runtime/Engine/Private/Components/SkeletalMeshComponent.cpp
@@ -1271,8 +1271,9 @@ bool USkeletalMeshComponent::ShouldOnlyTickMontages(const float DeltaTime) const
 {
 	// Ignore DeltaSeconds == 0.f, as that is used when we want to force an update followed by RefreshBoneTransforms.
 	// RefreshBoneTransforms will need an updated graph.
-	return (VisibilityBasedAnimTickOption == EVisibilityBasedAnimTickOption::OnlyTickMontagesWhenNotRendered)
-		&& !bRecentlyRendered 
+	return (VisibilityBasedAnimTickOption == EVisibilityBasedAnimTickOption::OnlyTickMontagesWhenNotRendered
+		|| VisibilityBasedAnimTickOption == EVisibilityBasedAnimTickOption::OnlyTickMontagesAndRefreshBonesWhenPlayingMontages)
+		&& !bRecentlyRendered
 		&& (DeltaTime > 0.f);
 }
 
@@ -1412,7 +1413,11 @@ bool USkeletalMeshComponent::ShouldUpdateTransform(bool bLODHasChanged) const
 	// If forcing RefPose we can skip updating the skeleton for perf, except if it's using MorphTargets.
 	const bool bSkipBecauseOfRefPose = bForceRefpose && bOldForceRefPose && (MorphTargetCurves.Num() == 0) && ((AnimScriptInstance) ? !AnimScriptInstance->HasMorphTargetCurves() : true);
 
-	return (Super::ShouldUpdateTransform(bLODHasChanged) && !bNoSkeletonUpdate && !bSkipBecauseOfRefPose);
+	const bool bShouldUpdateTransform = Super::ShouldUpdateTransform(bLODHasChanged) ||
+			(GetAnimInstance() && GetAnimInstance()->IsAnyMontagePlaying()
+			&& VisibilityBasedAnimTickOption == EVisibilityBasedAnimTickOption::OnlyTickMontagesAndRefreshBonesWhenPlayingMontages);
+
+	return (bShouldUpdateTransform && !bNoSkeletonUpdate && !bSkipBecauseOfRefPose);
 }
 
 bool USkeletalMeshComponent::ShouldTickPose() const
