Temporary workaround for bUseFauxOrthoViewPos recalculations to ensure they are only applied to Ortho cameras and don't affect non-camera ortho projection view resolutions. Long term plan is to hopefully deprecate bUseFauxOrthoViewPos, but this negates issues with Ortho camera fixes in the short term.

diff --git a/Engine/Source/Runtime/Engine/Classes/GameFramework/Actor.h b/Engine/Source/Runtime/Engine/Classes/GameFramework/Actor.h
--- a/Engine/Source/Runtime/Engine/Classes/GameFramework/Actor.h
+++ b/Engine/Source/Runtime/Engine/Classes/GameFramework/Actor.h
@@ -3429,7 +3429,7 @@ public:
 	ENGINE_API virtual void CalcCamera(float DeltaTime, struct FMinimalViewInfo& OutResult);
 
 	/** Returns true if the actor contains an active camera component */
-	ENGINE_API virtual bool HasActiveCameraComponent() const;
+	ENGINE_API virtual bool HasActiveCameraComponent(bool bForceFindCamera = false) const;
 
 	/** Returns true if the actor contains an active locked to HMD camera component */
 	ENGINE_API virtual bool HasActivePawnControlCameraComponent() const;
diff --git a/Engine/Source/Runtime/Engine/Private/Actor.cpp b/Engine/Source/Runtime/Engine/Private/Actor.cpp
--- a/Engine/Source/Runtime/Engine/Private/Actor.cpp
+++ b/Engine/Source/Runtime/Engine/Private/Actor.cpp
@@ -3166,9 +3166,9 @@ void AActor::CalcCamera(float DeltaTime, FMinimalViewInfo& OutResult)
 	GetActorEyesViewPoint(OutResult.Location, OutResult.Rotation);
 }
 
-bool AActor::HasActiveCameraComponent() const
+bool AActor::HasActiveCameraComponent(bool bForceFindCamera) const
 {
-	if (bFindCameraComponentWhenViewTarget)
+	if (bFindCameraComponentWhenViewTarget || bForceFindCamera)
 	{
 		// Look for the first active camera component and use that for the view
 		for (const UActorComponent* Component : OwnedComponents)
diff --git a/Engine/Source/Runtime/Engine/Private/SceneView.cpp b/Engine/Source/Runtime/Engine/Private/SceneView.cpp
--- a/Engine/Source/Runtime/Engine/Private/SceneView.cpp
+++ b/Engine/Source/Runtime/Engine/Private/SceneView.cpp
@@ -23,6 +23,7 @@
 #include "RenderUtils.h"
 #include "StereoRenderUtils.h"
 #include "SceneRelativeViewMatrices.h"
+#include "Camera/CameraComponent.h"
 #include "UObject/Interface.h"
 
 DEFINE_LOG_CATEGORY(LogBufferVisualization);
@@ -672,7 +673,17 @@ FViewMatrices::FViewMatrices(const FSceneViewInitOptions& InitOptions) : FViewMa
 	Initializer.ViewOrigin           = InitOptions.ViewOrigin;
 	Initializer.ConstrainedViewRect  = InitOptions.GetConstrainedViewRect();
 	Initializer.StereoPass           = InitOptions.StereoPass;
-	Initializer.bUseFauxOrthoViewPos = InitOptions.bUseFauxOrthoViewPos;
+
+	//This is a bit of a hack to ensure orthographic camera views are resolving correctly, without affecting non-camera views 
+	//it will be removed once we have a more robust solution + can hopefully deprecate bUseFauxOrthoViewPos entirely
+	if (InitOptions.ProjectionMatrix.M[3][3] >= 1.0f && InitOptions.ViewActor && InitOptions.ViewActor->HasActiveCameraComponent(true))
+	{
+		Initializer.bUseFauxOrthoViewPos = true;
+	}
+	else
+	{
+		Initializer.bUseFauxOrthoViewPos = InitOptions.bUseFauxOrthoViewPos;
+	}
 
 	Init(Initializer);
 }
diff --git a/Engine/Source/Runtime/Engine/Public/SceneView.h b/Engine/Source/Runtime/Engine/Public/SceneView.h
--- a/Engine/Source/Runtime/Engine/Public/SceneView.h
+++ b/Engine/Source/Runtime/Engine/Public/SceneView.h
@@ -237,7 +237,7 @@ struct FSceneViewInitOptions : public FSceneViewProjectionData
 		, bUseFieldOfViewForLOD(true)
 		, FOV(90.f)
 		, DesiredFOV(90.f)
-		, bUseFauxOrthoViewPos(true)
+		, bUseFauxOrthoViewPos(false)
 		, bIsSceneCapture(false)
 		, bIsSceneCaptureCube(false)
 		, bSceneCaptureUsesRayTracing(false)
