Material Editor - Made material translation DDC query async to material translation (running at the same time), so that the process that completes first will be used instead of doing them both serially.

diff --git a/Engine/Source/Editor/MaterialEditor/Private/MaterialEditor.cpp b/Engine/Source/Editor/MaterialEditor/Private/MaterialEditor.cpp
--- a/Engine/Source/Editor/MaterialEditor/Private/MaterialEditor.cpp
+++ b/Engine/Source/Editor/MaterialEditor/Private/MaterialEditor.cpp
@@ -283,6 +283,12 @@ bool FMatExpressionPreview::ShouldCache(EShaderPlatform Platform, const FShaderT
 
 int32 FMatExpressionPreview::CompilePropertyAndSetMaterialProperty(EMaterialProperty Property, FMaterialCompiler* Compiler, EShaderFrequency OverrideShaderFrequency, bool bUsePreviousFrameTime) const
 {
+	// Early out if the compiler wishes to terminate translation.
+	if (Compiler->ShouldStopTranslating())
+	{
+		return INDEX_NONE;
+	}
+
 	// needs to be called in this function!!
 	Compiler->SetMaterialProperty(Property, OverrideShaderFrequency, bUsePreviousFrameTime);
 
diff --git a/Engine/Source/Runtime/Engine/Private/Materials/HLSLMaterialTranslator.cpp b/Engine/Source/Runtime/Engine/Private/Materials/HLSLMaterialTranslator.cpp
--- a/Engine/Source/Runtime/Engine/Private/Materials/HLSLMaterialTranslator.cpp
+++ b/Engine/Source/Runtime/Engine/Private/Materials/HLSLMaterialTranslator.cpp
@@ -49,18 +49,18 @@
 #include "ShaderCompiler.h"
 #include "Misc/FileHelper.h"
 #include "Misc/ScopeLock.h"
-
-#if WITH_EDITORONLY_DATA
-#include "Materials/MaterialExpressionSubstrate.h"
-#include "ShaderPlatformCachedIniValue.h"
 #include "Serialization/ObjectWriter.h"
 #include "Serialization/ObjectReader.h"
 #include "Serialization/BufferArchive.h"
 #include "Serialization/MemoryReader.h"
 #include "Serialization/ObjectAndNameAsStringProxyArchive.h"
-#include "DerivedDataCache.h"
 #include "DerivedDataRequestOwner.h"
+#include "DerivedDataCache.h"
 #include "MaterialCachedData.h"
+
+#if WITH_EDITORONLY_DATA
+#include "Materials/MaterialExpressionSubstrate.h"
+#include "ShaderPlatformCachedIniValue.h"
 #endif
 
 #if ENABLE_COOK_STATS && STATS
@@ -203,23 +203,23 @@ static FAutoConsoleVariableRef CVarPreshaderGapInterval(
 	TEXT("Insert an empty element in the preshader buffer every specified number of elements in the buffer.  Workaround for a shader compiler bug."));
 
 /* Controls whether to use the new GetMaterialShaderCode() and GetMaterialEnvironment() implementations. */
-static bool GUseMaterialTranslationResultsGrouping = true;
+static const bool GUseMaterialTranslationResultsGrouping = true;
 
-/* Controls whether DDC caching of material translation results is enabled. */
-static bool GJobDisableMaterialTranslateDDC = true;
-static FAutoConsoleVariableRef CVarJobDisableMaterialTranslateDDC(
-	TEXT("r.Material.DisableTranslateDDC"),
-	GJobDisableMaterialTranslateDDC,
-	TEXT("Whether to disable material translation DDC caching.\n"));
+/* Controls whether DDC caching of material translation results should be forcefully disabled. */
+static const bool GForceDisableMaterialTranslationDDC = false;
 
 UE::DerivedData::FCacheBucket MaterialTranslationDDCBucket = UE::DerivedData::FCacheBucket(TEXT("MaterialTranslation"));
 UE::DerivedData::FValueId MaterialCompilationOutputId = UE::DerivedData::FValueId::FromName("FHLSLMaterialTranslator_MaterialCompilationOutput");
 UE::DerivedData::FValueId MaterialResultsOutputId = UE::DerivedData::FValueId::FromName("FHLSLMaterialTranslator_Results");
 UE::DerivedData::FValueId EnvironmentDefinesId = UE::DerivedData::FValueId::FromName("FHLSLMaterialTranslator_EnvironmentDefines");
 
+/* Helper macro to check whether the DDC material translation data has arrived and
+ * therefore quit material translation. */
+#define CHECK_DDC_QUERY_FINISHED_ELSE_RETURN() if (DDCQueryHit) { return; }
+
 /* This version number models the layout of the data stored on the DDC after a material translation (e.g. FEnvironmentDefines)
  * It must be bumped whenever a change is made requires re-translation of all materials. */
-static constexpr int MaterialTranslationDDCVersion = 4;
+static constexpr int MaterialTranslationDDCVersion = 5;
 
 /** Data structure used to cache a part of material translation results. It contains all the generated
  *  defines that will be declared during the compilation of the generated material shader.
@@ -616,7 +616,6 @@ FHLSLMaterialTranslator::FHLSLMaterialTranslator(FMaterial* InMaterial,
 ,	CurrentScopeID(0u)
 ,	NextTempScopeID(SF_NumFrequencies)
 ,	Material(InMaterial)
-,	MaterialCompilationOutput(InMaterialCompilationOutput)
 ,	StaticParameters(InStaticParameters)
 ,	Platform(InPlatform)
 ,	QualityLevel(InQualityLevel)
@@ -675,6 +674,7 @@ FHLSLMaterialTranslator::FHLSLMaterialTranslator(FMaterial* InMaterial,
 ,	NumVtSamples(0)
 ,	TargetPlatform(InTargetPlatform)
 ,	SubstrateCompilationConfig()
+,   MaterialCompilationOutput(InMaterialCompilationOutput)
 {
 	check((sizeof(SubstrateCompilationContext)/sizeof(FSubstrateCompilationContext)) == ESubstrateCompilationContext::SCC_MAX);
 
@@ -748,6 +748,11 @@ FHLSLMaterialTranslator::~FHLSLMaterialTranslator()
 	ClearAllFunctionStacks();
 }
 
+bool FHLSLMaterialTranslator::ShouldStopTranslating() const
+{
+	return DDCQueryHit;
+}
+
 int32 FHLSLMaterialTranslator::GetNumUserTexCoords() const
 {
 	return AllocatedUserTexCoords.FindLast(true) + 1;
@@ -1121,14 +1126,17 @@ UE_TRACE_EVENT_END()
 
 bool FHLSLMaterialTranslator::Translate()
 {
+	STAT(FDateTime TranslationDateTime = FDateTime::Now());
+	STAT(double TotalTime = FPlatformTime::Seconds());
+	STAT(double TranslationOnlyTime = 0);
+
 #if CPUPROFILERTRACE_ENABLED
 	FString TraceMaterialName;
 	if (UE_TRACE_CHANNELEXPR_IS_ENABLED(CpuChannel))
 	{
 		TraceMaterialName = Material->GetMaterialInterface()->GetFullName();
 	}
-	UE_TRACE_LOG_SCOPED_T(Cpu, FHLSLMaterialTranslatorTranslate, CpuChannel)
-		<< FHLSLMaterialTranslatorTranslate.MaterialName(*TraceMaterialName);
+	UE_TRACE_LOG_SCOPED_T(Cpu, FHLSLMaterialTranslatorTranslate, CpuChannel) << FHLSLMaterialTranslatorTranslate.MaterialName(*TraceMaterialName);
 #endif
 
 	static const bool bVerbose = FParse::Param(FCommandLine::Get(), TEXT("verbosematerialtranslation"));
@@ -1137,52 +1145,66 @@ bool FHLSLMaterialTranslator::Translate()
 		UE_LOG(LogMaterial, Display, TEXT("Translating '%s'"), *Material->GetMaterialInterface()->GetFullName());
 	}
 
-	STAT(FDateTime TranslationDateTime = FDateTime::Now());
-	STAT(double HLSLTranslateTime = 0);
-	STAT(double DDCQueryTime = 0);
+	// We call FindObject to serialize the array of Parameter Collections used by this material in EnvironmentDefines,
+	// but this can happen during save. FindObject is illegal during save because if the discovered objects
+	// are serialized into the package it will cause a crash on package load. But we are not storing the results
+	// of FindObject into the package so it is okay to remove the restriction.
+	TOptional<TGuardValue<bool>> IsSavingPackageGuard;
+	if (IsInGameThread())
+	{
+		IsSavingPackageGuard.Emplace(GIsSavingPackage, false);
+	}
 
-	bSuccess = true;
-	bool bInCache = false;
+	static const bool bNoMaterialTranslationDDC = FParse::Param(FCommandLine::Get(), TEXT("nomaterialtranslationddc"));
+	bool bDisableTranslationDDC = GForceDisableMaterialTranslationDDC || bNoMaterialTranslationDDC || Material->IsPreview();
+
+	UE::DerivedData::FRequestOwner DDCRequestOwner{ UE::DerivedData::EPriority::Blocking };
 
+	// Asynchronously query the DDC for results
+	if (!bDisableTranslationDDC)
 	{
-		SCOPE_SECONDS_COUNTER(HLSLTranslateTime);
+		AsyncQueryDDC(DDCRequestOwner);
+	}
 
-		if (!Material->IsPreview())
-		{
-			SCOPE_SECONDS_COUNTER(DDCQueryTime);
-			bInCache = QueryDDCCachedTranslationResults();
-		}
+	// Synchronously begin translating the material
+	{
+		SCOPE_SECONDS_COUNTER(TranslationOnlyTime);
+		DoTranslate();
+	}
 
-		// If material not in cache, we need to translate the material now.
-		if (!bInCache)
-		{
-			DoTranslate();
+	// One of the two has terminated. This will be a NOOP if DDC query task has completed.
+	DDCRequestOwner.Cancel();
 
-			// If translation was succesful, finalize the results and push them to the DDC
-			if (bSuccess)
-			{
-				PrepareEnvironmentDefines();
-				PrepareMaterialSourceStringParameters();
-				PushResultsToDDCCache(); 
-				ClearAllFunctionStacks();
-			}
+	if (DDCQueryHit)
+	{
+		// Material is in cache, copy over the compilation output retrieved from the DDC.
+		STAT(GShaderCompilerStats->IncrementMaterialCacheHit());
+		MaterialCompilationOutput = DDCMaterialCompilationOutput;
+		bSuccess = true;
+	}
+	else if (bSuccess)
+	{
+		// Material not in cache. If translation was succesful, finalize the results and push them to the DDC
+		PrepareEnvironmentDefines();
+		PrepareMaterialSourceStringParameters();
+		if (!bDisableTranslationDDC)
+		{
+			PushResultsToDDC();
 		}
+		ClearAllFunctionStacks();
 	}
 
 	// Report timings to Material Cook Stats
 #if STATS
-	GShaderCompilerStats->IncrementMaterialTranslated(HLSLTranslateTime);
-	if (bInCache)
-	{
-		STAT(GShaderCompilerStats->IncrementMaterialCacheHit(DDCQueryTime));
-	}
+	TotalTime = FPlatformTime::Seconds() - TotalTime;
+	GShaderCompilerStats->IncrementMaterialTranslated(TotalTime, TranslationOnlyTime);
 #endif
 
-	INC_FLOAT_STAT_BY(STAT_ShaderCompiling_HLSLTranslation, (float)HLSLTranslateTime);
+	INC_FLOAT_STAT_BY(STAT_ShaderCompiling_HLSLTranslation, (float)TotalTime);
 
 #if ENABLE_COOK_STATS && STATS
 	// Write out a CSV file MaterialTranslationLog.txt containing info about all material translations.
-	FCsvLogFile::Get().AddEntry(Material->GetMaterialInterface()->GetFullName(), TranslationDateTime, HLSLTranslateTime);
+	FCsvLogFile::Get().AddEntry(Material->GetMaterialInterface()->GetFullName(), TranslationDateTime, TotalTime);
 #endif
 
 	return bSuccess;
@@ -1190,6 +1212,8 @@ bool FHLSLMaterialTranslator::Translate()
 
 void FHLSLMaterialTranslator::DoTranslate()
 {
+	bSuccess = true;
+	
 	// No cache hit, continue translating the material
 	check(ScopeStack.Num() == 0);
 
@@ -1225,6 +1249,7 @@ void FHLSLMaterialTranslator::DoTranslate()
 		TSet<UMaterialExpression*> VisitedExpressions;
 		for (UMaterialExpression* Expression : UMaterial->GetExpressions())
 		{
+			CHECK_DDC_QUERY_FINISHED_ELSE_RETURN();
 			if (Expression->ContainsInputLoop(VisitedExpressions))
 			{
 				AppendExpressionError(Expression, TEXT("Expression is part of a cycle. Please make sure the material graph is acyclic."));
@@ -1244,6 +1269,7 @@ void FHLSLMaterialTranslator::DoTranslate()
 	bSubstrateWritesEmissive = false;
 	bSubstrateWritesAmbientOcclusion = false;
 
+	CHECK_DDC_QUERY_FINISHED_ELSE_RETURN();
 	UMaterialExpression* ExpressionToPreview = Material->GetMaterialGraphNodePreviewExpression();
 	if (ExpressionToPreview)
 	{
@@ -1266,6 +1292,7 @@ void FHLSLMaterialTranslator::DoTranslate()
 		FrontMaterialExpr = FrontMaterialInput ? FrontMaterialInput->GetTracedInput().Expression : nullptr;
 		FrontMaterialOutputIndex = FrontMaterialInput ? FrontMaterialInput->OutputIndex : INDEX_NONE;
 	}
+
 	if (bSubstrateEnabled && FrontMaterialExpr)
 	{
 		// Temp code chunk scope (e.g.needed for the creation of static booleans from static switch parameter node, see UMaterialExpressionStaticSwitch::GetEffectiveInput).
@@ -1358,6 +1385,7 @@ void FHLSLMaterialTranslator::DoTranslate()
 		int32 NumMaterialLayersAttributes = 0;
 		for (UMaterialExpression* Expression : Expressions)
 		{
+			CHECK_DDC_QUERY_FINISHED_ELSE_RETURN();
 			if (UMaterialExpressionMaterialAttributeLayers* Layers = Cast<UMaterialExpressionMaterialAttributeLayers>(Expression))
 			{
 				++NumMaterialLayersAttributes;
@@ -1434,6 +1462,8 @@ void FHLSLMaterialTranslator::DoTranslate()
 		Chunk[MP_SurfaceThickness] = Material->CompilePropertyAndSetMaterialProperty(MP_SurfaceThickness		,this);
 		Chunk[MP_FrontMaterial] = Material->CompilePropertyAndSetMaterialProperty(MP_FrontMaterial			,this);
 
+		CHECK_DDC_QUERY_FINISHED_ELSE_RETURN();
+
 		// Now generate the code for the fully simplified MP_FrontMaterial right after MP_FrontMaterial
 		FullySimplifiedSubstrateFrontMaterialCodeChunk = SubstrateCreateAndRegisterNullMaterial();
 		{																							// This causes issues, material look different
@@ -1444,7 +1474,9 @@ void FHLSLMaterialTranslator::DoTranslate()
 			CurrentSubstrateCompilationContext = ESubstrateCompilationContext::SCC_Default;
 			FullySimplifiedFrontMaterialCodeChunkEnd = SharedPropertyCodeChunks[FrontMaterialShaderFrequency].Num();
 		}
-		}
+	}
+
+	CHECK_DDC_QUERY_FINISHED_ELSE_RETURN();
 
 	// Get shading models from compilation (or material).
 	FMaterialShadingModelField MaterialShadingModels = GetCompiledShadingModels();
@@ -1453,6 +1485,8 @@ void FHLSLMaterialTranslator::DoTranslate()
 
 	if (!bEnableExecutionFlow)
 	{
+		CHECK_DDC_QUERY_FINISHED_ELSE_RETURN();
+
 		if (Domain == MD_Volume || (Domain == MD_Surface && IsSubsurfaceShadingModel(MaterialShadingModels)))
 		{
 			// Note we don't test for the blend mode as you can have a translucent material using the subsurface shading model
@@ -1467,6 +1501,8 @@ void FHLSLMaterialTranslator::DoTranslate()
 			int32 CodeSubsurfaceProfile = ForceCast(ScalarParameter(NameSubsurfaceProfile, 1.0f), MCT_Float1);
 
 			Chunk[MP_SubsurfaceColor] = AppendVector(SubsurfaceColor, CodeSubsurfaceProfile);
+			
+			CHECK_DDC_QUERY_FINISHED_ELSE_RETURN();
 		}
 
 		Chunk[MP_CustomData0] = Material->CompilePropertyAndSetMaterialProperty(MP_CustomData0, this);
@@ -1477,6 +1513,8 @@ void FHLSLMaterialTranslator::DoTranslate()
 		{
 			// Cast to exact match is needed for float parameter to be correctly cast to float2.
 			int32 UserRefraction = ForceCast(Material->CompilePropertyAndSetMaterialProperty(MP_Refraction, this), MCT_Float2, MFCF_ExactMatch);
+			
+			CHECK_DDC_QUERY_FINISHED_ELSE_RETURN();
 			int32 RefractionDepthBias = ForceCast(ScalarParameter(FName(TEXT("RefractionDepthBias")), Material->GetRefractionDepthBiasValue()), MCT_Float1);
 
 			Chunk[MP_Refraction] = AppendVector(UserRefraction, RefractionDepthBias);
@@ -1498,6 +1536,7 @@ void FHLSLMaterialTranslator::DoTranslate()
 		Chunk[MP_Displacement] = Material->CompilePropertyAndSetMaterialProperty(MP_Displacement, this);
 	}
 
+	CHECK_DDC_QUERY_FINISHED_ELSE_RETURN();
 	ResourcesString = TEXT("");
 
 #if HANDLE_CUSTOM_OUTPUTS_AS_MATERIAL_ATTRIBUTES
@@ -1516,6 +1555,7 @@ void FHLSLMaterialTranslator::DoTranslate()
 
 			for (int32 OutputIndex = 0; OutputIndex < NumOutputs; ++OutputIndex)
 			{
+				CHECK_DDC_QUERY_FINISHED_ELSE_RETURN();
 				MaterialProperty = Attribute.Property;
 				ShaderFrequency = Attribute.ShaderFrequency;
 				FunctionStacks[ShaderFrequency].Empty();
@@ -1576,6 +1616,7 @@ void FHLSLMaterialTranslator::DoTranslate()
 	{
 		for (uint32 CustomUVIndex = MP_CustomizedUVs0; CustomUVIndex <= MP_CustomizedUVs7; CustomUVIndex++)
 		{
+			CHECK_DDC_QUERY_FINISHED_ELSE_RETURN();
 			// Only compile custom UV inputs for UV channels requested by the pixel shader inputs
 			// Any unconnected inputs will have a texcoord generated for them in Material->CompileProperty, which will pass through the vertex (uncustomized) texture coordinates
 			// Note: this is using NumUserTexCoords, which is set by translating all the pixel properties above
@@ -1589,6 +1630,7 @@ void FHLSLMaterialTranslator::DoTranslate()
 	// Output the implementation for any custom expressions we will call below.
 	for (int32 ExpressionIndex = 0; ExpressionIndex < CustomExpressions.Num(); ExpressionIndex++)
 	{
+		CHECK_DDC_QUERY_FINISHED_ELSE_RETURN();
 		ResourcesString += CustomExpressions[ExpressionIndex].Implementation + "\n\n";
 	}
 
@@ -1639,6 +1681,8 @@ void FHLSLMaterialTranslator::DoTranslate()
 					((MaterialShadingModels.HasShadingModel(MSM_SubsurfaceProfile) && IsMaterialPropertyUsed(MP_CustomData0, Chunk[MP_CustomData0], FLinearColor(1, 0, 0, 0), 1))
 					|| (MaterialShadingModels.HasShadingModel(MSM_Eye) && bOpacityPropertyIsUsed));
 
+	CHECK_DDC_QUERY_FINISHED_ELSE_RETURN();
+
 	// If Substrate is enabled or if this is a Substrate material cooked/used in non-Substrate mode, 
 	// we disable this warning as Substrate supports 'colored transmittance only' mode (i.e, modulate).
 	if (!bSubstrateEnabled && !FrontMaterialExpr && IsModulateBlendMode(BlendMode) && MaterialShadingModels.IsLit() && !Material->IsDeferredDecal())
@@ -1808,6 +1852,8 @@ void FHLSLMaterialTranslator::DoTranslate()
 		Errorf(TEXT("Blend Mode \"Modulate\" materials are not currently supported in the \"After Motion Blur\" translucency pass."));
 	}
 
+	CHECK_DDC_QUERY_FINISHED_ELSE_RETURN();
+
 	// Catch any modifications to NumUserTexCoords that will not seen by customized UVs
 	check(SavedNumUserTexCoords == GetNumUserTexCoords());
 
@@ -1832,6 +1878,7 @@ void FHLSLMaterialTranslator::DoTranslate()
 
 	if (bEnableExecutionFlow)
 	{
+		CHECK_DDC_QUERY_FINISHED_ELSE_RETURN();
 		if (VertexAttributesChunk != INDEX_NONE)
 		{
 			LinkParentScopes(SharedPropertyCodeChunks[SF_Vertex]);
@@ -1840,6 +1887,7 @@ void FHLSLMaterialTranslator::DoTranslate()
 			GetScopeCode(1, VertexAttributesChunk, SharedPropertyCodeChunks[SF_Vertex], EmittedChunks, TranslatedAttributesCodeChunks[SF_Vertex]);
 		}
 
+		CHECK_DDC_QUERY_FINISHED_ELSE_RETURN();
 		if (PixelAttributesChunk != INDEX_NONE)
 		{
 			LinkParentScopes(SharedPropertyCodeChunks[SF_Pixel]);
@@ -1852,6 +1900,7 @@ void FHLSLMaterialTranslator::DoTranslate()
 	{
 		for (int32 VariationIter = 0; VariationIter < CompiledPDV_MAX; VariationIter++)
 		{
+			CHECK_DDC_QUERY_FINISHED_ELSE_RETURN();
 			ECompiledPartialDerivativeVariation Variation = (ECompiledPartialDerivativeVariation)VariationIter;
 
 			// Do Normal Chunk first
@@ -1875,6 +1924,7 @@ void FHLSLMaterialTranslator::DoTranslate()
 			// Now the rest, skipping Normal
 			for (uint32 PropertyId = 0; PropertyId < MP_MAX; ++PropertyId)
 			{
+				CHECK_DDC_QUERY_FINISHED_ELSE_RETURN();
 				if (PropertyId == MP_MaterialAttributes || PropertyId == MP_Normal || PropertyId == MP_CustomOutput)
 				{
 					continue;
@@ -2001,7 +2051,8 @@ void FHLSLMaterialTranslator::DoTranslate()
 
 			for (uint32 PropertyId = MP_MAX; PropertyId < CompiledMP_MAX; ++PropertyId)
 			{
-						switch (PropertyId)
+				CHECK_DDC_QUERY_FINISHED_ELSE_RETURN();
+				switch (PropertyId)
 				{
 				case CompiledMP_EmissiveColorCS:
 					if (bCompileForComputeShader)
@@ -2020,12 +2071,13 @@ void FHLSLMaterialTranslator::DoTranslate()
 			}
 
 			// Output the implementation for any custom output expressions
-					for (int32 ExpressionIndex = 0; ExpressionIndex < DerivativeVariations[Variation].CustomOutputImplementations.Num(); ExpressionIndex++)
+			for (int32 ExpressionIndex = 0; ExpressionIndex < DerivativeVariations[Variation].CustomOutputImplementations.Num(); ExpressionIndex++)
 			{
-						ResourcesString += DerivativeVariations[Variation].CustomOutputImplementations[ExpressionIndex] + "\n\n";
-					}
-				}
+				CHECK_DDC_QUERY_FINISHED_ELSE_RETURN();
+				ResourcesString += DerivativeVariations[Variation].CustomOutputImplementations[ExpressionIndex] + "\n\n";
 			}
+		}
+	}
 
 	// Store the number of float4s
 	MaterialCompilationOutput.UniformExpressionSet.UniformPreshaderBufferSize = (UniformPreshaderOffset + 3u) / 4u;
@@ -2035,17 +2087,20 @@ void FHLSLMaterialTranslator::DoTranslate()
 		MaterialCompilationOutput.UniformExpressionSet.UniformTextureParameters[TypeIndex].Empty(UniformTextureExpressions[TypeIndex].Num());
 		for (FMaterialUniformExpressionTexture* TextureExpression : UniformTextureExpressions[TypeIndex])
 		{
+			CHECK_DDC_QUERY_FINISHED_ELSE_RETURN();
 			TextureExpression->GetTextureParameterInfo(MaterialCompilationOutput.UniformExpressionSet.UniformTextureParameters[TypeIndex].AddDefaulted_GetRef());
 		}
 	}
 	MaterialCompilationOutput.UniformExpressionSet.UniformExternalTextureParameters.Empty(UniformExternalTextureExpressions.Num());
 	for (FMaterialUniformExpressionExternalTexture* TextureExpression : UniformExternalTextureExpressions)
 	{
+		CHECK_DDC_QUERY_FINISHED_ELSE_RETURN();
 		TextureExpression->GetExternalTextureParameterInfo(MaterialCompilationOutput.UniformExpressionSet.UniformExternalTextureParameters.AddDefaulted_GetRef());
 	}
 
 	for (uint32 SubstrateCompilationContextIndex = 0; SubstrateCompilationContextIndex < ESubstrateCompilationContext::SCC_MAX; ++SubstrateCompilationContextIndex)
 	{
+		CHECK_DDC_QUERY_FINISHED_ELSE_RETURN();
 		FSubstrateCompilationContext& SubstrateCtx = SubstrateCompilationContext[SubstrateCompilationContextIndex];
 
 		FString TreeFunctionPostFix = TEXT("ERROR");
@@ -3782,6 +3837,8 @@ FString FHLSLMaterialTranslator::GetDefinitions(const TArray<FShaderCodeChunk>&
 // GetFixedParameterCode
 void FHLSLMaterialTranslator::GetFixedParameterCode(int32 StartChunk, int32 EndChunk, int32 ResultIndex, TArray<FShaderCodeChunk>& CodeChunks, FString& OutDefinitions, FString& OutValue, ECompiledPartialDerivativeVariation OriginalVariation, bool bReduceAfterReturnValue)
 {
+	CHECK_DDC_QUERY_FINISHED_ELSE_RETURN();;
+
 	// Only allow the analytic variation for pixel shaders.
 	ECompiledPartialDerivativeVariation Variation = OriginalVariation;
 
@@ -15313,95 +15370,56 @@ void FHLSLMaterialTranslator::PrepareEnvironmentDefines()
 	EnvironmentDefines->ParameterCollections = ParameterCollections;
 }
 
-bool FHLSLMaterialTranslator::QueryDDCCachedTranslationResults()
+void FHLSLMaterialTranslator::AsyncQueryDDC(UE::DerivedData::FRequestOwner& DDCRequestOwner)
 {
-	if (GJobDisableMaterialTranslateDDC)
-	{
-		return false;
-	}
-
-	// We call FindObject to serialize the array of Parameter Collections used by this material in EnvironmentDefines,
-	// but this can happen during save. FindObject is illegal during save because if the discovered objects
-	// are serialized into the package it will cause a crash on package load. But we are not storing the results
-	// of FindObject into the package so it is okay to remove the restriction.
-	TOptional<TGuardValue<bool>> IsSavingPackageGuard;
-	if (IsInGameThread())
-	{
-		IsSavingPackageGuard.Emplace(GIsSavingPackage, false);
-	}
-
-	// Try fetching the translation results from the DDC using the generated key hash
-	FSharedBuffer MaterialCompilationOutputBuffer;
-	FSharedBuffer TranslationResultsBuffer;
-	FSharedBuffer EnvironmentDefinesBuffer;
-
-	UE::DerivedData::FRequestOwner RequestOwner(UE::DerivedData::EPriority::Blocking);
 	UE::DerivedData::FCacheKey CacheKey{ MaterialTranslationDDCBucket, DDCKeyHash };
+	UE::DerivedData::FCacheGetRequest Request;
+	Request.Name = Material->GetMaterialInterface()->GetName();
+	Request.Key = CacheKey;
+	Request.Policy = UE::DerivedData::ECachePolicy::Default;
 
-	{
-		UE::DerivedData::FCacheGetRequest Request;
-		Request.Name = Material->GetMaterialInterface()->GetName();
-		Request.Key = CacheKey;
-		Request.Policy = UE::DerivedData::ECachePolicy::Default;
-
-		UE::DerivedData::GetCache().Get(
-			{ Request },
-			RequestOwner,
-			[&](UE::DerivedData::FCacheGetResponse&& Response)
+	UE::DerivedData::GetCache().Get(
+		{ Request },
+		DDCRequestOwner,
+		[&](UE::DerivedData::FCacheGetResponse&& Response)
+		{
+			if (Response.Status == UE::DerivedData::EStatus::Ok)
 			{
-				if (Response.Status == UE::DerivedData::EStatus::Ok)
+				// Try fetching the translation results from the DDC using the generated key hash
+				FSharedBuffer MaterialCompilationOutputBuffer = Response.Record.GetValue(MaterialCompilationOutputId).GetData().Decompress();
+				FSharedBuffer TranslationResultsBuffer = Response.Record.GetValue(MaterialResultsOutputId).GetData().Decompress();
+				FSharedBuffer EnvironmentDefinesBuffer = Response.Record.GetValue(EnvironmentDefinesId).GetData().Decompress();
+
+				// Load the results if we hit the cache.
+				if (!MaterialCompilationOutputBuffer.IsNull() && !TranslationResultsBuffer.IsNull() && !EnvironmentDefinesBuffer.IsNull())
 				{
-					MaterialCompilationOutputBuffer = Response.Record.GetValue(MaterialCompilationOutputId).GetData().Decompress();
-					TranslationResultsBuffer = Response.Record.GetValue(MaterialResultsOutputId).GetData().Decompress();
-					EnvironmentDefinesBuffer = Response.Record.GetValue(EnvironmentDefinesId).GetData().Decompress();
+					// Read the environment defines
+					FMemoryReaderView EnvironmentDefinesBufferReader{ TArrayView<uint8>{ (uint8*)EnvironmentDefinesBuffer.GetData(), (int)EnvironmentDefinesBuffer.GetSize() } };
+					FObjectAndNameAsStringProxyArchive EnvironmentDefinesBufferReaderProxy{ EnvironmentDefinesBufferReader, true };
+					EnvironmentDefines->Serialize(EnvironmentDefinesBufferReaderProxy);
+
+					// Read the material compilation output
+					FShaderMapPointerTable PointerTable;
+					FPlatformTypeLayoutParameters LayoutParams;
+					LayoutParams.InitializeForPlatform(GetTargetPlatform());
+					FMemoryReaderView MaterialCompilationOutputReader{ TArrayView<uint8>{ (uint8*)MaterialCompilationOutputBuffer.GetData(), (int)MaterialCompilationOutputBuffer.GetSize() } };
+					FMemoryImageObject LoadedContent = FMemoryImageResult::LoadFromArchive(MaterialCompilationOutputReader, StaticGetTypeLayoutDesc<FMaterialCompilationOutput>(), &PointerTable, LayoutParams);
+					DDCMaterialCompilationOutput = *(FMaterialCompilationOutput*)LoadedContent.Object;
+
+					// Read the array of material string parameters
+					FMemoryReaderView ResultsMemoryReader{ TArrayView<uint8>{ (uint8*)TranslationResultsBuffer.GetData(), (int)TranslationResultsBuffer.GetSize() } };
+					ResultsMemoryReader << MaterialSourceTemplateParams;
+
+					DDCQueryHit = true;
 				}
 			}
-		);
-	}
-
-	RequestOwner.Wait();
-
-	// Load the results if we hit the cache.
-	if (!MaterialCompilationOutputBuffer.IsNull() && !TranslationResultsBuffer.IsNull() && !EnvironmentDefinesBuffer.IsNull())
-	{
-		// Read the environment defines
-		FMemoryReaderView EnvironmentDefinesBufferReader{ TArrayView<uint8>{ (uint8*)EnvironmentDefinesBuffer.GetData(), (int)EnvironmentDefinesBuffer.GetSize() } };
-		FObjectAndNameAsStringProxyArchive EnvironmentDefinesBufferReaderProxy{ EnvironmentDefinesBufferReader, true };
-		EnvironmentDefines->Serialize(EnvironmentDefinesBufferReaderProxy);
-
-		// Read the material compilation output
-		FShaderMapPointerTable PointerTable;
-		FPlatformTypeLayoutParameters LayoutParams;
-		LayoutParams.InitializeForPlatform(GetTargetPlatform());
-		FMemoryReaderView MaterialCompilationOutputReader{ TArrayView<uint8>{ (uint8*)MaterialCompilationOutputBuffer.GetData(), (int)MaterialCompilationOutputBuffer.GetSize() } };
-		FMemoryImageObject LoadedContent = FMemoryImageResult::LoadFromArchive(MaterialCompilationOutputReader, StaticGetTypeLayoutDesc<FMaterialCompilationOutput>(), &PointerTable, LayoutParams);
-		MaterialCompilationOutput = *(FMaterialCompilationOutput*)LoadedContent.Object;
-
-		// Read the array of material string parameters
-		FMemoryReaderView ResultsMemoryReader{ TArrayView<uint8>{ (uint8*)TranslationResultsBuffer.GetData(), (int)TranslationResultsBuffer.GetSize() } };
-		ResultsMemoryReader << MaterialSourceTemplateParams;
-
-		return true;
-	}
-
-	return false;
+		}
+	);
+	
 }
 
-void FHLSLMaterialTranslator::PushResultsToDDCCache()
+void FHLSLMaterialTranslator::PushResultsToDDC()
 {
-	if (GJobDisableMaterialTranslateDDC)
-	{
-		GShaderCompilerStats->IncrementMaterialTranslationSkippedDDC();
-		return;
-	}
-
-	// See FHLSLMaterialTranslator::QueryDDCCachedTranslationResults()
-	TOptional<TGuardValue<bool>> IsSavingPackageGuard;
-	if (IsInGameThread())
-	{
-		IsSavingPackageGuard.Emplace(GIsSavingPackage, false);
-	}
-
 	UE::DerivedData::FCacheKey CacheKey{ MaterialTranslationDDCBucket, DDCKeyHash };
 	UE::DerivedData::FCacheRecordBuilder RecordBuilder{ CacheKey };
 
diff --git a/Engine/Source/Runtime/Engine/Private/Materials/HLSLMaterialTranslator.h b/Engine/Source/Runtime/Engine/Private/Materials/HLSLMaterialTranslator.h
--- a/Engine/Source/Runtime/Engine/Private/Materials/HLSLMaterialTranslator.h
+++ b/Engine/Source/Runtime/Engine/Private/Materials/HLSLMaterialTranslator.h
@@ -57,6 +57,11 @@
 class Error;
 class FAddUniformExpressionScope;
 
+namespace UE::DerivedData
+{
+	class FRequestOwner;
+}
+
 /**
  * Returns whether the specified class of material expression is permitted.
  * For instance, custom expressions are not permitted in certain UE editor configurations for client generated materials.
@@ -265,8 +270,7 @@ protected:
 
 	/** Material being compiled.  Only transient compilation output like error information can be stored on the FMaterial. */
 	FMaterial* Material;
-	/** Compilation output which will be stored in the DDC. */
-	FMaterialCompilationOutput& MaterialCompilationOutput;
+
 	FStaticParameterSet StaticParameters;
 	FMaterialLayersFunctions CachedMaterialLayers;
 	EShaderPlatform Platform;
@@ -587,6 +591,8 @@ public:
 
 	~FHLSLMaterialTranslator();
 
+	bool ShouldStopTranslating() const override;
+
 	int32 GetNumUserTexCoords() const;
 	int32 GetNumUserVertexTexCoords() const;
 
@@ -1383,12 +1389,12 @@ protected:
 	 * Queries the DDC cache for a cached translation.
 	 * @return Whether the speecified key is in the DDC.
 	 */
-	bool QueryDDCCachedTranslationResults();
+	void AsyncQueryDDC(UE::DerivedData::FRequestOwner& DDCRequestOwner);
 
 	/**
 	 * Pushes the final results to the DDC cache.
 	 */
-	void PushResultsToDDCCache();
+	void PushResultsToDDC();
 
 	/**
 	 * Prepares the material source generation parameters.
@@ -1400,12 +1406,21 @@ protected:
 	 */
 	void PrepareEnvironmentDefines();
 
+	/** The final compilation output. */
+	FMaterialCompilationOutput& MaterialCompilationOutput;
+
+	/** Compilation output fetched from the DDC */
+	FMaterialCompilationOutput DDCMaterialCompilationOutput;
+
 	/** The material shader source template parameters */
 	TMap<FString, FString> MaterialSourceTemplateParams;
 
 	/** The output material shader defines */
 	TUniquePtr<FEnvironmentDefines> EnvironmentDefines;
 
+	/** Signals when the async DDC query task has completed AND there was a hit */
+	TAtomic<bool> DDCQueryHit;
+
 	friend class FAddUniformExpressionScope;
 };
 
diff --git a/Engine/Source/Runtime/Engine/Private/Materials/Material.cpp b/Engine/Source/Runtime/Engine/Private/Materials/Material.cpp
--- a/Engine/Source/Runtime/Engine/Private/Materials/Material.cpp
+++ b/Engine/Source/Runtime/Engine/Private/Materials/Material.cpp
@@ -180,6 +180,11 @@ FMaterialResource::~FMaterialResource()
 int32 FMaterialResource::CompilePropertyAndSetMaterialProperty(EMaterialProperty Property, FMaterialCompiler* Compiler, EShaderFrequency OverrideShaderFrequency, bool bUsePreviousFrameTime) const
 {
 #if WITH_EDITOR
+	if (Compiler->ShouldStopTranslating())
+	{
+		return INDEX_NONE;
+	}
+
 	// needs to be called in this function!!
 	// sets CurrentShaderFrequency
 	Compiler->SetMaterialProperty(Property, OverrideShaderFrequency, bUsePreviousFrameTime);
diff --git a/Engine/Source/Runtime/Engine/Private/ShaderCompiler/ShaderCompiler.cpp b/Engine/Source/Runtime/Engine/Private/ShaderCompiler/ShaderCompiler.cpp
--- a/Engine/Source/Runtime/Engine/Private/ShaderCompiler/ShaderCompiler.cpp
+++ b/Engine/Source/Runtime/Engine/Private/ShaderCompiler/ShaderCompiler.cpp
@@ -5190,23 +5190,22 @@ void FShaderCompilerStats::FMaterialCounters::WriteStatSummary(const TCHAR* Aggr
 	UE_LOG(LogShaderCompilers, Display, TEXT("=== Material stats%s ==="), AggregatedSuffix);
 	UE_LOG(LogShaderCompilers, Display, TEXT("Materials Cooked:        %d"), NumMaterialsCooked);
 	UE_LOG(LogShaderCompilers, Display, TEXT("Materials Translated:    %d"), MaterialTranslateCalls);
-	UE_LOG(LogShaderCompilers, Display, TEXT("Material Translate Time: %.2f s"), MaterialTranslateTimeSec);
-	UE_LOG(LogShaderCompilers, Display, TEXT("Material Cache Time: %.2f s"), MaterialCacheTimeSec);
+	UE_LOG(LogShaderCompilers, Display, TEXT("Material Total Translate Time: %.2f s"), MaterialTranslateTotalTimeSec);
 
-	int HitsPercentage = MaterialTranslateCalls ? MaterialCacheHits * 100 / MaterialTranslateCalls : 0;
-	UE_LOG(LogShaderCompilers, Display, TEXT("Material Cache Hits: %d (%d%%)"), MaterialCacheHits, HitsPercentage);
+	int TranslationTimePercentage = (int)round(MaterialTranslateTranslationOnlyTimeSec / FMath::Max(1e-6, MaterialTranslateTotalTimeSec) * 100);
+	UE_LOG(LogShaderCompilers, Display, TEXT("Material Translation Only: %.2f s (%d%%)"), MaterialTranslateTranslationOnlyTimeSec, TranslationTimePercentage);
 
-	UE_LOG(LogShaderCompilers, Display, TEXT("Material translations not on DDC: %d"), MaterialTranslationSkippedDDCCount);
+	int HitsPercentage = MaterialTranslateCalls ? (int)roundf(float(MaterialCacheHits) / MaterialTranslateCalls * 100) : 0;
+	UE_LOG(LogShaderCompilers, Display, TEXT("Material Cache Hits: %d (%d%%)"), MaterialCacheHits, HitsPercentage);
 }
 
 void FShaderCompilerStats::FMaterialCounters::GatherAnalytics(TArray<FAnalyticsEventAttribute>& Attributes)
 {
 	Attributes.Emplace(TEXT("Material_NumMaterialsCooked"), NumMaterialsCooked);
 	Attributes.Emplace(TEXT("Material_MaterialTranslateCalls"), MaterialTranslateCalls);
-	Attributes.Emplace(TEXT("Material_MaterialTranslateTimeSec"), MaterialTranslateTimeSec);
-	Attributes.Emplace(TEXT("Material_MaterialCacheTimeSec"), MaterialCacheTimeSec);
+	Attributes.Emplace(TEXT("Material_MaterialTranslateTimeSec"), MaterialTranslateTotalTimeSec);
+	Attributes.Emplace(TEXT("Material_MaterialTranslateTranslationOnlyTimeSec"), MaterialTranslateTranslationOnlyTimeSec);
 	Attributes.Emplace(TEXT("Material_MaterialCacheHits"), MaterialCacheHits);
-	Attributes.Emplace(TEXT("Material_MaterialTranslationSkippedDDCCount"), MaterialTranslationSkippedDDCCount);
 }
 
 void FShaderCompilerStats::IncrementMaterialCook()
@@ -5215,23 +5214,18 @@ void FShaderCompilerStats::IncrementMaterialCook()
 	MaterialCounters.NumMaterialsCooked++;
 }
 
-void FShaderCompilerStats::IncrementMaterialTranslated(double InTime)
+void FShaderCompilerStats::IncrementMaterialTranslated(double InTotalTime, double InTranslationOnlyTime)
 {
 	FScopeLock Lock(&CompileStatsLock);
 	MaterialCounters.MaterialTranslateCalls++;
-	MaterialCounters.MaterialTranslateTimeSec += InTime;
+	MaterialCounters.MaterialTranslateTotalTimeSec += InTotalTime;
+	MaterialCounters.MaterialTranslateTranslationOnlyTimeSec += InTranslationOnlyTime;
 }
 
-void FShaderCompilerStats::IncrementMaterialCacheHit(double InCacheTime)
+void FShaderCompilerStats::IncrementMaterialCacheHit()
 {
 	FScopeLock Lock(&CompileStatsLock);
 	MaterialCounters.MaterialCacheHits++;
-	MaterialCounters.MaterialCacheTimeSec += InCacheTime;
-}
-
-void FShaderCompilerStats::IncrementMaterialTranslationSkippedDDC()
-{
-	MaterialCounters.MaterialTranslationSkippedDDCCount++;
 }
 
 void FShaderCompilerStats::RegisterCookedShaders(uint32 NumCooked, float CompileTime, EShaderPlatform Platform, const FString MaterialPath, FString PermutationString)
diff --git a/Engine/Source/Runtime/Engine/Public/MaterialCompiler.h b/Engine/Source/Runtime/Engine/Public/MaterialCompiler.h
--- a/Engine/Source/Runtime/Engine/Public/MaterialCompiler.h
+++ b/Engine/Source/Runtime/Engine/Public/MaterialCompiler.h
@@ -87,6 +87,9 @@ class FMaterialCompiler
 public:
 	virtual ~FMaterialCompiler() { }
 
+	/** Whether material translation should abort */
+	virtual bool ShouldStopTranslating() const = 0;
+
 	// sets internal state CurrentShaderFrequency 
 	// @param OverrideShaderFrequency SF_NumFrequencies to not override
 	virtual void SetMaterialProperty(EMaterialProperty InProperty, EShaderFrequency OverrideShaderFrequency = SF_NumFrequencies, bool bUsePreviousFrameTime = false) = 0;
@@ -752,7 +755,8 @@ public:
 	{}
 
 	// Simple pass through all other material operations unmodified.
-
+	
+	virtual bool ShouldStopTranslating() const override { return false; }
 	virtual FMaterialShadingModelField GetMaterialShadingModels() const { return Compiler->GetMaterialShadingModels(); }
 	virtual FMaterialShadingModelField GetCompiledShadingModels() const { return Compiler->GetCompiledShadingModels(); }
 	virtual EMaterialValueType GetParameterType(int32 Index) const { return Compiler->GetParameterType(Index); }
diff --git a/Engine/Source/Runtime/Engine/Public/ShaderCompiler.h b/Engine/Source/Runtime/Engine/Public/ShaderCompiler.h
--- a/Engine/Source/Runtime/Engine/Public/ShaderCompiler.h
+++ b/Engine/Source/Runtime/Engine/Public/ShaderCompiler.h
@@ -365,9 +365,8 @@ public:
 	};
 
 	void IncrementMaterialCook();
-	void IncrementMaterialTranslated(double InTime);
-	void IncrementMaterialCacheHit(double InTime);
-	void IncrementMaterialTranslationSkippedDDC();
+	void IncrementMaterialTranslated(double InTotalTime, double InTranslationOnlyTime);
+	void IncrementMaterialCacheHit();
 
 	ENGINE_API void RegisterCookedShaders(uint32 NumCooked, float CompileTime, EShaderPlatform Platform, const FString MaterialPath, FString PermutationString = FString(""));
 	ENGINE_API void RegisterCompiledShaders(uint32 NumPermutations, EShaderPlatform Platform, const FString MaterialPath, FString PermutationString = FString(""));
@@ -565,24 +564,21 @@ private:
 		int32 MaterialTranslateCalls = 0;
 
 		/** The total time in seconds to translate all materials.  */
-		double MaterialTranslateTimeSec = 0.0;
+		double MaterialTranslateTotalTimeSec = 0.0;
+
+		/** The total time spent actually translating materials (rather than for instance accessing the DDC cache). */
+		double MaterialTranslateTranslationOnlyTimeSec = 0.0;
 
 		/** The total number times a material translation was skipped because the the results were in the DDC. */
 		int32 MaterialCacheHits = 0;
 		
-		/** The total time spent querying the DDC cache. */
-		double MaterialCacheTimeSec = 0.0;
-
-		int32 MaterialTranslationSkippedDDCCount = 0;
-
 		FMaterialCounters& operator+=(const FMaterialCounters& Other)
 		{
 			NumMaterialsCooked += Other.NumMaterialsCooked;
 			MaterialTranslateCalls += Other.MaterialTranslateCalls;
-			MaterialTranslateTimeSec += Other.MaterialTranslateTimeSec;
+			MaterialTranslateTotalTimeSec += Other.MaterialTranslateTotalTimeSec;
+			MaterialTranslateTranslationOnlyTimeSec += Other.MaterialTranslateTranslationOnlyTimeSec;
 			MaterialCacheHits += Other.MaterialCacheHits;
-			MaterialCacheTimeSec += Other.MaterialCacheTimeSec;
-			MaterialTranslationSkippedDDCCount += Other.MaterialTranslationSkippedDDCCount;
 
 			return *this;
 		}
