Fixed a failing test getting animation notifies, which corrects a recent error that could result in getting the wrong notifies when requesting a negative delta move from 0.0.  This was caused by a recent change to prevent an infinite loop in the same function.

diff --git a/Engine/Source/Runtime/Engine/Private/Animation/AnimSequenceBase.cpp b/Engine/Source/Runtime/Engine/Private/Animation/AnimSequenceBase.cpp
--- a/Engine/Source/Runtime/Engine/Private/Animation/AnimSequenceBase.cpp
+++ b/Engine/Source/Runtime/Engine/Private/Animation/AnimSequenceBase.cpp
@@ -419,22 +419,30 @@ void UAnimSequenceBase::GetAnimNotifies(const float& StartTime, const float& Del
 	float PreviousPosition = StartTime;
 	float CurrentPosition = StartTime;
 	float DesiredDeltaMove = DeltaTime;
+	const float PlayLength = GetPlayLength();
 
-	do 
+	// previous behaviour could get the same notify multiple times  - support this within reasonable limits
+	uint32_t MaxLoopCount = 2;
+	if (PlayLength > 0.0f && FMath::Abs(DeltaTime) > PlayLength)
+	{
+		MaxLoopCount = FMath::Clamp(uint32_t(DesiredDeltaMove / PlayLength), 2, 1000);
+	}
+
+	for (uint32_t i = 0; i < MaxLoopCount; i++)
 	{
 		// Disable looping here. Advance to desired position, or beginning / end of animation
-		const ETypeAdvanceAnim AdvanceType = FAnimationRuntime::AdvanceTime(false, DesiredDeltaMove, CurrentPosition, GetPlayLength());
+		const ETypeAdvanceAnim AdvanceType = FAnimationRuntime::AdvanceTime(false, DesiredDeltaMove, CurrentPosition, PlayLength);
 
 		// Verify position assumptions
-		ensureMsgf(bPlayingBackwards ? (CurrentPosition <= PreviousPosition) : (CurrentPosition >= PreviousPosition), TEXT("in Animation %s(Skeleton %s) : bPlayingBackwards(%d), PreviousPosition(%0.2f), Current Position(%0.2f)"), 
+		ensureMsgf(bPlayingBackwards ? (CurrentPosition <= PreviousPosition) : (CurrentPosition >= PreviousPosition), TEXT("in Animation %s(Skeleton %s) : bPlayingBackwards(%d), PreviousPosition(%0.2f), Current Position(%0.2f)"),
 			*GetName(), *GetNameSafe(GetSkeleton()), bPlayingBackwards, PreviousPosition, CurrentPosition);
-		
+
 		GetAnimNotifiesFromDeltaPositions(PreviousPosition, CurrentPosition, NotifyContext);
-	
-		const float ActualDeltaMove = (CurrentPosition - PreviousPosition);
+
 		// If we've hit the end of the animation, and we're allowed to loop, keep going.
-		if( (AdvanceType == ETAA_Finished) &&  NotifyContext.TickRecord && NotifyContext.TickRecord->bLooping && ActualDeltaMove != 0.f)
+		if ((AdvanceType == ETAA_Finished) && NotifyContext.TickRecord && NotifyContext.TickRecord->bLooping)
 		{
+			const float ActualDeltaMove = (CurrentPosition - PreviousPosition);
 			DesiredDeltaMove -= ActualDeltaMove;
 			PreviousPosition = bPlayingBackwards ? GetPlayLength() : 0.f;
 			CurrentPosition = PreviousPosition;
@@ -443,8 +451,7 @@ void UAnimSequenceBase::GetAnimNotifies(const float& StartTime, const float& Del
 		{
 			break;
 		}
-	} 
-	while( true );
+	}
 }
 
 void UAnimSequenceBase::GetAnimNotifiesFromDeltaPositions(const float& PreviousPosition, const float& CurrentPosition, TArray<const FAnimNotifyEvent *> & OutActiveNotifies) const
