Turn off GAllowActorScriptExecutionInEditor when we execute cheats to avoid recursive calls
#tests PIE
#rnx

diff --git a/Engine/Source/Runtime/Engine/Private/CheatManager.cpp b/Engine/Source/Runtime/Engine/Private/CheatManager.cpp
--- a/Engine/Source/Runtime/Engine/Private/CheatManager.cpp
+++ b/Engine/Source/Runtime/Engine/Private/CheatManager.cpp
@@ -123,6 +123,7 @@ bool UCheatManager::ProcessConsoleExec(const TCHAR* Cmd, FOutputDevice& Ar, UObj
 
 				if ((Function != nullptr) && Function->HasAnyFunctionFlags(FUNC_BlueprintAuthorityOnly))
 				{
+					TGuardValue AutoRestore(GAllowActorScriptExecutionInEditor, false);
 					MyPC->ServerExec(Cmd);
 					return true;
 				}
