Refactor AudioStreamingMemoryCountedFeature

Move/Rename to: AudioStreamCacheMemoryHandle

Make AudioStreamCacheMemoryHandle methods protected and make AudioStreamCacheMemoryHandle a friend class. That way the the "External Features" API is Protected and requires users to create an instance of an AudioStreamCacheMemoryHandle class, which ensures that we properly add/remove memory counts, and guarantees we remove all counted memory when the "MemoryHandle" instance is destroyed.

#tests PIE

diff --git a/Engine/Source/Runtime/Engine/Private/AudioStreamCacheMemoryHandle.cpp b/Engine/Source/Runtime/Engine/Private/AudioStreamCacheMemoryHandle.cpp
new file mode 100644
--- /dev/null
+++ b/Engine/Source/Runtime/Engine/Private/AudioStreamCacheMemoryHandle.cpp
@@ -0,0 +1,37 @@
+// Copyright Epic Games, Inc. All Rights Reserved.
+
+#include "AudioStreamCacheMemoryHandle.h"
+#include "ContentStreaming.h"
+
+FAudioStreamCacheMemoryHandle::FAudioStreamCacheMemoryHandle(FName InFeatureName, uint64 InMemoryUseInBytes)
+: FeatureName(InFeatureName)
+, MemoryUseInBytes(InMemoryUseInBytes)
+{
+	if (MemoryUseInBytes != 0)
+	{
+		IStreamingManager::Get().GetAudioStreamingManager().AddMemoryCountedFeature(*this);
+	}
+}
+
+FAudioStreamCacheMemoryHandle::~FAudioStreamCacheMemoryHandle()
+{
+	if (MemoryUseInBytes != 0)
+	{
+		IStreamingManager::Get().GetAudioStreamingManager().RemoveMemoryCountedFeature(*this);
+	}
+}
+
+void FAudioStreamCacheMemoryHandle::ResetMemoryUseInBytes(uint64 InMemoryUseInBytes)
+{
+	if (MemoryUseInBytes != 0)
+	{
+		IStreamingManager::Get().GetAudioStreamingManager().RemoveMemoryCountedFeature(*this);
+	}
+		
+	MemoryUseInBytes = InMemoryUseInBytes;
+
+	if (MemoryUseInBytes != 0)
+	{
+		IStreamingManager::Get().GetAudioStreamingManager().AddMemoryCountedFeature(*this);
+	}
+}
\ No newline at end of file
diff --git a/Engine/Source/Runtime/Engine/Private/AudioStreamingCache.cpp b/Engine/Source/Runtime/Engine/Private/AudioStreamingCache.cpp
--- a/Engine/Source/Runtime/Engine/Private/AudioStreamingCache.cpp
+++ b/Engine/Source/Runtime/Engine/Private/AudioStreamingCache.cpp
@@ -5,6 +5,7 @@ AudioStreaming.cpp: Implementation of audio streaming classes.
 =============================================================================*/
 
 #include "AudioStreamingCache.h"
+#include "AudioStreamCacheMemoryHandle.h"
 #include "Async/Async.h"
 #include "Audio.h"
 #include "AudioCompressionSettingsUtils.h"
@@ -443,7 +444,7 @@ void FCachedAudioStreamingManager::RemoveForceInlineSoundWave(const FSoundWavePr
 	}
 }
 
-void FCachedAudioStreamingManager::AddMemoryCountedFeature(const FAudioStreamingMemoryCountedFeature& Feature)
+void FCachedAudioStreamingManager::AddMemoryCountedFeature(const FAudioStreamCacheMemoryHandle& Feature)
 {
 	// add memory count to the first cache
 	if (ensure(CacheArray.Num() > 0))
@@ -452,7 +453,7 @@ void FCachedAudioStreamingManager::AddMemoryCountedFeature(const FAudioStreaming
 	}
 }
 
-void FCachedAudioStreamingManager::RemoveMemoryCountedFeature(const FAudioStreamingMemoryCountedFeature& Feature)
+void FCachedAudioStreamingManager::RemoveMemoryCountedFeature(const FAudioStreamCacheMemoryHandle& Feature)
 {
 	// remove memory count from the first cache
 	if (ensure(CacheArray.Num() > 0))
@@ -1063,7 +1064,7 @@ void FAudioChunkCache::RemoveForceInlineSoundWave(const FSoundWaveProxyPtr& Soun
 	ForceInlineMemoryCounterBytes -= Data ? Data->GetBulkDataSize() : 0;
 }
 
-void FAudioChunkCache::AddMemoryCountedFeature(const FAudioStreamingMemoryCountedFeature& Feature)
+void FAudioChunkCache::AddMemoryCountedFeature(const FAudioStreamCacheMemoryHandle& Feature)
 {
 	UE_LOG(LogAudioStreamCaching, Log, TEXT("Adding Memory Counted Feature (%s) Memory Usage: %d bytes"), *Feature.GetFeatureName().ToString(), (int32)Feature.GetMemoryUseInBytes());
 	const uint64 OldMemoryCount = FeatureMemoryCounterBytes.AddExchange(Feature.GetMemoryUseInBytes());
@@ -1085,7 +1086,7 @@ void FAudioChunkCache::AddMemoryCountedFeature(const FAudioStreamingMemoryCounte
 	}
 }
 
-void FAudioChunkCache::RemoveMemoryCountedFeature(const FAudioStreamingMemoryCountedFeature& Feature)
+void FAudioChunkCache::RemoveMemoryCountedFeature(const FAudioStreamCacheMemoryHandle& Feature)
 {
 	UE_LOG(LogAudioStreamCaching, Log, TEXT("Removing Memory Counted Feature (%s) Memory Usage: %d"), *Feature.GetFeatureName().ToString(), (int32)Feature.GetMemoryUseInBytes());
 	checkf(FeatureMemoryCounterBytes.Load() >= Feature.GetMemoryUseInBytes(), TEXT("Count (%lu) < Remove (%lu)"), FeatureMemoryCounterBytes.Load(), Feature.GetMemoryUseInBytes());
diff --git a/Engine/Source/Runtime/Engine/Public/AudioStreamCacheMemoryHandle.h b/Engine/Source/Runtime/Engine/Public/AudioStreamCacheMemoryHandle.h
new file mode 100644
--- /dev/null
+++ b/Engine/Source/Runtime/Engine/Public/AudioStreamCacheMemoryHandle.h
@@ -0,0 +1,33 @@
+// Copyright Epic Games, Inc. All Rights Reserved.
+
+#pragma once
+#include "Templates/UnrealTemplate.h"
+#include "UObject/NameTypes.h"
+#include "HAL/Platform.h"
+
+/**
+ * Class for utilizing memory in the stream cache budget for an unrelated, temporary audio-based feature.
+ * Allows us to borrow from budgeted memory for audio features that would otherwise not fit in the overall memory budget.
+ *
+ * Usage:
+ * - Create an instance of this class on an object or subsystem where want to track memory usage
+ * - Memory usage will immediately be taken out of the Audio Stream Cache budget on construction of the object
+ * - Update memory usage via this class as necessary.
+ * - Deleting the instance will automatically reset the memory usage to 0
+*/
+class ENGINE_API FAudioStreamCacheMemoryHandle : public FNoncopyable
+{
+public:
+	FAudioStreamCacheMemoryHandle(FName InFeatureName, uint64 InMemoryUseInBytes);
+	
+	~FAudioStreamCacheMemoryHandle();
+	
+	FORCEINLINE uint64 GetMemoryUseInBytes() const { return MemoryUseInBytes; } 
+	FORCEINLINE FName GetFeatureName() const { return FeatureName; }
+
+	void ResetMemoryUseInBytes(uint64 InMemoryUseInBytes);
+	
+private:
+	const FName FeatureName;
+	uint64 MemoryUseInBytes;
+};
\ No newline at end of file
diff --git a/Engine/Source/Runtime/Engine/Public/AudioStreamingCache.h b/Engine/Source/Runtime/Engine/Public/AudioStreamingCache.h
--- a/Engine/Source/Runtime/Engine/Public/AudioStreamingCache.h
+++ b/Engine/Source/Runtime/Engine/Public/AudioStreamingCache.h
@@ -25,7 +25,7 @@ AudioStreaming.h: Definitions of classes used for audio streaming.
 
 ENGINE_API DECLARE_LOG_CATEGORY_EXTERN(LogAudioStreamCaching, Display, All);
 
-class FAudioStreamingMemoryCountedFeature;
+class FAudioStreamCacheMemoryHandle;
 
 // Basic fixed-size LRU cache for retaining chunks of compressed audio data.
 class FAudioChunkCache
@@ -131,9 +131,9 @@ public:
 
 	void RemoveForceInlineSoundWave(const FSoundWaveProxyPtr&);
 
-	void AddMemoryCountedFeature(const FAudioStreamingMemoryCountedFeature& Feature);
+	void AddMemoryCountedFeature(const FAudioStreamCacheMemoryHandle& Feature);
 
-	void RemoveMemoryCountedFeature(const FAudioStreamingMemoryCountedFeature& Feature);
+	void RemoveMemoryCountedFeature(const FAudioStreamCacheMemoryHandle& Feature);
 
 	// This function will reclaim memory by freeing as many chunks as needed to free BytesToFree.
 	// returns the amount of bytes we were actually able to free.
@@ -528,8 +528,6 @@ public:
 	virtual void RemoveStreamingSoundWave(const FSoundWaveProxyPtr& SoundWave) override;
 	virtual void AddForceInlineSoundWave(const FSoundWaveProxyPtr& SoundWave) override;
 	virtual void RemoveForceInlineSoundWave(const FSoundWaveProxyPtr& SoundWave) override;
-	virtual void AddMemoryCountedFeature(const FAudioStreamingMemoryCountedFeature& Feature) override;
-	virtual void RemoveMemoryCountedFeature(const FAudioStreamingMemoryCountedFeature& Feature) override;
 	virtual void AddDecoder(ICompressedAudioInfo* CompressedAudioInfo) override;
 	virtual void RemoveDecoder(ICompressedAudioInfo* CompressedAudioInfo) override;
 	virtual bool IsManagedStreamingSoundWave(const FSoundWaveProxyPtr&  SoundWave) const override;
@@ -556,6 +554,10 @@ protected:
 	virtual void AddReferenceToChunk(const FAudioChunkHandle& InHandle) override;
 	virtual void RemoveReferenceToChunk(const FAudioChunkHandle& InHandle) override;
 
+	// These are used to update the memory count of Memory Counted Features
+	virtual void AddMemoryCountedFeature(const FAudioStreamCacheMemoryHandle& Feature) override;
+	virtual void RemoveMemoryCountedFeature(const FAudioStreamCacheMemoryHandle& Feature) override;
+
 	/**
 	 * Returns which cache this sound wave should be in,
 	 * based on the size of this sound wave's chunk,
diff --git a/Engine/Source/Runtime/Engine/Public/ContentStreaming.h b/Engine/Source/Runtime/Engine/Public/ContentStreaming.h
--- a/Engine/Source/Runtime/Engine/Public/ContentStreaming.h
+++ b/Engine/Source/Runtime/Engine/Public/ContentStreaming.h
@@ -47,7 +47,7 @@ class UStaticMesh;
 class USkeletalMesh;
 class UStreamableRenderAsset;
 class FSoundSource;
-class FAudioStreamingMemoryCountedFeature;
+class FAudioStreamCacheMemoryHandle;
 struct FWaveInstance;
 struct FRenderAssetStreamingManager;
 
@@ -527,10 +527,6 @@ struct IAudioStreamingManager : public IStreamingManager
 	/** Removes the memory usage of the force inline sound from the streaming cache budget */
 	virtual void RemoveForceInlineSoundWave(const FSoundWaveProxyPtr& SoundWave) { };
 
-	virtual void AddMemoryCountedFeature(const FAudioStreamingMemoryCountedFeature& Feature) { };
-
-	virtual void RemoveMemoryCountedFeature(const FAudioStreamingMemoryCountedFeature& Feature) { };
-
 	/** Adds the decoder to the streaming manager to prevent stream chunks from getting reaped from underneath it */
 	virtual void AddDecoder(ICompressedAudioInfo* CompressedAudioInfo) = 0;
 
@@ -600,6 +596,7 @@ struct IAudioStreamingManager : public IStreamingManager
 
 protected:
 	friend FAudioChunkHandle;
+	friend FAudioStreamCacheMemoryHandle;
 
 	/** This can be called by implementers of IAudioStreamingManager to construct an FAudioChunkHandle using an otherwise inaccessible constructor. */
 	static FAudioChunkHandle BuildChunkHandle(const uint8* InData, uint32 NumBytes, const FSoundWaveProxyPtr& InSoundWave, const FName& SoundWaveName, uint32 InChunkIndex, uint64 CacheLookupID);
@@ -613,6 +610,17 @@ protected:
 	 * This can be used to decrement reference counted handles to audio chunks. Called by the destructor of FAudioChunkHandle.
 	 */
 	virtual void RemoveReferenceToChunk(const FAudioChunkHandle& InHandle) = 0;
+
+	/**
+     * This can be used to increase the memory count for external features. Called by FAudioStreamCacheMemoryHandle.
+     * The pattern for _changing_ the amount of memory of an already added feature is to first remove and then add again with the new number
+     */
+	virtual void AddMemoryCountedFeature(const FAudioStreamCacheMemoryHandle& Feature) { };
+
+	/**
+	* This can be used to decrease the memory count for external features. Called by FAudioStreamCacheMemoryHandle.
+	*/
+	virtual void RemoveMemoryCountedFeature(const FAudioStreamCacheMemoryHandle& Feature) { };
 };
 
 /**
@@ -856,47 +864,4 @@ protected:
 	// Locks out any audio streaming manager call when we are re-initializing the audio streaming manager.
 	mutable FCriticalSection AudioStreamingManagerCriticalSection;
 #endif
-};
-
-class ENGINE_API FAudioStreamingMemoryCountedFeature : public TSharedFromThis<FAudioStreamingMemoryCountedFeature>
-{
-public:
-	FAudioStreamingMemoryCountedFeature(FName InFeatureName, uint64 InMemoryUseInBytes)
-	: FeatureName(InFeatureName)
-	, MemoryUseInBytes(InMemoryUseInBytes)
-	{
-		if (MemoryUseInBytes != 0)
-		{
-			IStreamingManager::Get().GetAudioStreamingManager().AddMemoryCountedFeature(*this);
-		}
-	}
-	
-	~FAudioStreamingMemoryCountedFeature()
-	{
-		if (MemoryUseInBytes != 0)
-		{
-			IStreamingManager::Get().GetAudioStreamingManager().RemoveMemoryCountedFeature(*this);
-		}
-	}
-	
-	uint64 GetMemoryUseInBytes() const { return MemoryUseInBytes; } 
-	FName GetFeatureName() const { return FeatureName; }
-
-	void ResetMemoryUseInBytes(uint64 InMemoryUseInBytes)
-	{
-		if (MemoryUseInBytes != 0)
-		{
-			IStreamingManager::Get().GetAudioStreamingManager().RemoveMemoryCountedFeature(*this);
-		}
-		
-		MemoryUseInBytes = InMemoryUseInBytes;
-
-		if (MemoryUseInBytes != 0)
-		{
-			IStreamingManager::Get().GetAudioStreamingManager().AddMemoryCountedFeature(*this);
-		}
-	}
-private:
-	const FName FeatureName;
-	uint64 MemoryUseInBytes;
 };
\ No newline at end of file
