Fixes a bug where OneOverPreExposure was being applied only to PostProcessInput0

diff --git a/Engine/Shaders/Private/PostProcessMaterialShaders.usf b/Engine/Shaders/Private/PostProcessMaterialShaders.usf
--- a/Engine/Shaders/Private/PostProcessMaterialShaders.usf
+++ b/Engine/Shaders/Private/PostProcessMaterialShaders.usf
@@ -7,6 +7,8 @@
 #include "Common.ush"
 #include "ScreenPass.ush"
 
+#pragma message("UESHADERMETADATA_VERSION C79EC12D-4DE8-4D26-B0C7-BD44503987C5")
+
 #ifndef POST_PROCESS_MATERIAL
 #define POST_PROCESS_MATERIAL 0
 #endif
diff --git a/Engine/Source/Runtime/Engine/Private/Materials/HLSLMaterialTranslator.cpp b/Engine/Source/Runtime/Engine/Private/Materials/HLSLMaterialTranslator.cpp
--- a/Engine/Source/Runtime/Engine/Private/Materials/HLSLMaterialTranslator.cpp
+++ b/Engine/Source/Runtime/Engine/Private/Materials/HLSLMaterialTranslator.cpp
@@ -8159,7 +8159,7 @@ int32 FHLSLMaterialTranslator::SceneTextureLookup(int32 ViewportUV, uint32 InSce
 	SubstrateCtx.SubstrateMaterialComplexity.bIsSimple = false;
 	SubstrateCtx.SubstrateMaterialComplexity.bIsSingle = false;
 	
-	if (SceneTextureId == PPI_PostProcessInput0 && Material->GetMaterialDomain() == MD_PostProcess && Material->GetBlendableLocation() != BL_SceneColorAfterTonemapping)
+	if (SceneTextureId >= PPI_PostProcessInput0 && SceneTextureId <= PPI_PostProcessInput6 && Material->GetMaterialDomain() == MD_PostProcess && Material->GetBlendableLocation() != BL_SceneColorAfterTonemapping)
 	{
 		return AddInlinedCodeChunk(MCT_Float4, TEXT("(float4(View.OneOverPreExposure.xxx, 1) * %s)"), *CoerceParameter(LookUp, MCT_Float4));
 	}
diff --git a/Engine/Source/Runtime/Engine/Private/Materials/MaterialHLSLTree.cpp b/Engine/Source/Runtime/Engine/Private/Materials/MaterialHLSLTree.cpp
--- a/Engine/Source/Runtime/Engine/Private/Materials/MaterialHLSLTree.cpp
+++ b/Engine/Source/Runtime/Engine/Private/Materials/MaterialHLSLTree.cpp
@@ -2312,7 +2312,7 @@ void FExpressionSceneTexture::EmitValueShader(FEmitContext& Context, FEmitScope&
 		EmitLookup = Context.EmitExpression(Scope, Shader::EValueType::Float4, TEXT("MobileSceneTextureLookup(Parameters, %, %)"), (int)SceneTextureId, EmitTexCoord);
 	}
 
-	if (SceneTextureId == PPI_PostProcessInput0 && Context.Material->GetMaterialDomain() == MD_PostProcess && Context.Material->GetBlendableLocation() != BL_SceneColorAfterTonemapping)
+	if (SceneTextureId >= PPI_PostProcessInput0 && SceneTextureId <= PPI_PostProcessInput6 && Context.Material->GetMaterialDomain() == MD_PostProcess && Context.Material->GetBlendableLocation() != BL_SceneColorAfterTonemapping)
 	{
 		EmitLookup = Context.EmitExpression(Scope, Shader::EValueType::Float4, TEXT("(float4(View.OneOverPreExposure.xxx, 1) * %)"), EmitLookup);
 	}
