Fix Ray Tracing Far Field:
* INSTANCE_SCENE_DATA_FLAG_IS_RAYTRACING_FAR_FIELD is no longer set after recent changes to instance data
* Add that flag to primitive instead as PRIMITIVE_SCENE_DATA_FLAG_RAYTRACING_FAR_FIELD

diff --git a/Engine/Shaders/Private/RayTracing/RayTracingInstanceBufferUtil.usf b/Engine/Shaders/Private/RayTracing/RayTracingInstanceBufferUtil.usf
--- a/Engine/Shaders/Private/RayTracing/RayTracingInstanceBufferUtil.usf
+++ b/Engine/Shaders/Private/RayTracing/RayTracingInstanceBufferUtil.usf
@@ -57,10 +57,11 @@ void RayTracingBuildInstanceBufferCS(uint3 DispatchThreadId : SV_DispatchThreadI
 
 #if USE_GPUSCENE
 		FInstanceSceneData InstanceSceneData = GetInstanceSceneData(InputDesc.GPUSceneInstanceOrTransformIndex, InstanceSceneDataSOAStride);
+		FPrimitiveSceneData PrimitiveData = GetPrimitiveData(InstanceSceneData.PrimitiveId);
 
 		float4x4 LocalToTranslatedWorld4x4 = LWCMultiplyTranslation(InstanceSceneData.LocalToWorld, PreViewTranslation);
 
-		if (InstanceSceneData.Flags & INSTANCE_SCENE_DATA_FLAG_IS_RAYTRACING_FAR_FIELD)
+		if (PrimitiveData.Flags & PRIMITIVE_SCENE_DATA_FLAG_RAYTRACING_FAR_FIELD)
 		{
 			LocalToTranslatedWorld4x4[3].xyz += FarFieldReferencePos;
 		}
diff --git a/Engine/Shaders/Shared/SceneDefinitions.h b/Engine/Shaders/Shared/SceneDefinitions.h
--- a/Engine/Shaders/Shared/SceneDefinitions.h
+++ b/Engine/Shaders/Shared/SceneDefinitions.h
@@ -52,6 +52,7 @@
 #define PRIMITIVE_SCENE_DATA_FLAG_SPLINE_MESH							0x800000
 #define PRIMITIVE_SCENE_DATA_FLAG_INSTANCE_CULLING_OCCLUSION_QUERIES	0x1000000 // Whether per-instance occlusion queries are allowed when performing GPUScene instance culling
 #define PRIMITIVE_SCENE_DATA_FLAG_HAS_PIXEL_ANIMATION					0x2000000
+#define PRIMITIVE_SCENE_DATA_FLAG_RAYTRACING_FAR_FIELD					0x4000000
 
 
 // Primitive visibility flags
@@ -87,8 +88,7 @@
 #define INSTANCE_SCENE_DATA_FLAG_HAS_LOCAL_BOUNDS			0x40
 #define INSTANCE_SCENE_DATA_FLAG_HAS_PAYLOAD_EXTENSION		0x80
 #define INSTANCE_SCENE_DATA_FLAG_HAS_EDITOR_DATA			0x100
-#define INSTANCE_SCENE_DATA_FLAG_IS_RAYTRACING_FAR_FIELD	0x200
-#define INSTANCE_SCENE_DATA_FLAG_HIDDEN						0x400
+#define INSTANCE_SCENE_DATA_FLAG_HIDDEN						0x200
 
 #define INSTANCE_SCENE_DATA_FLAG_PAYLOAD_MASK (			\
 	INSTANCE_SCENE_DATA_FLAG_HAS_RANDOM					\
diff --git a/Engine/Source/Runtime/Engine/Private/PrimitiveSceneProxy.cpp b/Engine/Source/Runtime/Engine/Private/PrimitiveSceneProxy.cpp
--- a/Engine/Source/Runtime/Engine/Private/PrimitiveSceneProxy.cpp
+++ b/Engine/Source/Runtime/Engine/Private/PrimitiveSceneProxy.cpp
@@ -846,7 +846,8 @@ void FPrimitiveSceneProxy::BuildUniformShaderParameters(FPrimitiveUniformShaderP
 			.PrimitiveComponentId(GetPrimitiveComponentId().PrimIDValue)
 			.EditorColors(GetWireframeColor(), GetLevelColor())
 			.SplineMesh(IsSplineMesh())
-			.HasPixelAnimation(AnyMaterialHasPixelAnimation());
+			.HasPixelAnimation(AnyMaterialHasPixelAnimation())
+			.RayTracingFarField(IsRayTracingFarField());
 
 		if (PrimitiveSceneInfo != nullptr)
 		{
diff --git a/Engine/Source/Runtime/Engine/Public/PrimitiveUniformShaderParametersBuilder.h b/Engine/Source/Runtime/Engine/Public/PrimitiveUniformShaderParametersBuilder.h
--- a/Engine/Source/Runtime/Engine/Public/PrimitiveUniformShaderParametersBuilder.h
+++ b/Engine/Source/Runtime/Engine/Public/PrimitiveUniformShaderParametersBuilder.h
@@ -61,6 +61,7 @@ public:
 		bSplineMesh									= false;
 		bAllowInstanceCullingOcclusionQueries		= false;
 		bHasPixelAnimation                          = false;
+		bRayTracingFarField							= false;
 
 		Parameters.MaxWPOExtent						= 0.0f;
 		Parameters.MinMaterialDisplacement			= 0.0f;
@@ -129,6 +130,7 @@ public:
 	PRIMITIVE_UNIFORM_BUILDER_FLAG_METHOD(bool,			AllowInstanceCullingOcclusionQueries);
 	PRIMITIVE_UNIFORM_BUILDER_FLAG_METHOD(bool,			HasAlwaysEvaluateWPOMaterials);
 	PRIMITIVE_UNIFORM_BUILDER_FLAG_METHOD(bool,			HasPixelAnimation);
+	PRIMITIVE_UNIFORM_BUILDER_FLAG_METHOD(bool,			RayTracingFarField);
 
 	PRIMITIVE_UNIFORM_BUILDER_METHOD(uint32,			InstanceSceneDataOffset);
 	PRIMITIVE_UNIFORM_BUILDER_METHOD(uint32,			NumInstanceSceneDataEntries);
@@ -399,6 +401,7 @@ public:
 		Parameters.Flags |= bSplineMesh ? PRIMITIVE_SCENE_DATA_FLAG_SPLINE_MESH : 0u;
 		Parameters.Flags |= bAllowInstanceCullingOcclusionQueries ? PRIMITIVE_SCENE_DATA_FLAG_INSTANCE_CULLING_OCCLUSION_QUERIES: 0u;
 		Parameters.Flags |= bHasPixelAnimation ? PRIMITIVE_SCENE_DATA_FLAG_HAS_PIXEL_ANIMATION : 0u;
+		Parameters.Flags |= bRayTracingFarField ? PRIMITIVE_SCENE_DATA_FLAG_RAYTRACING_FAR_FIELD : 0u;
 		
 		Parameters.VisibilityFlags = 0;
 		Parameters.VisibilityFlags |= bCastHiddenShadow ? PRIMITIVE_VISIBILITY_FLAG_CAST_HIDDEN_SHADOW : 0u;
@@ -463,4 +466,5 @@ private:
 	uint32 bSplineMesh : 1;
 	uint32 bAllowInstanceCullingOcclusionQueries : 1;
 	uint32 bHasPixelAnimation : 1;
+	uint32 bRayTracingFarField : 1;
 };
diff --git a/Engine/Source/Runtime/Renderer/Private/GPUScene.cpp b/Engine/Source/Runtime/Renderer/Private/GPUScene.cpp
--- a/Engine/Source/Runtime/Renderer/Private/GPUScene.cpp
+++ b/Engine/Source/Runtime/Renderer/Private/GPUScene.cpp
@@ -461,8 +461,6 @@ struct FUploadDataSourceAdapterScenePrimitives
 	#if WITH_EDITOR
 		PackedFlags |= Flags.bHasPerInstanceEditorData  ? INSTANCE_SCENE_DATA_FLAG_HAS_EDITOR_DATA         : 0u;
 	#endif
-		// TODO: what is this one for/from?
-		// Flags |= IsRayTracingFarField()          ? INSTANCE_SCENE_DATA_FLAG_IS_RAYTRACING_FAR_FIELD : 0u;
 
 		return PackedFlags;
 	}
