Fix missing InstanceCustomData

The new optimized logic to skip work in vertex shaders was added to GetMaterialEnvironmentOld but not GetMaterialEnvironment.

diff --git a/Engine/Source/Runtime/Engine/Private/Materials/HLSLMaterialTranslator.cpp b/Engine/Source/Runtime/Engine/Private/Materials/HLSLMaterialTranslator.cpp
--- a/Engine/Source/Runtime/Engine/Private/Materials/HLSLMaterialTranslator.cpp
+++ b/Engine/Source/Runtime/Engine/Private/Materials/HLSLMaterialTranslator.cpp
@@ -253,6 +253,7 @@ struct FHLSLMaterialTranslator::FEnvironmentDefines
 	bool bUsesEyeAdaptation;
 	bool bVirtualTextureOutput;
 	bool bUsesPerInstanceCustomData;
+	bool bUsesPerInstanceCustomDataPS;
 	bool bUsesPerInstanceRandomPS;
 	bool bUsesPerInstanceFadeAmount;
 	bool bUsesVertexInterpolator;
@@ -2403,6 +2404,9 @@ void FHLSLMaterialTranslator::GetMaterialEnvironment(EShaderPlatform InPlatform,
 	}
 
 	OutEnvironment.SetDefine(TEXT("USES_PER_INSTANCE_CUSTOM_DATA"), EnvironmentDefines->bUsesPerInstanceCustomData);
+	// Whether we need to pass custom data paramaters to PS through interpolators
+	OutEnvironment.SetDefine(TEXT("USES_PER_INSTANCE_CUSTOM_DATA_PS"), EnvironmentDefines->bUsesPerInstanceCustomDataPS && Material->IsUsedWithInstancedStaticMeshes());
+
 	OutEnvironment.SetDefine(TEXT("USES_PER_INSTANCE_FADE_AMOUNT"), EnvironmentDefines->bUsesPerInstanceFadeAmount);
 	OutEnvironment.SetDefine(TEXT("USES_VERTEX_INTERPOLATOR"), EnvironmentDefines->bUsesVertexInterpolator);
 
@@ -15015,6 +15019,7 @@ void FHLSLMaterialTranslator::PrepareEnvironmentDefines()
 	EnvironmentDefines->bUsesEyeAdaptation = MaterialCompilationOutput.bUsesEyeAdaptation;
 	EnvironmentDefines->bVirtualTextureOutput = MaterialCompilationOutput.bHasRuntimeVirtualTextureOutputNode;
 	EnvironmentDefines->bUsesPerInstanceCustomData = MaterialCompilationOutput.bUsesPerInstanceCustomData && Material->IsUsedWithInstancedStaticMeshes();
+	EnvironmentDefines->bUsesPerInstanceCustomDataPS = bUsesPerInstanceCustomDataPS && Material->IsUsedWithInstancedStaticMeshes();
 	EnvironmentDefines->bUsesPerInstanceRandomPS = bUsesPerInstanceRandomPS && Material->IsUsedWithInstancedStaticMeshes();
 	EnvironmentDefines->bUsesPerInstanceFadeAmount = bUsesPerInstanceFadeAmount && Material->IsUsedWithInstancedStaticMeshes();
 	EnvironmentDefines->bUsesVertexInterpolator = MaterialCompilationOutput.bUsesVertexInterpolator;
