Added an explicit string -> invariant text node for BP

This is the same logic as the implicit string -> text conversion in BP (as that also produces invariant text), but can make the intent of the invariant result more explicit for things that are validating text setters.

diff --git a/Engine/Source/Runtime/Engine/Classes/Kismet/KismetTextLibrary.h b/Engine/Source/Runtime/Engine/Classes/Kismet/KismetTextLibrary.h
--- a/Engine/Source/Runtime/Engine/Classes/Kismet/KismetTextLibrary.h
+++ b/Engine/Source/Runtime/Engine/Classes/Kismet/KismetTextLibrary.h
@@ -117,7 +117,7 @@ class UKismetTextLibrary : public UBlueprintFunctionLibrary
 	UFUNCTION(BlueprintPure, meta = (DisplayName = "To String (Text)", CompactNodeTitle = "->", BlueprintAutocast), Category = "Utilities|String")
 	static ENGINE_API FString Conv_TextToString(const FText& InText);
 
-	/** Converts string to culture invariant text. Use Format or Make Literal Text to create localizable text */
+	/** Converts string to culture invariant text. Use 'Make Literal Text' to create localizable text, or 'Format' if concatenating localized text */
 	UFUNCTION(BlueprintPure, meta=(DisplayName = "To Text (String)", BlueprintAutocast), Category="Utilities|Text")
 	static ENGINE_API FText Conv_StringToText(const FString& InString);
 
@@ -125,6 +125,10 @@ class UKismetTextLibrary : public UBlueprintFunctionLibrary
 	UFUNCTION(BlueprintPure, meta=(DisplayName = "To Text (Name)", BlueprintAutocast), Category="Utilities|Text")
 	static ENGINE_API FText Conv_NameToText(FName InName);
 
+	/** Converts string to culture invariant text. Use 'Make Literal Text' to create localizable text, or 'Format' if concatenating localized text */
+	UFUNCTION(BlueprintPure, Category="Utilities|Text")
+	static ENGINE_API FText MakeInvariantText(const FString& InString);
+
 	/** Returns true if text is empty. */
 	UFUNCTION(BlueprintPure, Category="Utilities|Text")
 	static ENGINE_API bool TextIsEmpty(const FText& InText);
diff --git a/Engine/Source/Runtime/Engine/Private/KismetTextLibrary.cpp b/Engine/Source/Runtime/Engine/Private/KismetTextLibrary.cpp
--- a/Engine/Source/Runtime/Engine/Private/KismetTextLibrary.cpp
+++ b/Engine/Source/Runtime/Engine/Private/KismetTextLibrary.cpp
@@ -94,6 +94,11 @@ FText UKismetTextLibrary::Conv_NameToText(FName InName)
 	return FText::AsCultureInvariant(InName.ToString());
 }
 
+FText UKismetTextLibrary::MakeInvariantText(const FString& InString)
+{
+	return FText::AsCultureInvariant(InString);
+}
+
 bool UKismetTextLibrary::TextIsEmpty(const FText& InText)
 {
 	return InText.IsEmpty();
