Texture HasAlpha nop just comments

#rnx

diff --git a/Engine/Source/Editor/TextureEditor/Private/TextureEditorToolkit.cpp b/Engine/Source/Editor/TextureEditor/Private/TextureEditorToolkit.cpp
--- a/Engine/Source/Editor/TextureEditor/Private/TextureEditorToolkit.cpp
+++ b/Engine/Source/Editor/TextureEditor/Private/TextureEditorToolkit.cpp
@@ -880,6 +880,8 @@ void FTextureEditorToolkit::PopulateQuickInfo( )
 		FormatText->SetText(FText::Format(NSLOCTEXT("TextureEditor", "QuickInfo_Format", "Format: {0}"), FText::FromString(GPixelFormats[(uint8)TextureFormat].Name)));
 	}
 
+	// This "Has Alpha Channel" is whether the GPU format can represent alpha in the format (eg. is it DXT1 vs DXT5)
+	//	it does not tell you if the texture actually has non-opaque alpha
 	EPixelFormatChannelFlags ValidTextureChannels = GetPixelFormatValidChannels(TextureFormat);
 	HasAlphaChannelText->SetText(FText::Format(NSLOCTEXT("TextureEditor", "QuickInfo_HasAlphaChannel", "Has Alpha Channel: {0}"),
 		EnumHasAnyFlags(ValidTextureChannels, EPixelFormatChannelFlags::A) ? NSLOCTEXT("TextureEditor", "True", "True") : NSLOCTEXT("TextureEditor", "False", "False")));
diff --git a/Engine/Source/Runtime/Engine/Private/Texture2D.cpp b/Engine/Source/Runtime/Engine/Private/Texture2D.cpp
--- a/Engine/Source/Runtime/Engine/Private/Texture2D.cpp
+++ b/Engine/Source/Runtime/Engine/Private/Texture2D.cpp
@@ -696,6 +696,9 @@ void UTexture2D::GetAssetRegistryTags(TArray<FAssetRegistryTag>& OutTags) const
 
 	const FString DimensionsStr = FString::Printf(TEXT("%dx%d"), SourceSize.X, SourceSize.Y);
 	OutTags.Add( FAssetRegistryTag("Dimensions", DimensionsStr, FAssetRegistryTag::TT_Dimensional) );
+	
+	// This "Has Alpha Channel" is whether the GPU format can represent alpha in the format (eg. is it DXT1 vs DXT5)
+	//	it does not tell you if the texture actually has non-opaque alpha
 	OutTags.Add( FAssetRegistryTag("HasAlphaChannel", HasAlphaChannel() ? TEXT("True") : TEXT("False"), FAssetRegistryTag::TT_Alphabetical) );
 	OutTags.Add( FAssetRegistryTag("Format", GPixelFormats[GetPixelFormat()].Name, FAssetRegistryTag::TT_Alphabetical) );
 
@@ -950,6 +953,9 @@ bool UTexture2D::HasSameSourceArt(UTexture2D* InTexture)
 
 bool UTexture2D::HasAlphaChannel() const
 {
+	// This "Has Alpha Channel" is whether the GPU format can represent alpha in the format (eg. is it DXT1 vs DXT5)
+	//	it does not tell you if the texture actually has non-opaque alpha
+
 	if (PrivatePlatformData)
 	{
 #if WITH_EDITOR
