C4702 warnings

diff --git a/Engine/Plugins/Animation/DeformerGraph/Source/OptimusCore/Private/Nodes/OptimusNode_ComputeKernelBase.cpp b/Engine/Plugins/Animation/DeformerGraph/Source/OptimusCore/Private/Nodes/OptimusNode_ComputeKernelBase.cpp
--- a/Engine/Plugins/Animation/DeformerGraph/Source/OptimusCore/Private/Nodes/OptimusNode_ComputeKernelBase.cpp
+++ b/Engine/Plugins/Animation/DeformerGraph/Source/OptimusCore/Private/Nodes/OptimusNode_ComputeKernelBase.cpp
@@ -31,7 +31,7 @@ static FString GetShaderParamPinValueString(
 	)
 {
 	return InPin->GetDataType()->ShaderValueType->GetZeroValueAsString();
-	
+#if 0
 	// FIXME: Need property storage.
 	const FShaderValueType& ValueType = *InPin->GetDataType()->ShaderValueType;
 	TArrayView<UOptimusNodePin* const> SubPins = InPin->GetSubPins();
@@ -43,6 +43,7 @@ static FString GetShaderParamPinValueString(
 	
 	// FIXME: Support all types properly. Should probably be moved to a better place.
 	return FString::Printf(TEXT("%s(%s)"), *ValueType.ToString(), *InPin->GetValueAsString());
+#endif
 }
 
 
diff --git a/Engine/Plugins/Experimental/GeometryScripting/Source/GeometryScriptingCore/Private/MeshAssetFunctions.cpp b/Engine/Plugins/Experimental/GeometryScripting/Source/GeometryScriptingCore/Private/MeshAssetFunctions.cpp
--- a/Engine/Plugins/Experimental/GeometryScripting/Source/GeometryScriptingCore/Private/MeshAssetFunctions.cpp
+++ b/Engine/Plugins/Experimental/GeometryScripting/Source/GeometryScriptingCore/Private/MeshAssetFunctions.cpp
@@ -241,9 +241,6 @@ UDynamicMesh*  UGeometryScriptLibrary_StaticMeshFunctions::CopyMeshFromStaticMes
 #else
 	return CopyMeshFromStaticMesh_RenderData(FromStaticMeshAsset, ToDynamicMesh, AssetOptions, RequestedLOD, Outcome, Debug);	
 #endif
-
-
-	return ToDynamicMesh;
 }
 
 
diff --git a/Engine/Source/Runtime/Core/Public/HAL/IPlatformFileManagedStorageWrapper.h b/Engine/Source/Runtime/Core/Public/HAL/IPlatformFileManagedStorageWrapper.h
--- a/Engine/Source/Runtime/Core/Public/HAL/IPlatformFileManagedStorageWrapper.h
+++ b/Engine/Source/Runtime/Core/Public/HAL/IPlatformFileManagedStorageWrapper.h
@@ -489,7 +489,6 @@ public:
 		}
 
 		return Categories.GetCategories()[File.Category].RemoveFile(File.FullFilename);
-		File.Clear();
 	}
 
 	int64 GetTotalUsedSize() const
diff --git a/Engine/Source/Runtime/Engine/Private/UnrealEngine.cpp b/Engine/Source/Runtime/Engine/Private/UnrealEngine.cpp
--- a/Engine/Source/Runtime/Engine/Private/UnrealEngine.cpp
+++ b/Engine/Source/Runtime/Engine/Private/UnrealEngine.cpp
@@ -9661,7 +9661,6 @@ bool UEngine::HandleObjCommand( const TCHAR* Cmd, FOutputDevice& Ar )
 		// OBJ command but not supported here
 		return false;
 	}
-	return false;
 }
 
 bool UEngine::HandleDirCommand( const TCHAR* Cmd, FOutputDevice& Ar )
@@ -10615,6 +10614,8 @@ bool UEngine::PerformError(const TCHAR* Cmd, FOutputDevice& Ar)
 		ENQUEUE_RENDER_COMMAND(CauseGpuCrash)(&FRender::GpuCrash);
 		return true;
 	}
+#pragma warning(push)
+#pragma warning(disable:4702)
 	else if (FParse::Command(&Cmd, TEXT("TERMINATE")))
 	{
 		UE_LOG(LogEngine, Warning, TEXT("Printed warning to log."));
@@ -10631,6 +10632,7 @@ bool UEngine::PerformError(const TCHAR* Cmd, FOutputDevice& Ar)
 		std::abort();
 		return true;
 	}
+#pragma warning(pop)
 	else if (FParse::Command(&Cmd, TEXT("CHECK")))
 	{
 		UE_LOG(LogEngine, Warning, TEXT("Printed warning to log."));
@@ -10963,6 +10965,8 @@ bool UEngine::PerformError(const TCHAR* Cmd, FOutputDevice& Ar)
 		);
 		return true;
 	}
+#pragma warning(push)
+#pragma warning(disable:4702)
 	else if (FParse::Command(&Cmd, TEXT("SOFTLOCK")))
 	{
 		Ar.Log(TEXT("Hanging the current thread"));
@@ -10982,6 +10986,7 @@ bool UEngine::PerformError(const TCHAR* Cmd, FOutputDevice& Ar)
 		}
 		return true;
 	}
+#pragma warning(pop)
 	else if (FParse::Command(&Cmd, TEXT("SLEEP")))
 	{
 		Ar.Log(TEXT("Sleep for 1 hour. This should crash after a few seconds in cooked builds."));
diff --git a/Engine/Source/Runtime/Engine/Private/UserInterface/PlayerInput.cpp b/Engine/Source/Runtime/Engine/Private/UserInterface/PlayerInput.cpp
--- a/Engine/Source/Runtime/Engine/Private/UserInterface/PlayerInput.cpp
+++ b/Engine/Source/Runtime/Engine/Private/UserInterface/PlayerInput.cpp
@@ -381,8 +381,6 @@ bool UPlayerInput::InputKey(const FInputKeyParams& Params)
 
 		return true;
 	}
-	
-	return false;
 }
 
 bool UPlayerInput::InputAxis(FKey Key, float Delta, float DeltaTime, int32 NumSamples, bool bGamepad)
