Fix per-platform enable/disable of RVT volume for editor preview.
Previous code always used windows as current platform.

diff --git a/Engine/Source/Runtime/Engine/Private/Components/RuntimeVirtualTextureComponent.cpp b/Engine/Source/Runtime/Engine/Private/Components/RuntimeVirtualTextureComponent.cpp
--- a/Engine/Source/Runtime/Engine/Private/Components/RuntimeVirtualTextureComponent.cpp
+++ b/Engine/Source/Runtime/Engine/Private/Components/RuntimeVirtualTextureComponent.cpp
@@ -8,8 +8,6 @@
 #include "Engine/World.h"
 #include "Logging/MessageLog.h"
 #include "GameFramework/Actor.h"
-#include "Interfaces/ITargetPlatform.h"
-#include "Interfaces/ITargetPlatformManagerModule.h"
 #include "Misc/UObjectToken.h"
 #include "Misc/MapErrors.h"
 #include "SceneInterface.h"
@@ -160,17 +158,10 @@ bool URuntimeVirtualTextureComponent::IsEnabledInScene() const
 	{
 		if (World->WorldType == EWorldType::Game || World->WorldType == EWorldType::PIE)
 		{
-#if WITH_EDITOR
-			if (!EnableInGamePerPlatform.GetValueForPlatform(*GetTargetPlatformManagerRef().GetRunningTargetPlatform()->IniPlatformName()))
-			{
-				return false;
-			}
-#else
 			if (!EnableInGamePerPlatform.GetValue())
 			{
 				return false;
 			}
-#endif
 		}
 	}
 
