Modified FNaniteSplineMeshSceneProxy to respect r.RayTracing.Geometry.SplineMeshes.

diff --git a/Engine/Source/Runtime/Engine/Private/SplineMeshSceneProxy.cpp b/Engine/Source/Runtime/Engine/Private/SplineMeshSceneProxy.cpp
--- a/Engine/Source/Runtime/Engine/Private/SplineMeshSceneProxy.cpp
+++ b/Engine/Source/Runtime/Engine/Private/SplineMeshSceneProxy.cpp
@@ -369,6 +369,16 @@ ERayTracingPrimitiveFlags FNaniteSplineMeshSceneProxy::GetCachedRayTracingInstan
 	return FPrimitiveSceneProxy::GetCachedRayTracingInstance(OutRayTracingInstance);
 }
 
+void FNaniteSplineMeshSceneProxy::GetDynamicRayTracingInstances(struct FRayTracingMaterialGatheringContext& Context, TArray<FRayTracingInstance>& OutRayTracingInstances)
+{
+	if (CVarRayTracingSplineMeshes.GetValueOnRenderThread() == 0)
+	{
+		return;
+	}
+
+	return Nanite::FSceneProxy::GetDynamicRayTracingInstances(Context, OutRayTracingInstances);
+}
+
 void FNaniteSplineMeshSceneProxy::SetupRayTracingMaterials(int32 LODIndex, TArray<FMeshBatch>& Materials, bool bUseNaniteVertexFactory) const
 {
 	Nanite::FSceneProxy::SetupRayTracingMaterials(LODIndex, Materials, bUseNaniteVertexFactory);
diff --git a/Engine/Source/Runtime/Engine/Public/SplineMeshSceneProxy.h b/Engine/Source/Runtime/Engine/Public/SplineMeshSceneProxy.h
--- a/Engine/Source/Runtime/Engine/Public/SplineMeshSceneProxy.h
+++ b/Engine/Source/Runtime/Engine/Public/SplineMeshSceneProxy.h
@@ -206,6 +206,7 @@ public:
 #if RHI_RAYTRACING
 	virtual bool IsRayTracingStaticRelevant() const override { return false; }
 	virtual ERayTracingPrimitiveFlags GetCachedRayTracingInstance(FRayTracingInstance& OutRayTracingInstance) override;
+	virtual void GetDynamicRayTracingInstances(FRayTracingMaterialGatheringContext& Context, TArray<FRayTracingInstance>& OutRayTracingInstances) override;
 	virtual void SetupRayTracingMaterials(int32 LODIndex, TArray<FMeshBatch>& Materials, bool bUseNaniteVertexFactory) const override;
 #endif
 };
