UEWAV style cleanup pass.

diff --git a/Engine/Source/Runtime/Engine/Classes/Sound/SoundWave.h b/Engine/Source/Runtime/Engine/Classes/Sound/SoundWave.h
--- a/Engine/Source/Runtime/Engine/Classes/Sound/SoundWave.h
+++ b/Engine/Source/Runtime/Engine/Classes/Sound/SoundWave.h
@@ -904,6 +904,11 @@ public:
 	struct FEditorAudioBulkData
 	{
 		UE::Serialization::FEditorBulkData RawData;
+
+		// The container soundwave for this raw data. This must be non-null for any non-metadata instances.
+		// We need this in order to parse the multichannel layout of the raw data - and we also use this 
+		// in place of the BulkData Owner parameter in many places because many call sites pass null incorrectly,
+		// resulting in the bulk data not being correlated with our asset.
 		USoundWave* SoundWave;
 
 		FEditorAudioBulkData()
@@ -927,6 +932,7 @@ public:
 		// Deprecated unused API forwarding for potential backwards compatability issues.
 		// As the raw data needs to be converted before use or storage, always access it via the above functions.
 		//
+#pragma region Deprecated Pass Thru
 		UE_DEPRECATED(5.4, "CreateLegacyUniqueIdentifier is provided just for API backwards compatibility.")
 		void CreateLegacyUniqueIdentifier(UObject* Owner)
 		{
@@ -1042,6 +1048,7 @@ public:
 		{
 			RawData.UpdateRegistrationOwner(Owner);
 		}
+#pragma endregion
 	} RawData;
 
 	/** Waveform edits to be applied to this SoundWave on cook (editing transformations will trigger a cook) */
diff --git a/Engine/Source/Runtime/Engine/Private/SoundWave.cpp b/Engine/Source/Runtime/Engine/Private/SoundWave.cpp
--- a/Engine/Source/Runtime/Engine/Private/SoundWave.cpp
+++ b/Engine/Source/Runtime/Engine/Private/SoundWave.cpp
@@ -4852,125 +4852,184 @@ bool USoundWave::HasError() const
 #if WITH_EDITORONLY_DATA
 TFuture<FSharedBuffer> USoundWave::FEditorAudioBulkData::GetPayload() const
 {
-	TFuture<FSharedBuffer> BufferFuture = RawData.GetPayload();
-	const uint8* Data = (const uint8*) BufferFuture.Get().GetData(); // Will block.
-	uint64 DataSize = BufferFuture.Get().GetSize();
-	FSharedBuffer Buffer = FSharedBuffer::Clone(Data, DataSize);
-	Data = (const uint8*) Buffer.GetData();
-	DataSize = Buffer.GetSize();
+	if (SoundWave == nullptr)
+	{
+		// Sanity - should never happen! Just forward the data.
+		UE_LOG(LogAudio, Error, TEXT("SoundWave pointer is null in FEditorAudioBulkData!"));
+		return RawData.GetPayload();
+	}
+
+	// Per discussions - while the Payload system was designed to return async stuff, it's practically
+	// always in a dedicated task and immediately blocked on, so this isn't as bad as it looks. If we wanted
+	// to make this properly async we'd just need to grab a copy of the channel offsets in case the containing
+	// sound wave wanders off and then drop all of this is a task.
+
+	// Here's our blockage on a potential rehydration.
+	FSharedBuffer BulkDataBuffer = RawData.GetPayload().Get();
+
+	// We need it mutable. Unless the bulk data system holds a copy this should be a steal.
+	FUniqueBuffer MutablePayloadBuffer = BulkDataBuffer.MoveToUnique();
+	FMutableMemoryView MutablePayload(MutablePayloadBuffer);
+
 	FWaveModInfo WaveInfo;
-	if (SoundWave && SoundWave->ChannelOffsets.Num() > 0)
+	if (SoundWave->ChannelOffsets.Num() > 0)
 	{
-		TArray<int16> scratch_buffer;
-		for (int i = 0; i < SoundWave->ChannelOffsets.Num(); ++i)
+		TArray<int16> Scratch;
+		for (int32 i = 0; i < SoundWave->ChannelOffsets.Num(); ++i)
 		{
 			// If the channel doesn't exist, skip it.
 			if (!SoundWave->ChannelSizes[i])
 			{
 				continue;
 			}
-			if (!WaveInfo.ReadWaveInfo(Data + SoundWave->ChannelOffsets[i], SoundWave->ChannelSizes[i]))
+
+			FMutableMemoryView ChannelPayload = MutablePayload.Mid(SoundWave->ChannelOffsets[i], SoundWave->ChannelSizes[i]);
+			if (ChannelPayload.GetSize() != SoundWave->ChannelSizes[i])
+			{
+				UE_LOG(LogAudio, Error, TEXT("Sound wave payload in asset %s has inconsitent multi channel size for channel %d: ChannelSizes = %d, Available: %lld. Attempting to continue..."),
+					*SoundWave->GetName(), i, SoundWave->ChannelSizes[i], ChannelPayload.GetSize());
+
+				// On import we fail here, not sure what to do other than try and continue. Presumably this will die
+				// during the derived data build or the wave parse below.
+			}
+
+			if (!WaveInfo.ReadWaveInfo((const uint8*)ChannelPayload.GetData(), ChannelPayload.GetSize()))
 			{
-				UE_LOG(LogAudio, Warning, TEXT("Failed to read wave data out of '%s'."), *SoundWave->GetFullName());
+				UE_LOG(LogAudio, Error, TEXT("Sound wave payload in asset %s failed to parse channel %d as a wave!"), *SoundWave->GetName(), i);
 			}
-			if (*WaveInfo.pFormatTag == FWaveModInfo::WAVE_INFO_FORMAT_OODLE_WAVE)
+			else if (*WaveInfo.pFormatTag == FWaveModInfo::WAVE_INFO_FORMAT_OODLE_WAVE)
 			{
 				// Convert UEWavComp data back to a WAV file
 				*WaveInfo.pFormatTag = FWaveModInfo::WAVE_INFO_FORMAT_PCM;
-				int16* samples = (int16*)WaveInfo.SampleDataStart;
-				int64 num_samples = WaveInfo.GetNumSamples();
-				int64 num_channels = *WaveInfo.pChannels;
-				scratch_buffer.Reset();
-				scratch_buffer.AddUninitialized(num_samples);
-				uewav_decode16(samples, scratch_buffer.GetData(), num_samples, num_channels);
+				int64 PayloadSampleCount = WaveInfo.GetNumSamples();
+				int64 PayloadChannelCount = *WaveInfo.pChannels;
+				Scratch.SetNumUninitialized(PayloadSampleCount);
+				uewav_decode16((int16*)WaveInfo.SampleDataStart, Scratch.GetData(), PayloadSampleCount, PayloadChannelCount);
 			}
+			// else - it wasn't encoded as uewav so leave it alone.
 		}
 	}
 	else
 	{
-		if (!WaveInfo.ReadWaveInfo(Data, DataSize)) 
+		if (!WaveInfo.ReadWaveInfo((const uint8*)MutablePayload.GetData(), MutablePayload.GetSize())) 
 		{
-			UE_LOG(LogAudio, Warning, TEXT("Failed to read wave data."));
+			UE_LOG(LogAudio, Error, TEXT("Sound wave payload in asset %s failed to parse as a wave!"), *SoundWave->GetName());
 		}
-		if (*WaveInfo.pFormatTag == FWaveModInfo::WAVE_INFO_FORMAT_OODLE_WAVE)
+		else if (*WaveInfo.pFormatTag == FWaveModInfo::WAVE_INFO_FORMAT_OODLE_WAVE)
 		{
 			// Convert UEWavComp data back to a WAV file
 			*WaveInfo.pFormatTag = FWaveModInfo::WAVE_INFO_FORMAT_PCM;
-			int16* samples = (int16*)WaveInfo.SampleDataStart;
-			int64 num_samples = WaveInfo.GetNumSamples();
-			int64 num_channels = *WaveInfo.pChannels;
-			TArray<int16> scratch_buffer;
-			scratch_buffer.AddUninitialized(num_samples);
-			uewav_decode16(samples, scratch_buffer.GetData(), num_samples, num_channels);
+			int64 PayloadSampleCount = WaveInfo.GetNumSamples();
+			int64 PayloadChannelCount = *WaveInfo.pChannels;
+			TArray<int16> Scratch;
+			Scratch.AddUninitialized(PayloadSampleCount);
+			uewav_decode16((int16*)WaveInfo.SampleDataStart, Scratch.GetData(), PayloadSampleCount, PayloadChannelCount);
 		}
+		// else - not uewav encoded so leave it alone.
 	}
+
 	TPromise<FSharedBuffer> promise;
-	promise.EmplaceValue(Buffer);
+	promise.EmplaceValue(MutablePayloadBuffer.MoveToShared());
 	return promise.GetFuture();
 }
 
-void USoundWave::FEditorAudioBulkData::UpdatePayload(FSharedBuffer InPayload, UObject* Owner)
+void USoundWave::FEditorAudioBulkData::UpdatePayload(FSharedBuffer InPayload, UObject* /*Owner - see comments in header declaration*/)
 {
-	const uint8* Data = (const uint8*) InPayload.GetData(); // Will block.
-	uint64 DataSize = InPayload.GetSize();
-	FSharedBuffer Buffer = FSharedBuffer::Clone(Data, DataSize);
-	Data = (const uint8*) Buffer.GetData();
-	DataSize = Buffer.GetSize();
+	if (SoundWave == nullptr)
+	{
+		// Sanity - should never happen! Just forward the data.
+		UE_LOG(LogAudio, Error, TEXT("SoundWave pointer is null in UpdatePayload!"));
+		return RawData.UpdatePayload(InPayload, nullptr);
+	}
+
 	bool bEnableUEWavComp = false;
 	GConfig->GetBool(TEXT("AudioImporter"), TEXT("EnableUEWavComp"), bEnableUEWavComp, GEditorIni);
-	if (SoundWave && SoundWave->ChannelOffsets.Num() > 0) 
+	if (!bEnableUEWavComp)
 	{
-		TArray<int16> scratch_buffer;
-		for (int i = 0; i < SoundWave->ChannelOffsets.Num(); ++i)
+		// Passthru if we are disabled.
+		return RawData.UpdatePayload(InPayload, SoundWave);
+	}
+
+	// We modify in place so we have to move (or hopefully steal) to a mutable copy.
+	FUniqueBuffer Buffer = InPayload.MoveToUnique();
+	FMutableMemoryView MutablePayload(Buffer);
+
+	if (SoundWave->ChannelOffsets.Num() > 0) 
+	{
+		// Multichannel - N mono RIFF files packed back to back.
+		TArray<int16> Scratch;
+		for (int32 i = 0; i < SoundWave->ChannelOffsets.Num(); ++i)
 		{
 			// If the channel doesn't exist, skip it.
 			if (!SoundWave->ChannelSizes[i])
 			{
 				continue;
 			}
+
+			FMutableMemoryView ChannelPayload = MutablePayload.Mid(SoundWave->ChannelOffsets[i], SoundWave->ChannelSizes[i]);
+			if (ChannelPayload.GetSize() != SoundWave->ChannelSizes[i])
+			{
+				UE_LOG(LogAudio, Error, TEXT("New sound wave payload in asset %s has inconsistent multi channel size for channel %d: ChannelSizes = %d, Available: %lld, payload not updated."),
+					*SoundWave->GetName(), i, SoundWave->ChannelSizes[i], ChannelPayload.GetSize());	
+				return;
+			}
+
 			FWaveModInfo WaveInfo;
-			if (!WaveInfo.ReadWaveInfo(Data + SoundWave->ChannelOffsets[i], SoundWave->ChannelSizes[i]))
+			if (!WaveInfo.ReadWaveInfo((const uint8*)ChannelPayload.GetData(), ChannelPayload.GetSize()))
 			{
-				UE_LOG(LogAudio, Warning, TEXT("Failed to read wave data out of '%s'."), *SoundWave->GetFullName());
+				UE_LOG(LogAudio, Error, TEXT("New sound wave payload failed to parse channel %d as a wave in asset %s! Payload not updated!"), i, *SoundWave->GetFullName());
+				return;
 			}
+
 			if (*WaveInfo.pFormatTag == FWaveModInfo::WAVE_INFO_FORMAT_PCM)
 			{
-				if (bEnableUEWavComp)
-				{
-					int16* samples = (int16*)WaveInfo.SampleDataStart;
-					int64 num_samples = WaveInfo.GetNumSamples();
-					int64 num_channels = *WaveInfo.pChannels;
-					scratch_buffer.Reset();
-					scratch_buffer.AddUninitialized(num_samples);
-					uewav_encode16(samples, scratch_buffer.GetData(), num_samples, num_channels);
-					*WaveInfo.pFormatTag = FWaveModInfo::WAVE_INFO_FORMAT_OODLE_WAVE;
-				}
+				int64 PayloadSampleCount = WaveInfo.GetNumSamples();
+				int64 PayloadChannelCount = *WaveInfo.pChannels;
+				Scratch.SetNumUninitialized(PayloadSampleCount);
+
+				// Samples is modified in place.
+				uewav_encode16((int16*)WaveInfo.SampleDataStart, Scratch.GetData(), PayloadSampleCount, PayloadChannelCount);
+
+				// Mark in the wave format that we are OODLE so we know to decode.
+				*WaveInfo.pFormatTag = FWaveModInfo::WAVE_INFO_FORMAT_OODLE_WAVE;
+			}
+			else
+			{
+				UE_LOG(LogAudio, Error, TEXT("New sound wave payload isn't PCM format for channel %d in asset %s (reported = %d) - payload not updated!"), i, *SoundWave->GetFullName(), *WaveInfo.pFormatTag);
+				return;
 			}
 		}
 	} 
 	else
 	{
 		FWaveModInfo WaveInfo;
-		if (!WaveInfo.ReadWaveInfo(Data, DataSize, 0)) 
+		if (!WaveInfo.ReadWaveInfo((const uint8*)MutablePayload.GetData(), MutablePayload.GetSize(), 0)) 
 		{
-			UE_LOG(LogAudio, Warning, TEXT("Failed to read wave data."));
-		}
+			UE_LOG(LogAudio, Error, TEXT("USoundWave new payload failed to parse as a wave in asset %s! Payload not updated!"), *SoundWave->GetName());
+			return;
+		}		
+		
 		if (*WaveInfo.pFormatTag == FWaveModInfo::WAVE_INFO_FORMAT_PCM)
 		{
-			if (bEnableUEWavComp)
-			{
-				int16* samples = (int16*)WaveInfo.SampleDataStart;
-				int64 num_samples = WaveInfo.GetNumSamples();
-				int64 num_channels = *WaveInfo.pChannels;
-				TArray<int16> scratch_buffer;
-				scratch_buffer.AddUninitialized(num_samples);
-				uewav_encode16(samples, scratch_buffer.GetData(), num_samples, num_channels);
-				*WaveInfo.pFormatTag = FWaveModInfo::WAVE_INFO_FORMAT_OODLE_WAVE;
-			}
+			int64 PayloadSampleCount = WaveInfo.GetNumSamples();
+			int64 PayloadChannelCount = *WaveInfo.pChannels;
+			TArray<int16> Scratch;
+			Scratch.AddUninitialized(PayloadSampleCount);
+
+			// Samples is modified in place.
+			uewav_encode16((int16*)WaveInfo.SampleDataStart, Scratch.GetData(), PayloadSampleCount, PayloadChannelCount);
+				
+			// Mark in the wave format that we are OODLE so we know to decode.
+			*WaveInfo.pFormatTag = FWaveModInfo::WAVE_INFO_FORMAT_OODLE_WAVE;
+		}
+		else
+		{
+			UE_LOG(LogAudio, Error, TEXT("New sound wave payload isn't PCM format in asset %s (reported = %d) - payload not updated!"), *SoundWave->GetFullName(), *WaveInfo.pFormatTag);
+			return;
 		}
 	}
 
-	RawData.UpdatePayload(Buffer, Owner);
+	RawData.UpdatePayload(Buffer.MoveToShared(), SoundWave);
 }
 
 bool USoundWave::FEditorAudioBulkData::HasPayloadData() const
@@ -4978,14 +5037,14 @@ bool USoundWave::FEditorAudioBulkData::HasPayloadData() const
 	return RawData.HasPayloadData();
 }
 
-void USoundWave::FEditorAudioBulkData::Serialize(FArchive& Ar, UObject* Owner, bool bAllowRegister)
+void USoundWave::FEditorAudioBulkData::Serialize(FArchive& Ar, UObject* /*Owner - see comments in header re: SoundWave */, bool bAllowRegister)
 {
-	RawData.Serialize(Ar, Owner, bAllowRegister);
+	RawData.Serialize(Ar, SoundWave, bAllowRegister);
 }
 
-void USoundWave::FEditorAudioBulkData::CreateFromBulkData(FBulkData& InBulkData, const FGuid& InGuid, UObject* Owner)
+void USoundWave::FEditorAudioBulkData::CreateFromBulkData(FBulkData& InBulkData, const FGuid& InGuid, UObject* /*Owner - see comments in header re: SoundWave */)
 {
-	RawData.CreateFromBulkData(InBulkData, InGuid, Owner);
+	RawData.CreateFromBulkData(InBulkData, InGuid, SoundWave);
 }
 #endif
 
