nop just comments about RenderTarget ReadPixels functions

#rnx

diff --git a/Engine/Source/Runtime/Engine/Public/UnrealClient.h b/Engine/Source/Runtime/Engine/Public/UnrealClient.h
--- a/Engine/Source/Runtime/Engine/Public/UnrealClient.h
+++ b/Engine/Source/Runtime/Engine/Public/UnrealClient.h
@@ -95,6 +95,13 @@ public:
 	* @return True if the read succeeded.
 	*
 	* This will convert whatever the pixel format is to FColor
+	* Prefer using FImageUtils::GetRenderTargetImage rather than calling this directly.
+	*
+	* The default value for InFlags specifies RCM_UNorm which will cause values to be scaled into [0,1] ; use RCM_MinMax to retrieve values without change.
+	*
+	* If the RenderTarget surface is float linear, it will converted to SRGB FColor, if InFlags.bLinearToGamma is set (which is on by default).
+	* If the RenderTarget surface is U8, then the SRGB/not state is unchanged, the U8 values are retrieved unchanged in either Linear or SRGB.
+	* Gamma is handled correctly automatically by FImageUtils::GetRenderTargetImage
 	*/
 	ENGINE_API virtual bool ReadPixels(TArray<FColor>& OutImageData, FReadSurfaceDataFlags InFlags = FReadSurfaceDataFlags(RCM_UNorm, CubeFace_MAX), FIntRect InSrcRect = FIntRect(0, 0, 0, 0));
 
@@ -105,7 +112,8 @@ public:
 	* @param InSrcRect - InSrcRect not specified means the whole rect
 	* @return True if the read succeeded.
 	*
-	* Ptr variant of this API just does an extra memcpy; prefer the TArray variant
+	* Ptr variant of this API just does an extra memcpy; prefer the TArray variant.
+	* Prefer using FImageUtils::GetRenderTargetImage rather than calling this directly.
 	*/
 	ENGINE_API bool ReadPixelsPtr(FColor* OutImageBytes, FReadSurfaceDataFlags InFlags = FReadSurfaceDataFlags(RCM_UNorm, CubeFace_MAX), FIntRect InSrcRect = FIntRect(0, 0, 0, 0));
 
@@ -117,6 +125,7 @@ public:
 	 */
 	UE_DEPRECATED(5.4, "Use the other ReadFloat16Pixels variant (ECubeFace can be set in FReadSurfaceDataFlags)")
 	ENGINE_API bool ReadFloat16Pixels(TArray<FFloat16Color>& OutImageData, ECubeFace CubeFace);
+
 	/**
 	 * Reads the render target's displayed pixels into the given color buffer.
 	 * @param OutImageData - RGBA16F values will be stored in this buffer
@@ -124,7 +133,10 @@ public:
 	 * @param InSrcRect - InSrcRect not specified means the whole rect
 	 * @return True if the read succeeded.
 	 *
-	 * this only works if surface is PF_FloatRGBA exactly ; it does not convert
+	 * The default value for InFlags specifies RCM_UNorm which will cause values to be scaled into [0,1] ; use RCM_MinMax to retrieve values without change.
+	 *
+	 * Unlike other RenderTarget Read functions, this only works if surface is PF_FloatRGBA exactly ; it does not convert.
+	 * Prefer using FImageUtils::GetRenderTargetImage rather than calling this directly.
 	 */
 	ENGINE_API virtual bool ReadFloat16Pixels(TArray<FFloat16Color>& OutImageData, FReadSurfaceDataFlags InFlags = FReadSurfaceDataFlags(RCM_UNorm, CubeFace_MAX), FIntRect InSrcRect = FIntRect(0, 0, 0, 0));
 
@@ -135,7 +147,10 @@ public:
 	 * @param InSrcRect - InSrcRect not specified means the whole rect
 	 * @return True if the read succeeded.
 	 *
-	 * This will convert whatever the pixel format is to FLinearColor
+	 * The default value for InFlags specifies RCM_UNorm which will cause values to be scaled into [0,1] ; use RCM_MinMax to retrieve values without change.
+	 *
+	 * This will convert whatever the pixel format is to FLinearColor (if supported).
+	 * Prefer using FImageUtils::GetRenderTargetImage rather than calling this directly.
 	 */
 	ENGINE_API virtual bool ReadLinearColorPixels(TArray<FLinearColor>& OutImageData, FReadSurfaceDataFlags InFlags = FReadSurfaceDataFlags(RCM_MinMax, CubeFace_MAX), FIntRect InSrcRect = FIntRect(0, 0, 0, 0));
 
@@ -146,7 +161,8 @@ public:
 	 * @param InSrcRect - InSrcRect not specified means the whole rect
 	 * @return True if the read succeeded.
 	 *
-	 * Ptr variant of this API just does an extra memcpy; prefer the TArray variant
+	 * Ptr variant of this API just does an extra memcpy; prefer the TArray variant.
+	 * Prefer using FImageUtils::GetRenderTargetImage rather than calling this directly.
 	 */
 	ENGINE_API bool ReadLinearColorPixelsPtr(FLinearColor* OutImageBytes, FReadSurfaceDataFlags InFlags = FReadSurfaceDataFlags(RCM_MinMax, CubeFace_MAX), FIntRect InSrcRect = FIntRect(0, 0, 0, 0));
 
diff --git a/Engine/Source/Runtime/RHI/Public/RHIDefinitions.h b/Engine/Source/Runtime/RHI/Public/RHIDefinitions.h
--- a/Engine/Source/Runtime/RHI/Public/RHIDefinitions.h
+++ b/Engine/Source/Runtime/RHI/Public/RHIDefinitions.h
@@ -642,16 +642,19 @@ enum EResourceLockMode
 };
 
 /** limited to 8 types in FReadSurfaceDataFlags */
+// RCM_UNorm is the default
+// RCM_MinMax means "leave the values alone" and is recommended as what you should use
+// RCM_SNorm and RCM_MinMaxNorm seem to be unsupported
 enum ERangeCompressionMode
 {
 	// 0 .. 1
-	RCM_UNorm,
+	RCM_UNorm, // if you read values that go outside [0,1], they are scaled to fit inside [0,1]
 	// -1 .. 1
 	RCM_SNorm,
 	// 0 .. 1 unless there are smaller values than 0 or bigger values than 1, then the range is extended to the minimum or the maximum of the values
 	RCM_MinMaxNorm,
 	// minimum .. maximum (each channel independent)
-	RCM_MinMax,
+	RCM_MinMax, // read values without changing them
 };
 
 enum class EPrimitiveTopologyType : uint8
diff --git a/Engine/Source/Runtime/RHI/Public/RHITypes.h b/Engine/Source/Runtime/RHI/Public/RHITypes.h
--- a/Engine/Source/Runtime/RHI/Public/RHITypes.h
+++ b/Engine/Source/Runtime/RHI/Public/RHITypes.h
@@ -16,6 +16,9 @@ class FReadSurfaceDataFlags
 {
 public:
 	// @param InCompressionMode defines the value input range that is mapped to output range
+	//			The default CompressionMode of RCM_UNorm will modify values to fit in [0,1]
+	//			it is recommended to always use RCM_MinMax instead, which leaves values unchanged
+	//			if you do want scaling into [0,1] do it after the fact using ScaleChannelsSoMinMaxIsInZeroToOne
 	// @param InCubeFace defined which cubemap side is used, only required for cubemap content, then it needs to be a valid side
 	FReadSurfaceDataFlags(ERangeCompressionMode InCompressionMode = RCM_UNorm, ECubeFace InCubeFace = CubeFace_MAX)
 		:CubeFace(InCubeFace), CompressionMode(InCompressionMode)
