Path Tracing: Change Indirect Emissive lighting component back to its original name

Including this flag as part of the lighting components is rather confusing as it needs to be included all the time. It does not work like the other booleans in this section do. The original name "Emissive Materials" represents what this boolean does more accurately.

diff --git a/Engine/Source/Runtime/Engine/Classes/Engine/Scene.h b/Engine/Source/Runtime/Engine/Classes/Engine/Scene.h
--- a/Engine/Source/Runtime/Engine/Classes/Engine/Scene.h
+++ b/Engine/Source/Runtime/Engine/Classes/Engine/Scene.h
@@ -1287,6 +1287,9 @@ struct FPostProcessSettings
 	UPROPERTY(EditAnywhere, BlueprintReadWrite, Category = Overrides, meta = (PinHiddenByDefault, InlineEditConditionToggle))
 	uint32 bOverride_PathTracingMaxPathExposure : 1;
 
+	UPROPERTY(EditAnywhere, BlueprintReadWrite, Category = Overrides, meta = (PinHiddenByDefault, InlineEditConditionToggle))
+	uint32 bOverride_PathTracingEnableEmissiveMaterials : 1;
+
 	UPROPERTY(EditAnywhere, BlueprintReadWrite, Category = Overrides, meta = (PinHiddenByDefault, InlineEditConditionToggle))
 	uint32 bOverride_PathTracingEnableReferenceDOF : 1;
 
@@ -1299,9 +1302,6 @@ struct FPostProcessSettings
 	UPROPERTY(EditAnywhere, BlueprintReadWrite, Category = Overrides, meta = (PinHiddenByDefault, InlineEditConditionToggle))
 	uint32 bOverride_PathTracingIncludeEmissive : 1;
 
-	UPROPERTY(EditAnywhere, BlueprintReadWrite, Category = Overrides, meta = (PinHiddenByDefault, InlineEditConditionToggle))
-	uint32 bOverride_PathTracingIncludeIndirectEmissive : 1;
-
 	UPROPERTY(EditAnywhere, BlueprintReadWrite, Category = Overrides, meta = (PinHiddenByDefault, InlineEditConditionToggle))
 	uint32 bOverride_PathTracingIncludeDiffuse : 1;
 
@@ -2210,6 +2210,10 @@ struct FPostProcessSettings
 	UPROPERTY(interp, EditAnywhere, BlueprintReadWrite, Category = "Path Tracing", meta = (ClampMin = "-10.0", ClampMax = "30.0", editcondition = "bOverride_PathTracingMaxPathExposure", DisplayName = "Max Path Exposure"))
 	float PathTracingMaxPathExposure;
 
+	/** Should emissive materials contribute to scene lighting? */
+	UPROPERTY(EditAnywhere, BlueprintReadWrite, Category = "Path Tracing", meta = (editcondition = "bOverride_PathTracingEnableEmissiveMaterials", DisplayName = "Emissive Materials"))
+	uint32 PathTracingEnableEmissiveMaterials : 1;
+
 	/** Enables a reference quality depth-of-field which replaces the post-process effect. */
 	UPROPERTY(EditAnywhere, BlueprintReadWrite, Category = "Path Tracing", meta = (editcondition = "bOverride_PathTracingEnableReferenceDOF", DisplayName = "Reference Depth Of Field"))
 	uint32 PathTracingEnableReferenceDOF : 1;
@@ -2226,10 +2230,6 @@ struct FPostProcessSettings
 	UPROPERTY(EditAnywhere, BlueprintReadWrite, Category = "Path Tracing|Lighting Components", meta = (editcondition = "bOverride_PathTracingIncludeEmissive", DisplayName = "Emissive"))
 	uint32 PathTracingIncludeEmissive : 1;
 
-	/** Should the render include indirectly visible emissive elements? */
-	UPROPERTY(EditAnywhere, BlueprintReadWrite, Category = "Path Tracing|Lighting Components", meta = (editcondition = "bOverride_PathTracingIncludeIndirectEmissive", DisplayName = "Indirect Emissive"))
-	uint32 PathTracingIncludeIndirectEmissive : 1;
-
 	/** Should the render include diffuse lighting contributions? */
 	UPROPERTY(EditAnywhere, BlueprintReadWrite, Category = "Path Tracing|Lighting Components", meta = (editcondition = "bOverride_PathTracingIncludeDiffuse", DisplayName = "Diffuse"))
 	uint32 PathTracingIncludeDiffuse : 1;
diff --git a/Engine/Source/Runtime/Engine/Private/Scene.cpp b/Engine/Source/Runtime/Engine/Private/Scene.cpp
--- a/Engine/Source/Runtime/Engine/Private/Scene.cpp
+++ b/Engine/Source/Runtime/Engine/Private/Scene.cpp
@@ -630,12 +630,12 @@ FPostProcessSettings::FPostProcessSettings()
 	PathTracingMaxBounces = 32;
 	PathTracingSamplesPerPixel = 2048;
 	PathTracingMaxPathExposure = 30.0f;
+	PathTracingEnableEmissiveMaterials = 1;
 	PathTracingEnableReferenceDOF = 0;
 	PathTracingEnableReferenceAtmosphere = 0;
 	PathTracingEnableDenoiser = 1;
 
 	PathTracingIncludeEmissive = 1;
-	PathTracingIncludeIndirectEmissive = 1;
 	PathTracingIncludeDiffuse = 1;
 	PathTracingIncludeIndirectDiffuse = 1;
 	PathTracingIncludeSpecular = 1;
diff --git a/Engine/Source/Runtime/Engine/Private/SceneView.cpp b/Engine/Source/Runtime/Engine/Private/SceneView.cpp
--- a/Engine/Source/Runtime/Engine/Private/SceneView.cpp
+++ b/Engine/Source/Runtime/Engine/Private/SceneView.cpp
@@ -1675,11 +1675,11 @@ void FSceneView::OverridePostProcessSettings(const FPostProcessSettings& Src, fl
 		SET_PP(PathTracingMaxBounces);
 		SET_PP(PathTracingSamplesPerPixel);
 		LERP_PP(PathTracingMaxPathExposure);
+		SET_PP(PathTracingEnableEmissiveMaterials);
 		SET_PP(PathTracingEnableReferenceDOF);
 		SET_PP(PathTracingEnableReferenceAtmosphere);
 		SET_PP(PathTracingEnableDenoiser);
 		SET_PP(PathTracingIncludeEmissive);
-		SET_PP(PathTracingIncludeIndirectEmissive);
 		SET_PP(PathTracingIncludeDiffuse);
 		SET_PP(PathTracingIncludeIndirectDiffuse);
 		SET_PP(PathTracingIncludeSpecular);
diff --git a/Engine/Source/Runtime/Renderer/Private/PathTracing.cpp b/Engine/Source/Runtime/Renderer/Private/PathTracing.cpp
--- a/Engine/Source/Runtime/Renderer/Private/PathTracing.cpp
+++ b/Engine/Source/Runtime/Renderer/Private/PathTracing.cpp
@@ -697,8 +697,9 @@ static void PreparePathTracingData(const FScene* Scene, const FViewInfo& View, F
 	PathTracingData.EnabledIndirectLightingContributions = 0;
 	if (ShowFlags.GlobalIllumination != 0)
 	{
-		const bool bEnableEmissive = CVarPathTracingEnableEmissive.GetValueOnRenderThread() != 0;
-		PathTracingData.EnabledIndirectLightingContributions |= (PPV.PathTracingIncludeIndirectEmissive != 0 && bEnableEmissive        ) ? PATHTRACER_CONTRIBUTION_EMISSIVE : 0;
+		const int EnableEmissiveCVar = CVarPathTracingEnableEmissive.GetValueOnRenderThread();
+		const bool bEnableEmissive = EnableEmissiveCVar < 0 ? PPV.PathTracingEnableEmissiveMaterials : EnableEmissiveCVar != 0;
+		PathTracingData.EnabledIndirectLightingContributions |= (bEnableEmissive                                                       ) ? PATHTRACER_CONTRIBUTION_EMISSIVE : 0;
 		PathTracingData.EnabledIndirectLightingContributions |= (PPV.PathTracingIncludeIndirectDiffuse  != 0 && ShowFlags.Diffuse  != 0) ? PATHTRACER_CONTRIBUTION_DIFFUSE  : 0;
 		PathTracingData.EnabledIndirectLightingContributions |= (PPV.PathTracingIncludeIndirectSpecular != 0 && ShowFlags.Specular != 0) ? PATHTRACER_CONTRIBUTION_SPECULAR : 0;
 		PathTracingData.EnabledIndirectLightingContributions |= (PPV.PathTracingIncludeIndirectVolume   != 0                           ) ? PATHTRACER_CONTRIBUTION_VOLUME   : 0;
