Reorganize code to allow disable/restore unreachable code pragmas to work (consistently)
#rnx

diff --git a/Engine/Source/Runtime/Engine/Private/UnrealEngine.cpp b/Engine/Source/Runtime/Engine/Private/UnrealEngine.cpp
--- a/Engine/Source/Runtime/Engine/Private/UnrealEngine.cpp
+++ b/Engine/Source/Runtime/Engine/Private/UnrealEngine.cpp
@@ -10276,6 +10276,48 @@
 
 extern CORE_API bool GIsGPUCrashed;
 
+PRAGMA_DISABLE_UNREACHABLE_CODE_WARNINGS
+bool PerformBlockingError(const TCHAR* Cmd, FOutputDevice& Ar)
+{
+	if (FParse::Command(&Cmd, TEXT("TERMINATE")))
+	{
+		UE_LOG(LogEngine, Warning, TEXT("Printed warning to log."));
+		FGenericCrashContext::SetCrashTrigger(ECrashTrigger::Debug);
+		UE_LOG(LogEngine, Warning, TEXT("%s"), TEXT("Terminating from the gamethread at your request"));
+		std::terminate();
+		return true;
+	}
+	else if (FParse::Command(&Cmd, TEXT("ABORT")))
+	{
+		UE_LOG(LogEngine, Warning, TEXT("Printed warning to log."));
+		FGenericCrashContext::SetCrashTrigger(ECrashTrigger::Debug);
+		UE_LOG(LogEngine, Warning, TEXT("%s"), TEXT("Aborting from the gamethread at your request"));
+		std::abort();
+		return true;
+	}
+	else if (FParse::Command(&Cmd, TEXT("SOFTLOCK")))
+	{
+		Ar.Log(TEXT("Hanging the current thread"));
+		FGenericCrashContext::SetCrashTrigger(ECrashTrigger::Debug);
+		while (1)
+		{
+			FPlatformProcess::Sleep(1.0f);
+		}
+		return true;
+	}
+	else if (FParse::Command(&Cmd, TEXT("INFINITELOOP")))
+	{
+		Ar.Log(TEXT("Hanging the current thread (CPU-intensive)"));
+		FGenericCrashContext::SetCrashTrigger(ECrashTrigger::Debug);
+		for (;;)
+		{
+		}
+		return true;
+	}
+	return false;
+}
+PRAGMA_RESTORE_UNREACHABLE_CODE_WARNINGS
+
 bool UEngine::PerformError(const TCHAR* Cmd, FOutputDevice& Ar)
 {
 #if !UE_BUILD_SHIPPING
@@ -10614,24 +10656,6 @@
 		ENQUEUE_RENDER_COMMAND(CauseGpuCrash)(&FRender::GpuCrash);
 		return true;
 	}
-PRAGMA_DISABLE_UNREACHABLE_CODE_WARNINGS
-	else if (FParse::Command(&Cmd, TEXT("TERMINATE")))
-	{
-		UE_LOG(LogEngine, Warning, TEXT("Printed warning to log."));
-		FGenericCrashContext::SetCrashTrigger(ECrashTrigger::Debug);
-		UE_LOG(LogEngine, Warning, TEXT("%s"), TEXT("Terminating from the gamethread at your request"));
-		std::terminate();
-		return true;
-	}
-	else if (FParse::Command(&Cmd, TEXT("ABORT")))
-	{
-		UE_LOG(LogEngine, Warning, TEXT("Printed warning to log."));
-		FGenericCrashContext::SetCrashTrigger(ECrashTrigger::Debug);
-		UE_LOG(LogEngine, Warning, TEXT("%s"), TEXT("Aborting from the gamethread at your request"));
-		std::abort();
-		return true;
-	}
-PRAGMA_RESTORE_UNREACHABLE_CODE_WARNINGS
 	else if (FParse::Command(&Cmd, TEXT("CHECK")))
 	{
 		UE_LOG(LogEngine, Warning, TEXT("Printed warning to log."));
@@ -10964,27 +10988,7 @@
 		);
 		return true;
 	}
-PRAGMA_DISABLE_UNREACHABLE_CODE_WARNINGS
-	else if (FParse::Command(&Cmd, TEXT("SOFTLOCK")))
-	{
-		Ar.Log(TEXT("Hanging the current thread"));
-		FGenericCrashContext::SetCrashTrigger(ECrashTrigger::Debug);
-		while (1)
-		{
-			FPlatformProcess::Sleep(1.0f);
-		}
-		return true;
-	}
-	else if (FParse::Command(&Cmd, TEXT("INFINITELOOP")))
-	{
-		Ar.Log(TEXT("Hanging the current thread (CPU-intensive)"));
-		FGenericCrashContext::SetCrashTrigger(ECrashTrigger::Debug);
-		for(;;)
-		{
-		}
-		return true;
-	}
-PRAGMA_RESTORE_UNREACHABLE_CODE_WARNINGS
+
 	else if (FParse::Command(&Cmd, TEXT("SLEEP")))
 	{
 		Ar.Log(TEXT("Sleep for 1 hour. This should crash after a few seconds in cooked builds."));
@@ -11029,6 +11033,10 @@
 		FGenericCrashContext::SetCrashTrigger(ECrashTrigger::Debug);
 		UE_LOG(LogEngine, Fatal, TEXT("Crashing the worker thread at your request"));
 	}
+	else if (PerformBlockingError(Cmd, Ar))
+	{
+		return true;
+	}
 #if USING_ADDRESS_SANITISER
 	else if (FParse::Command(&Cmd, TEXT("USEAFTERFREE")))
 	{
