Fix shipping
#rnx

diff --git a/Engine/Source/Runtime/Engine/Private/UnrealEngine.cpp b/Engine/Source/Runtime/Engine/Private/UnrealEngine.cpp
--- a/Engine/Source/Runtime/Engine/Private/UnrealEngine.cpp
+++ b/Engine/Source/Runtime/Engine/Private/UnrealEngine.cpp
@@ -10276,6 +10276,7 @@ private:
 
 extern CORE_API bool GIsGPUCrashed;
 
+#if !UE_BUILD_SHIPPING
 PRAGMA_DISABLE_UNREACHABLE_CODE_WARNINGS
 bool PerformBlockingError(const TCHAR* Cmd, FOutputDevice& Ar)
 {
@@ -10317,6 +10318,7 @@ bool PerformBlockingError(const TCHAR* Cmd, FOutputDevice& Ar)
 	return false;
 }
 PRAGMA_RESTORE_UNREACHABLE_CODE_WARNINGS
+#endif // !UE_BUILD_SHIPPING
 
 bool UEngine::PerformError(const TCHAR* Cmd, FOutputDevice& Ar)
 {
