Fix ADPCM seeking discontinuities.

Make sure to clear out decode buffer after performing a seek. Make sure we update decode position as well when starting at the beginning

#tests PIE

diff --git a/Engine/Source/Runtime/Engine/Private/ADPCMAudioInfo.cpp b/Engine/Source/Runtime/Engine/Private/ADPCMAudioInfo.cpp
--- a/Engine/Source/Runtime/Engine/Private/ADPCMAudioInfo.cpp
+++ b/Engine/Source/Runtime/Engine/Private/ADPCMAudioInfo.cpp
@@ -107,13 +107,20 @@ void FADPCMAudioInfo::SeekToFrameInternal(const uint32 InSeekFrame)
 	// Reset chunk handle in preperation for a new chunk.
 	CurCompressedChunkData = nullptr;
 
+	// zero out any existing decode data as it will stale after seeking
+	if (UncompressedBlockData)
+	{
+		FMemory::Memzero(UncompressedBlockData, NumChannels * UncompressedBlockSize);
+	}
+	
 	float SeekTimeSeconds = InSeekFrame / static_cast<float>(*WaveInfo.pSamplesPerSec);
 	UE_LOG(LogAudio, Verbose, TEXT("Seeking ADPCM source to %.3f sec"), SeekTimeSeconds);
 
 	if (InSeekFrame == 0)
 	{
 		CurrentCompressedBlockIndex = 0;
-		CurrentUncompressedBlockSampleIndex = 0;
+		// This is set to the max value to trigger the decompression of the first audio block
+		CurrentUncompressedBlockSampleIndex = UncompressedBlockSize / sizeof(uint16);;
 		CurrentChunkIndex = FirstChunkSampleDataIndex;
 		CurrentChunkBufferOffset = FirstChunkSampleDataOffset;
 		TotalSamplesStreamed = 0;
