Fix logic to discard mismatched components when ran via the diff tool, which does not compile blueprints

diff --git a/Engine/Source/Runtime/Engine/Private/InheritableComponentHandler.cpp b/Engine/Source/Runtime/Engine/Private/InheritableComponentHandler.cpp
--- a/Engine/Source/Runtime/Engine/Private/InheritableComponentHandler.cpp
+++ b/Engine/Source/Runtime/Engine/Private/InheritableComponentHandler.cpp
@@ -149,7 +149,7 @@ UActorComponent* UInheritableComponentHandler::CreateOverridenComponentTemplate(
 		// but if it is, just consign it to oblivion as its purpose is no longer required with the allocation
 		// of an object of the same name
 		UActorComponent* ExistingComp = Cast<UActorComponent>(ExistingObj);
-		if (ensure(ExistingComp) && ensure(UnnecessaryComponents.RemoveSwap(ExistingComp) > 0))
+		if (ensure(ExistingComp) && ensure(UnnecessaryComponents.RemoveSwap(ExistingComp) > 0 || GetPackage()->HasAnyPackageFlags(PKG_ForDiffing)))
 		{
 			ExistingObj->Rename(nullptr, GetTransientPackage(), REN_DontCreateRedirectors | REN_ForceNoResetLoaders);
 			ExistingObj->MarkAsGarbage();
