Engine-side massaging of FActorInstanceHandle creation, extending the range of cases where IActorInstanceManagerInterface would get utilized (if implemented).

diff --git a/Engine/Plugins/Experimental/GameplayTargetingSystem/Source/GameplayTargetingSystem/Private/Tasks/TargetingSelectionTask_AOE.cpp b/Engine/Plugins/Experimental/GameplayTargetingSystem/Source/GameplayTargetingSystem/Private/Tasks/TargetingSelectionTask_AOE.cpp
--- a/Engine/Plugins/Experimental/GameplayTargetingSystem/Source/GameplayTargetingSystem/Private/Tasks/TargetingSelectionTask_AOE.cpp
+++ b/Engine/Plugins/Experimental/GameplayTargetingSystem/Source/GameplayTargetingSystem/Private/Tasks/TargetingSelectionTask_AOE.cpp
@@ -264,7 +264,7 @@ void UTargetingSelectionTask_AOE::ProcessOverlapResults(const FTargetingRequestH
 			if (bAddResult)
 			{
 				FTargetingDefaultResultData* ResultData = new(TargetingResults.TargetResults) FTargetingDefaultResultData();
-				ResultData->HitResult.HitObjectHandle = FActorInstanceHandle(OverlapResult.GetActor());
+				ResultData->HitResult.HitObjectHandle = FActorInstanceHandle(OverlapResult.GetActor(), OverlapResult.GetComponent(), OverlapResult.ItemIndex);
 				ResultData->HitResult.Component = OverlapResult.GetComponent();
 				ResultData->HitResult.bBlockingHit = OverlapResult.bBlockingHit;
 				ResultData->HitResult.TraceStart = SourceLocation;
diff --git a/Engine/Source/Runtime/Engine/Classes/Engine/ActorInstanceHandle.h b/Engine/Source/Runtime/Engine/Classes/Engine/ActorInstanceHandle.h
--- a/Engine/Source/Runtime/Engine/Classes/Engine/ActorInstanceHandle.h
+++ b/Engine/Source/Runtime/Engine/Classes/Engine/ActorInstanceHandle.h
@@ -19,10 +19,11 @@ struct FActorInstanceHandle
 {
 	GENERATED_BODY()
 
-	ENGINE_API FActorInstanceHandle();
+	ENGINE_API FActorInstanceHandle() = default;
 
 	ENGINE_API explicit FActorInstanceHandle(AActor* InActor);
-	ENGINE_API FActorInstanceHandle(UObject* InManager, const UPrimitiveComponent* RelevantComponent, int32 CollisionInstanceIndex);
+	ENGINE_API FActorInstanceHandle(const UPrimitiveComponent* RelevantComponent, int32 CollisionInstanceIndex);
+	ENGINE_API FActorInstanceHandle(AActor* InActor, const UPrimitiveComponent* RelevantComponent, int32 CollisionInstanceIndex);
 	ENGINE_API FActorInstanceHandle(FActorInstanceManagerInterface InManagerInterface, int32 InstanceIndex);
 	ENGINE_API FActorInstanceHandle(const FActorInstanceHandle& Other);
 
@@ -89,7 +90,7 @@ struct FActorInstanceHandle
 	}
 
 private:
-	void SetInternal(IActorInstanceManagerInterface& InManagerInterface, const UPrimitiveComponent* RelevantComponent, int32 CollisionInstanceIndex);
+	friend struct FActorInstanceHandleInternalHelper;
 
 	/**
 	 * helper functions that let us treat the actor pointer as a UObject in templated functions
@@ -109,7 +110,7 @@ private:
 	FActorInstanceManagerInterface ManagerInterface;
 
 	/** Identifies the instance within the manager */
-	int32 InstanceIndex;
+	int32 InstanceIndex = INDEX_NONE;
 };
 
 template<typename T>
diff --git a/Engine/Source/Runtime/Engine/Private/Actor.cpp b/Engine/Source/Runtime/Engine/Private/Actor.cpp
--- a/Engine/Source/Runtime/Engine/Private/Actor.cpp
+++ b/Engine/Source/Runtime/Engine/Private/Actor.cpp
@@ -5070,7 +5070,6 @@ void AActor::DispatchPhysicsCollisionHit(const FRigidBodyCollisionInfo& MyInfo,
 	Result.Normal = Result.ImpactNormal = ContactInfo.ContactNormal;
 	Result.PenetrationDepth = ContactInfo.ContactPenetration;
 	Result.PhysMaterial = ContactInfo.PhysMaterial[1];
-	Result.HitObjectHandle = FActorInstanceHandle(OtherInfo.Actor.Get());
 	Result.Component = OtherInfo.Component;
 	Result.MyItem = MyInfo.BodyIndex;
 	Result.Item = OtherInfo.BodyIndex;
@@ -5078,7 +5077,11 @@ void AActor::DispatchPhysicsCollisionHit(const FRigidBodyCollisionInfo& MyInfo,
 	Result.MyBoneName = MyInfo.BoneName;
 	Result.bBlockingHit = true;
 
-	NotifyHit(MyInfo.Component.Get(), OtherInfo.Actor.Get(), OtherInfo.Component.Get(), true, Result.Location, Result.Normal, RigidCollisionData.TotalNormalImpulse, Result);
+	AActor* Actor = OtherInfo.Actor.Get();
+	UPrimitiveComponent* Component = OtherInfo.Component.Get();
+	Result.HitObjectHandle = FActorInstanceHandle(Actor, Component, OtherInfo.BodyIndex);
+
+	NotifyHit(MyInfo.Component.Get(), Actor, Component, true, Result.Location, Result.Normal, RigidCollisionData.TotalNormalImpulse, Result);
 
 	// Execute delegates if bound
 
diff --git a/Engine/Source/Runtime/Engine/Private/Collision/CollisionConversions.cpp b/Engine/Source/Runtime/Engine/Private/Collision/CollisionConversions.cpp
--- a/Engine/Source/Runtime/Engine/Private/Collision/CollisionConversions.cpp
+++ b/Engine/Source/Runtime/Engine/Private/Collision/CollisionConversions.cpp
@@ -185,8 +185,7 @@ static void SetHitResultFromShapeAndFaceIndex(const FPhysicsShape& Shape,  const
 	if( OwningComponent )
 	{
 		OutResult.Component = OwningComponent;
-		AActor* Owner = OwningComponent->GetOwner();
-		OutResult.HitObjectHandle = FActorInstanceHandle(Owner);
+		OutResult.HitObjectHandle = FActorInstanceHandle(OwningComponent, OutResult.Item);
 
 		if (bReturnPhysMat)
 		{
@@ -603,15 +602,7 @@ void ConvertQueryOverlap(const FPhysicsShape& Shape, const FPhysicsActor& Actor,
         BodyInst = FPhysicsInterface::ShapeToOriginalBodyInstance(BodyInst, &Shape);
 		if (const UPrimitiveComponent* OwnerComponent = BodyInst->OwnerComponent.Get())
 		{
-			UObject* Outer = OwnerComponent->GetOuter();
-			if (Outer && Outer->Implements<UActorInstanceManagerInterface>())
-			{
-				OutOverlap.OverlapObjectHandle = FActorInstanceHandle(Outer, OwnerComponent, BodyInst->InstanceBodyIndex);
-			}
-			else
-			{
-				OutOverlap.OverlapObjectHandle = FActorInstanceHandle(OwnerComponent->GetOwner());
-			}
+			OutOverlap.OverlapObjectHandle = FActorInstanceHandle(OwnerComponent->GetOwner(), OwnerComponent, BodyInst->InstanceBodyIndex);
 			OutOverlap.Component = BodyInst->OwnerComponent; // Copying weak pointer is faster than assigning raw pointer.
 			OutOverlap.ItemIndex = OwnerComponent->bMultiBodyOverlap ? BodyInst->InstanceBodyIndex : INDEX_NONE;
 		}
diff --git a/Engine/Source/Runtime/Engine/Private/Components/PrimitiveComponent.cpp b/Engine/Source/Runtime/Engine/Private/Components/PrimitiveComponent.cpp
--- a/Engine/Source/Runtime/Engine/Private/Components/PrimitiveComponent.cpp
+++ b/Engine/Source/Runtime/Engine/Private/Components/PrimitiveComponent.cpp
@@ -141,7 +141,10 @@ DECLARE_CYCLE_STAT(TEXT("PrimComp DispatchBlockingHit"), STAT_DispatchBlockingHi
 FOverlapInfo::FOverlapInfo(UPrimitiveComponent* InComponent, int32 InBodyIndex)
 	: bFromSweep(false)
 {
-	OverlapInfo.HitObjectHandle = FActorInstanceHandle(InComponent ? InComponent->GetOwner() : nullptr);
+	if (InComponent)
+	{
+		OverlapInfo.HitObjectHandle = FActorInstanceHandle(InComponent->GetOwner(), InComponent, InBodyIndex);
+	}
 	OverlapInfo.Component = InComponent;
 	OverlapInfo.Item = InBodyIndex;
 }
diff --git a/Engine/Source/Runtime/Engine/Private/Engine/ActorInstanceHandle.cpp b/Engine/Source/Runtime/Engine/Private/Engine/ActorInstanceHandle.cpp
--- a/Engine/Source/Runtime/Engine/Private/Engine/ActorInstanceHandle.cpp
+++ b/Engine/Source/Runtime/Engine/Private/Engine/ActorInstanceHandle.cpp
@@ -1,49 +1,80 @@
 // Copyright Epic Games, Inc. All Rights Reserved.
 
 #include "Engine/ActorInstanceHandle.h"
+#include "PhysicsPublic.h"
 #include "Engine/World.h"
 #include "Engine/ActorInstanceManagerInterface.h"
 #include "GameFramework/Actor.h"
+#include "Components/PrimitiveComponent.h"
 #include "UObject/FortniteValkyrieBranchObjectVersion.h"
 
 #include UE_INLINE_GENERATED_CPP_BY_NAME(ActorInstanceHandle)
 
+struct FActorInstanceHandleInternalHelper
+{
+	inline static void SetUpAsInterface(FActorInstanceHandle& InstanceHandle, IActorInstanceManagerInterface& InManagerInterface, const UPrimitiveComponent* RelevantComponent, int32 CollisionInstanceIndex)
+	{
+		InstanceHandle.InstanceIndex = InManagerInterface.ConvertCollisionIndexToInstanceIndex(CollisionInstanceIndex, RelevantComponent);
+		InstanceHandle.Actor = InManagerInterface.FindActor(InstanceHandle);
+	}
+
+	inline static void SetUpWithActor(FActorInstanceHandle& InstanceHandle, AActor* InActor, const UPrimitiveComponent* RelevantComponent, int32 CollisionInstanceIndex)
+	{
+		InstanceHandle.ManagerInterface = FActorInstanceManagerInterface(InActor);
+		if (IActorInstanceManagerInterface* ManagerInterfacePtr = InstanceHandle.ManagerInterface.Get())
+		{
+			SetUpAsInterface(InstanceHandle, *ManagerInterfacePtr, RelevantComponent, CollisionInstanceIndex);
+		}
+		else
+		{
+			InstanceHandle.Actor = InActor;
+		}
+	}
+};
 
 //-----------------------------------------------------------------------------
 // FActorInstanceHandle
 //-----------------------------------------------------------------------------
-FActorInstanceHandle::FActorInstanceHandle()
-	: Actor(nullptr)
-	, InstanceIndex(INDEX_NONE)
+FActorInstanceHandle::FActorInstanceHandle(AActor* InActor)
+	: Actor(InActor)
 {
-	// do nothing
 }
 
-FActorInstanceHandle::FActorInstanceHandle(AActor* InActor)
-	: Actor(InActor)
-	, InstanceIndex(INDEX_NONE)
+FActorInstanceHandle::FActorInstanceHandle(const UPrimitiveComponent* RelevantComponent, int32 CollisionInstanceIndex)
 {
+	if (UNLIKELY(!ensureMsgf(RelevantComponent, TEXT("Calling FActorInstanceHandle(UPrimitiveComponent, int32) constructor is pointless with RelevantComponent == nullptr"))))
+	{
+		return;
+	}
+
+	if (AActor* OwnerActor = RelevantComponent->GetOwner())
+	{
+		FActorInstanceHandleInternalHelper::SetUpWithActor(*this, OwnerActor, RelevantComponent, CollisionInstanceIndex);
+	}
 }
 
-FActorInstanceHandle::FActorInstanceHandle(UObject* InManager, const UPrimitiveComponent* RelevantComponent, int32 CollisionInstanceIndex)
-	: Actor(nullptr)
-	, ManagerInterface(InManager)
-	, InstanceIndex(CollisionInstanceIndex)
+FActorInstanceHandle::FActorInstanceHandle(AActor* InActor, const UPrimitiveComponent* RelevantComponent, int32 CollisionInstanceIndex)
 {
-	if (IActorInstanceManagerInterface* ManagerInterfacePtr = ManagerInterface.Get())
+	if (LIKELY(InActor))
+	{
+		FActorInstanceHandleInternalHelper::SetUpWithActor(*this, InActor, RelevantComponent, CollisionInstanceIndex);
+	}
+	else if (RelevantComponent)
 	{
-		SetInternal(*ManagerInterfacePtr, RelevantComponent, CollisionInstanceIndex);
+		FActorInstanceHandle(RelevantComponent, CollisionInstanceIndex);
+	}
+	else
+	{
+		Actor = InActor;
 	}
 }
 
 FActorInstanceHandle::FActorInstanceHandle(FActorInstanceManagerInterface InManagerInterface, int32 CollisionInstanceIndex)
-	: Actor(nullptr)
-	, ManagerInterface(InManagerInterface)
-	, InstanceIndex(CollisionInstanceIndex)
+	: ManagerInterface(InManagerInterface)
 {
 	if (IActorInstanceManagerInterface* ManagerInterfacePtr = ManagerInterface.Get())
 	{
-		SetInternal(*ManagerInterfacePtr, /*RelevantComponent=*/nullptr, CollisionInstanceIndex);
+		FActorInstanceHandleInternalHelper::SetUpAsInterface(*this, *ManagerInterfacePtr, /*RelevantComponent=*/nullptr, CollisionInstanceIndex);
 	}
 }
 
@@ -64,14 +95,6 @@ FActorInstanceHandle FActorInstanceHandle::MakeDehydratedActorHandle(UObject& Ma
 	return ReturnHandle;
 }
 
-void FActorInstanceHandle::SetInternal(IActorInstanceManagerInterface& InManagerInterface, const UPrimitiveComponent* RelevantComponent, int32 CollisionInstanceIndex)
-{
-	check(ManagerInterface.Get() == &InManagerInterface);
-
-	InstanceIndex = InManagerInterface.ConvertCollisionIndexToInstanceIndex(CollisionInstanceIndex, RelevantComponent);
-	Actor = InManagerInterface.FindActor(*this);
-}
-
 bool FActorInstanceHandle::IsValid() const
 {
 	return (ManagerInterface.IsValid() && InstanceIndex != INDEX_NONE) || IsActorValid();
diff --git a/Engine/Source/Runtime/Engine/Private/GameplayStatics.cpp b/Engine/Source/Runtime/Engine/Private/GameplayStatics.cpp
--- a/Engine/Source/Runtime/Engine/Private/GameplayStatics.cpp
+++ b/Engine/Source/Runtime/Engine/Private/GameplayStatics.cpp
@@ -1424,7 +1424,7 @@ FHitResult UGameplayStatics::MakeHitResult(bool bBlockingHit, bool bInitialOverl
 	Hit.Normal = Normal;
 	Hit.ImpactNormal = ImpactNormal;
 	Hit.PhysMaterial = PhysMat;
-	Hit.HitObjectHandle = FActorInstanceHandle(HitActor);
+	Hit.HitObjectHandle = FActorInstanceHandle(HitActor, HitComponent, HitItem);
 	Hit.Component = HitComponent;
 	Hit.BoneName = HitBoneName;
 	Hit.MyBoneName = BoneName;
