Add DebugViewModeHelpers.Enable CVar to allow disabling debug view shaders for editor builds that don't require them.

diff --git a/Engine/Source/Runtime/Engine/Private/DebugViewModeHelpers.cpp b/Engine/Source/Runtime/Engine/Private/DebugViewModeHelpers.cpp
--- a/Engine/Source/Runtime/Engine/Private/DebugViewModeHelpers.cpp
+++ b/Engine/Source/Runtime/Engine/Private/DebugViewModeHelpers.cpp
@@ -21,6 +21,12 @@
 
 #define LOCTEXT_NAMESPACE "LogDebugViewMode"
 
+static TAutoConsoleVariable<bool> CVarEnableDebugViewModeHelpers(
+	TEXT("DebugViewModeHelpers.Enable"),
+	true,
+	TEXT("Specifies whether to enable the debug view mode shaders. Typically only disabled for a special case editor build, if it doesn't require them"),
+	ECVF_Default);
+
 const TCHAR* DebugViewShaderModeToString(EDebugViewShaderMode InShaderMode)
 {
 	switch (InShaderMode)
@@ -66,6 +72,10 @@ static bool PlatformSupportsDebugViewShaders(EShaderPlatform Platform)
 
 bool AllowDebugViewVSDSHS(EShaderPlatform Platform)
 {
+	if(!CVarEnableDebugViewModeHelpers.GetValueOnAnyThread())
+	{
+		return false;
+	}
 	return IsPCPlatform(Platform); 
 }
 
