Issue 405783: PVS warning V603 fix for ActorInstanceHandle.cpp

diff --git a/Engine/Source/Runtime/Engine/Private/Engine/ActorInstanceHandle.cpp b/Engine/Source/Runtime/Engine/Private/Engine/ActorInstanceHandle.cpp
--- a/Engine/Source/Runtime/Engine/Private/Engine/ActorInstanceHandle.cpp
+++ b/Engine/Source/Runtime/Engine/Private/Engine/ActorInstanceHandle.cpp
@@ -61,7 +61,7 @@ FActorInstanceHandle::FActorInstanceHandle(AActor* InActor, const UPrimitiveComp
 	}
 	else if (RelevantComponent)
 	{
-		FActorInstanceHandle(RelevantComponent, CollisionInstanceIndex);
+		this->FActorInstanceHandle::FActorInstanceHandle(RelevantComponent, CollisionInstanceIndex);
 	}
 	else
 	{
