Linux compilation fix for FActorInstanceHandle

diff --git a/Engine/Source/Runtime/Engine/Private/Engine/ActorInstanceHandle.cpp b/Engine/Source/Runtime/Engine/Private/Engine/ActorInstanceHandle.cpp
--- a/Engine/Source/Runtime/Engine/Private/Engine/ActorInstanceHandle.cpp
+++ b/Engine/Source/Runtime/Engine/Private/Engine/ActorInstanceHandle.cpp
@@ -61,7 +61,7 @@ FActorInstanceHandle::FActorInstanceHandle(AActor* InActor, const UPrimitiveComp
 	}
 	else if (RelevantComponent)
 	{
-		this->FActorInstanceHandle::FActorInstanceHandle(RelevantComponent, CollisionInstanceIndex);
+		*this = FActorInstanceHandle(RelevantComponent, CollisionInstanceIndex);
 	}
 	else
 	{
