Skeletal meshes with morph targets don't use pass through vertex factory

diff --git a/Engine/Source/Runtime/Engine/Private/SkeletalRenderGPUSkin.cpp b/Engine/Source/Runtime/Engine/Private/SkeletalRenderGPUSkin.cpp
--- a/Engine/Source/Runtime/Engine/Private/SkeletalRenderGPUSkin.cpp
+++ b/Engine/Source/Runtime/Engine/Private/SkeletalRenderGPUSkin.cpp
@@ -1909,17 +1909,21 @@ static void UpdateVertexFactoryMorph(TArray<TUniquePtr<FGPUBaseSkinVertexFactory
 
 // APEX cloth
 
-static const FVertexFactoryType* GetVertexFactoryTypeCloth(FSkeletalMeshLODRenderData& LODRenderData, ERHIFeatureLevel::Type FeatureLevel)
+static FPSOPrecacheVertexFactoryData GetVertexFactoryDataCloth(FSkeletalMeshLODRenderData& LODRenderData, FGPUSkinDataType& GPUSkinDataType, ERHIFeatureLevel::Type FeatureLevel)
 {
-	GPUSkinBoneInfluenceType BoneInfluenceType = LODRenderData.SkinWeightVertexBuffer.GetBoneInfluenceType();
-	if (BoneInfluenceType == GPUSkinBoneInfluenceType::DefaultBoneInfluence)
+	const FVertexFactoryType* VertexFactoryType;
+	FVertexDeclarationElementList VertexElements;
+	if (LODRenderData.SkinWeightVertexBuffer.GetBoneInfluenceType() == GPUSkinBoneInfluenceType::DefaultBoneInfluence)
 	{
-		return &TGPUSkinAPEXClothVertexFactory<GPUSkinBoneInfluenceType::DefaultBoneInfluence>::StaticType;
+		VertexFactoryType = &TGPUSkinAPEXClothVertexFactory<GPUSkinBoneInfluenceType::DefaultBoneInfluence>::StaticType;
+		TGPUSkinAPEXClothVertexFactory<GPUSkinBoneInfluenceType::DefaultBoneInfluence>::GetVertexElements(FeatureLevel, EVertexInputStreamType::Default, GPUSkinDataType, VertexElements);
 	}
 	else
 	{
-		return &TGPUSkinAPEXClothVertexFactory<GPUSkinBoneInfluenceType::UnlimitedBoneInfluence>::StaticType;
+		VertexFactoryType = &TGPUSkinAPEXClothVertexFactory<GPUSkinBoneInfluenceType::UnlimitedBoneInfluence>::StaticType;
+		TGPUSkinAPEXClothVertexFactory<GPUSkinBoneInfluenceType::UnlimitedBoneInfluence>::GetVertexElements(FeatureLevel, EVertexInputStreamType::Default, GPUSkinDataType, VertexElements);
 	}
+	return FPSOPrecacheVertexFactoryData(VertexFactoryType, VertexElements);
 }
 
 /**
@@ -2027,14 +2031,15 @@ void FSkeletalMeshObjectGPUSkin::GetUsedVertexFactoryData(
 	{
 		VertexFactoryDataList.AddUnique(FPSOPrecacheVertexFactoryData(&FGPUSkinPassthroughVertexFactory::StaticType));
 	}
-	else
+	
+	if (!bHasMeshDeformer)
 	{
 		// Add GPU skin cloth vertex factory type is needed
 		const EShaderPlatform ShaderPlatform = GetFeatureLevelShaderPlatform(InFeatureLevel);
 		const bool bClothEnabled = FGPUBaseSkinAPEXClothVertexFactory::IsClothEnabled(ShaderPlatform);
 		if (bClothEnabled && RenderSection.HasClothingData())
 		{
-			VertexFactoryDataList.AddUnique(FPSOPrecacheVertexFactoryData(GetVertexFactoryTypeCloth(LODRenderData, InFeatureLevel)));
+			VertexFactoryDataList.AddUnique(GetVertexFactoryDataCloth(LODRenderData, GPUSkinDataType, InFeatureLevel));
 		}
 		else
 		{
