Change UIMax for displacement Center property in Materials to match the clamped max. This was causing warnings in log from having min and max set to 0.

#rnx

diff --git a/Engine/Source/Runtime/Engine/Classes/Engine/EngineTypes.h b/Engine/Source/Runtime/Engine/Classes/Engine/EngineTypes.h
--- a/Engine/Source/Runtime/Engine/Classes/Engine/EngineTypes.h
+++ b/Engine/Source/Runtime/Engine/Classes/Engine/EngineTypes.h
@@ -2975,7 +2975,7 @@ struct FDisplacementScaling
 	UPROPERTY(EditAnywhere, BlueprintReadWrite, Category = Displacement, meta = (NoSpinbox = true, ClampMin = "0.0", UIMin = "0.0"))
 	float Magnitude;
 
-	UPROPERTY(EditAnywhere, BlueprintReadWrite, Category = Displacement, meta = (NoSpinbox = true, ClampMin = "0.0", ClampMax = "1.0", UIMin = "0.0", UIMax = "0.0"))
+	UPROPERTY(EditAnywhere, BlueprintReadWrite, Category = Displacement, meta = (NoSpinbox = true, ClampMin = "0.0", ClampMax = "1.0", UIMin = "0.0", UIMax = "1.0"))
 	float Center;
 
 public:
