Replaced recursion with a loop in the CalculateComponentSpaceTransform

diff --git a/Engine/Source/Runtime/Engine/Public/BonePose.h b/Engine/Source/Runtime/Engine/Public/BonePose.h
--- a/Engine/Source/Runtime/Engine/Public/BonePose.h
+++ b/Engine/Source/Runtime/Engine/Public/BonePose.h
@@ -596,46 +596,56 @@ void FCSPose<PoseType>::CalculateComponentSpaceTransform(BoneIndexType BoneIndex
 	checkSlow(Pose.IsValid());
 	check(ComponentSpaceFlags[BoneIndex] == 0);
 
-	// root is already verified, so root should not come here
-	// check AllocateLocalPoses
-	const BoneIndexType ParentIndex = Pose.GetParentBoneIndex(BoneIndex);
+	TArray<BoneIndexType> BoneIndexStack;
+	BoneIndexStack.Reserve(ComponentSpaceFlags.Num());
+	BoneIndexStack.Add(BoneIndex);	//Add a dummy index to avoid last element checks in the loop
 
-	// if Parent already has been calculated, use it
-	if (ComponentSpaceFlags[ParentIndex] == 0)
+	do
 	{
-		// if Parent hasn't been calculated, also calculate parents
-		CalculateComponentSpaceTransform(ParentIndex);
-	}
+		// root is already verified, so root should not come here
+		// check AllocateLocalPoses
+		const BoneIndexType ParentIndex = Pose.GetParentBoneIndex(BoneIndex);
 
-	// current Bones(Index) should contain LocalPoses.
-	FTransform& Bone = Pose[BoneIndex];
-	FTransform& ParentBone = Pose[ParentIndex];
-	check(!Pose[BoneIndex].ContainsNaN());
-	check(!Pose[ParentIndex].ContainsNaN());
+		// if Parent already has been calculated, use it
+		if (ComponentSpaceFlags[ParentIndex] == 0)
+		{
+			BoneIndexStack.Add(BoneIndex);
+			BoneIndex = ParentIndex;
+			continue;
+		}
 
-	FTransform ComponentTransform = Pose[BoneIndex] * Pose[ParentIndex];
-	if (ComponentTransform.ContainsNaN())
-	{
-		// We've failed, output as much info as we can....
-		// Added for Jira UE-55511
-		auto BoolToStr = [](const bool& bValue) { return bValue ? TEXT("true") : TEXT("false"); };
-		
-		const TCHAR* BoneHasNaN = BoolToStr(Pose[BoneIndex].ContainsNaN());
-		const TCHAR* ParentHasNaN = BoolToStr(Pose[ParentIndex].ContainsNaN());
-		FString ErrorMsg = FString(TEXT("NaN created in during FTransform Multiplication\n"));
-		ErrorMsg += FString::Format(TEXT("\tBoneIndex {0} : ParentBoneIndex {1} BoneTransformNaN={2} : ParentTransformNaN={3}\n"), { BoneIndex.GetInt(), ParentIndex.GetInt(), BoneHasNaN, ParentHasNaN });
-		ErrorMsg += FString::Format(TEXT("\tBone {0}\n"), { Pose[BoneIndex].ToString() });
-		ErrorMsg += FString::Format(TEXT("\tParent {0}\n"), { Pose[ParentIndex].ToString() });
-		ErrorMsg += FString::Format(TEXT("\tResult {0}\n"), { ComponentTransform.ToString() });
-		ErrorMsg += FString::Format(TEXT("\tBone B64 {0}\n"), { FBase64::Encode((uint8*)&Pose[BoneIndex], sizeof(FTransform)) });
-		ErrorMsg += FString::Format(TEXT("\tParent B64 {0}\n"), { FBase64::Encode((uint8*)&Pose[ParentIndex], sizeof(FTransform)) });
-		ErrorMsg += FString::Format(TEXT("\tResult B64 {0}\n"), { FBase64::Encode((uint8*)&ComponentTransform, sizeof(FTransform)) });
-		checkf(false, TEXT("Error during CalculateComponentSpaceTransform\n%s"), *ErrorMsg); // Failed during multiplication
-	}
-	Pose[BoneIndex] = ComponentTransform;
-	Pose[BoneIndex].NormalizeRotation();
-	check(!Pose[BoneIndex].ContainsNaN());
-	ComponentSpaceFlags[BoneIndex] = 1;
+		// current Bones(Index) should contain LocalPoses.
+		FTransform& Bone = Pose[BoneIndex];
+		FTransform& ParentBone = Pose[ParentIndex];
+		check(!Pose[BoneIndex].ContainsNaN());
+		check(!Pose[ParentIndex].ContainsNaN());
+
+		FTransform ComponentTransform = Pose[BoneIndex] * Pose[ParentIndex];
+		if (ComponentTransform.ContainsNaN())
+		{
+			// We've failed, output as much info as we can....
+			// Added for Jira UE-55511
+			auto BoolToStr = [](const bool& bValue) { return bValue ? TEXT("true") : TEXT("false"); };
+
+			const TCHAR* BoneHasNaN = BoolToStr(Pose[BoneIndex].ContainsNaN());
+			const TCHAR* ParentHasNaN = BoolToStr(Pose[ParentIndex].ContainsNaN());
+			FString ErrorMsg = FString(TEXT("NaN created in during FTransform Multiplication\n"));
+			ErrorMsg += FString::Format(TEXT("\tBoneIndex {0} : ParentBoneIndex {1} BoneTransformNaN={2} : ParentTransformNaN={3}\n"), { BoneIndex.GetInt(), ParentIndex.GetInt(), BoneHasNaN, ParentHasNaN });
+			ErrorMsg += FString::Format(TEXT("\tBone {0}\n"), { Pose[BoneIndex].ToString() });
+			ErrorMsg += FString::Format(TEXT("\tParent {0}\n"), { Pose[ParentIndex].ToString() });
+			ErrorMsg += FString::Format(TEXT("\tResult {0}\n"), { ComponentTransform.ToString() });
+			ErrorMsg += FString::Format(TEXT("\tBone B64 {0}\n"), { FBase64::Encode((uint8*)&Pose[BoneIndex], sizeof(FTransform)) });
+			ErrorMsg += FString::Format(TEXT("\tParent B64 {0}\n"), { FBase64::Encode((uint8*)&Pose[ParentIndex], sizeof(FTransform)) });
+			ErrorMsg += FString::Format(TEXT("\tResult B64 {0}\n"), { FBase64::Encode((uint8*)&ComponentTransform, sizeof(FTransform)) });
+			checkf(false, TEXT("Error during CalculateComponentSpaceTransform\n%s"), *ErrorMsg); // Failed during multiplication
+		}
+		Pose[BoneIndex] = ComponentTransform;
+		Pose[BoneIndex].NormalizeRotation();
+		check(!Pose[BoneIndex].ContainsNaN());
+		ComponentSpaceFlags[BoneIndex] = 1;
+
+		BoneIndex = BoneIndexStack.Pop(false);
+	} while (BoneIndexStack.Num());
 }
 
 template<class PoseType>
