Fix for ParticleSubUv material expression working on non-vertex fetch shader profiles when used with Niagara mesh rendererer.  The default UV coordinate needs to be reserved on the material in order for the VF to connect the data for rendering

diff --git a/Engine/Source/Runtime/Engine/Private/Materials/MaterialExpressions.cpp b/Engine/Source/Runtime/Engine/Private/Materials/MaterialExpressions.cpp
--- a/Engine/Source/Runtime/Engine/Private/Materials/MaterialExpressions.cpp
+++ b/Engine/Source/Runtime/Engine/Private/Materials/MaterialExpressions.cpp
@@ -4641,6 +4641,11 @@ int32 UMaterialExpressionTextureSampleParameterSubUV::Compile(class FMaterialCom
 		return Compiler->Errorf(TEXT("%s"), *SamplerTypeError);
 	}
 
+	// while this expression does provide a TextureCoordinate input pin, it is, and has always been, ignored.  And only
+	// supports using UV0.  Further, in order to support non-vertex fetch implementations we need to be sure to register
+	// the use of the first texture slot
+	Compiler->TextureCoordinate(0 /*Explit dependency on the 1st uv channel*/, false, false);
+
 	int32 TextureCodeIndex = Compiler->TextureParameter(ParameterName, Texture, SamplerType);
 	return ParticleSubUV(Compiler, TextureCodeIndex, SamplerType, CompileMipValue0(Compiler), CompileMipValue1(Compiler), MipValueMode,	bBlend);
 }
