Fix bug with CurveSource node where it would affect curves not present in the curve source asset

diff --git a/Engine/Source/Runtime/AnimGraphRuntime/Private/AnimNodes/AnimNode_CurveSource.cpp b/Engine/Source/Runtime/AnimGraphRuntime/Private/AnimNodes/AnimNode_CurveSource.cpp
--- a/Engine/Source/Runtime/AnimGraphRuntime/Private/AnimNodes/AnimNode_CurveSource.cpp
+++ b/Engine/Source/Runtime/AnimGraphRuntime/Private/AnimNodes/AnimNode_CurveSource.cpp
@@ -108,7 +108,7 @@ void FAnimNode_CurveSource::Evaluate_AnyThread(FPoseContext& Output)
 		}
 #endif
 
-		Output.Curve.LerpTo(Curve, ClampedAlpha);
+		Output.Curve.LerpToValid(Curve, ClampedAlpha);
 	}
 }
 
diff --git a/Engine/Source/Runtime/Engine/Public/Animation/AnimCurveTypes.h b/Engine/Source/Runtime/Engine/Public/Animation/AnimCurveTypes.h
--- a/Engine/Source/Runtime/Engine/Public/Animation/AnimCurveTypes.h
+++ b/Engine/Source/Runtime/Engine/Public/Animation/AnimCurveTypes.h
@@ -699,7 +699,7 @@ public:
 
 public:
 	/**
-	 * Blend (A, B) using Alpha
+	 * Blend (A, B) using Alpha. Lone valid elements are not preserved.
 	 */
 	template<typename AllocatorA, typename ElementTypeA, typename AllocatorB, typename ElementTypeB>
 	void Lerp(const TBaseBlendedCurve<AllocatorA, ElementTypeA>& A, const TBaseBlendedCurve<AllocatorB, ElementTypeB>& B, float Alpha)
@@ -728,7 +728,7 @@ public:
 	}
 
 	/**
-	 * Blend with Other using Alpha, same as Lerp
+	 * Blend with Other using Alpha. Lone valid elements are not preserved. Same as Lerp
 	 */
 	template<typename OtherAllocator, typename OtherElementType>
 	void LerpTo(const TBaseBlendedCurve<OtherAllocator, OtherElementType>& Other, float Alpha)
@@ -755,6 +755,59 @@ public:
 		}
 	}
 
+	/**
+	* Blend with Other using Alpha when both elements are valid, otherwise preserve valid element.
+	*/
+	template<typename AllocatorA, typename ElementTypeA, typename AllocatorB, typename ElementTypeB>
+	void LerpValid(const TBaseBlendedCurve<AllocatorA, ElementTypeA>& A, const TBaseBlendedCurve<AllocatorB, ElementTypeB>& B, float Alpha)
+	{
+		CURVE_PROFILE_CYCLE_COUNTER(TBaseBlendedCurve_LerpValid);
+		
+		// Combine using Lerp. Result is a merged set of curves in 'this'. Elements where only one input is valid will be preserved.
+		UE::Anim::FNamedValueArrayUtils::Union(*this, A, B, [&Alpha](UE::Anim::FCurveElement& OutResult, const ElementTypeA& InElement0, const ElementTypeB& InElement1, UE::Anim::ENamedValueUnionFlags InFlags)
+		{
+			if (EnumHasAllFlags(InFlags, UE::Anim::ENamedValueUnionFlags::BothArgsValid))
+			{
+				OutResult.Value = FMath::Lerp(InElement0.Value, InElement1.Value, Alpha);
+				OutResult.Flags = InElement0.Flags | InElement1.Flags;
+			}
+			else if(EnumHasAllFlags(InFlags, UE::Anim::ENamedValueUnionFlags::ValidArg0))
+			{
+				OutResult.Value = InElement0.Value;
+				OutResult.Flags = InElement0.Flags;
+			}
+			else if(EnumHasAllFlags(InFlags, UE::Anim::ENamedValueUnionFlags::ValidArg1))
+			{
+				OutResult.Value = InElement1.Value;
+				OutResult.Flags = InElement1.Flags;
+			}				
+		});
+	}
+
+	/**
+	* Blend with Other using Alpha when both elements are valid, otherwise preserve valid element. Same as LerpValid
+	*/
+	template<typename OtherAllocator, typename OtherElementType>
+	void LerpToValid(const TBaseBlendedCurve<OtherAllocator, OtherElementType>& Other, float Alpha)
+	{
+		CURVE_PROFILE_CYCLE_COUNTER(TBaseBlendedCurve_LerpToValid);
+		
+		// Combine using Lerp. Result is a merged set of curves in 'this'. Elements where only one input is valid will be preserved.
+		UE::Anim::FNamedValueArrayUtils::Union(*this, Other, [&Alpha](ElementType& InOutThisElement, const OtherElementType& InOtherElement, UE::Anim::ENamedValueUnionFlags InFlags)
+		{
+			if (EnumHasAllFlags(InFlags, UE::Anim::ENamedValueUnionFlags::BothArgsValid))
+			{
+				InOutThisElement.Value = FMath::Lerp(InOutThisElement.Value, InOtherElement.Value, Alpha);
+				InOutThisElement.Flags |= InOtherElement.Flags;
+			}
+			else if(EnumHasAllFlags(InFlags, UE::Anim::ENamedValueUnionFlags::ValidArg1))
+			{
+				InOutThisElement.Value = InOtherElement.Value;
+				InOutThisElement.Flags = InOtherElement.Flags;
+			}			
+		});
+	}
+
 	/**
 	 * Convert current curves to Additive (this - BaseCurve) if overlapping entries are found
 	 */
