Fix incorrect mips streaming and assert firing when RVT has both low mips removed and a low mip streaming texture.

diff --git a/Engine/Source/Runtime/Engine/Classes/VT/RuntimeVirtualTexture.h b/Engine/Source/Runtime/Engine/Classes/VT/RuntimeVirtualTexture.h
--- a/Engine/Source/Runtime/Engine/Classes/VT/RuntimeVirtualTexture.h
+++ b/Engine/Source/Runtime/Engine/Classes/VT/RuntimeVirtualTexture.h
@@ -68,7 +68,7 @@ protected:
 	bool bContinuousUpdate = false;
 
 	/** Number of low mips to cut from the virtual texture. This can reduce peak virtual texture update cost but will also increase the probability of mip shimmering. */
-	UPROPERTY(EditAnywhere, AdvancedDisplay, Category = Layout, meta = (UIMin = "0", UIMax = "6", DisplayName = "Number of low mips to remove from the virtual texture"))
+	UPROPERTY(EditAnywhere, AdvancedDisplay, Category = Layout, meta = (UIMin = "0", UIMax = "5", DisplayName = "Number of low mips to remove from the virtual texture"))
 	int32 RemoveLowMips = 0;
 
 	/** Texture group this texture belongs to */
diff --git a/Engine/Source/Runtime/Engine/Private/VT/RuntimeVirtualTexture.cpp b/Engine/Source/Runtime/Engine/Private/VT/RuntimeVirtualTexture.cpp
--- a/Engine/Source/Runtime/Engine/Private/VT/RuntimeVirtualTexture.cpp
+++ b/Engine/Source/Runtime/Engine/Private/VT/RuntimeVirtualTexture.cpp
@@ -670,14 +670,15 @@ namespace RuntimeVirtualTexture
 					// Note that streaming data may have mips added/removed during cook.
 					const uint32 BlockWidthInTiles = VTData->GetWidthInTiles();
 					const uint32 BlockHeightInTiles = VTData->GetHeightInTiles();
-					const uint32 MaxLevel = FMath::CeilLogTwo(FMath::Max(BlockWidthInTiles, BlockHeightInTiles));
+					const uint32 NumLevels = FMath::CeilLogTwo(FMath::Max(BlockWidthInTiles, BlockHeightInTiles));
+					const uint32 NumOwnerLevels = FMath::CeilLogTwo(FMath::Max(InOwnerProducerDesc.BlockWidthInTiles, InOwnerProducerDesc.BlockHeightInTiles));
 
 					// Clamp the streaming texture size to the runtime virtual texture.
-					const uint32 FirstMipToUse = MaxLevel > InOwnerProducerDesc.MaxLevel ? MaxLevel - InOwnerProducerDesc.MaxLevel : 0;
+					const uint32 FirstMipToUse = NumLevels > NumOwnerLevels ? NumLevels - NumOwnerLevels : 0;
 
 					OutStreamingProducerDesc.BlockWidthInTiles = BlockWidthInTiles >> FirstMipToUse;
 					OutStreamingProducerDesc.BlockHeightInTiles = BlockHeightInTiles >> FirstMipToUse;
-					OutStreamingProducerDesc.MaxLevel = MaxLevel - FirstMipToUse;
+					OutStreamingProducerDesc.MaxLevel = NumLevels - FirstMipToUse;
 
 					return new FUploadingVirtualTexture(InStreamingTexture->GetFName(), VTData, FirstMipToUse);
 				}
diff --git a/Engine/Source/Runtime/Renderer/Private/VT/RuntimeVirtualTextureSceneProxy.cpp b/Engine/Source/Runtime/Renderer/Private/VT/RuntimeVirtualTextureSceneProxy.cpp
--- a/Engine/Source/Runtime/Renderer/Private/VT/RuntimeVirtualTextureSceneProxy.cpp
+++ b/Engine/Source/Runtime/Renderer/Private/VT/RuntimeVirtualTextureSceneProxy.cpp
@@ -110,8 +110,10 @@ FRuntimeVirtualTextureSceneProxy::FRuntimeVirtualTextureSceneProxy(URuntimeVirtu
 					FVTProducerDescription StreamingProducerDesc;
 					IVirtualTexture* StreamingProducer = RuntimeVirtualTexture::CreateStreamingTextureProducer(StreamingTexture, ProducerDesc, StreamingProducerDesc);
 
-					ensure(ProducerDesc.MaxLevel >= StreamingProducerDesc.MaxLevel);
-					const int32 TransitionLevel = ProducerDesc.MaxLevel - StreamingProducerDesc.MaxLevel;
+					const int32 NumLevels = (int32)FMath::CeilLogTwo(FMath::Max(ProducerDesc.BlockWidthInTiles, ProducerDesc.BlockHeightInTiles));
+					const int32 NumStreamingLevels = (int32)FMath::CeilLogTwo(FMath::Max(StreamingProducerDesc.BlockWidthInTiles, StreamingProducerDesc.BlockHeightInTiles));
+					ensure(NumLevels >= NumStreamingLevels);
+					const int32 TransitionLevel = NumLevels - NumStreamingLevels;
 
 					Producer = RuntimeVirtualTexture::BindStreamingTextureProducer(Producer, StreamingProducer, TransitionLevel);
 
