Hollow out AudioStreamingManager whenever not rendering audio (e.g. server).

diff --git a/Engine/Source/Runtime/Engine/Private/ContentStreaming.cpp b/Engine/Source/Runtime/Engine/Private/ContentStreaming.cpp
--- a/Engine/Source/Runtime/Engine/Private/ContentStreaming.cpp
+++ b/Engine/Source/Runtime/Engine/Private/ContentStreaming.cpp
@@ -782,14 +782,22 @@ FStreamingManagerCollection::FStreamingManagerCollection()
 
 	AddOrRemoveTextureStreamingManagerIfNeeded(true);
 
-	if (FPlatformCompressionUtilities::IsCurrentPlatformUsingStreamCaching())
+	if (FApp::CanEverRenderAudio())
 	{
-		FCachedAudioStreamingManagerParams Params = FPlatformCompressionUtilities::BuildCachedStreamingManagerParams();
-		AudioStreamingManager = new FCachedAudioStreamingManager(Params);
+		if (FPlatformCompressionUtilities::IsCurrentPlatformUsingStreamCaching())
+		{
+			FCachedAudioStreamingManagerParams Params = FPlatformCompressionUtilities::BuildCachedStreamingManagerParams();
+			AudioStreamingManager = new FCachedAudioStreamingManager(Params);
+		}
+		else
+		{
+			AudioStreamingManager = new FLegacyAudioStreamingManager();
+		}
 	}
 	else
 	{
-		AudioStreamingManager = new FLegacyAudioStreamingManager();
+		// cannot render any audio, but code still expects this class to exist.
+		AudioStreamingManager = new FDummyAudioStreamingManager();
 	}
 	
 	AddStreamingManager( AudioStreamingManager );
diff --git a/Engine/Source/Runtime/Engine/Public/ContentStreaming.h b/Engine/Source/Runtime/Engine/Public/ContentStreaming.h
--- a/Engine/Source/Runtime/Engine/Public/ContentStreaming.h
+++ b/Engine/Source/Runtime/Engine/Public/ContentStreaming.h
@@ -623,6 +623,46 @@ protected:
 	virtual void RemoveMemoryCountedFeature(const FAudioStreamCacheMemoryHandle& Feature) { };
 };
 
+/**
+ * Dummy audio streaming manager used on the servers and whenever we cannot render audio
+ */
+struct FDummyAudioStreamingManager final : public IAudioStreamingManager
+{
+	virtual void UpdateResourceStreaming(float DeltaTime, bool bProcessEverything = false) override {}
+	virtual int32 BlockTillAllRequestsFinished(float TimeLimit = 0.0f, bool bLogResults = false) override { return 0; }
+	virtual void CancelForcedResources() override {}
+	virtual void NotifyLevelChange() override {}
+	virtual void SetDisregardWorldResourcesForFrames(int32 NumFrames) {}
+	virtual void AddLevel(class ULevel* Level) {}
+	virtual void RemoveLevel(class ULevel* Level) {}
+	virtual void NotifyLevelOffset(class ULevel* Level, const FVector& Offset) {}
+
+	virtual void AddStreamingSoundWave(const FSoundWaveProxyPtr& SoundWave) override {}
+	virtual void RemoveStreamingSoundWave(const FSoundWaveProxyPtr& SoundWave) override {}
+	virtual void AddForceInlineSoundWave(const FSoundWaveProxyPtr& SoundWave) override {}
+	virtual void RemoveForceInlineSoundWave(const FSoundWaveProxyPtr& SoundWave) override {}
+	virtual void AddMemoryCountedFeature(const FAudioStreamCacheMemoryHandle& Feature) override {}
+	virtual void RemoveMemoryCountedFeature(const FAudioStreamCacheMemoryHandle& Feature) override {}
+	virtual void AddDecoder(ICompressedAudioInfo* CompressedAudioInfo) override {}
+	virtual void RemoveDecoder(ICompressedAudioInfo* CompressedAudioInfo) override {}
+	virtual bool IsManagedStreamingSoundWave(const FSoundWaveProxyPtr& SoundWave) const override { return false; }
+	virtual bool IsStreamingInProgress(const FSoundWaveProxyPtr& SoundWave) override { return false; }
+	virtual bool CanCreateSoundSource(const FWaveInstance* WaveInstance) const override { return false; }
+	virtual void AddStreamingSoundSource(FSoundSource* SoundSource) override {}
+	virtual void RemoveStreamingSoundSource(FSoundSource* SoundSource) override {}
+	virtual bool IsManagedStreamingSoundSource(const FSoundSource* SoundSource) const override { return false; }
+	virtual bool RequestChunk(const FSoundWaveProxyPtr& SoundWave, uint32 ChunkIndex, TFunction<void(EAudioChunkLoadResult)> OnLoadCompleted = [](EAudioChunkLoadResult) {}, ENamedThreads::Type ThreadToCallOnLoadCompletedOn = ENamedThreads::AnyThread, bool bForImmediatePlayback = false) override { return false; }
+	virtual FAudioChunkHandle GetLoadedChunk(const FSoundWaveProxyPtr& SoundWave, uint32 ChunkIndex, bool bBlockForLoad = false, bool bForImmediatePlayback = false) const override { return FAudioChunkHandle(); }
+	virtual uint64 TrimMemory(uint64 NumBytesToFree) override { return 0; }
+	virtual int32 RenderStatAudioStreaming(UWorld* World, FViewport* Viewport, FCanvas* Canvas, int32 X, int32 Y, const FVector* ViewLocation, const FRotator* ViewRotation) override { return 0; }
+	virtual FString GenerateMemoryReport() override { return TEXT(""); }
+	virtual void SetProfilingMode(bool bEnabled) override {}
+
+protected:
+	virtual void AddReferenceToChunk(const FAudioChunkHandle& InHandle) override {}
+	virtual void RemoveReferenceToChunk(const FAudioChunkHandle& InHandle) override {}
+};
+
 /**
  * Interface to add functions specifically related to animation streaming
  */
