Fix crash when unlinking a linked anim layer that nests a linked anim graph

Ensure that all linked anim instances have valid object setup before calling into their anim graphs on link.
Prior to this we only ensured that the instance for the layer in question had its objects set up, but that layer can potentially link out to others that were not set up or may have stale objects

#tests Tested code path in PIE

diff --git a/Engine/Source/Runtime/Engine/Private/Animation/AnimInstance.cpp b/Engine/Source/Runtime/Engine/Private/Animation/AnimInstance.cpp
--- a/Engine/Source/Runtime/Engine/Private/Animation/AnimInstance.cpp
+++ b/Engine/Source/Runtime/Engine/Private/Animation/AnimInstance.cpp
@@ -3046,9 +3046,14 @@ void UAnimInstance::PerformLinkedLayerOverlayOperation(TSubclassOf<UAnimInstance
 	if (IAnimClassInterface* AnimBlueprintClass = IAnimClassInterface::GetFromClass(GetClass()))
 	{
 		UClass* NewClass = InClass.Get();
+		USkeletalMeshComponent* MeshComp = GetSkelMeshComponent();
 
-		// Make sure we have valid objects as initialization can route back out of linked instances into this outer graph
-		GetProxyOnAnyThread<FAnimInstanceProxy>().InitializeObjects(this);
+		MeshComp->ForEachAnimInstance([](UAnimInstance* InInstance)
+		{
+			// Make sure we have valid objects on all instances as initialization can route back
+			// out of linked instances into other graphs, including 'this'
+			InInstance->GetProxyOnAnyThread<FAnimInstanceProxy>().InitializeObjects(InInstance);
+		});
 
 		// Map of group name->nodes, per class, to run under that group instance
 		TMap<UClass*, TMap<FName, TArray<FAnimNode_LinkedAnimLayer*, TInlineAllocator<4>>, TInlineSetAllocator<4>>, TInlineSetAllocator<4>> LayerNodesToSet;
@@ -3085,8 +3090,6 @@ void UAnimInstance::PerformLinkedLayerOverlayOperation(TSubclassOf<UAnimInstance
 			}
 		}
 
-		USkeletalMeshComponent* MeshComp = GetSkelMeshComponent();
-
 		auto UnlinkLayerNodesInInstance = [](UAnimInstance* InAnimInstance, TArrayView<FAnimNode_LinkedAnimLayer*> InLayerNodes)
 		{
 			const IAnimClassInterface* const NewLinkedInstanceClass = IAnimClassInterface::GetFromClass(InAnimInstance->GetClass());
