Substrate - support for inline values on material graph root node during conversion.

diff --git a/Engine/Source/Runtime/Engine/Private/Materials/Material.cpp b/Engine/Source/Runtime/Engine/Private/Materials/Material.cpp
--- a/Engine/Source/Runtime/Engine/Private/Materials/Material.cpp
+++ b/Engine/Source/Runtime/Engine/Private/Materials/Material.cpp
@@ -80,6 +80,7 @@
 #include "Materials/MaterialExpressionClearCoatNormalCustomOutput.h"
 #include "Materials/MaterialExpressionTangentOutput.h"
 #include "Materials/MaterialExpressionConstant.h"
+#include "Materials/MaterialExpressionConstant3Vector.h"
 #include "Materials/MaterialExpressionBreakMaterialAttributes.h"
 #include "MaterialCachedData.h"
 #include "Misc/OutputDeviceArchiveWrapper.h"
@@ -103,7 +104,7 @@
 #include UE_INLINE_GENERATED_CPP_BY_NAME(Material)
 
 #define LOCTEXT_NAMESPACE "Material"
-
+UE_DISABLE_OPTIMIZATION
 static TAutoConsoleVariable<int32> CVarMaterialParameterLegacyChecks(
 	TEXT("r.MaterialParameterLegacyChecks"),
 	0,
@@ -3173,6 +3174,9 @@ EBlendMode ConvertLegacyBlendMode(EBlendMode InBlendMode, FMaterialShadingModelF
 	return InBlendMode;
 }
 
+#define SUBSTRATE_MOVE_CONNECTION 0
+#define SUBSTRATE_COPY_CONNECTION 1
+
 void UMaterial::ConvertMaterialToSubstrateMaterial()
 {
 	/*
@@ -3213,19 +3217,80 @@ void UMaterial::ConvertMaterialToSubstrateMaterial()
 		CurrentNodePosX = FMath::Max(NodeToReplace->MaterialExpressionEditorX + TranslationOffsetX, CurrentNodePosX);
 	};
 
-	auto MoveConnectionTo = [](auto& OldNodeInput, UMaterialExpression* NewNode, uint32 NewInputIndex)
+
+	// for ExpressionInput
+	auto ConnectionTo = [](auto& OldNodeInput, UMaterialExpression* NewNode, uint32 NewInputIndex, uint32 OperationType = SUBSTRATE_MOVE_CONNECTION)
 	{
 		if (OldNodeInput.IsConnected())
 		{
 			NewNode->GetInput(NewInputIndex)->Connect(OldNodeInput.OutputIndex, OldNodeInput.Expression);
-			OldNodeInput.Expression = nullptr;
+			if (OperationType == SUBSTRATE_MOVE_CONNECTION)
+			{
+				OldNodeInput.Expression = nullptr;
+			}
 		}
 	};
-	auto CopyConnectionTo = [](auto& OldNodeInput, UMaterialExpression* NewNode, uint32 NewInputIndex)
+
+	// For material input
+	auto ScalarMatInputConnectionTo = [&](FScalarMaterialInput& OldNodeInput, UMaterialExpression* NewNode, uint32 NewInputIndex, EMaterialProperty Property, uint32 OperationType = SUBSTRATE_MOVE_CONNECTION)
 	{
+		FVector4f DefaultPinValue = FMaterialAttributeDefinitionMap::GetDefaultValue(Property);
+
 		if (OldNodeInput.IsConnected())
 		{
-			NewNode->GetInput(NewInputIndex)->Connect(OldNodeInput.OutputIndex, OldNodeInput.Expression);
+			ConnectionTo(OldNodeInput, NewNode, NewInputIndex, OperationType);
+		}
+		else if (OldNodeInput.UseConstant && DefaultPinValue.X != OldNodeInput.Constant)
+		{
+			UMaterialExpressionConstant* ExpressionConstantScalar = NewObject<UMaterialExpressionConstant>(this);
+			ExpressionConstantScalar->R = OldNodeInput.Constant;
+			NewNode->GetInput(NewInputIndex)->Connect(0, ExpressionConstantScalar);
+			if (OperationType == SUBSTRATE_MOVE_CONNECTION)
+			{
+				OldNodeInput.Expression = nullptr;
+			}
+		}
+	};
+	auto ColorMatInputConnectionTo = [&](FColorMaterialInput& OldNodeInput, UMaterialExpression* NewNode, uint32 NewInputIndex, EMaterialProperty Property, uint32 OperationType = SUBSTRATE_MOVE_CONNECTION)
+	{
+		FVector4f DefaultPinValue = FMaterialAttributeDefinitionMap::GetDefaultValue(Property);
+		FLinearColor DefaultPinLinearColor = FLinearColor(DefaultPinValue.X, DefaultPinValue.Y, DefaultPinValue.Z);
+		FLinearColor CurrentPinLinearColor = OldNodeInput.Constant.ReinterpretAsLinear();
+
+		if (OldNodeInput.IsConnected())
+		{
+			ConnectionTo(OldNodeInput, NewNode, NewInputIndex, OperationType);
+		}
+		else if (OldNodeInput.UseConstant && DefaultPinLinearColor != CurrentPinLinearColor)
+		{
+			UMaterialExpressionConstant3Vector* ExpressionConstantScalar = NewObject<UMaterialExpressionConstant3Vector>(this);
+			ExpressionConstantScalar->Constant = CurrentPinLinearColor;
+			NewNode->GetInput(NewInputIndex)->Connect(0, ExpressionConstantScalar);
+			if (OperationType == SUBSTRATE_MOVE_CONNECTION)
+			{
+				OldNodeInput.Expression = nullptr;
+			}
+		}
+	};
+	auto Vector3MatInputConnectionTo = [&](FVectorMaterialInput& OldNodeInput, UMaterialExpression* NewNode, uint32 NewInputIndex, EMaterialProperty Property, uint32 OperationType = SUBSTRATE_MOVE_CONNECTION)
+	{
+		FVector4f DefaultPinValue = FMaterialAttributeDefinitionMap::GetDefaultValue(Property);
+		FVector3f DefaultPinVec3f = FVector3f(DefaultPinValue.X, DefaultPinValue.Y, DefaultPinValue.Z);
+		FVector3f CurrentPinVec3f = FVector3f(OldNodeInput.Constant.X, OldNodeInput.Constant.Y, OldNodeInput.Constant.Z);
+
+		if (OldNodeInput.IsConnected())
+		{
+			ConnectionTo(OldNodeInput, NewNode, NewInputIndex, OperationType);
+		}
+		else if (OldNodeInput.UseConstant && DefaultPinVec3f != CurrentPinVec3f)
+		{
+			UMaterialExpressionConstant3Vector* ExpressionConstantScalar = NewObject<UMaterialExpressionConstant3Vector>(this);
+			ExpressionConstantScalar->Constant = CurrentPinVec3f;
+			NewNode->GetInput(NewInputIndex)->Connect(0, ExpressionConstantScalar);
+			if (OperationType == SUBSTRATE_MOVE_CONNECTION)
+			{
+				OldNodeInput.Expression = nullptr;
+			}
 		}
 	};
 
@@ -3297,7 +3362,7 @@ void UMaterial::ConvertMaterialToSubstrateMaterial()
 		// * Leave the material attribute existing connection plugged to the root node, 
 		//   so that other input (PixelDepthOffset, WorldPositionOffset, ...) get pull 
 		//   from the material attributes node
-		CopyConnectionTo(EditorOnly->MaterialAttributes, ConvertAttributeNode, 0);
+		ConnectionTo(EditorOnly->MaterialAttributes, ConvertAttributeNode, 0, SUBSTRATE_COPY_CONNECTION);
 
 		// Reconnect custom output to material attribute conversion node
 		{
@@ -3306,22 +3371,22 @@ void UMaterial::ConvertMaterialToSubstrateMaterial()
 
 			if (ThinTranslucentOutput)
 			{
-				MoveConnectionTo(*ThinTranslucentOutput->GetInput(0), ConvertAttributeNode, 1);	 // TransmittanceColor
+				ConnectionTo(*ThinTranslucentOutput->GetInput(0), ConvertAttributeNode, 1);	 // TransmittanceColor
 			}
 			if (SingleLayerWaterOutput)
 			{
-				MoveConnectionTo(*SingleLayerWaterOutput->GetInput(0), ConvertAttributeNode, 2); // WaterScatteringCoefficients
-				MoveConnectionTo(*SingleLayerWaterOutput->GetInput(1), ConvertAttributeNode, 3); // WaterAbsorptionCoefficients
-				MoveConnectionTo(*SingleLayerWaterOutput->GetInput(2), ConvertAttributeNode, 4); // WaterPhaseG
-				MoveConnectionTo(*SingleLayerWaterOutput->GetInput(3), ConvertAttributeNode, 5); // ColorScaleBehindWater
+				ConnectionTo(*SingleLayerWaterOutput->GetInput(0), ConvertAttributeNode, 2); // WaterScatteringCoefficients
+				ConnectionTo(*SingleLayerWaterOutput->GetInput(1), ConvertAttributeNode, 3); // WaterAbsorptionCoefficients
+				ConnectionTo(*SingleLayerWaterOutput->GetInput(2), ConvertAttributeNode, 4); // WaterPhaseG
+				ConnectionTo(*SingleLayerWaterOutput->GetInput(3), ConvertAttributeNode, 5); // ColorScaleBehindWater
 			}
 			if (ClearCoatBottomNormalOutput)
 			{
-				CopyConnectionTo(*ClearCoatBottomNormalOutput->GetInput(0), ConvertAttributeNode, 6); // ClearCoatNormal
+				ConnectionTo(*ClearCoatBottomNormalOutput->GetInput(0), ConvertAttributeNode, 6, SUBSTRATE_COPY_CONNECTION); // ClearCoatNormal
 			}
 			if (TangentOutput)
 			{
-				CopyConnectionTo(*TangentOutput->GetInput(0), ConvertAttributeNode, 7);	// TangentOutput
+				ConnectionTo(*TangentOutput->GetInput(0), ConvertAttributeNode, 7, SUBSTRATE_COPY_CONNECTION);	// TangentOutput
 			}
 		}
 
@@ -3382,6 +3447,14 @@ void UMaterial::ConvertMaterialToSubstrateMaterial()
 				MFCallNode->UpdateFromFunctionResource();
 				MFCallNode->PostEditChange();
 				MFCallNode->ConditionalPostLoad();
+
+				ColorMatInputConnectionTo(EditorOnly->BaseColor,		MFCallNode, 0, MP_BaseColor);
+				ScalarMatInputConnectionTo(EditorOnly->Metallic,		MFCallNode, 1, MP_Metallic);
+				ScalarMatInputConnectionTo(EditorOnly->Specular,		MFCallNode, 2, MP_Specular);
+				ScalarMatInputConnectionTo(EditorOnly->Roughness,		MFCallNode, 3, MP_Roughness);
+				Vector3MatInputConnectionTo(EditorOnly->Normal,			MFCallNode, 4, MP_Normal);
+				ColorMatInputConnectionTo(EditorOnly->EmissiveColor,	MFCallNode, 5, MP_EmissiveColor);
+				ScalarMatInputConnectionTo(EditorOnly->Opacity,			MFCallNode, 6, MP_Opacity);
 			}
 		}
 		else
@@ -3410,9 +3483,9 @@ void UMaterial::ConvertMaterialToSubstrateMaterial()
 					// Create metalness to Slab parameterisation conveersion node
 					UMaterialExpressionSubstrateMetalnessToDiffuseAlbedoF0* SubstrateMetalnessToDiffuseAlbedoF0 = NewObject<UMaterialExpressionSubstrateMetalnessToDiffuseAlbedoF0>(this);
 					SetPosXAndMoveReferenceToTheRight(SubstrateMetalnessToDiffuseAlbedoF0);
-					MoveConnectionTo(EditorOnly->BaseColor, SubstrateMetalnessToDiffuseAlbedoF0, 0);					// BaseColor
-					MoveConnectionTo(EditorOnly->Metallic, SubstrateMetalnessToDiffuseAlbedoF0, 1);					// Metallic
-					MoveConnectionTo(EditorOnly->Specular, SubstrateMetalnessToDiffuseAlbedoF0, 2);					// Specular
+					ColorMatInputConnectionTo(EditorOnly->BaseColor, SubstrateMetalnessToDiffuseAlbedoF0, 0, MP_BaseColor);
+					ScalarMatInputConnectionTo(EditorOnly->Metallic, SubstrateMetalnessToDiffuseAlbedoF0, 1, MP_Metallic);
+					ScalarMatInputConnectionTo(EditorOnly->Specular, SubstrateMetalnessToDiffuseAlbedoF0, 2, MP_Specular);
 					
 					// Top slab BSDF as a simple Disney material
 					UMaterialExpressionSubstrateSlabBSDF* BottomSlabBSDF = NewObject<UMaterialExpressionSubstrateSlabBSDF>(this);
@@ -3421,21 +3494,21 @@ void UMaterial::ConvertMaterialToSubstrateMaterial()
 					BottomSlabBSDF->GetInput(0)->Connect(0, SubstrateMetalnessToDiffuseAlbedoF0);
 					BottomSlabBSDF->GetInput(1)->Connect(1, SubstrateMetalnessToDiffuseAlbedoF0);
 					BottomSlabBSDF->GetInput(2)->Connect(2, SubstrateMetalnessToDiffuseAlbedoF0);
-					MoveConnectionTo(EditorOnly->Roughness, BottomSlabBSDF, 3);					// Roughness
-					CopyConnectionTo(EditorOnly->Anisotropy, BottomSlabBSDF, 4);				// Anisotropy
-					MoveConnectionTo(EditorOnly->Tangent, BottomSlabBSDF, 6);					// Tangent
+					ScalarMatInputConnectionTo(EditorOnly->Roughness, BottomSlabBSDF, 3, MP_Roughness);
+					ScalarMatInputConnectionTo(EditorOnly->Anisotropy, BottomSlabBSDF, 4, MP_Anisotropy, SUBSTRATE_COPY_CONNECTION);
+					Vector3MatInputConnectionTo(EditorOnly->Tangent, BottomSlabBSDF, 6, MP_Tangent);
 
 					check(ClearCoatBottomNormalOutput);
-					CopyConnectionTo(*ClearCoatBottomNormalOutput->GetInput(0), BottomSlabBSDF, 5);// ClearColorBottomNormal -> BottomSlabBSDF.Normal
+					ConnectionTo(*ClearCoatBottomNormalOutput->GetInput(0), BottomSlabBSDF, 5, SUBSTRATE_COPY_CONNECTION);// ClearColorBottomNormal -> BottomSlabBSDF.Normal
 
 					// Now weight the top base material by opacity.
 					UMaterialExpressionSubstrateSlabBSDF* TopSlabBSDF = NewObject<UMaterialExpressionSubstrateSlabBSDF>(this);
 					TopSlabBSDF->Material = this;
 					TopSlabBSDF->MaterialExpressionEditorX = BottomSlabBSDF->MaterialExpressionEditorX;
 					TopSlabBSDF->MaterialExpressionEditorY = BottomSlabBSDF->MaterialExpressionEditorY + 650;
-					MoveConnectionTo(EditorOnly->EmissiveColor, TopSlabBSDF, 10);				// Emissive
-					MoveConnectionTo(EditorOnly->ClearCoatRoughness, TopSlabBSDF, 3);			// ClearCoatRoughness => Roughness
-					MoveConnectionTo(EditorOnly->Normal, TopSlabBSDF, 5);						// Normal
+					ColorMatInputConnectionTo(EditorOnly->EmissiveColor, TopSlabBSDF, 10, MP_EmissiveColor);
+					ScalarMatInputConnectionTo(EditorOnly->ClearCoatRoughness, TopSlabBSDF, 3, MP_CustomData0);	// ClearCoatRoughness => Roughness
+					Vector3MatInputConnectionTo(EditorOnly->Normal, TopSlabBSDF, 5, MP_Normal);
 
 					//  The top layer has a hard coded specular value of 0.5 (F0 = 0.04)
 					UMaterialExpressionConstant* ConstantHalf = NewObject<UMaterialExpressionConstant>(this);
@@ -3464,8 +3537,9 @@ void UMaterial::ConvertMaterialToSubstrateMaterial()
 					// Now weight the top base material by ClearCoat
 					UMaterialExpressionSubstrateWeight* TopSlabBSDFWithCoverage = NewObject<UMaterialExpressionSubstrateWeight>(this);
 					SetPosXAndMoveReferenceToTheRight(TopSlabBSDFWithCoverage);
-					TopSlabBSDFWithCoverage->GetInput(0)->Connect(0, TopSlabBSDF);				// TopSlabBSDF -> A
-					MoveConnectionTo(EditorOnly->ClearCoat, TopSlabBSDFWithCoverage, 1);		// ClearCoat -> Weight
+					TopSlabBSDFWithCoverage->GetInput(0)->Connect(0, TopSlabBSDF);												// TopSlabBSDF -> A
+					ScalarMatInputConnectionTo(EditorOnly->ClearCoat, TopSlabBSDFWithCoverage, 1, MP_CustomData0);				// ClearCoat -> Weight
+					ScalarMatInputConnectionTo(EditorOnly->ClearCoatRoughness, TopSlabBSDFWithCoverage, 1, MP_CustomData1);		// ClearCoat -> Weight
 
 					UMaterialExpressionSubstrateVerticalLayering* VerticalLayering = NewObject<UMaterialExpressionSubstrateVerticalLayering>(this);
 					SetPosXAndMoveReferenceToTheRight(VerticalLayering);
@@ -3484,18 +3558,18 @@ void UMaterial::ConvertMaterialToSubstrateMaterial()
 				ConvertNode->Material = this;
 				SetPosXAndMoveReferenceToTheRight(ConvertNode);
 				ConvertNode->SubsurfaceProfile = bRequireNoSubsurfaceProfile ? nullptr : SubsurfaceProfile;
-				MoveConnectionTo(EditorOnly->BaseColor, ConvertNode, 0);
-				MoveConnectionTo(EditorOnly->Metallic, ConvertNode, 1);
-				MoveConnectionTo(EditorOnly->Specular, ConvertNode, 2);
-				MoveConnectionTo(EditorOnly->Roughness, ConvertNode, 3);
-				MoveConnectionTo(EditorOnly->Anisotropy, ConvertNode, 4);
-				MoveConnectionTo(EditorOnly->EmissiveColor, ConvertNode, 5);
-				CopyConnectionTo(EditorOnly->Normal, ConvertNode, 6);
-				MoveConnectionTo(EditorOnly->Tangent, ConvertNode, 7);
-				MoveConnectionTo(EditorOnly->SubsurfaceColor, ConvertNode, 8);
-				MoveConnectionTo(EditorOnly->ClearCoat, ConvertNode, 9);
-				MoveConnectionTo(EditorOnly->ClearCoatRoughness, ConvertNode, 10);
-				CopyConnectionTo(EditorOnly->Opacity, ConvertNode, 11);	// We only copy, to keep Opacity on the root node in case BLEND_AlphaComposite is selected.
+				ColorMatInputConnectionTo(EditorOnly->BaseColor, ConvertNode, 0, MP_BaseColor);
+				ScalarMatInputConnectionTo(EditorOnly->Metallic, ConvertNode, 1, MP_Metallic);
+				ScalarMatInputConnectionTo(EditorOnly->Specular, ConvertNode, 2, MP_Specular);
+				ScalarMatInputConnectionTo(EditorOnly->Roughness, ConvertNode, 3, MP_Roughness);
+				ScalarMatInputConnectionTo(EditorOnly->Anisotropy, ConvertNode, 4, MP_Anisotropy);
+				ColorMatInputConnectionTo(EditorOnly->EmissiveColor, ConvertNode, 5, MP_EmissiveColor);
+				Vector3MatInputConnectionTo(EditorOnly->Normal, ConvertNode, 6, MP_Normal, SUBSTRATE_COPY_CONNECTION);
+				Vector3MatInputConnectionTo(EditorOnly->Tangent, ConvertNode, 7, MP_Tangent);
+				ColorMatInputConnectionTo(EditorOnly->SubsurfaceColor, ConvertNode, 8, MP_SubsurfaceColor);
+				ScalarMatInputConnectionTo(EditorOnly->ClearCoat, ConvertNode, 9, MP_CustomData0);
+				ScalarMatInputConnectionTo(EditorOnly->ClearCoatRoughness, ConvertNode, 10, MP_CustomData1);
+				ScalarMatInputConnectionTo(EditorOnly->Opacity, ConvertNode, 11, MP_Opacity, SUBSTRATE_COPY_CONNECTION);	// We only copy, to keep Opacity on the root node in case BLEND_AlphaComposite is selected.
 				bRelinkCustomOutputNodes = true;
 			
 				// Shading Model
@@ -3539,10 +3613,10 @@ void UMaterial::ConvertMaterialToSubstrateMaterial()
 			UMaterialExpressionSubstrateVolumetricFogCloudBSDF* VolBSDF = NewObject<UMaterialExpressionSubstrateVolumetricFogCloudBSDF>(this);
 			VolBSDF->Material = this;
 			SetPosXAndMoveReferenceToTheRight(VolBSDF);
-			MoveConnectionTo(EditorOnly->BaseColor, VolBSDF, 0);		// Albedo
-			MoveConnectionTo(EditorOnly->SubsurfaceColor, VolBSDF, 1);	// Extinction
-			MoveConnectionTo(EditorOnly->EmissiveColor, VolBSDF, 2);	// EmissiveColor
-			MoveConnectionTo(EditorOnly->AmbientOcclusion, VolBSDF, 3);	// AmbientOcclusion
+			ColorMatInputConnectionTo(EditorOnly->BaseColor, VolBSDF, 0, MP_BaseColor);	
+			ColorMatInputConnectionTo(EditorOnly->SubsurfaceColor, VolBSDF, 1, MP_SubsurfaceColor);
+			ColorMatInputConnectionTo(EditorOnly->EmissiveColor, VolBSDF, 2, MP_EmissiveColor);	
+			ScalarMatInputConnectionTo(EditorOnly->AmbientOcclusion, VolBSDF, 3, MP_AmbientOcclusion);
 
 			// SUBSTRATE_TODO remove the VolumetricAdvancedOutput node and add the input onto FogCloudBSDF even if only used by the cloud renderer?
 			EditorOnly->FrontMaterial.Connect(0, VolBSDF);
@@ -3559,7 +3633,7 @@ void UMaterial::ConvertMaterialToSubstrateMaterial()
 			UMaterialExpressionSubstrateLightFunction* LightFunctionNode = NewObject<UMaterialExpressionSubstrateLightFunction>(this);
 			LightFunctionNode->Material = this;
 			SetPosXAndMoveReferenceToTheRight(LightFunctionNode);
-			MoveConnectionTo(EditorOnly->EmissiveColor, LightFunctionNode, 0);
+			ColorMatInputConnectionTo(EditorOnly->EmissiveColor, LightFunctionNode, 0, MP_EmissiveColor);
 
 			EditorOnly->FrontMaterial.Connect(0, LightFunctionNode);
 			bInvalidateShader = true;
@@ -3580,8 +3654,8 @@ void UMaterial::ConvertMaterialToSubstrateMaterial()
 			PostProcNode->Material = this;
 			SetPosXAndMoveReferenceToTheRight(PostProcNode);
 
-			MoveConnectionTo(EditorOnly->EmissiveColor, PostProcNode, 0);
-			CopyConnectionTo(EditorOnly->Opacity, PostProcNode, 1);	// We only copy, to keep Opacity on the root node in case BLEND_AlphaComposite is selected.
+			ColorMatInputConnectionTo(EditorOnly->EmissiveColor, PostProcNode, 0, MP_EmissiveColor);
+			ScalarMatInputConnectionTo(EditorOnly->Opacity, PostProcNode, 1, MP_Opacity, SUBSTRATE_COPY_CONNECTION);	// We only copy, to keep Opacity on the root node in case BLEND_AlphaComposite is selected.
 
 			EditorOnly->FrontMaterial.Connect(0, PostProcNode);
 			bInvalidateShader = true;
@@ -3596,18 +3670,18 @@ void UMaterial::ConvertMaterialToSubstrateMaterial()
 			ConvertNode = NewObject<UMaterialExpressionSubstrateShadingModels>(this);
 			ConvertNode->Material = this;
 			SetPosXAndMoveReferenceToTheRight(ConvertNode);
-			MoveConnectionTo(EditorOnly->BaseColor, ConvertNode, 0);
-			MoveConnectionTo(EditorOnly->Metallic, ConvertNode, 1);
-			MoveConnectionTo(EditorOnly->Specular, ConvertNode, 2);
-			MoveConnectionTo(EditorOnly->Roughness, ConvertNode, 3);
-			MoveConnectionTo(EditorOnly->Anisotropy, ConvertNode, 4);
-			MoveConnectionTo(EditorOnly->EmissiveColor, ConvertNode, 5);
-			CopyConnectionTo(EditorOnly->Normal, ConvertNode, 6);
-			MoveConnectionTo(EditorOnly->Tangent, ConvertNode, 7);
-			MoveConnectionTo(EditorOnly->SubsurfaceColor, ConvertNode, 8);
-			MoveConnectionTo(EditorOnly->ClearCoat, ConvertNode, 9);
-			MoveConnectionTo(EditorOnly->ClearCoatRoughness, ConvertNode, 10);
-			CopyConnectionTo(EditorOnly->Opacity, ConvertNode, 11);	// We only copy, to keep Opacity on the root node in case BLEND_AlphaComposite is selected.
+			ColorMatInputConnectionTo(EditorOnly->BaseColor, ConvertNode, 0, MP_BaseColor);
+			ScalarMatInputConnectionTo(EditorOnly->Metallic, ConvertNode, 1, MP_Metallic);
+			ScalarMatInputConnectionTo(EditorOnly->Specular, ConvertNode, 2, MP_Specular);
+			ScalarMatInputConnectionTo(EditorOnly->Roughness, ConvertNode, 3, MP_Roughness);
+			ScalarMatInputConnectionTo(EditorOnly->Anisotropy, ConvertNode, 4, MP_Anisotropy);
+			ColorMatInputConnectionTo(EditorOnly->EmissiveColor, ConvertNode, 5, MP_EmissiveColor);
+			Vector3MatInputConnectionTo(EditorOnly->Normal, ConvertNode, 6, MP_Normal, SUBSTRATE_COPY_CONNECTION);
+			Vector3MatInputConnectionTo(EditorOnly->Tangent, ConvertNode, 7, MP_Tangent);
+			ColorMatInputConnectionTo(EditorOnly->SubsurfaceColor, ConvertNode, 8, MP_SubsurfaceColor);
+			ScalarMatInputConnectionTo(EditorOnly->ClearCoat, ConvertNode, 9, MP_CustomData0);
+			ScalarMatInputConnectionTo(EditorOnly->ClearCoatRoughness, ConvertNode, 10, MP_CustomData1);
+			ScalarMatInputConnectionTo(EditorOnly->Opacity, ConvertNode, 11, MP_Opacity, SUBSTRATE_COPY_CONNECTION);	// We only copy, to keep Opacity on the root node in case BLEND_AlphaComposite is selected.
 
 			// Add constant for the Unlit shading model
 			ConvertNode->ShadingModelOverride = ShadingModel;
@@ -3632,8 +3706,8 @@ void UMaterial::ConvertMaterialToSubstrateMaterial()
 			UMaterialExpressionSubstrateUI* UINode = NewObject<UMaterialExpressionSubstrateUI>(this);
 			UINode->Material = this;
 			SetPosXAndMoveReferenceToTheRight(UINode);
-			MoveConnectionTo(EditorOnly->EmissiveColor, UINode, 0);
-			CopyConnectionTo(EditorOnly->Opacity, UINode, 1);	// We only copy, to keep Opacity on the root node in case BLEND_AlphaComposite is selected.
+			ColorMatInputConnectionTo(EditorOnly->EmissiveColor, UINode, 0, MP_EmissiveColor);
+			ScalarMatInputConnectionTo(EditorOnly->Opacity, UINode, 1, MP_Opacity, SUBSTRATE_COPY_CONNECTION);	// We only copy, to keep Opacity on the root node in case BLEND_AlphaComposite is selected.
 
 			EditorOnly->FrontMaterial.Connect(0, UINode);
 			bInvalidateShader = true;
@@ -3650,22 +3724,22 @@ void UMaterial::ConvertMaterialToSubstrateMaterial()
 
 		if (ThinTranslucentOutput)
 		{
-			MoveConnectionTo(*ThinTranslucentOutput->GetInput(0), ConvertNode, 12);	 // TransmittanceColor
+			ConnectionTo(*ThinTranslucentOutput->GetInput(0), ConvertNode, 12);	 // TransmittanceColor
 		}
 		if (SingleLayerWaterOutput)
 		{
-			MoveConnectionTo(*SingleLayerWaterOutput->GetInput(0), ConvertNode, 13); // WaterScatteringCoefficients
-			MoveConnectionTo(*SingleLayerWaterOutput->GetInput(1), ConvertNode, 14); // WaterAbsorptionCoefficients
-			MoveConnectionTo(*SingleLayerWaterOutput->GetInput(2), ConvertNode, 15); // WaterPhaseG
-			MoveConnectionTo(*SingleLayerWaterOutput->GetInput(3), ConvertNode, 16); // ColorScaleBehindWater
+			ConnectionTo(*SingleLayerWaterOutput->GetInput(0), ConvertNode, 13); // WaterScatteringCoefficients
+			ConnectionTo(*SingleLayerWaterOutput->GetInput(1), ConvertNode, 14); // WaterAbsorptionCoefficients
+			ConnectionTo(*SingleLayerWaterOutput->GetInput(2), ConvertNode, 15); // WaterPhaseG
+			ConnectionTo(*SingleLayerWaterOutput->GetInput(3), ConvertNode, 16); // ColorScaleBehindWater
 		}
 		if (ClearCoatBottomNormalOutput)
 		{
-			CopyConnectionTo(*ClearCoatBottomNormalOutput->GetInput(0), ConvertNode, 17); // ClearCoatNormal
+			ConnectionTo(*ClearCoatBottomNormalOutput->GetInput(0), ConvertNode, 17, SUBSTRATE_COPY_CONNECTION); // ClearCoatNormal
 		}
 		if (TangentOutput)
 		{
-			CopyConnectionTo(*TangentOutput->GetInput(0), ConvertNode, 18);	// TangentOutput
+			ConnectionTo(*TangentOutput->GetInput(0), ConvertNode, 18, SUBSTRATE_COPY_CONNECTION);	// TangentOutput
 		}
 	}
 
