Removed left disabled optimisation.

diff --git a/Engine/Source/Runtime/Engine/Private/Materials/Material.cpp b/Engine/Source/Runtime/Engine/Private/Materials/Material.cpp
--- a/Engine/Source/Runtime/Engine/Private/Materials/Material.cpp
+++ b/Engine/Source/Runtime/Engine/Private/Materials/Material.cpp
@@ -104,7 +104,7 @@
 #include UE_INLINE_GENERATED_CPP_BY_NAME(Material)
 
 #define LOCTEXT_NAMESPACE "Material"
-UE_DISABLE_OPTIMIZATION
+
 static TAutoConsoleVariable<int32> CVarMaterialParameterLegacyChecks(
 	TEXT("r.MaterialParameterLegacyChecks"),
 	0,
