Iris - Fixed name of csv stats for DataStreamChannel PostTickDispatch send

#rnx

diff --git a/Engine/Source/Runtime/Engine/Private/Net/Experimental/Iris/DataStreamChannel.cpp b/Engine/Source/Runtime/Engine/Private/Net/Experimental/Iris/DataStreamChannel.cpp
--- a/Engine/Source/Runtime/Engine/Private/Net/Experimental/Iris/DataStreamChannel.cpp
+++ b/Engine/Source/Runtime/Engine/Private/Net/Experimental/Iris/DataStreamChannel.cpp
@@ -197,10 +197,10 @@ void UDataStreamChannel::PostTickDispatch()
 	}
 
 #if UE_NET_IRIS_CSV_STATS
-	CSV_SCOPED_TIMING_STAT_EXCLUSIVE(UDataStreamChannel_Tick);
+	CSV_SCOPED_TIMING_STAT_EXCLUSIVE(UDataStreamChannel_PostTickDispatch);
 #endif
 
-	IRIS_PROFILER_SCOPE(UDataStreamChannel_Tick);
+	IRIS_PROFILER_SCOPE(UDataStreamChannel_PostTickDispatch);
 	LLM_SCOPE_BYTAG(Iris);
 
 	WriteData(EDataStreamWriteMode::PostTickDispatch);
