Add option to show in PhAT when a pose violates the joint limits

diff --git a/Engine/Source/Editor/ComponentVisualizers/Private/ConstraintComponentVisualizer.cpp b/Engine/Source/Editor/ComponentVisualizers/Private/ConstraintComponentVisualizer.cpp
--- a/Engine/Source/Editor/ComponentVisualizers/Private/ConstraintComponentVisualizer.cpp
+++ b/Engine/Source/Editor/ComponentVisualizers/Private/ConstraintComponentVisualizer.cpp
@@ -72,7 +72,7 @@ void FConstraintComponentVisualizer::DrawVisualization( const UActorComponent* C
 		FBox Body2Box = ConstraintComp->GetBodyBox(EConstraintFrame::Frame2);
 
 		// Draw constraint information
-		Instance.DrawConstraint(PDI, 1.f, 1.f, true, true, Con1Frame, Con2Frame, false);
+		Instance.DrawConstraint(PDI, 1.f, 1.f, true, true, Con1Frame, Con2Frame, false, false);
 
 		// Draw boxes to indicate bodies connected by joint.
 		if(Body1Box.IsValid)
diff --git a/Engine/Source/Editor/Persona/Private/PhysicsAssetRenderUtils.cpp b/Engine/Source/Editor/Persona/Private/PhysicsAssetRenderUtils.cpp
--- a/Engine/Source/Editor/Persona/Private/PhysicsAssetRenderUtils.cpp
+++ b/Engine/Source/Editor/Persona/Private/PhysicsAssetRenderUtils.cpp
@@ -55,6 +55,7 @@ FPhysicsAssetRenderSettings::FPhysicsAssetRenderSettings()
 	, bRenderOnlySelectedConstraints(false)
 	, bShowCOM(false)
 	, bShowConstraintsAsPoints(false)
+	, bDrawViolatedLimits(false)
 	, BoneUnselectedColor(170, 155, 225)
 	, NoCollisionColor(200, 200, 200)
 	, COMRenderColor(255, 255, 100)
@@ -260,6 +261,7 @@ void FPhysicsAssetRenderSettings::ResetEditorViewportOptions()
 	bHideSimulatedBodies = DefaultObject.bHideSimulatedBodies;
 	bRenderOnlySelectedConstraints = DefaultObject.bRenderOnlySelectedConstraints;
 	bShowConstraintsAsPoints = DefaultObject.bShowConstraintsAsPoints;
+	bDrawViolatedLimits = DefaultObject.bDrawViolatedLimits;
 }
 
 namespace PhysicsAssetRender
@@ -478,7 +480,7 @@ namespace PhysicsAssetRender
 							Con1Frame.RemoveScaling();
 							Con2Frame.RemoveScaling();
 
-							ConstraintSetup->DefaultInstance.DrawConstraint(PDI, RenderSettings->ConstraintDrawSize, DrawScale, bDrawLimits, bDrawSelected, Con1Frame, Con2Frame, RenderSettings->bShowConstraintsAsPoints);
+							ConstraintSetup->DefaultInstance.DrawConstraint(PDI, RenderSettings->ConstraintDrawSize, DrawScale, bDrawLimits, bDrawSelected, Con1Frame, Con2Frame, RenderSettings->bShowConstraintsAsPoints, RenderSettings->bDrawViolatedLimits);
 
 						}
 					}
diff --git a/Engine/Source/Editor/Persona/Public/PhysicsAssetRenderUtils.h b/Engine/Source/Editor/Persona/Public/PhysicsAssetRenderUtils.h
--- a/Engine/Source/Editor/Persona/Public/PhysicsAssetRenderUtils.h
+++ b/Engine/Source/Editor/Persona/Public/PhysicsAssetRenderUtils.h
@@ -109,6 +109,9 @@ public:
 	UPROPERTY()
 	bool bShowConstraintsAsPoints;
 
+	UPROPERTY()
+	bool bDrawViolatedLimits;
+
 	// Draw colors
 	UPROPERTY()
 	FColor BoneUnselectedColor;
diff --git a/Engine/Source/Editor/PhysicsAssetEditor/Private/PhysicsAssetEditor.cpp b/Engine/Source/Editor/PhysicsAssetEditor/Private/PhysicsAssetEditor.cpp
--- a/Engine/Source/Editor/PhysicsAssetEditor/Private/PhysicsAssetEditor.cpp
+++ b/Engine/Source/Editor/PhysicsAssetEditor/Private/PhysicsAssetEditor.cpp
@@ -830,6 +830,7 @@ void FPhysicsAssetEditor::ExtendViewportMenus()
 				{
 					FToolMenuSection& Section = InSubMenu->AddSection("PhysicsAssetEditorConstraints", LOCTEXT("ConstraintHeader", "Constraints"));
 					Section.AddMenuEntry(Commands.DrawConstraintsAsPoints);
+					Section.AddMenuEntry(Commands.DrawViolatedLimits);
 					Section.AddMenuEntry(Commands.RenderOnlySelectedConstraints);
 					Section.AddEntry(FToolMenuEntry::InitWidget(TEXT("ConstraintScale"), WeakPhysicsAssetEditor.Pin()->MakeConstraintScaleWidget(), LOCTEXT("ConstraintScaleLabel", "Constraint Scale")));
 				}
@@ -1457,6 +1458,12 @@ void FPhysicsAssetEditor::BindCommands()
 		FCanExecuteAction(),
 		FIsActionChecked::CreateSP(this, &FPhysicsAssetEditor::IsDrawingConstraintsAsPoints));
 
+	ViewportCommandList->MapAction(
+		Commands.DrawViolatedLimits,
+		FExecuteAction::CreateSP(this, &FPhysicsAssetEditor::ToggleDrawViolatedLimits),
+		FCanExecuteAction(),
+		FIsActionChecked::CreateSP(this, &FPhysicsAssetEditor::IsDrawingViolatedLimits));
+
 	ViewportCommandList->MapAction(
 		Commands.RenderOnlySelectedConstraints,
 		FExecuteAction::CreateSP(this, &FPhysicsAssetEditor::ToggleRenderOnlySelectedConstraints),
@@ -2746,6 +2753,17 @@ bool FPhysicsAssetEditor::IsDrawingConstraintsAsPoints() const
 	return SharedData->EditorOptions->bShowConstraintsAsPoints;
 }
 
+void FPhysicsAssetEditor::ToggleDrawViolatedLimits()
+{
+	SharedData->EditorOptions->bDrawViolatedLimits = !SharedData->EditorOptions->bDrawViolatedLimits;
+	SharedData->EditorOptions->SaveConfig();
+}
+
+bool FPhysicsAssetEditor::IsDrawingViolatedLimits() const
+{
+	return SharedData->EditorOptions->bDrawViolatedLimits;
+}
+
 void FPhysicsAssetEditor::ToggleRenderOnlySelectedConstraints()
 {
 	SharedData->EditorOptions->bRenderOnlySelectedConstraints = !SharedData->EditorOptions->bRenderOnlySelectedConstraints;
diff --git a/Engine/Source/Editor/PhysicsAssetEditor/Private/PhysicsAssetEditor.h b/Engine/Source/Editor/PhysicsAssetEditor/Private/PhysicsAssetEditor.h
--- a/Engine/Source/Editor/PhysicsAssetEditor/Private/PhysicsAssetEditor.h
+++ b/Engine/Source/Editor/PhysicsAssetEditor/Private/PhysicsAssetEditor.h
@@ -244,6 +244,8 @@ private:
 	bool IsConstraintRenderingMode(EPhysicsAssetEditorConstraintViewMode Mode, bool bSimulation) const;
 	void ToggleDrawConstraintsAsPoints();
 	bool IsDrawingConstraintsAsPoints() const;
+	void ToggleDrawViolatedLimits();
+	bool IsDrawingViolatedLimits() const;
 	void ToggleRenderOnlySelectedConstraints();
 	bool IsRenderingOnlySelectedConstraints() const;
 	void ToggleRenderOnlySelectedSolid();
diff --git a/Engine/Source/Editor/PhysicsAssetEditor/Private/PhysicsAssetEditorActions.cpp b/Engine/Source/Editor/PhysicsAssetEditor/Private/PhysicsAssetEditorActions.cpp
--- a/Engine/Source/Editor/PhysicsAssetEditor/Private/PhysicsAssetEditorActions.cpp
+++ b/Engine/Source/Editor/PhysicsAssetEditor/Private/PhysicsAssetEditorActions.cpp
@@ -54,6 +54,7 @@ void FPhysicsAssetEditorCommands::RegisterCommands()
 	UI_COMMAND(HideKinematicBodies, "Hide Kinematic Bodies", "Hide rendering for kinematic bodies", EUserInterfaceActionType::ToggleButton, FInputChord());
 	UI_COMMAND(RenderOnlySelectedConstraints, "Only Selected Constraints", "Draw only selected constraints.", EUserInterfaceActionType::ToggleButton, FInputChord());
 	UI_COMMAND(DrawConstraintsAsPoints, "Draw Constraints As Points", "Draw Constraints As Points", EUserInterfaceActionType::ToggleButton, FInputChord());
+	UI_COMMAND(DrawViolatedLimits, "Draw Violated Limits", "Draw Violated Limits", EUserInterfaceActionType::ToggleButton, FInputChord());
 	UI_COMMAND(ToggleMassProperties, "Mass Properties", "Show Mass Properties For Bodies When Simulation Is Enabled", EUserInterfaceActionType::ToggleButton, FInputChord());
 	UI_COMMAND(DisableCollision, "Disable Collision", "Disable collision between the currently selected bodies", EUserInterfaceActionType::Button, FInputChord(EKeys::RightBracket));
 	UI_COMMAND(DisableCollisionAll, "Disable Collision All", "Disable collision between the currently selected bodies and all bodies", EUserInterfaceActionType::Button, FInputChord(EModifierKey::Control, EKeys::RightBracket));
diff --git a/Engine/Source/Editor/PhysicsAssetEditor/Private/PhysicsAssetEditorActions.h b/Engine/Source/Editor/PhysicsAssetEditor/Private/PhysicsAssetEditorActions.h
--- a/Engine/Source/Editor/PhysicsAssetEditor/Private/PhysicsAssetEditorActions.h
+++ b/Engine/Source/Editor/PhysicsAssetEditor/Private/PhysicsAssetEditorActions.h
@@ -66,6 +66,7 @@ public:
 	TSharedPtr<FUICommandInfo> HideKinematicBodies;
 	TSharedPtr<FUICommandInfo> RenderOnlySelectedConstraints;
 	TSharedPtr<FUICommandInfo> DrawConstraintsAsPoints;
+	TSharedPtr<FUICommandInfo> DrawViolatedLimits;
 	TSharedPtr<FUICommandInfo> ToggleMassProperties;
 	TSharedPtr<FUICommandInfo> DisableCollision;
 	TSharedPtr<FUICommandInfo> DisableCollisionAll;
diff --git a/Engine/Source/Editor/PhysicsAssetEditor/Private/PhysicsAssetEditorSkeletalMeshComponent.cpp b/Engine/Source/Editor/PhysicsAssetEditor/Private/PhysicsAssetEditorSkeletalMeshComponent.cpp
--- a/Engine/Source/Editor/PhysicsAssetEditor/Private/PhysicsAssetEditorSkeletalMeshComponent.cpp
+++ b/Engine/Source/Editor/PhysicsAssetEditor/Private/PhysicsAssetEditorSkeletalMeshComponent.cpp
@@ -127,6 +127,7 @@ void UPhysicsAssetEditorSkeletalMeshComponent::DebugDraw(const FSceneView* View,
 		RenderSettings->bHideSimulatedBodies = SharedData->EditorOptions->bHideSimulatedBodies;
 		RenderSettings->bRenderOnlySelectedConstraints = SharedData->EditorOptions->bRenderOnlySelectedConstraints;
 		RenderSettings->bShowConstraintsAsPoints = SharedData->EditorOptions->bShowConstraintsAsPoints;
+		RenderSettings->bDrawViolatedLimits = SharedData->EditorOptions->bDrawViolatedLimits;
 
 		// Draw Bodies.
 		{
diff --git a/Engine/Source/Editor/UnrealEd/Classes/Preferences/PhysicsAssetEditorOptions.h b/Engine/Source/Editor/UnrealEd/Classes/Preferences/PhysicsAssetEditorOptions.h
--- a/Engine/Source/Editor/UnrealEd/Classes/Preferences/PhysicsAssetEditorOptions.h
+++ b/Engine/Source/Editor/UnrealEd/Classes/Preferences/PhysicsAssetEditorOptions.h
@@ -107,6 +107,10 @@ class UPhysicsAssetEditorOptions : public UObject
 	UPROPERTY(config)
 	uint32 bShowConstraintsAsPoints:1;
 
+	/** Whether to highlight limits that have been violated */
+	UPROPERTY(config)
+	uint32 bDrawViolatedLimits:1;
+
 	/** Whether to only render selected constraints */
 	UPROPERTY(config)
 	uint32 bRenderOnlySelectedConstraints:1;
diff --git a/Engine/Source/Editor/UnrealEd/Private/PreferenceStubs.cpp b/Engine/Source/Editor/UnrealEd/Private/PreferenceStubs.cpp
--- a/Engine/Source/Editor/UnrealEd/Private/PreferenceStubs.cpp
+++ b/Engine/Source/Editor/UnrealEd/Private/PreferenceStubs.cpp
@@ -43,6 +43,7 @@ UPhysicsAssetEditorOptions::UPhysicsAssetEditorOptions(const FObjectInitializer&
 	InterpolationSpeed = 50.f;
 
 	bShowConstraintsAsPoints = false;
+	bDrawViolatedLimits = false;
 	bSimulationFloorCollisionEnabled = true;
 	ConstraintDrawSize = 1.0f;
 
diff --git a/Engine/Source/Runtime/Engine/Classes/PhysicsEngine/ConstraintInstance.h b/Engine/Source/Runtime/Engine/Classes/PhysicsEngine/ConstraintInstance.h
--- a/Engine/Source/Runtime/Engine/Classes/PhysicsEngine/ConstraintInstance.h
+++ b/Engine/Source/Runtime/Engine/Classes/PhysicsEngine/ConstraintInstance.h
@@ -934,16 +934,16 @@ public:
 	// @todo document
 	void DrawConstraint(int32 ViewIndex, class FMeshElementCollector& Collector,
 		float Scale, float LimitDrawScale, bool bDrawLimits, bool bDrawSelected,
-		const FTransform& Con1Frame, const FTransform& Con2Frame, bool bDrawAsPoint) const
+		const FTransform& Con1Frame, const FTransform& Con2Frame, bool bDrawAsPoint, bool bDrawViolatedLimits) const
 	{
-		DrawConstraintImp(FPDIOrCollector(ViewIndex, Collector), Scale, LimitDrawScale, bDrawLimits, bDrawSelected, Con1Frame, Con2Frame, bDrawAsPoint);
+		DrawConstraintImp(FPDIOrCollector(ViewIndex, Collector), Scale, LimitDrawScale, bDrawLimits, bDrawSelected, Con1Frame, Con2Frame, bDrawAsPoint, bDrawViolatedLimits);
 	}
 
 	void DrawConstraint(FPrimitiveDrawInterface* PDI,
 		float Scale, float LimitDrawScale, bool bDrawLimits, bool bDrawSelected,
-		const FTransform& Con1Frame, const FTransform& Con2Frame, bool bDrawAsPoint) const
+		const FTransform& Con1Frame, const FTransform& Con2Frame, bool bDrawAsPoint, bool bDrawViolatedLimits) const
 	{
-		DrawConstraintImp(FPDIOrCollector(PDI), Scale, LimitDrawScale, bDrawLimits, bDrawSelected, Con1Frame, Con2Frame, bDrawAsPoint);
+		DrawConstraintImp(FPDIOrCollector(PDI), Scale, LimitDrawScale, bDrawLimits, bDrawSelected, Con1Frame, Con2Frame, bDrawAsPoint, bDrawViolatedLimits);
 	}
 
 	ENGINE_API void GetUsedMaterials(TArray<UMaterialInterface*>& Materials);
@@ -1050,7 +1050,7 @@ private:
 
 	ENGINE_API void DrawConstraintImp(const FPDIOrCollector& PDIOrCollector,
 		float Scale, float LimitDrawScale, bool bDrawLimits, bool bDrawSelected,
-		const FTransform& Con1Frame, const FTransform& Con2Frame, bool bDrawAsPoint) const;
+		const FTransform& Con1Frame, const FTransform& Con2Frame, bool bDrawAsPoint, bool bDrawViolatedLimits = false) const;
 
 	ENGINE_API void UpdateBreakable();
 	ENGINE_API void UpdatePlasticity();
diff --git a/Engine/Source/Runtime/Engine/Private/PhysicsEngine/PhysDrawing.cpp b/Engine/Source/Runtime/Engine/Private/PhysicsEngine/PhysDrawing.cpp
--- a/Engine/Source/Runtime/Engine/Private/PhysicsEngine/PhysDrawing.cpp
+++ b/Engine/Source/Runtime/Engine/Private/PhysicsEngine/PhysDrawing.cpp
@@ -20,6 +20,9 @@
 #include "Chaos/Levelset.h"
 #include "Chaos/UniformGrid.h"
 #include "Chaos/WeightedLatticeImplicitObject.h"
+#include "Chaos/PBDJointConstraintUtilities.h"
+
+//UE_DISABLE_OPTIMIZATION
 #include "Misc/ScopeLock.h"
 
 static const int32 DrawCollisionSides = 32;
@@ -33,6 +36,9 @@ static const float UnselectedJointRenderSize = 4.f;
 static const float SelectedJointRenderSize = 10.f;
 static const float LimitRenderSize = 0.16f;
 
+static const float ConstraintPointSize = 4.0f;
+static const float ConstraintLimitViolatedPointSize = 6.0f;
+
 static const FColor JointUnselectedColor(255, 0, 255);
 static const FColor JointRed(FColor::Red);
 static const FColor JointGreen(FColor::Green);
@@ -42,6 +48,7 @@ static const FColor	JointLimitColor(FColor::Green);
 static const FColor	JointRefColor(FColor::Yellow);
 static const FColor JointLockedColor(255,128,10);
 
+static const FColor JointLimitsViolatedColor(FColor::Purple);
 
 static int SkinnedLatticeBoneWeight = -1;
 static FAutoConsoleVariableRef CVarClothVizDrawSkinnedLattice(TEXT("p.PhysDrawing.SkinnedLatticeBoneWeight"), SkinnedLatticeBoneWeight, TEXT("Draw skinned lattice bone weight. -1 = all lattice points"));
@@ -1332,7 +1339,7 @@ void UPhysicsAsset::DrawConstraints(int32 ViewIndex, FMeshElementCollector& Coll
 		}
 
 
-		Instance.DrawConstraint(ViewIndex, Collector, Scale, 1.f, true, true, Con1Frame, Con2Frame, false);
+		Instance.DrawConstraint(ViewIndex, Collector, Scale, 1.f, true, true, Con1Frame, Con2Frame, false, false);
 	}
 }
 static void DrawLinearLimit(FPrimitiveDrawInterface* PDI, const FVector& Origin, const FVector& Axis, const FVector& Orth, float LinearLimitRadius, bool bLinearLimited, float DrawScale)
@@ -1433,22 +1440,22 @@ void FConstraintInstance::GetUsedMaterials(TArray<UMaterialInterface*>& Material
 	Materials.AddUnique(GEngine->ConstraintLimitMaterialZAxis);
 }
 
-void FConstraintInstance::DrawConstraintImp(const FPDIOrCollector& PDIOrCollector, float Scale, float LimitDrawScale, bool bDrawLimits, bool bDrawSelected, const FTransform& Con1Frame, const FTransform& Con2Frame, bool bDrawAsPoint) const
+void FConstraintInstance::DrawConstraintImp(const FPDIOrCollector& PDIOrCollector, float Scale, float LimitDrawScale, bool bDrawLimits, bool bDrawSelected, const FTransform& Con1Frame, const FTransform& Con2Frame, bool bDrawAsPoint, bool bDrawViolatedLimits) const
 {
 	// Do nothing if we're shipping
 #if !UE_BUILD_SHIPPING
 	const ESceneDepthPriorityGroup Layer = ESceneDepthPriorityGroup::SDPG_World;
 	FPrimitiveDrawInterface* PDI = PDIOrCollector.GetPDI();
 
-	check((GEngine->ConstraintLimitMaterialX != nullptr) && (GEngine->ConstraintLimitMaterialY != nullptr) && (GEngine->ConstraintLimitMaterialZ != nullptr));
+	check(GEngine->ConstraintLimitMaterialX && GEngine->ConstraintLimitMaterialY && GEngine->ConstraintLimitMaterialZ);
+
+	static UMaterialInterface* LimitMaterialX = GEngine->ConstraintLimitMaterialX;
+	static UMaterialInterface* LimitMaterialXAxis = GEngine->ConstraintLimitMaterialXAxis;
+	static UMaterialInterface* LimitMaterialY = GEngine->ConstraintLimitMaterialY;
+	static UMaterialInterface* LimitMaterialYAxis = GEngine->ConstraintLimitMaterialYAxis;
+	static UMaterialInterface* LimitMaterialZ = GEngine->ConstraintLimitMaterialZ;
+	static UMaterialInterface* LimitMaterialZAxis = GEngine->ConstraintLimitMaterialZAxis;
 
-	static UMaterialInterface * LimitMaterialX = GEngine->ConstraintLimitMaterialX;
-	static UMaterialInterface * LimitMaterialXAxis = GEngine->ConstraintLimitMaterialXAxis;
-	static UMaterialInterface * LimitMaterialY = GEngine->ConstraintLimitMaterialY;
-	static UMaterialInterface * LimitMaterialYAxis = GEngine->ConstraintLimitMaterialYAxis;
-	static UMaterialInterface * LimitMaterialZ = GEngine->ConstraintLimitMaterialZ;
-	static UMaterialInterface * LimitMaterialZAxis = GEngine->ConstraintLimitMaterialZAxis;
-	
 	const FVector Con1Pos = Con1Frame.GetTranslation();
 	const FVector Con2Pos = Con2Frame.GetTranslation();
 
@@ -1458,8 +1465,8 @@ void FConstraintInstance::DrawConstraintImp(const FPDIOrCollector& PDIOrCollecto
 	// Special mode for drawing joints just as points..
 	if(bDrawAsPoint && !bDrawSelected)
 	{
-		PDI->DrawPoint( Con1Frame.GetTranslation(), JointUnselectedColor, 4.f, ESceneDepthPriorityGroup::SDPG_Foreground );
-		PDI->DrawPoint( Con2Frame.GetTranslation(), JointUnselectedColor, 4.f, ESceneDepthPriorityGroup::SDPG_Foreground );
+		PDI->DrawPoint( Con1Frame.GetTranslation(), JointUnselectedColor, ConstraintPointSize, ESceneDepthPriorityGroup::SDPG_Foreground );
+		PDI->DrawPoint( Con2Frame.GetTranslation(), JointUnselectedColor, ConstraintPointSize, ESceneDepthPriorityGroup::SDPG_Foreground );
 
 		// do nothing else in this mode.
 		return;
@@ -1467,6 +1474,20 @@ void FConstraintInstance::DrawConstraintImp(const FPDIOrCollector& PDIOrCollecto
 
 	if (bDrawLimits)
 	{
+		bool bTwistViolated = false;
+		bool bSwing1Violated = false;
+		bool bSwing2Violated = false;
+		if (bDrawViolatedLimits)
+		{
+			Chaos::FReal TwistAngle, Swing1Angle, Swing2Angle;
+			const FQuat ParentQ = Con2Frame.GetRotation();
+			FQuat ChildQ = Con1Frame.GetRotation();
+			ChildQ.EnforceShortestArcWith(ParentQ);
+			Chaos::FPBDJointUtilities::GetSwingTwistAngles(ParentQ, ChildQ, TwistAngle, Swing1Angle, Swing2Angle);
+			bTwistViolated = GetAngularTwistMotion() == ACM_Limited && FMath::Abs(TwistAngle) > FMath::DegreesToRadians(GetAngularTwistLimit());
+			bSwing1Violated = GetAngularSwing1Motion() == ACM_Limited && FMath::Abs(Swing1Angle) > FMath::DegreesToRadians(GetAngularSwing1Limit());
+			bSwing2Violated = GetAngularSwing2Motion() == ACM_Limited && FMath::Abs(Swing2Angle) > FMath::DegreesToRadians(GetAngularSwing2Limit());
+		}
 
 		//////////////////////////////////////////////////////////////////////////
 		// ANGULAR DRAWING
@@ -1503,6 +1524,11 @@ void FConstraintInstance::DrawConstraintImp(const FPDIOrCollector& PDIOrCollecto
 			FTransform ArrowTM = Con1Frame;
 			ArrowTM.SetTranslation(Con2Pos);
 			PDIOrCollector.DrawArrow(ArrowTM.ToMatrixWithScale(), Length, Thickness, DrawConeLimitSides, JointLimitColor, LimitMaterialX->GetRenderProxy(), Layer);
+
+			if (bSwing1Violated || bSwing2Violated)
+			{
+				PDI->DrawPoint(ArrowTM.TransformPosition(FVector(Length, 0, 0)), JointLimitsViolatedColor, ConstraintLimitViolatedPointSize, ESceneDepthPriorityGroup::SDPG_Foreground);
+			}
 		}
 
 		// Draw the twist limit - A green arc that shows the allowed range of rotation about the parent frame's x axis.
@@ -1517,6 +1543,12 @@ void FConstraintInstance::DrawConstraintImp(const FPDIOrCollector& PDIOrCollecto
 			const FVector TwistIndicator = FVector::PointPlaneProject(Con1Frame.GetScaledAxis(EAxis::Y), FVector::ZeroVector, Con2Frame.GetUnitAxis(EAxis::X)).GetSafeNormal(); // project the y axis of the child frame into the parent frame's yz plane
 			const FTransform ArrowTM(TwistIndicator, Con2Frame.GetScaledAxis(EAxis::X), TwistIndicator ^ Con2Frame.GetScaledAxis(EAxis::X), Con2Frame.GetTranslation());
 			PDIOrCollector.DrawArrow(ArrowTM.ToMatrixWithScale(), Length, Thickness, DrawConeLimitSides, JointLimitColor, LimitMaterialYAxis->GetRenderProxy(), Layer);
+
+			if (bTwistViolated)
+			{
+				PDI->DrawPoint(ArrowTM.TransformPosition(FVector(Length, 0, 0)), JointLimitsViolatedColor, ConstraintLimitViolatedPointSize, ESceneDepthPriorityGroup::SDPG_Foreground);
+			}
+
 		}
 	}
 
diff --git a/Engine/Source/Runtime/Experimental/Chaos/Public/Chaos/PBDJointConstraintUtilities.h b/Engine/Source/Runtime/Experimental/Chaos/Public/Chaos/PBDJointConstraintUtilities.h
--- a/Engine/Source/Runtime/Experimental/Chaos/Public/Chaos/PBDJointConstraintUtilities.h
+++ b/Engine/Source/Runtime/Experimental/Chaos/Public/Chaos/PBDJointConstraintUtilities.h
@@ -13,13 +13,13 @@ namespace Chaos
 	class FPBDJointUtilities
 	{
 	public:
-		static void GetSphericalAxisDelta(
+		static CHAOS_API void GetSphericalAxisDelta(
 			const FVec3& X0,
 			const FVec3& X1,
 			FVec3& Axis,
 			FReal& Delta);
 
-		static void GetCylindricalAxesDeltas(
+		static CHAOS_API void GetCylindricalAxesDeltas(
 			const FRotation3& R0,
 			const FVec3& X0,
 			const FVec3& X1,
@@ -29,7 +29,7 @@ namespace Chaos
 			FVec3& RadialAxis,
 			FReal& RadialDelta);
 
-		static void GetPlanarAxisDelta(
+		static CHAOS_API void GetPlanarAxisDelta(
 			const FRotation3& R0,
 			const FVec3& X0,
 			const FVec3& X1,
@@ -37,43 +37,43 @@ namespace Chaos
 			FVec3& Axis,
 			FReal& Delta);
 
-		static void DecomposeSwingTwistLocal(
+		static CHAOS_API void DecomposeSwingTwistLocal(
 			const FRotation3& R0, 
 			const FRotation3& R1, 
 			FRotation3& R01Swing, 
 			FRotation3& R01Twist);
 		
-		static void GetSwingTwistAngles(
+		static void CHAOS_API GetSwingTwistAngles(
 			const FRotation3& R0, 
 			const FRotation3& R1, 
 			FReal& TwistAngle, 
 			FReal& Swing1Angle, 
 			FReal& Swing2Angle);
 
-		static FReal GetTwistAngle(
+		static CHAOS_API FReal GetTwistAngle(
 			const FRotation3& InTwist);
 		
-		static void GetTwistAxisAngle(
+		static CHAOS_API void GetTwistAxisAngle(
 			const FRotation3& R0,
 			const FRotation3& R1,
 			FVec3& Axis,
 			FReal& Angle);
 
-		static void GetConeAxisAngleLocal(
+		static CHAOS_API void GetConeAxisAngleLocal(
 			const FRotation3& R0,
 			const FRotation3& R1,
 			const FReal AngleTolerance,
 			FVec3& AxisLocal,
 			FReal& Angle);
 
-		static void GetCircularConeAxisErrorLocal(
+		static CHAOS_API void GetCircularConeAxisErrorLocal(
 			const FRotation3& R0,
 			const FRotation3& R1,
 			const FReal SwingLimit,
 			FVec3& AxisLocal,
 			FReal& Error);
 
-		static void GetEllipticalConeAxisErrorLocal(
+		static CHAOS_API void GetEllipticalConeAxisErrorLocal(
 			const FRotation3& R0,
 			const FRotation3& R1,
 			const FReal SwingLimitY,
@@ -81,21 +81,21 @@ namespace Chaos
 			FVec3& AxisLocal,
 			FReal& Error);
 
-		static void GetLockedSwingAxisAngle(
+		static CHAOS_API void GetLockedSwingAxisAngle(
 			const FRotation3& R0,
 			const FRotation3& R1,
 			const EJointAngularConstraintIndex SwingConstraintIndex,
 			FVec3& Axis,
 			FReal& Angle);
 
-		static void GetDualConeSwingAxisAngle(
+		static CHAOS_API void GetDualConeSwingAxisAngle(
 			const FRotation3& R0,
 			const FRotation3& R1,
 			const EJointAngularConstraintIndex SwingConstraintIndex,
 			FVec3& Axis,
 			FReal& Angle);
 
-		static void GetSwingAxisAngle(
+		static CHAOS_API void GetSwingAxisAngle(
 			const FRotation3& R0,
 			const FRotation3& R1,
 			const FReal AngleTolerance,
@@ -103,14 +103,14 @@ namespace Chaos
 			FVec3& Axis,
 			FReal& Angle);
 
-		static void GetLockedRotationAxes(
+		static CHAOS_API void GetLockedRotationAxes(
 			const FRotation3& R0, 
 			const FRotation3& R1, 
 			FVec3& Axis0, 
 			FVec3& Axis1, 
 			FVec3& Axis2);
 
-		static FReal GetConeAngleLimit(
+		static CHAOS_API FReal GetConeAngleLimit(
 			const FPBDJointSettings& JointSettings,
 			const FVec3& SwingAxisLocal,
 			const FReal SwingAngle);
