Worker thread parallel anim eval CSV + stats

diff --git a/Engine/Source/Runtime/Engine/Private/Components/SkeletalMeshComponent.cpp b/Engine/Source/Runtime/Engine/Private/Components/SkeletalMeshComponent.cpp
--- a/Engine/Source/Runtime/Engine/Private/Components/SkeletalMeshComponent.cpp
+++ b/Engine/Source/Runtime/Engine/Private/Components/SkeletalMeshComponent.cpp
@@ -111,6 +111,74 @@ FAutoConsoleTaskPriority CPrio_ParallelAnimationEvaluationTask(
 /** Static Multicaster fired when SkeletalMeshComponent finalizes the regeneration of the required bones list for the current LOD*/
 /*static*/ FOnLODRequiredBonesUpdateMulticast USkeletalMeshComponent::OnLODRequiredBonesUpdate;
 
+#if !UE_BUILD_SHIPPING
+CSV_DEFINE_CATEGORY(AnimationParallelEvaluation, true);
+struct FParallelAnimationEvaluationStats
+{
+	static FParallelAnimationEvaluationStats& Get()
+	{
+		static FParallelAnimationEvaluationStats Stats;
+		return Stats;
+	}
+
+	void AddTiming(float Value)
+	{
+		TotalTime.Store(TotalTime.Load() + Value);
+		MinTime.Store(FMath::Min(MinTime.Load(), Value));
+		MaxTime.Store(FMath::Max(MaxTime.Load(), Value));
+		NumberOfTasks.IncrementExchange();
+	}
+
+private:
+	FParallelAnimationEvaluationStats()
+	{
+		EndOfFrameHandle = FCoreDelegates::OnEndFrame.AddRaw(this, &FParallelAnimationEvaluationStats::OnEndFrame);
+		MinTime = TNumericLimits<float>::Max();
+		MaxTime = 0.f;
+		TotalTime = 0.f;
+		NumberOfTasks = 0;
+	}
+	
+	~FParallelAnimationEvaluationStats()
+	{
+		FCoreDelegates::OnEndFrame.Remove(EndOfFrameHandle);
+	}
+
+	void OnEndFrame()
+	{
+		// Only set stats when a task has actually run
+		if (NumberOfTasks > 0)
+		{	
+		    CSV_CUSTOM_STAT(AnimationParallelEvaluation, TotalTaskTime, TotalTime, ECsvCustomStatOp::Set);	
+		    CSV_CUSTOM_STAT(AnimationParallelEvaluation, AverageTaskTime, NumberOfTasks == 0 || FMath::IsNearlyZero(TotalTime) ? 0.f : TotalTime / static_cast<float>(NumberOfTasks), ECsvCustomStatOp::Set);				
+		    CSV_CUSTOM_STAT(AnimationParallelEvaluation, NumberOfTasks, NumberOfTasks, ECsvCustomStatOp::Set);
+				    
+		    CSV_CUSTOM_STAT(AnimationParallelEvaluation, MinTaskTime, MinTime, ECsvCustomStatOp::Set);
+		    CSV_CUSTOM_STAT(AnimationParallelEvaluation, MaxTaskTime, MaxTime, ECsvCustomStatOp::Set);
+		}
+
+		MinTime = TNumericLimits<float>::Max();
+		MaxTime = 0.f;
+		TotalTime = 0.f;
+		NumberOfTasks = 0;
+	}
+
+	TAtomic<float> MinTime;
+	TAtomic<float> MaxTime;
+	TAtomic<float> TotalTime;
+	TAtomic<int32> NumberOfTasks;
+
+	FDelegateHandle EndOfFrameHandle;
+};
+
+// When the object system has been completely loaded, register the OnEndFrame delegate
+static FDelayedAutoRegisterHelper GParallelAnimationEvaluationStatsHelper(EDelayedRegisterRunPhase::EndOfEngineInit, []() -> void
+{
+	FParallelAnimationEvaluationStats::Get();
+});
+
+#endif
+
 class FParallelAnimationEvaluationTask
 {
 	TWeakObjectPtr<USkeletalMeshComponent> SkeletalMeshComponent;
@@ -151,7 +219,14 @@ public:
 				GInitRunaway();
 			}
 
+#if !UE_BUILD_SHIPPING			
+			const uint64 StartTime = FPlatformTime::Cycles64();
 			Comp->ParallelAnimationEvaluation();
+			const float EvaluationTimeMS = FPlatformTime::ToMilliseconds64(FPlatformTime::Cycles64() - StartTime);
+			FParallelAnimationEvaluationStats::Get().AddTiming(EvaluationTimeMS);
+#else
+			Comp->ParallelAnimationEvaluation();
+#endif
 		}
 	}
 };
