Moved calling GRayTracingGeometry->Tick() to FRendererModule::PostRenderAllViewports()

diff --git a/Engine/Source/Runtime/Engine/Private/Rendering/RayTracingGeometryManager.cpp b/Engine/Source/Runtime/Engine/Private/Rendering/RayTracingGeometryManager.cpp
--- a/Engine/Source/Runtime/Engine/Private/Rendering/RayTracingGeometryManager.cpp
+++ b/Engine/Source/Runtime/Engine/Private/Rendering/RayTracingGeometryManager.cpp
@@ -4,6 +4,7 @@
 
 #include "PrimitiveSceneProxy.h"
 #include "SceneInterface.h"
+#include "ComponentRecreateRenderStateContext.h"
 
 #include "RHIResources.h"
 #include "RHICommandList.h"
@@ -13,6 +14,24 @@
 
 #if RHI_RAYTRACING
 
+static bool bHasRayTracingEnableChanged = false;
+static TAutoConsoleVariable<int32> CVarRayTracingEnable(
+	TEXT("r.RayTracing.Enable"),
+	1,
+	TEXT("Runtime toggle for switching raytracing on/off (experimental)."),
+	FConsoleVariableDelegate::CreateLambda([](IConsoleVariable* InVariable)
+		{
+			FGlobalComponentRecreateRenderStateContext Context;
+			ENQUEUE_RENDER_COMMAND(RayTracingToggledCmd)(
+				[](FRHICommandListImmediate&)
+				{
+					bHasRayTracingEnableChanged = true;
+				}
+			);
+		}),
+	ECVF_RenderThreadSafe
+);
+
 static int32 GRayTracingMaxBuiltPrimitivesPerFrame = -1;
 static FAutoConsoleVariableRef CVarRayTracingMaxBuiltPrimitivesPerFrame(
 	TEXT("r.RayTracing.Geometry.MaxBuiltPrimitivesPerFrame"),
@@ -107,7 +126,7 @@ void FRayTracingGeometryManager::ReleaseRayTracingGeometryHandle(RayTracingGeome
 	}	
 }
 
-void FRayTracingGeometryManager::Tick(FRHICommandList& RHICmdList, bool bHasRayTracingEnableChanged)
+void FRayTracingGeometryManager::Tick(FRHICommandList& RHICmdList)
 {
 	if (GetRayTracingMode() != ERayTracingMode::Dynamic)
 	{
@@ -141,14 +160,29 @@ void FRayTracingGeometryManager::Tick(FRHICommandList& RHICmdList, bool bHasRayT
 		return;
 	}
 
+	bHasRayTracingEnableChanged = false;
+
 	if (IsRayTracingEnabled())
 	{
-		FScopeLock ScopeLock(&RequestCS);
-		for (FRayTracingGeometry* Geometry : RegisteredGeometries)
 		{
-			if (Geometry->IsEvicted())
+			FScopeLock ScopeLock(&RequestCS);
+			for (FRayTracingGeometry* Geometry : RegisteredGeometries)
+			{
+				if (Geometry->IsEvicted())
+				{
+					Geometry->MakeResident(RHICmdList);
+				}
+			}
+		}
+
+		{
+			FScopeLock ScopeLock(&CachedRayTracingStateProxiesCS);
+			for (TPair<const UStaticMesh*, TSet<FPrimitiveSceneProxy*>>& ProxiesSet : CachedRayTracingStateProxiesMap)
 			{
-				Geometry->MakeResident(RHICmdList);
+				for (FPrimitiveSceneProxy* Proxy : ProxiesSet.Value)
+				{
+					Proxy->GetScene().UpdateCachedRayTracingState(Proxy);
+				}
 			}
 		}
 	}
diff --git a/Engine/Source/Runtime/Engine/Public/Rendering/RayTracingGeometryManager.h b/Engine/Source/Runtime/Engine/Public/Rendering/RayTracingGeometryManager.h
--- a/Engine/Source/Runtime/Engine/Public/Rendering/RayTracingGeometryManager.h
+++ b/Engine/Source/Runtime/Engine/Public/Rendering/RayTracingGeometryManager.h
@@ -27,7 +27,7 @@ public:
 	ENGINE_API virtual RayTracingGeometryHandle RegisterRayTracingGeometry(FRayTracingGeometry* InGeometry) override;
 	ENGINE_API virtual void ReleaseRayTracingGeometryHandle(RayTracingGeometryHandle Handle) override;
 
-	ENGINE_API virtual void Tick(FRHICommandList& RHICmdList, bool bHasRayTracingEnableChanged) override;
+	ENGINE_API virtual void Tick(FRHICommandList& RHICmdList) override;
 
 	void RegisterProxyWithCachedRayTracingState(FPrimitiveSceneProxy* Proxy, const UStaticMesh* StaticMesh);
 	void UnregisterProxyWithCachedRayTracingState(FPrimitiveSceneProxy* Proxy, const UStaticMesh* StaticMesh);
diff --git a/Engine/Source/Runtime/RenderCore/Public/RayTracingGeometryManagerInterface.h b/Engine/Source/Runtime/RenderCore/Public/RayTracingGeometryManagerInterface.h
--- a/Engine/Source/Runtime/RenderCore/Public/RayTracingGeometryManagerInterface.h
+++ b/Engine/Source/Runtime/RenderCore/Public/RayTracingGeometryManagerInterface.h
@@ -38,7 +38,7 @@ public:
 	RENDERCORE_API virtual RayTracingGeometryHandle RegisterRayTracingGeometry(FRayTracingGeometry* InGeometry) = 0;
 	RENDERCORE_API virtual void ReleaseRayTracingGeometryHandle(RayTracingGeometryHandle Handle) = 0;
 
-	RENDERCORE_API virtual void Tick(FRHICommandList& RHICmdList, bool bHasRayTracingEnableChanged) = 0;
+	RENDERCORE_API virtual void Tick(FRHICommandList& RHICmdList) = 0;
 };
 
 extern RENDERCORE_API IRayTracingGeometryManager* GRayTracingGeometryManager;
diff --git a/Engine/Source/Runtime/Renderer/Private/DeferredShadingRenderer.cpp b/Engine/Source/Runtime/Renderer/Private/DeferredShadingRenderer.cpp
--- a/Engine/Source/Runtime/Renderer/Private/DeferredShadingRenderer.cpp
+++ b/Engine/Source/Runtime/Renderer/Private/DeferredShadingRenderer.cpp
@@ -155,24 +155,6 @@ static TAutoConsoleVariable<int32> CVarRayTracing(
 	TEXT(" 1: on"),
 	ECVF_RenderThreadSafe | ECVF_ReadOnly);
 
-static bool bHasRayTracingEnableChanged = false;
-static TAutoConsoleVariable<int32> CVarRayTracingEnable(
-	TEXT("r.RayTracing.Enable"),
-	1,
-	TEXT("Runtime toggle for switching raytracing on/off (experimental)."),
-	FConsoleVariableDelegate::CreateLambda([](IConsoleVariable* InVariable)
-		{
-			FGlobalComponentRecreateRenderStateContext Context;		 
-			ENQUEUE_RENDER_COMMAND(RefreshRayTracingMeshCommandsCmd)(
-				[](FRHICommandListImmediate&)
-				{
-					bHasRayTracingEnableChanged = true;
-				}
-			);
-		}),
-	ECVF_RenderThreadSafe
-);
-
 int32 GRayTracingUseTextureLod = 0;
 static TAutoConsoleVariable<int32> CVarRayTracingTextureLod(
 	TEXT("r.RayTracing.UseTextureLod"),
@@ -1613,8 +1595,6 @@ void FDeferredShadingSceneRenderer::Render(FRDGBuilder& GraphBuilder)
 		// Initialize ray tracing flags, in case they weren't initialized in the CreateSceneRenderers code path
 		InitializeRayTracingFlags_RenderThread();
 
-		GRayTracingGeometryManager->Tick(GraphBuilder.RHICmdList, bHasRayTracingEnableChanged);
-
 		// Now that we have updated all the PrimitiveSceneInfos, update the RayTracing mesh commands cache if needed
 		{
 			const ERayTracingMeshCommandsMode CurrentMode = ViewFamily.EngineShowFlags.PathTracing ? ERayTracingMeshCommandsMode::PATH_TRACING : ERayTracingMeshCommandsMode::RAY_TRACING;
@@ -1627,7 +1607,6 @@ void FDeferredShadingSceneRenderer::Render(FRDGBuilder& GraphBuilder)
 			if (CurrentMode != Scene->CachedRayTracingMeshCommandsMode
 				|| bNaniteCoarseMeshStreamingModeChanged
 				|| bNaniteRayTracingModeChanged
-				|| bHasRayTracingEnableChanged
 				|| bUpdateCachedRayTracingState)
 			{
 				Scene->WaitForCacheRayTracingPrimitivesTask();
@@ -1636,7 +1615,6 @@ void FDeferredShadingSceneRenderer::Render(FRDGBuilder& GraphBuilder)
 				// This operation is a bit expensive but only happens once as we transition between modes which should be rare.
 				Scene->CachedRayTracingMeshCommandsMode = CurrentMode;
 				Scene->RefreshRayTracingMeshCommandCache();
-				bHasRayTracingEnableChanged = false;
 				bUpdateCachedRayTracingState = false;
 			}
 
diff --git a/Engine/Source/Runtime/Renderer/Private/SceneRendering.cpp b/Engine/Source/Runtime/Renderer/Private/SceneRendering.cpp
--- a/Engine/Source/Runtime/Renderer/Private/SceneRendering.cpp
+++ b/Engine/Source/Runtime/Renderer/Private/SceneRendering.cpp
@@ -70,6 +70,7 @@
 #include "Rendering/StaticLightingSystemInterface.h"
 #endif
 #include "RayTracing/RayTracingScene.h"
+#include "Rendering/RayTracingGeometryManager.h"
 #include "FXSystem.h"
 #include "Lumen/Lumen.h"
 #include "Nanite/Nanite.h"
@@ -4970,14 +4971,17 @@ void FRendererModule::PostRenderAllViewports()
 #if RHI_RAYTRACING
 	// Update the resource state after all viewports are done with rendering - all info collected for all views
 	Nanite::FCoarseMeshStreamingManager* CoarseMeshSM = IStreamingManager::Get().GetNaniteCoarseMeshStreamingManager();
-	if (CoarseMeshSM)
-	{
-		ENQUEUE_RENDER_COMMAND(NaniteCoarseMeshUpdateResourceStates)(
-			[CoarseMeshSM](FRHICommandListImmediate& RHICmdList)
+
+	ENQUEUE_RENDER_COMMAND(PostRenderAllViewports_RenderThread)(
+ 		[CoarseMeshSM](FRHICommandListImmediate& RHICmdList)
+ 		{
+			if (CoarseMeshSM)
 			{
 				CoarseMeshSM->UpdateResourceStates();
-			});
-	}
+			}
+
+ 			GRayTracingGeometryManager->Tick(RHICmdList);
+ 		});
 #endif //#if RHI_RAYTRACING
 }
 
