- Fix for "FX.AllowGPUParticles 0" causing crashes

diff --git a/Engine/Source/Runtime/Engine/Private/Particles/ParticleGpuSimulation.cpp b/Engine/Source/Runtime/Engine/Private/Particles/ParticleGpuSimulation.cpp
--- a/Engine/Source/Runtime/Engine/Private/Particles/ParticleGpuSimulation.cpp
+++ b/Engine/Source/Runtime/Engine/Private/Particles/ParticleGpuSimulation.cpp
@@ -4220,12 +4220,9 @@ void FFXSystem::InitGPUResources()
 
 void FFXSystem::ReleaseGPUResources()
 {
-	if (RHISupportsGPUParticles())
+	if(ParticleSimulationResources)
 	{
-		if(ParticleSimulationResources)
-		{
-			ParticleSimulationResources->Release();
-		}
+		ParticleSimulationResources->Release();
 	}
 }
 
