SkeletalMesh: Add support for morph targets to the skeletal mesh attributes. Morph targets store both position and normal/tangent-z deltas, although normals are currently not imported. Updated the conversion code for FSkeletalMeshImportData to FMeshDescription to support morph targets and alternate influences.

diff --git a/Engine/Source/Runtime/Engine/Private/Rendering/SkeletalMeshLODImporterData.cpp b/Engine/Source/Runtime/Engine/Private/Rendering/SkeletalMeshLODImporterData.cpp
--- a/Engine/Source/Runtime/Engine/Private/Rendering/SkeletalMeshLODImporterData.cpp
+++ b/Engine/Source/Runtime/Engine/Private/Rendering/SkeletalMeshLODImporterData.cpp
@@ -8,8 +8,10 @@
 #include "Engine/SkinnedAssetCommon.h"
 #include "ImportUtils/SkeletalMeshImportUtils.h"
 #include "Materials/MaterialInterface.h"
-#include "Misc/ScopedSlowTask.h"
+#include "MeshElementContainer.h"
 #include "Misc/ScopeRWLock.h"
+#include "Misc/ScopedSlowTask.h"
+#include "Modules/ModuleManager.h"
 #include "Rendering/SkeletalMeshModel.h"
 #include "Serialization/BulkDataReader.h"
 #include "Serialization/BulkDataWriter.h"
@@ -17,8 +19,6 @@
 #include "Serialization/MemoryReader.h"
 #include "SkeletalMeshAttributes.h"
 #include "SkeletalMeshOperations.h"
-#include "MeshElementContainer.h"
-
 
 DEFINE_LOG_CATEGORY_STATIC(LogSkeletalMeshLODImporterData, Log, All);
 
@@ -1541,7 +1541,7 @@ void FSkeletalMeshImportData::SplitVerticesBySmoothingGroups()
 	//
 	// Connectivity: triangles with non-matching smoothing groups will be physically split.
 	//
-	// -> Splitting involves: the UV+material-contaning vertex AND the 3d point.
+	// -> Splitting involves: the UV+material-containing vertex AND the 3d point.
 	//
 	// -> Tally smoothing groups for each and every (textured) vertex.
 	//
@@ -1735,6 +1735,54 @@ void FSkeletalMeshImportData::SplitVerticesBySmoothingGroups()
 	}
 }
 
+static void CopySkinWeightsToAttribute(
+	const TArray<SkeletalMeshImportData::FRawBoneInfluence>& InInfluences,
+	const TArray<FVertexID>& InVertexIDMap,
+	const TMap<int32, int32>* InBoneIndexMap,
+	FSkinWeightsVertexAttributesRef OutSkinWeightsAttribute
+	)
+{
+	using namespace UE::AnimationCore;
+	
+	// The weights are stored with links back to the vertices, rather than being compact.
+	// Make a copy of the weights, sort them by vertex id and go by equal vertex-id strides.
+	// We could do an indirection but the traversal + setup cost is probably not worth it.
+	TArray<SkeletalMeshImportData::FRawBoneInfluence> SortedInfluences(InInfluences);
+	SortedInfluences.Sort([](const SkeletalMeshImportData::FRawBoneInfluence &A, const SkeletalMeshImportData::FRawBoneInfluence &B)
+	{
+		return A.VertexIndex < B.VertexIndex;
+	});
+
+	// Do the base skin weights first. We do the alternative skin weights later, since they may require geometric remapping. 
+	TArray<FBoneWeight> BoneWeights;
+	for(int32 StartStride = 0, EndStride = 0; EndStride != SortedInfluences.Num(); StartStride = EndStride)
+	{
+		const int32 VertexIndex = SortedInfluences[StartStride].VertexIndex;
+		
+		EndStride = StartStride + 1;
+		while (EndStride < SortedInfluences.Num() && VertexIndex == SortedInfluences[EndStride].VertexIndex)
+		{
+			EndStride++;
+		}
+
+		BoneWeights.Reset(0);
+		for (int32 Idx = StartStride; Idx < EndStride; Idx++)
+		{
+			const SkeletalMeshImportData::FRawBoneInfluence &RawInfluence = SortedInfluences[Idx];
+			int32 BoneIndex = RawInfluence.BoneIndex;
+			if (InBoneIndexMap)
+			{
+				// The map should be complete.
+				BoneIndex = (*InBoneIndexMap)[BoneIndex];
+			}
+			
+			FBoneWeight BoneWeight(static_cast<FBoneIndexType>(BoneIndex), RawInfluence.Weight);
+			BoneWeights.Add(BoneWeight);
+		}
+
+		OutSkinWeightsAttribute.Set(InVertexIDMap[VertexIndex], BoneWeights);		
+	}
+}
 
 bool FSkeletalMeshImportData::GetMeshDescription(FMeshDescription& OutMeshDescription) const
 {
@@ -1765,7 +1813,70 @@ bool FSkeletalMeshImportData::GetMeshDescription(FMeshDescription& OutMeshDescri
 	FSkeletalMeshAttributes::FBoneParentIndexAttributesRef BoneParentIndices = MeshAttributes.GetBoneParentIndices();
 	FSkeletalMeshAttributes::FBonePoseAttributesRef BonePoses = MeshAttributes.GetBonePoses();
 
-	TMap<FString, FName> ValidAttributes;
+	// Register morph targets.
+	TSet<FString> ValidMorphTargets;
+	if (ensure(MorphTargetNames.Num() == MorphTargets.Num()))
+	{
+		for (int32 MorphIndex = 0; MorphIndex < MorphTargets.Num(); MorphIndex++)
+		{
+			const FSkeletalMeshImportData& MorphTargetMesh = MorphTargets[MorphIndex];
+			
+			// As long as we have the same point count, we have a valid morph target mesh. 
+			if (!ensure(MorphTargetMesh.Points.Num() == Points.Num()))
+			{
+				continue;
+			}
+
+			const FString MorphTargetName(MorphTargetNames[MorphIndex]);
+			if (ensure(MeshAttributes.RegisterMorphTargetAttribute(*MorphTargetName)))
+			{
+				ValidMorphTargets.Add(MorphTargetName);
+			}
+		}
+	}
+
+	// Register alternate skin weight profiles.
+	TSet<FString> ValidSkinWeights;
+	if (ensure(AlternateInfluenceProfileNames.Num() == AlternateInfluences.Num()))
+	{
+		for (int32 AlternativeInfluenceIndex = 0; AlternativeInfluenceIndex < AlternateInfluences.Num(); AlternativeInfluenceIndex++)
+		{
+			const FSkeletalMeshImportData& AlternativeInfluenceMesh = AlternateInfluences[AlternativeInfluenceIndex];
+			
+			// We minimally need the same point count.
+			// Q: Should we match on topology as well?
+			if (!ensure(AlternativeInfluenceMesh.Points.Num() == Points.Num()))
+			{
+				continue;
+			}
+
+			// The set of bones on the alternative influence mesh must be a fully included in the set of the bones on the
+			// base mesh. Otherwise we cannot map an influence across.
+			TSet<FString> BaseBones, AltBones;
+			for (int32 BoneIndex = 0; BoneIndex < RefBonesBinary.Num(); BoneIndex++)
+			{
+				BaseBones.Add(RefBonesBinary[BoneIndex].Name);
+			}
+			for (int32 BoneIndex = 0; BoneIndex < AlternativeInfluenceMesh.RefBonesBinary.Num(); BoneIndex++)
+			{
+				AltBones.Add(AlternativeInfluenceMesh.RefBonesBinary[BoneIndex].Name);
+			}
+
+			if (!ensure(BaseBones.Includes(AltBones)))
+			{
+				continue;
+			}
+
+			if (ensure(MeshAttributes.RegisterSkinWeightAttribute(*AlternateInfluenceProfileNames[AlternativeInfluenceIndex])))
+			{
+				ValidSkinWeights.Add(AlternateInfluenceProfileNames[AlternativeInfluenceIndex]);
+			}
+
+		}
+	}
+
+	// Register vertex attributes.
+	TSet<FString> ValidAttributes;
 	for (int32 AttributeIndex = 0; AttributeIndex < VertexAttributes.Num(); AttributeIndex++)
 	{
 		const FString& VertexAttributeName = VertexAttributeNames[AttributeIndex];
@@ -1785,26 +1896,30 @@ bool FSkeletalMeshImportData::GetMeshDescription(FMeshDescription& OutMeshDescri
 		{
 			continue;
 		}
-		
+
+		bool bIsValidAttribute = false;
 		switch(VertexAttribute.ComponentCount)
 		{
 		case 1:
-			OutMeshDescription.VertexAttributes().RegisterAttribute<float>(RegisteredName, 1, 0.0f, DefaultAttributeFlags);
+			bIsValidAttribute = OutMeshDescription.VertexAttributes().RegisterAttribute<float>(RegisteredName, 1, 0.0f, DefaultAttributeFlags).IsValid();
 			break;
 		case 2:
-			OutMeshDescription.VertexAttributes().RegisterAttribute<FVector2f>(RegisteredName, 1, FVector2f::Zero(), DefaultAttributeFlags);
+			bIsValidAttribute = OutMeshDescription.VertexAttributes().RegisterAttribute<FVector2f>(RegisteredName, 1, FVector2f::Zero(), DefaultAttributeFlags).IsValid();
 			break;
 		case 3:
-			OutMeshDescription.VertexAttributes().RegisterAttribute<FVector3f>(RegisteredName, 1, FVector3f::Zero(), DefaultAttributeFlags);
+			bIsValidAttribute = OutMeshDescription.VertexAttributes().RegisterAttribute<FVector3f>(RegisteredName, 1, FVector3f::Zero(), DefaultAttributeFlags).IsValid();
 			break;
 		case 4:
-			OutMeshDescription.VertexAttributes().RegisterAttribute<FVector4f>(RegisteredName, 1, FVector4f::Zero(), DefaultAttributeFlags);
+			bIsValidAttribute = OutMeshDescription.VertexAttributes().RegisterAttribute<FVector4f>(RegisteredName, 1, FVector4f::Zero(), DefaultAttributeFlags).IsValid();
 			break;
 		default:
 			continue;
 		}
-		
-		ValidAttributes.Add(VertexAttributeName, RegisteredName);
+
+		if (ensure(bIsValidAttribute))
+		{
+			ValidAttributes.Add(VertexAttributeName);
+		}
 	}
 	
 	VertexInstanceUVs.SetNumChannels(NumTexCoords);
@@ -1815,6 +1930,7 @@ bool FSkeletalMeshImportData::GetMeshDescription(FMeshDescription& OutMeshDescri
 	OutMeshDescription.ReserveNewVertexInstances(Wedges.Num());
 	OutMeshDescription.ReserveNewVertices(Points.Num());
 	MeshAttributes.ReserveNewBones(RefBonesBinary.Num());
+	
 
 	// Copy the vertex positions first and maintain a map so that we can go from the import data's raw vertex index
 	// to the mesh description's VertexID.
@@ -1827,37 +1943,7 @@ bool FSkeletalMeshImportData::GetMeshDescription(FMeshDescription& OutMeshDescri
 		VertexPositions.Set(VertexID, Points[Idx]);
 	}
 
-	// The weights are stored with links back to the vertices, rather than being compact.
-	// Make a copy of the weights, sort them by vertex id and go by equal vertex-id strides.
-	// We could do an indirection but the traversal + setup cost is probably not worth it.
-	TArray<SkeletalMeshImportData::FRawBoneInfluence> SortedInfluences(Influences);
-	SortedInfluences.Sort([](const SkeletalMeshImportData::FRawBoneInfluence &A, const SkeletalMeshImportData::FRawBoneInfluence &B)
-	{
-		return A.VertexIndex < B.VertexIndex;
-	});
-
-	// TODO: Alternative profiles.
-	TArray<FBoneWeight> BoneWeights;
-	for(int32 StartStride = 0, EndStride = 0; EndStride != SortedInfluences.Num(); StartStride = EndStride)
-	{
-		const int32 VertexIndex = SortedInfluences[StartStride].VertexIndex;
-		
-		EndStride = StartStride + 1;
-		while (EndStride < SortedInfluences.Num() && VertexIndex == SortedInfluences[EndStride].VertexIndex)
-		{
-			EndStride++;
-		}
-
-		BoneWeights.Reset(0);
-		for (int32 Idx = StartStride; Idx < EndStride; Idx++)
-		{
-			const SkeletalMeshImportData::FRawBoneInfluence &RawInfluence = SortedInfluences[Idx];
-			FBoneWeight BoneWeight(FBoneIndexType(RawInfluence.BoneIndex), RawInfluence.Weight);
-			BoneWeights.Add(BoneWeight);
-		}
-
-		VertexSkinWeights.Set(VertexIDMap[VertexIndex], BoneWeights);		
-	}
+	CopySkinWeightsToAttribute(Influences, VertexIDMap, nullptr, VertexSkinWeights);
 
 	// Set Bone Attributes
 	for (int Idx = 0; Idx < RefBonesBinary.Num(); ++Idx)
@@ -1892,28 +1978,42 @@ bool FSkeletalMeshImportData::GetMeshDescription(FMeshDescription& OutMeshDescri
 	
 	for(int32 StartStride = 0, EndStride = 0; EndStride != FaceIndices.Num(); StartStride = EndStride)
 	{
-		const int32 MaterialIndex = Faces[FaceIndices[StartStride]].MatIndex;
+		FPolygonGroupID PolygonGroupID;
+		FName MaterialName;
+		if (!Materials.IsEmpty())
+		{
+			const int32 MaterialIndex = Faces[FaceIndices[StartStride]].MatIndex;
 		
-		EndStride = StartStride + 1;
-		while (EndStride < FaceIndices.Num() && MaterialIndex == Faces[FaceIndices[EndStride]].MatIndex)
+			EndStride = StartStride + 1;
+			while (EndStride < FaceIndices.Num() && MaterialIndex == Faces[FaceIndices[EndStride]].MatIndex)
+			{
+				EndStride++;
+			}
+
+			// Create a section for each material index. We re-use vertex instances if they are
+			// referred to multiple times by the FTriangle object. However, because the tangents
+			// are stored on the triangle, it's possible to end up with a recycled vertex instance
+			// that has a different tangent. This is a limitation for now. 
+			// Along the way we track smoothing groups and use that to define hard edges once the
+			// entire mesh is defined.
+			PolygonGroupID = MaterialIndex;
+			MaterialName = FName(*Materials[MaterialIndex].MaterialImportName);
+		}
+		else
 		{
-			EndStride++;
+			// No materials defined, in which case we set a dummy material on the whole mesh.
+			// This can happen when getting a mesh description from alternate influences mesh.
+			PolygonGroupID = 0;
+			EndStride = FaceIndices.Num();
 		}
-
-		// Create a section for each material index. We re-use vertex instances if they are
-		// referred to multiple times by the FTriangle object. However, because the tangents
-		// are stored on the triangle, it's possible to end up with a recycled vertex instance
-		// that has a different tangent. This is a limitation for now. 
-		// Along the way we track smoothing groups and use that to define hard edges once the
-		// entire mesh is defined.
-		const FPolygonGroupID PolygonGroupID(MaterialIndex);
-
+		
 		if (!OutMeshDescription.IsPolygonGroupValid(PolygonGroupID))
 		{
 			OutMeshDescription.CreatePolygonGroupWithID(PolygonGroupID);
 		}
+		PolygonGroupMaterialSlotNames.Set(PolygonGroupID, MaterialName);
 
-		PolygonGroupMaterialSlotNames.Set(PolygonGroupID, FName(*Materials[MaterialIndex].MaterialImportName));
+		
 		
 		for (int32 Idx = StartStride; Idx < EndStride; Idx++)
 		{
@@ -1933,9 +2033,9 @@ bool FSkeletalMeshImportData::GetMeshDescription(FMeshDescription& OutMeshDescri
 					if (bHasVertexColors)
 					{
 						// Don't perform sRGB conversion (which mirrors what CreateFromMeshDescription does).
-						VertexInstanceColors.Set(VertexInstanceID, (FVector4f)Wedge.Color.ReinterpretAsLinear());
+						VertexInstanceColors.Set(VertexInstanceID, Wedge.Color.ReinterpretAsLinear());
 					}
-					for (int32 UVIndex = 0; UVIndex < int32(NumTexCoords); UVIndex++)
+					for (int32 UVIndex = 0; UVIndex < static_cast<int32>(NumTexCoords); UVIndex++)
 					{
 						VertexInstanceUVs.Set(VertexInstanceID, UVIndex, Wedge.UVs[UVIndex]);
 					}
@@ -1960,7 +2060,41 @@ bool FSkeletalMeshImportData::GetMeshDescription(FMeshDescription& OutMeshDescri
 			FaceSmoothingMasks[PolygonID.GetValue()] = Triangle.SmoothingGroups;
 		}
 	}
+
+	// Convert morph targets
+	for (int32 MorphTargetIndex = 0; MorphTargetIndex < MorphTargets.Num(); MorphTargetIndex++)
+	{
+		FString MorphTargetName(MorphTargetNames[MorphTargetIndex]);
+		if (!ValidMorphTargets.Contains(MorphTargetName))
+		{
+			continue;
+		}
+
+		const TArray<FVector3f>& MorphPoints = MorphTargets[MorphTargetIndex].Points;
+		const TSet<uint32>& ModifiedPoints = MorphTargetModifiedPoints[MorphTargetIndex]; 
+		FMorphTargetVertexAttributesRef MorphTargetRef = MeshAttributes.GetVertexMorphTarget(*MorphTargetName);
+
+		for (uint32 PointIndex: ModifiedPoints)
+		{
+			const FVertexID VertexID = VertexIDMap[PointIndex];
+			MorphTargetRef.SetPositionDelta(VertexID, MorphPoints[PointIndex] - Points[PointIndex]);
+		}
+	}
+
+	// Convert optional skin weights
+	for (int32 AlternateInfluenceIndex = 0; AlternateInfluenceIndex < AlternateInfluences.Num(); AlternateInfluenceIndex++)
+	{
+		FString AlternateInfluenceProfileName(AlternateInfluenceProfileNames[AlternateInfluenceIndex]);
+		if (!ValidSkinWeights.Contains(AlternateInfluenceProfileName))
+		{
+			continue;
+		}
+
+		CopySkinWeightsToMeshDescription(FName(AlternateInfluenceProfileName), AlternateInfluences[AlternateInfluenceIndex], VertexIDMap, OutMeshDescription);
+	}
 	
+
+	// Convert vertex attributes.
 	for (int32 AttributeIndex = 0; AttributeIndex < VertexAttributes.Num(); AttributeIndex++)
 	{
 		const FString& VertexAttributeName = VertexAttributeNames[AttributeIndex];
@@ -1971,7 +2105,7 @@ bool FSkeletalMeshImportData::GetMeshDescription(FMeshDescription& OutMeshDescri
 			continue;
 		}
 		
-		FName RegisteredName(ValidAttributes[VertexAttributeName]);
+		FName RegisteredName(VertexAttributeName);
 		switch(VertexAttribute.ComponentCount)
 		{
 		case 1:
@@ -2025,6 +2159,88 @@ bool FSkeletalMeshImportData::GetMeshDescription(FMeshDescription& OutMeshDescri
 	return true;
 }
 
+
+void FSkeletalMeshImportData::CopySkinWeightsToMeshDescription(
+	const FName InSkinWeightName, 
+	const FSkeletalMeshImportData& InSkinWeightMesh,
+	const TArray<FVertexID>& InVertexIDMap,
+	FMeshDescription& OutMeshDescription
+	) const
+{
+	// The invariant here is that the point count is the same as on the base mesh.
+
+	// - First we check if the topology is the same. For some reason, the two meshes could have been triangulated differently on 
+	//   import, in which case we will follow the next set of steps. However, if the topology is exactly the same, we just assume  
+	//   that if there are point position differences, they're caused by subtle deformations during output, rather than some
+	//   explicit differences.
+	// - If the topologies differ, we use a closest-triangle method implemented in FSkeletalMeshOperations::CopySkinWeightAttributeFromMesh.
+	FSkeletalMeshAttributes MeshAttributes(OutMeshDescription);
+	FSkinWeightsVertexAttributesRef VertexSkinWeights = MeshAttributes.GetVertexSkinWeights(InSkinWeightName);
+
+	// Create a mapping from the alt influence mesh bones to the base mesh bones.
+	TMap<FString, int32> BaseBoneToIndexMap;
+	for (int32 BoneIndex = 0; BoneIndex < RefBonesBinary.Num(); BoneIndex++)
+	{
+		BaseBoneToIndexMap.Add(RefBonesBinary[BoneIndex].Name, BoneIndex);
+	}
+	TMap<int32, int32> AltMeshBoneToBaseBoneMap;
+	for (int32 BoneIndex = 0; BoneIndex < InSkinWeightMesh.RefBonesBinary.Num(); BoneIndex++)
+	{
+		const int32 BaseBoneIndex = BaseBoneToIndexMap[InSkinWeightMesh.RefBonesBinary[BoneIndex].Name];
+		AltMeshBoneToBaseBoneMap.Add(BoneIndex, BaseBoneIndex);
+	}
+	
+	bool bTopologySame = true;
+	if (InSkinWeightMesh.Faces.Num() == Faces.Num())
+	{
+		for (int32 FaceIndex = 0; FaceIndex < Faces.Num(); FaceIndex++)
+		{
+			if (Faces[FaceIndex].WedgeIndex[0] != InSkinWeightMesh.Faces[FaceIndex].WedgeIndex[0] ||
+				Faces[FaceIndex].WedgeIndex[1] != InSkinWeightMesh.Faces[FaceIndex].WedgeIndex[1] ||
+				Faces[FaceIndex].WedgeIndex[2] != InSkinWeightMesh.Faces[FaceIndex].WedgeIndex[2])
+			{
+				bTopologySame = false;
+				break;
+			}
+		}
+	}
+	else
+	{
+		bTopologySame = false;
+	}
+
+	if (bTopologySame && InSkinWeightMesh.Wedges.Num() == Wedges.Num())
+	{
+		for (int32 WedgeIndex = 0; WedgeIndex < Wedges.Num(); WedgeIndex++)
+		{
+			if (Wedges[WedgeIndex].VertexIndex != InSkinWeightMesh.Wedges[WedgeIndex].VertexIndex)
+			{
+				bTopologySame = false;
+				break;
+			}
+		}
+	}
+	else
+	{
+		bTopologySame = false;
+	}
+
+	if (bTopologySame)
+	{
+		CopySkinWeightsToAttribute(InSkinWeightMesh.Influences, InVertexIDMap, &AltMeshBoneToBaseBoneMap, VertexSkinWeights);
+		return;
+	}
+
+	// The topologies don't match, proceed as above.
+	FMeshDescription AlternateInfluenceMesh;
+	InSkinWeightMesh.GetMeshDescription(AlternateInfluenceMesh);
+	
+	FSkeletalMeshOperations::CopySkinWeightAttributeFromMesh(
+	AlternateInfluenceMesh, OutMeshDescription, NAME_None, InSkinWeightName, &AltMeshBoneToBaseBoneMap);   
+}
+
+
+
 template<typename T>
 struct FCopyAttributeElement {};
 
@@ -2121,6 +2337,7 @@ private:
 	const FVertexArray& Vertices;
 };
 
+
 FSkeletalMeshImportData FSkeletalMeshImportData::CreateFromMeshDescription(const FMeshDescription& InMeshDescription)
 {
 	FSkeletalMeshImportData SkelMeshImportData;
@@ -2274,6 +2491,65 @@ FSkeletalMeshImportData FSkeletalMeshImportData::CreateFromMeshDescription(const
 		}
 	}
 
+	// Copy morph targets.
+	for (const FName MorphTargetName: Attributes.GetMorphTargetNames())
+	{
+		FMorphTargetVertexAttributesConstRef MorphTargetAttribute = Attributes.GetVertexMorphTarget(MorphTargetName);
+
+		SkelMeshImportData.MorphTargetNames.Add(MorphTargetName.ToString());
+		FSkeletalMeshImportData& MorphTarget = SkelMeshImportData.MorphTargets.AddDefaulted_GetRef();
+		TSet<uint32> MorphTargetModifiedPoints = SkelMeshImportData.MorphTargetModifiedPoints.AddDefaulted_GetRef();
+
+		MorphTarget.Points.SetNumZeroed(SkelMeshImportData.Points.Num());
+		for (FVertexID VertexID : InMeshDescription.Vertices().GetElementIDs())
+		{
+			FVector3f PositionDelta = MorphTargetAttribute.GetPositionDelta(VertexID);
+			if (!PositionDelta.IsNearlyZero())
+			{
+				MorphTarget.Points[VertexID.GetValue()] = MorphTargetAttribute.GetPositionDelta(VertexID) + SkelMeshImportData.Points[VertexID.GetValue()];
+				MorphTargetModifiedPoints.Add(VertexID.GetValue());
+			}
+		}
+	}
+
+	// Copy alternate influences.
+	for (const FName SkinWeightProfileName: Attributes.GetSkinWeightProfileNames())
+	{
+		// The default profile was already handled above.
+		if (SkinWeightProfileName == FSkeletalMeshAttributes::DefaultSkinWeightProfileName)
+		{
+			continue;
+		}
+		
+		FSkinWeightsVertexAttributesConstRef SkinWeightsAttribute = Attributes.GetVertexSkinWeights(SkinWeightProfileName);
+
+		SkelMeshImportData.AlternateInfluenceProfileNames.Add(SkinWeightProfileName.ToString());
+		FSkeletalMeshImportData& AlternateInfluenceMesh = SkelMeshImportData.AlternateInfluences.AddDefaulted_GetRef();
+		
+		// Copy over all the relevant mesh data aside from influences themselves.
+		AlternateInfluenceMesh.Points = SkelMeshImportData.Points;
+		AlternateInfluenceMesh.Wedges = SkelMeshImportData.Wedges;
+		AlternateInfluenceMesh.Faces = SkelMeshImportData.Faces;
+		AlternateInfluenceMesh.RefBonesBinary = SkelMeshImportData.RefBonesBinary;
+
+		for (FVertexID VertexID : InMeshDescription.Vertices().GetElementIDs())
+		{
+			FVertexBoneWeightsConst BoneWeights = SkinWeightsAttribute.Get(VertexID);
+			const int32 InfluenceCount = BoneWeights.Num();
+
+			const int32 InfluenceOffsetIndex = AlternateInfluenceMesh.Influences.Num();
+			AlternateInfluenceMesh.Influences.AddDefaulted(InfluenceCount);
+			for (int32 InfluenceIndex = 0; InfluenceIndex < InfluenceCount; ++InfluenceIndex)
+			{
+				SkeletalMeshImportData::FRawBoneInfluence& BoneInfluence = AlternateInfluenceMesh.Influences[InfluenceOffsetIndex + InfluenceIndex];
+				BoneInfluence.VertexIndex = VertexID.GetValue();
+				BoneInfluence.BoneIndex = BoneWeights[InfluenceIndex].GetBoneIndex();
+				BoneInfluence.Weight = BoneWeights[InfluenceIndex].GetWeight();
+			}
+		}
+	}
+	
+
 	SkelMeshImportData.CleanUpUnusedMaterials();
 	SkelMeshImportData.SplitVerticesBySmoothingGroups();
 	
diff --git a/Engine/Source/Runtime/Engine/Public/Rendering/SkeletalMeshLODImporterData.h b/Engine/Source/Runtime/Engine/Public/Rendering/SkeletalMeshLODImporterData.h
--- a/Engine/Source/Runtime/Engine/Public/Rendering/SkeletalMeshLODImporterData.h
+++ b/Engine/Source/Runtime/Engine/Public/Rendering/SkeletalMeshLODImporterData.h
@@ -1,6 +1,7 @@
 // Copyright Epic Games, Inc. All Rights Reserved.
 
 #pragma once
+#include "MeshDescription.h"
 
 #if WITH_EDITOR
 
@@ -420,18 +421,25 @@ public:
 	/**
 	 * Returns a mesh description from the import data
 	 */
-	ENGINE_API bool GetMeshDescription(FMeshDescription &OutMeshDescription) const;
+	ENGINE_API bool GetMeshDescription(FMeshDescription& OutMeshDescription) const;
 
 	/**
 	 * @note MeshDescription always contains color, normal and tangent data by default. Therefore, while iterating
 	 * over the vertices, we check if at least one normal/tangent vector is not a zero vector and the color is
 	 * not white, then we set the corresponding bHasNormals/bHasTangent/bHasVertexColors flags to true.
 	 */
-	static ENGINE_API FSkeletalMeshImportData CreateFromMeshDescription(const FMeshDescription &InMeshDescription);
+	static ENGINE_API FSkeletalMeshImportData CreateFromMeshDescription(const FMeshDescription& InMeshDescription);
 
 private:
-	ENGINE_API void CleanUpUnusedMaterials();
-	ENGINE_API void SplitVerticesBySmoothingGroups();
+	void CopySkinWeightsToMeshDescription(
+		const FName InSkinWeightName,
+		const FSkeletalMeshImportData& InSkinWeightMesh,
+		const TArray<FVertexID>& InVertexIDMap,
+		FMeshDescription& OutMeshDescription
+		) const;
+
+	void CleanUpUnusedMaterials();
+	void SplitVerticesBySmoothingGroups();
 	
 	friend FArchive& operator<<(FArchive& Ar, FSkeletalMeshImportData& RawMesh);
 };
diff --git a/Engine/Source/Runtime/SkeletalMeshDescription/Private/SkeletalMeshAttributes.cpp b/Engine/Source/Runtime/SkeletalMeshDescription/Private/SkeletalMeshAttributes.cpp
--- a/Engine/Source/Runtime/SkeletalMeshDescription/Private/SkeletalMeshAttributes.cpp
+++ b/Engine/Source/Runtime/SkeletalMeshDescription/Private/SkeletalMeshAttributes.cpp
@@ -5,7 +5,10 @@
 
 namespace MeshAttribute
 {
-	const FName Vertex::SkinWeights("SkinWeights");
+	namespace Vertex
+	{
+		const FName SkinWeights("SkinWeights");
+	}
 	
 	namespace Bone
 	{
@@ -21,6 +24,9 @@ FName FSkeletalMeshAttributesShared::DefaultSkinWeightProfileName("Default");
 
 FName FSkeletalMeshAttributesShared::BonesElementName("BonesElementName");
 
+static FString MorphTargetAttributeNamePrefix("Morph-");
+
+
 static FString SkinWeightAttributeNamePrefix()
 {
 	return MeshAttribute::Vertex::SkinWeights.ToString() + TEXT("-");
@@ -70,7 +76,8 @@ void FSkeletalMeshAttributes::RegisterColorAttribute()
 	BoneAttributes().RegisterAttribute<FVector4f>(MeshAttribute::Bone::Color, 1, FVector4f(1.0f, 1.0f, 1.0f, 1.0f), EMeshAttributeFlags::Mandatory);
 }
 
-bool FSkeletalMeshAttributes::RegisterSkinWeightAttribute(const FName& InProfileName)
+
+bool FSkeletalMeshAttributes::RegisterSkinWeightAttribute(const FName InProfileName)
 {
 	if (!IsValidSkinWeightProfileName(InProfileName))
 	{
@@ -78,29 +85,27 @@ bool FSkeletalMeshAttributes::RegisterSkinWeightAttribute(const FName& InProfile
 	}	
 	
 	const FName AttributeName = CreateSkinWeightAttributeName(InProfileName);
-	if (ensure(AttributeName.IsValid()))
+	if (!ensure(AttributeName.IsValid()))
 	{
-		TArray<FName> AllAttributeNames;
-		MeshDescription.VertexAttributes().GetAttributeNames(AllAttributeNames);
-		if (AllAttributeNames.Contains(AttributeName))
-		{
-			return false;
-		}
-
-		return MeshDescription.VertexAttributes().RegisterAttribute<int32[]>(AttributeName, 1, 0, EMeshAttributeFlags::Mandatory).IsValid();
+		return false;
 	}
-	else
+	
+	if (MeshDescription.VertexAttributes().HasAttribute(AttributeName))
 	{
 		return false;
-	}	
+	}
+
+	return MeshDescription.VertexAttributes().RegisterAttribute<int32[]>(AttributeName, 1, 0, EMeshAttributeFlags::None).IsValid();
 }
 
-FSkinWeightsVertexAttributesRef FSkeletalMeshAttributes::GetVertexSkinWeights(const FName& InProfileName)
+
+FSkinWeightsVertexAttributesRef FSkeletalMeshAttributes::GetVertexSkinWeights(const FName InProfileName)
 {
 	return MeshDescription.VertexAttributes().GetAttributesRef<TArrayAttribute<int32>>(CreateSkinWeightAttributeName(InProfileName));
 }
 
-FSkinWeightsVertexAttributesRef FSkeletalMeshAttributes::GetVertexSkinWeightsFromAttributeName(const FName& InAttributeName)
+
+FSkinWeightsVertexAttributesRef FSkeletalMeshAttributes::GetVertexSkinWeightsFromAttributeName(const FName InAttributeName)
 {
 	if (IsSkinWeightAttribute(InAttributeName))
 	{
@@ -112,6 +117,44 @@ FSkinWeightsVertexAttributesRef FSkeletalMeshAttributes::GetVertexSkinWeightsFro
 	}
 }
 
+
+bool FSkeletalMeshAttributes::RegisterMorphTargetAttribute(const FName InMorphTargetName)
+{
+	if (InMorphTargetName.IsNone())
+	{
+		return false;
+	}
+
+	const FName AttributeName = CreateMorphTargetAttributeName(InMorphTargetName);
+	if (!ensure(AttributeName.IsValid()))
+	{
+		return false;
+	}
+
+	// Already has this attribute?
+	if (MeshDescription.VertexAttributes().HasAttribute(AttributeName))
+	{
+		return false;
+	}
+
+	return MeshDescription.VertexAttributes().RegisterAttribute<FVector3f[2]>(AttributeName, 1, FVector3f::ZeroVector, EMeshAttributeFlags::None).IsValid();
+}
+
+FMorphTargetVertexAttributesRef FSkeletalMeshAttributes::GetVertexMorphTarget(const FName InMorphTargetName)
+{
+	return MeshDescription.VertexAttributes().GetAttributesRef<TArrayView<FVector3f>>(CreateMorphTargetAttributeName(InMorphTargetName));
+}
+
+FMorphTargetVertexAttributesRef FSkeletalMeshAttributes::GetVertexMorphTargetFromAttributeName(const FName InAttributeName)
+{
+	if (!IsMorphTargetAttribute(InAttributeName))
+	{
+		return {};
+	}
+	return MeshDescription.VertexAttributes().GetAttributesRef<TArrayView<FVector3f>>(InAttributeName);
+}
+
+
 FSkeletalMeshAttributes::FBoneArray& FSkeletalMeshAttributes::Bones() 
 { 
 	return static_cast<TMeshElementContainer<FBoneID>&>(BoneElements->Get()); 
@@ -178,12 +221,12 @@ MeshDescriptionShared(InMeshDescription)
 	}
 }
 
-FSkinWeightsVertexAttributesConstRef FSkeletalMeshAttributesShared::GetVertexSkinWeights(const FName& InProfileName) const
+FSkinWeightsVertexAttributesConstRef FSkeletalMeshAttributesShared::GetVertexSkinWeights(const FName InProfileName) const
 {
 	return MeshDescriptionShared.VertexAttributes().GetAttributesRef<TArrayAttribute<int32>>(CreateSkinWeightAttributeName(InProfileName));
 }
 
-FSkinWeightsVertexAttributesConstRef FSkeletalMeshAttributesShared::GetVertexSkinWeightsFromAttributeName(const FName& InAttributeName) const
+FSkinWeightsVertexAttributesConstRef FSkeletalMeshAttributesShared::GetVertexSkinWeightsFromAttributeName(const FName InAttributeName) const
 {
 	if (IsSkinWeightAttribute(InAttributeName))
 	{
@@ -203,7 +246,7 @@ TArray<FName> FSkeletalMeshAttributesShared::GetSkinWeightProfileNames() const
 	TArray<FName> SkinWeightProfileNames;
 	bool bHasDefault = false; 
 
-	for (const FName& AttributeName: AllAttributeNames)
+	for (const FName AttributeName: AllAttributeNames)
 	{
 		if (AttributeName == MeshAttribute::Vertex::SkinWeights)
 		{
@@ -224,18 +267,18 @@ TArray<FName> FSkeletalMeshAttributesShared::GetSkinWeightProfileNames() const
 	return SkinWeightProfileNames;
 }
 
-bool FSkeletalMeshAttributesShared::IsValidSkinWeightProfileName(const FName& InProfileName)
+bool FSkeletalMeshAttributesShared::IsValidSkinWeightProfileName(const FName InProfileName)
 {
 	return !InProfileName.IsNone() && !InProfileName.IsEqual(DefaultSkinWeightProfileName, ENameCase::IgnoreCase);
 }
 
-bool FSkeletalMeshAttributesShared::IsSkinWeightAttribute(const FName& InAttributeName)
+bool FSkeletalMeshAttributesShared::IsSkinWeightAttribute(const FName InAttributeName)
 {
 	return InAttributeName == MeshAttribute::Vertex::SkinWeights ||
 		   InAttributeName.ToString().StartsWith(SkinWeightAttributeNamePrefix());
 }
 
-FName FSkeletalMeshAttributesShared::GetProfileNameFromAttribute(const FName& InAttributeName)
+FName FSkeletalMeshAttributesShared::GetProfileNameFromAttribute(const FName InAttributeName)
 {
 	if (InAttributeName == MeshAttribute::Vertex::SkinWeights)
 	{
@@ -253,7 +296,7 @@ FName FSkeletalMeshAttributesShared::GetProfileNameFromAttribute(const FName& In
 	}
 }
 
-FName FSkeletalMeshAttributesShared::CreateSkinWeightAttributeName(const FName& InProfileName)
+FName FSkeletalMeshAttributesShared::CreateSkinWeightAttributeName(const FName InProfileName)
 {
 	// If it's the default profile, then return the base skin weights attribute name.
 	if (InProfileName.IsNone() || InProfileName.IsEqual(FSkeletalMeshAttributesShared::DefaultSkinWeightProfileName, ENameCase::IgnoreCase))
@@ -264,6 +307,62 @@ FName FSkeletalMeshAttributesShared::CreateSkinWeightAttributeName(const FName&
 	return FName(SkinWeightAttributeNamePrefix() + InProfileName.ToString());
 }
 
+FName FSkeletalMeshAttributesShared::CreateMorphTargetAttributeName(const FName InMorphTargetName)
+{
+	if (ensure(!InMorphTargetName.IsNone()))
+	{
+		return FName(MorphTargetAttributeNamePrefix + InMorphTargetName.ToString());
+	}
+	
+	return NAME_None;
+}
+
+
+TArray<FName> FSkeletalMeshAttributesShared::GetMorphTargetNames() const
+{
+	TArray<FName> AllAttributeNames;
+	MeshDescriptionShared.VertexAttributes().GetAttributeNames(AllAttributeNames);
+
+	TArray<FName> AllMorphTargetNames;
+	for (const FName AttributeName: AllAttributeNames)
+	{
+		if (IsMorphTargetAttribute(AttributeName))
+		{
+			AllMorphTargetNames.Add(GetMorphTargetNameFromAttribute(AttributeName));
+		}
+	}
+	return AllMorphTargetNames;
+}
+
+bool FSkeletalMeshAttributesShared::IsMorphTargetAttribute(const FName InAttributeName)
+{
+	return InAttributeName.ToString().StartsWith(MorphTargetAttributeNamePrefix);
+}
+
+FName FSkeletalMeshAttributesShared::GetMorphTargetNameFromAttribute(const FName InAttributeName)
+{
+	if (ensure(IsMorphTargetAttribute(InAttributeName)))
+	{
+		return FName(InAttributeName.ToString().Mid(MorphTargetAttributeNamePrefix.Len()));
+	}
+	return NAME_None;
+}
+
+FMorphTargetVertexAttributesConstRef FSkeletalMeshAttributesShared::GetVertexMorphTarget(const FName InMorphTargetName) const
+{
+	return MeshDescriptionShared.VertexAttributes().GetAttributesRef<TArrayView<FVector3f>>(CreateMorphTargetAttributeName(InMorphTargetName));
+}
+
+FMorphTargetVertexAttributesConstRef FSkeletalMeshAttributesShared::GetVertexMorphTargetFromAttributeName(const FName InAttributeName) const
+{
+	if (!IsMorphTargetAttribute(InAttributeName))
+	{
+		return {};
+	}
+	return MeshDescriptionShared.VertexAttributes().GetAttributesRef<TArrayView<FVector3f>>(InAttributeName);
+}
+
+
 bool FSkeletalMeshAttributesShared::HasBoneColorAttribute() const
 {
 	return BoneAttributes().HasAttribute(MeshAttribute::Bone::Color);
@@ -299,6 +398,11 @@ int32 FSkeletalMeshAttributesShared::GetNumBones() const
 	return HasBones() ? BoneElementsShared->Get().Num() : 0;
 }
 
+bool FSkeletalMeshAttributesShared::IsBoneValid(const FBoneID BoneID) const
+{
+	return BoneElementsShared->Get().IsValid(BoneID.GetValue());
+}
+
 FSkeletalMeshAttributesShared::FBoneNameAttributesConstRef FSkeletalMeshAttributesShared::GetBoneNames() const
 {
 	return BoneAttributes().GetAttributesRef<FName>(MeshAttribute::Bone::Name);
diff --git a/Engine/Source/Runtime/SkeletalMeshDescription/Private/SkeletalMeshOperations.cpp b/Engine/Source/Runtime/SkeletalMeshDescription/Private/SkeletalMeshOperations.cpp
--- a/Engine/Source/Runtime/SkeletalMeshDescription/Private/SkeletalMeshOperations.cpp
+++ b/Engine/Source/Runtime/SkeletalMeshDescription/Private/SkeletalMeshOperations.cpp
@@ -2,8 +2,10 @@
 
 #include "SkeletalMeshOperations.h"
 
-#include "SkeletalMeshAttributes.h"
 #include "BoneWeights.h"
+#include "MeshDescriptionAdapter.h"
+#include "SkeletalMeshAttributes.h"
+#include "Spatial/MeshAABBTree3.h"
 
 
 DEFINE_LOG_CATEGORY(LogSkeletalMeshOperations);
@@ -112,4 +114,112 @@ void FSkeletalMeshOperations::AppendSkinWeight(const FMeshDescription& SourceMes
 	TargetMesh.ResumeVertexIndexing();
 }
 
-#undef LOCTEXT_NAMESPACE
\ No newline at end of file
+
+bool FSkeletalMeshOperations::CopySkinWeightAttributeFromMesh(
+	const FMeshDescription& InSourceMesh,
+	FMeshDescription& InTargetMesh,
+	const FName InSourceProfile,
+	const FName InTargetProfile,
+	const TMap<int32, int32>* SourceBoneIndexToTargetBoneIndexMap
+	)
+{
+	// This is effectively a slower and dumber version of FTransferBoneWeights.
+	using namespace UE::AnimationCore;
+	using namespace UE::Geometry;
+
+	FSkeletalMeshConstAttributes SourceAttributes(InSourceMesh);
+	FSkeletalMeshAttributes TargetAttributes(InTargetMesh);
+	
+	FSkinWeightsVertexAttributesConstRef SourceWeights = SourceAttributes.GetVertexSkinWeights(InSourceProfile);
+	FSkinWeightsVertexAttributesRef TargetWeights = TargetAttributes.GetVertexSkinWeights(InTargetProfile);
+	TVertexAttributesConstRef<FVector3f> TargetPositions = TargetAttributes.GetVertexPositions();
+
+	if (!SourceWeights.IsValid() || !TargetWeights.IsValid())
+	{
+		return false;
+	}
+	
+	FMeshDescriptionTriangleMeshAdapter MeshAdapter(&InSourceMesh);
+	TMeshAABBTree3<FMeshDescriptionTriangleMeshAdapter> BVH(&MeshAdapter);
+
+	auto RemapBoneWeights = [SourceBoneIndexToTargetBoneIndexMap](const FVertexBoneWeightsConst& InWeights) -> FBoneWeights
+	{
+		TArray<FBoneWeight, TInlineAllocator<MaxInlineBoneWeightCount>> Weights;
+
+		if (SourceBoneIndexToTargetBoneIndexMap)
+		{
+			for (FBoneWeight OriginalWeight: InWeights)
+			{
+				FBoneWeight NewWeight(static_cast<FBoneIndexType>((*SourceBoneIndexToTargetBoneIndexMap)[OriginalWeight.GetBoneIndex()]), OriginalWeight.GetRawWeight());
+				Weights.Add(NewWeight);
+			}
+		}
+		else
+		{
+			for (FBoneWeight Weight: InWeights)
+			{
+				Weights.Add(Weight);
+			}
+		}
+		return FBoneWeights::Create(Weights);
+	};
+	
+	auto InterpolateWeights = [&MeshAdapter, &SourceWeights, &RemapBoneWeights](int32 InTriangleIndex, const FVector3d& InTargetPoint) -> FBoneWeights
+	{
+		const FDistPoint3Triangle3d Query = TMeshQueries<FMeshDescriptionTriangleMeshAdapter>::TriangleDistance(MeshAdapter, InTriangleIndex, InTargetPoint);
+
+		const FIndex3i TriangleVertexes = MeshAdapter.GetTriangle(InTriangleIndex);
+		const FVector3f BaryCoords(VectorUtil::BarycentricCoords(Query.ClosestTrianglePoint, MeshAdapter.GetVertex(TriangleVertexes.A), MeshAdapter.GetVertex(TriangleVertexes.B), MeshAdapter.GetVertex(TriangleVertexes.C)));
+		const FBoneWeights WeightsA = RemapBoneWeights(SourceWeights.Get(TriangleVertexes.A));
+		const FBoneWeights WeightsB = RemapBoneWeights(SourceWeights.Get(TriangleVertexes.B));
+		const FBoneWeights WeightsC = RemapBoneWeights(SourceWeights.Get(TriangleVertexes.C));
+
+		FBoneWeights BoneWeights = FBoneWeights::Blend(WeightsA, WeightsB, WeightsC, BaryCoords.X, BaryCoords.Y, BaryCoords.Z);
+		
+		// Blending can leave us with zero weights. Let's strip them out here.
+		BoneWeights.Renormalize();
+		return BoneWeights;
+	};
+
+	TArray<FBoneWeights> TargetBoneWeights;
+	TargetBoneWeights.SetNum(InTargetMesh.Vertices().GetArraySize());
+
+	ParallelFor(InTargetMesh.Vertices().GetArraySize(), [&BVH, &InTargetMesh, &TargetPositions, &TargetBoneWeights, &InterpolateWeights](int32 InVertexIndex)
+	{
+		const FVertexID VertexID(InVertexIndex);
+		if (!InTargetMesh.Vertices().IsValid(VertexID))
+		{
+			return;
+		}
+
+		const FVector3d TargetPoint(TargetPositions.Get(VertexID));
+
+		const IMeshSpatial::FQueryOptions Options;
+		double NearestDistanceSquared;
+		const int32 NearestTriangleIndex = BVH.FindNearestTriangle(TargetPoint, NearestDistanceSquared, Options);
+
+		if (!ensure(NearestTriangleIndex != IndexConstants::InvalidID))
+		{
+			return;
+		}
+
+		TargetBoneWeights[InVertexIndex] = InterpolateWeights(NearestTriangleIndex, TargetPoint);
+	});
+
+	// Transfer the computed bone weights to the target mesh.
+	for (FVertexID TargetVertexID: InTargetMesh.Vertices().GetElementIDs())
+	{
+		FBoneWeights& BoneWeights = TargetBoneWeights[TargetVertexID];
+		if (BoneWeights.Num() == 0)
+		{
+			// Bind to root so that we have something.
+			BoneWeights.SetBoneWeight(FBoneIndexType{0}, 1.0);
+		}
+
+		TargetWeights.Set(TargetVertexID, BoneWeights);
+	}
+
+	return true;
+}
+
+#undef LOCTEXT_NAMESPACE
diff --git a/Engine/Source/Runtime/SkeletalMeshDescription/Public/MorphTargetAttributesRef.h b/Engine/Source/Runtime/SkeletalMeshDescription/Public/MorphTargetAttributesRef.h
new file mode 100644
--- /dev/null
+++ b/Engine/Source/Runtime/SkeletalMeshDescription/Public/MorphTargetAttributesRef.h
@@ -0,0 +1,118 @@
+﻿// Copyright Epic Games, Inc. All Rights Reserved.
+
+#pragma once
+
+#include "MeshDescription.h"
+
+class FMorphTargetVertexAttributesRef
+{
+public:
+	FMorphTargetVertexAttributesRef() = default;
+	FMorphTargetVertexAttributesRef(const FMorphTargetVertexAttributesRef&) = default;
+	FMorphTargetVertexAttributesRef(FMorphTargetVertexAttributesRef&&) = default;
+	FMorphTargetVertexAttributesRef& operator=(const FMorphTargetVertexAttributesRef&) = default;
+	FMorphTargetVertexAttributesRef& operator=(FMorphTargetVertexAttributesRef&&) = default;
+
+	bool IsValid() const
+	{
+		return AttributesRef.IsValid();
+	}
+
+	FVector3f GetPositionDelta(const FVertexID InVertexIndex) const
+	{
+		return AttributesRef.Get(InVertexIndex)[0]; 
+	}
+
+	FVector3f GetTangentZDelta(const FVertexID InVertexIndex) const
+	{
+		return AttributesRef.Get(InVertexIndex)[1]; 
+	}
+
+	void SetPositionDelta(
+		const FVertexID InVertexIndex,
+		FVector3f InPositionDelta
+		)
+	{
+		AttributesRef.Get(InVertexIndex)[0] = InPositionDelta;
+	}
+	
+	void SetTangentZDelta(
+		const FVertexID InVertexIndex,
+		FVector3f InTangentZDelta
+		)
+	{
+		AttributesRef.Get(InVertexIndex)[1] = InTangentZDelta;
+	}
+	
+	void SetPositionAndTangentZDelta(
+		const FVertexID InVertexIndex,
+		FVector3f InPositionDelta,
+		FVector3f InTangentZDelta
+		)
+	{
+		AttributesRef.SetArrayView(InVertexIndex, {InPositionDelta, InTangentZDelta});
+	}
+
+protected:
+	friend class FSkeletalMeshAttributes;
+
+	FMorphTargetVertexAttributesRef(TVertexAttributesRef<TArrayView<FVector3f>> InAttributesRef)
+		: AttributesRef(InAttributesRef)
+	{}
+
+private:
+	friend class FMorphTargetVertexAttributesConstRef;
+	TVertexAttributesRef<TArrayView<FVector3f>> AttributesRef;
+};
+
+
+class FMorphTargetVertexAttributesConstRef
+{
+public:
+	FMorphTargetVertexAttributesConstRef() = default;
+	FMorphTargetVertexAttributesConstRef(const FMorphTargetVertexAttributesConstRef&) = default;
+	FMorphTargetVertexAttributesConstRef(FMorphTargetVertexAttributesConstRef&&) = default;
+	FMorphTargetVertexAttributesConstRef& operator=(const FMorphTargetVertexAttributesConstRef&) = default;
+	FMorphTargetVertexAttributesConstRef& operator=(FMorphTargetVertexAttributesConstRef&&) = default;
+
+	// Converting constructors from the non-const variant
+	explicit FMorphTargetVertexAttributesConstRef(const FMorphTargetVertexAttributesRef &InAttributesRef) :
+		AttributesConstRef(InAttributesRef.AttributesRef)
+	{}
+	
+	FMorphTargetVertexAttributesConstRef& operator=(const FMorphTargetVertexAttributesRef &InAttributesRef)
+	{
+		AttributesConstRef = InAttributesRef.AttributesRef;
+		return *this;
+	}
+	
+	bool IsValid() const
+	{
+		return AttributesConstRef.IsValid();
+	}
+
+	FVector3f GetPositionDelta(const FVertexID InVertexIndex) const
+	{
+		return AttributesConstRef.Get(InVertexIndex)[0]; 
+	}
+
+	FVector3f GetTangentZDelta(const FVertexID InVertexIndex) const
+	{
+		return AttributesConstRef.Get(InVertexIndex)[1]; 
+	}
+
+protected:
+	friend class FSkeletalMeshAttributes;
+	friend class FSkeletalMeshAttributesShared;
+
+	FMorphTargetVertexAttributesConstRef(TVertexAttributesRef<TArrayView<FVector3f>> InAttributesConstRef)
+		: AttributesConstRef(InAttributesConstRef)
+	{}
+
+	FMorphTargetVertexAttributesConstRef(TVertexAttributesRef<TArrayView<const FVector3f>> InAttributesConstRef)
+		: AttributesConstRef(InAttributesConstRef)
+	{}
+	
+private:
+	TVertexAttributesConstRef<TArrayView<FVector3f>> AttributesConstRef;
+};
diff --git a/Engine/Source/Runtime/SkeletalMeshDescription/Public/SkeletalMeshAttributes.h b/Engine/Source/Runtime/SkeletalMeshDescription/Public/SkeletalMeshAttributes.h
--- a/Engine/Source/Runtime/SkeletalMeshDescription/Public/SkeletalMeshAttributes.h
+++ b/Engine/Source/Runtime/SkeletalMeshDescription/Public/SkeletalMeshAttributes.h
@@ -5,7 +5,7 @@
 #include "Containers/Array.h"
 #include "HAL/Platform.h"
 #include "MeshAttributeArray.h"
-#include "MeshAttributes.h"
+#include "MorphTargetAttributesRef.h"
 #include "SkinWeightsAttributesRef.h"
 #include "StaticMeshAttributes.h"
 #include "UObject/NameTypes.h"
@@ -21,6 +21,7 @@ namespace MeshAttribute
 {
  	namespace Vertex
  	{
+ 		// Name of the default skin weights attribute.
 		extern SKELETALMESHDESCRIPTION_API const FName SkinWeights;
  	}
 
@@ -74,20 +75,36 @@ public:
 	
 	/// Returns \c true if the given identifier is a valid profile name. If the name is empty, or matches the default profile,
 	/// then the profile name is considered invalid. 
-	static SKELETALMESHDESCRIPTION_API bool IsValidSkinWeightProfileName(const FName& InProfileName);
+	static SKELETALMESHDESCRIPTION_API bool IsValidSkinWeightProfileName(const FName InProfileName);
 
 	/// Helper function that indicates whether an attribute name represents a skin weight attribute.
-	static SKELETALMESHDESCRIPTION_API bool IsSkinWeightAttribute(const FName& InAttributeName);
+	static SKELETALMESHDESCRIPTION_API bool IsSkinWeightAttribute(const FName InAttributeName);
 
 	/// Returns a skin profile name from the attribute name, if the attribute name is a valid skin weights
 	/// attribute.
-	static SKELETALMESHDESCRIPTION_API FName GetProfileNameFromAttribute(const FName& InAttributeName);
+	static SKELETALMESHDESCRIPTION_API FName GetProfileNameFromAttribute(const FName InAttributeName);
 
-	SKELETALMESHDESCRIPTION_API FSkinWeightsVertexAttributesConstRef GetVertexSkinWeights(const FName& InProfileName = NAME_None) const;
+	SKELETALMESHDESCRIPTION_API FSkinWeightsVertexAttributesConstRef GetVertexSkinWeights(const FName InProfileName = NAME_None) const;
 	
-	SKELETALMESHDESCRIPTION_API FSkinWeightsVertexAttributesConstRef GetVertexSkinWeightsFromAttributeName(const FName& InAttributeName = NAME_None) const;
+	SKELETALMESHDESCRIPTION_API FSkinWeightsVertexAttributesConstRef GetVertexSkinWeightsFromAttributeName(const FName InAttributeName = NAME_None) const;
 
+	//
+	// Morph Target Methods
+	//
+
+	/// Returns the list of all registered morph targets on this mesh.
+	SKELETALMESHDESCRIPTION_API TArray<FName> GetMorphTargetNames() const;
+	
+	/// Returns \c true if the given attribute name refers to a morph target attribute.
+	static SKELETALMESHDESCRIPTION_API bool IsMorphTargetAttribute(const FName InAttributeName);
 
+	/// Returns the name of a morph target given the attribute name. If the given attribute name is invalid, \c NAME_None is returned.
+	static SKELETALMESHDESCRIPTION_API FName GetMorphTargetNameFromAttribute(const FName InAttributeName);
+
+	SKELETALMESHDESCRIPTION_API FMorphTargetVertexAttributesConstRef GetVertexMorphTarget(const FName InMorphTargetName = NAME_None) const;
+	
+	SKELETALMESHDESCRIPTION_API FMorphTargetVertexAttributesConstRef GetVertexMorphTargetFromAttributeName(const FName InAttributeName = NAME_None) const;
+	
 	//
 	// Bones Methods
 	//
@@ -128,8 +145,11 @@ protected:
 	/// Construct a name for a skin weight attribute with the given skin weight profile name.
 	/// Each mesh description can hold different skin weight profiles, although the default
 	/// is always present.
-	static SKELETALMESHDESCRIPTION_API FName CreateSkinWeightAttributeName(const FName& InProfileName);
+	static SKELETALMESHDESCRIPTION_API FName CreateSkinWeightAttributeName(const FName InProfileName);
 
+	/// Construct a name for a morph target attribute with the given a user-visible morph target name.
+	static SKELETALMESHDESCRIPTION_API FName CreateMorphTargetAttributeName(const FName InMorphTargetName);
+	
 private:
 	const FMeshElementChannels* BoneElementsShared = nullptr;
 
@@ -146,31 +166,45 @@ public:
 	SKELETALMESHDESCRIPTION_API explicit FSkeletalMeshAttributes(FMeshDescription& InMeshDescription);
 
 	SKELETALMESHDESCRIPTION_API virtual void Register(bool bKeepExistingAttribute = false) override;
-	
-	//
-	// Skin Weights Methods
-	//
-
-	/// Register a new skin weight profile with the given name. The attribute name will encode the profile name and
-	/// it will be listed in GetSkinWeightProfileNames(). Returns \c true if the profile was successfully registered.
-	/// Returns \c false if the attribute was already registered or if IsValidSkinWeightProfileName() returned false.
-	SKELETALMESHDESCRIPTION_API bool RegisterSkinWeightAttribute(const FName& InProfileName);
 
-	static bool IsReservedAttributeName(const FName& InAttributeName)
+	static bool IsReservedAttributeName(const FName InAttributeName)
 	{
 		return FStaticMeshAttributes::IsReservedAttributeName(InAttributeName) ||
 			   IsSkinWeightAttribute(InAttributeName) ||
+			   IsMorphTargetAttribute(InAttributeName) ||
 			   InAttributeName == MeshAttribute::Bone::Name ||
 			   InAttributeName == MeshAttribute::Bone::ParentIndex ||
 			   InAttributeName == MeshAttribute::Bone::Pose ||
 			   InAttributeName == MeshAttribute::Bone::Color;
 	}
 	
+	//
+	// Skin Weights Methods
+	//
+
+	/// Register a new skin weight profile with the given name. The attribute name will encode the profile name and
+	/// it will be listed in GetSkinWeightProfileNames(). Returns \c true if the profile was successfully registered.
+	/// Returns \c false if the attribute was already registered or if IsValidSkinWeightProfileName() returned false.
+	SKELETALMESHDESCRIPTION_API bool RegisterSkinWeightAttribute(const FName InProfileName);
+
 	/// Returns the skin weight profile given by its name. NAME_None corresponds to the default profile.
-	SKELETALMESHDESCRIPTION_API FSkinWeightsVertexAttributesRef GetVertexSkinWeights(const FName& InProfileName = NAME_None);
+	SKELETALMESHDESCRIPTION_API FSkinWeightsVertexAttributesRef GetVertexSkinWeights(const FName InProfileName = NAME_None);
 
-	SKELETALMESHDESCRIPTION_API FSkinWeightsVertexAttributesRef GetVertexSkinWeightsFromAttributeName(const FName& InAttributeName = NAME_None);
+	SKELETALMESHDESCRIPTION_API FSkinWeightsVertexAttributesRef GetVertexSkinWeightsFromAttributeName(const FName InAttributeName = NAME_None);
 
+	//
+	// Morph Target Methods
+	//
+	
+	/// Register a new morph target with the given name. The attribute name will encode the user-defined morph target name and
+	/// it will be listed in GetMorphTargetNames(). Returns \c true if the morph target was successfully registered.
+	/// Returns \c false if the attribute was already registered or if \c InMorphTargetName is empty.
+	SKELETALMESHDESCRIPTION_API bool RegisterMorphTargetAttribute(const FName InMorphTargetName);
+	
+	SKELETALMESHDESCRIPTION_API FMorphTargetVertexAttributesRef GetVertexMorphTarget(const FName InMorphTargetName = NAME_None);
+	
+	SKELETALMESHDESCRIPTION_API FMorphTargetVertexAttributesRef GetVertexMorphTargetFromAttributeName(const FName InAttributeName = NAME_None);
+	
 	//
 	// Bones Methods
 	//
diff --git a/Engine/Source/Runtime/SkeletalMeshDescription/Public/SkeletalMeshOperations.h b/Engine/Source/Runtime/SkeletalMeshDescription/Public/SkeletalMeshOperations.h
--- a/Engine/Source/Runtime/SkeletalMeshDescription/Public/SkeletalMeshOperations.h
+++ b/Engine/Source/Runtime/SkeletalMeshDescription/Public/SkeletalMeshOperations.h
@@ -32,4 +32,23 @@ public:
 	};
 	
 	static SKELETALMESHDESCRIPTION_API void AppendSkinWeight(const FMeshDescription& SourceMesh, FMeshDescription& TargetMesh, FSkeletalMeshAppendSettings& AppendSettings);
+
+	/* Copies skin weight attribute from one mesh to another. Assumes the two geometries are identical or near-identical.
+	 * Uses closest triangle on the target mesh to interpolate skin weights to each of the points on the target mesh.
+	 * Attributes for the given profiles on both meshes should exist in order for this function to succeed. 
+	 * \param InSourceMesh The mesh to copy skin weights from.
+	 * \param InTargetMesh The mesh to copy skin weights to.
+	 * \param InSourceProfile The name of the skin weight profile on the source mesh to read from.
+	 * \param InTargetProfile The name of the skin weight profile on the target mesh to write to.
+	 * \param SourceBoneIndexToTargetBoneIndexMap An optional mapping table to map bone indexes on the source mesh to the target mesh.
+	 *     The table needs to be complete for all the source bone indexes to valid target bone indexes, otherwise the behavior
+	 *     is undefined. If the table is not given, the bone indexes on the source and target meshes are assumed to be the same.
+	 */
+	static SKELETALMESHDESCRIPTION_API bool CopySkinWeightAttributeFromMesh(
+		const FMeshDescription& InSourceMesh,
+		FMeshDescription& InTargetMesh,
+		const FName InSourceProfile,
+		const FName InTargetProfile,
+		const TMap<int32, int32>* SourceBoneIndexToTargetBoneIndexMap
+		);
 };
diff --git a/Engine/Source/Runtime/SkeletalMeshDescription/SkeletalMeshDescription.Build.cs b/Engine/Source/Runtime/SkeletalMeshDescription/SkeletalMeshDescription.Build.cs
--- a/Engine/Source/Runtime/SkeletalMeshDescription/SkeletalMeshDescription.Build.cs
+++ b/Engine/Source/Runtime/SkeletalMeshDescription/SkeletalMeshDescription.Build.cs
@@ -14,6 +14,14 @@ namespace UnrealBuildTool.Rules
 				}
 			);
 
+			PrivateIncludePathModuleNames.AddRange(
+				new string[]
+				{
+					"GeometryCore",
+					"MeshConversion"
+				}
+			);
+
 			PublicDependencyModuleNames.AddRange(
 				new string[]
 				{
