Rename/move SoundWaveProxyReader to `Engine/Source/Runtime/Engine/Classes/Sound`

#tests PIE

diff --git a/Engine/Plugins/Runtime/Metasound/Source/MetasoundEngine/Private/MetasoundWavePlayerNode.cpp b/Engine/Plugins/Runtime/Metasound/Source/MetasoundEngine/Private/MetasoundWavePlayerNode.cpp
--- a/Engine/Plugins/Runtime/Metasound/Source/MetasoundEngine/Private/MetasoundWavePlayerNode.cpp
+++ b/Engine/Plugins/Runtime/Metasound/Source/MetasoundEngine/Private/MetasoundWavePlayerNode.cpp
@@ -17,7 +17,7 @@
 #include "MetasoundTrigger.h"
 #include "MetasoundVertex.h"
 #include "MetasoundWave.h"
-#include "SoundWaveProxyReader.h"
+#include "Sound/SoundWaveProxyReader.h"
 
 #define LOCTEXT_NAMESPACE "MetasoundWaveNode"
 
diff --git a/Engine/Source/Runtime/AudioCodecEngine/Private/SoundWaveProxyReader.cpp b/Engine/Source/Runtime/AudioCodecEngine/Private/SoundWaveProxyReader.cpp
deleted file mode 100644
diff --git a/Engine/Source/Runtime/AudioCodecEngine/Public/SoundWaveProxyReader.h b/Engine/Source/Runtime/AudioCodecEngine/Public/SoundWaveProxyReader.h
deleted file mode 100644
diff --git a/Engine/Source/Runtime/Engine/Classes/Sound/SoundWaveProxyReader.h b/Engine/Source/Runtime/Engine/Classes/Sound/SoundWaveProxyReader.h
new file mode 100644
--- /dev/null
+++ b/Engine/Source/Runtime/Engine/Classes/Sound/SoundWaveProxyReader.h
@@ -0,0 +1,206 @@
+// Copyright Epic Games, Inc. All Rights Reserved.
+
+#pragma once
+
+#include "DSP/BufferVectorOperations.h"
+#include "AudioDecompress.h"
+#include "HAL/Platform.h"
+#include "IAudioCodec.h"
+#include "Templates/SharedPointer.h"
+#include "Templates/UniquePtr.h"
+
+// Forward declare
+class FSoundWaveProxy;
+
+/** FSoundWaveProxyReader reads a FWaveProxy and outputs 32 bit interleaved audio.
+*
+* FSoundWaveProxyReader provides controls for looping and relevant frame index
+* values.
+*/
+class FSoundWaveProxyReader
+{
+
+public:
+	using FSoundWaveProxyRef = TSharedRef<FSoundWaveProxy, ESPMode::ThreadSafe>;
+
+	/** Minimum number of frames to decode per a call to the decoder.  */
+	static constexpr uint32 DefaultMinDecodeSizeInFrames = 128;
+
+	/** Some codecs have strict requirements on decode size. In order to be
+	 * functional with all supported codecs, the decode size must be a multiple
+	 * of 128.
+	 */
+	static constexpr uint32 DecodeSizeQuantizationInFrames = 128;
+
+	static uint32 ConformDecodeSize(uint32 InMaxDesiredDecodeSizeInFrames);
+
+	/** Minimum duration of a loop. */
+	static constexpr float MinLoopDurationInSeconds = 0.05f;
+
+	/** Maximum duration of a loop.
+	 *
+	 * 10 years. ridiculously high. This exists to prevent floating point
+	 * undefined overflow behavior if performing calculations with
+	 * TNumericLimits<float>::Max().
+	 */
+	static constexpr float MaxLoopDurationInSeconds = 60.f * 60.f * 365.f * 10.f;
+
+	/** Settings for a FSoundWaveProxyReader. */
+	struct FSettings
+	{
+		uint32 MaxDecodeSizeInFrames = 1024;
+		float StartTimeInSeconds = 0.f;
+		bool bIsLooping = false;
+		float LoopStartTimeInSeconds = 0.f;
+		float LoopDurationInSeconds = 0.f;
+		bool bMaintainAudioSync = false;
+	};
+
+private:
+	/** Construct a wave proxy reader.
+	 *
+	 * @param InWaveProxy - A TSharedRef of a FSoundWaveProxy which is to be played.
+	 * @param InSettings - Reader settings.
+	 */
+	FSoundWaveProxyReader(FSoundWaveProxyRef InWaveProxy, const FSettings& InSettings);
+
+public:
+	/** Create a wave proxy reader.
+	 *
+	 * @param InWaveProxy - A TSharedRef of a FSoundWaveProxy which is to be played.
+	 * @param InSettings - Reader settings.
+	 */
+	static ENGINE_API TUniquePtr<FSoundWaveProxyReader> Create(FSoundWaveProxyRef InWaveProxy, const FSettings& InSettings);
+
+	/** Set whether the reader should loop the audio or not. */
+	ENGINE_API void SetIsLooping(bool bInIsLooping);
+
+	/** Returns true if the audio will be looped, false otherwise. */
+	FORCEINLINE bool IsLooping() const
+	{
+		return Settings.bIsLooping;
+	}
+
+	/** Sets the beginning position of the loop. */
+	ENGINE_API void SetLoopStartTime(float InLoopStartTimeInSeconds);
+
+	/** Sets the duration of the loop in seconds.
+	 *
+	 * If the value is negative, the MaxLoopDurationInSeconds will be used
+	 * which will effectively loop at the end of the file.
+	 */
+	ENGINE_API void SetLoopDuration(float InLoopDurationInSeconds);
+
+	FORCEINLINE float GetSampleRate() const
+	{
+		return SampleRate;
+	}
+
+	FORCEINLINE int32 GetNumChannels() const
+	{
+		return NumChannels;
+	}
+
+	/** Returns the index of the playhead within the complete wave. */
+	FORCEINLINE int32 GetFrameIndex() const
+	{
+		return CurrentFrameIndex;
+	}
+
+	FORCEINLINE int32 GetNumFramesInWave() const
+	{
+		return NumFramesInWave;
+	}
+
+	FORCEINLINE int32 GetNumFramesInLoop() const
+	{
+		return LoopEndFrameIndex - LoopStartFrameIndex;
+	}
+
+	FORCEINLINE int32 GetLoopStartFrameIndex() const
+	{
+		return LoopStartFrameIndex;
+	}
+
+	FORCEINLINE int32 GetLoopEndFrameIndex() const
+	{
+		return LoopEndFrameIndex;
+	}
+
+	/** Seeks to position in wave.
+	 *
+	 * @param InSeconds - The location to seek the playhead
+	 *
+	 * @return true on success, false on failure.
+	 */
+	ENGINE_API bool SeekToTime(float InSeconds);
+
+	/** Seeks to position in wave at a specific frame.
+	 *
+	 * @param InFrameNum - The specific frame to seek the playhead
+	 *
+	 * @return true on success, false on failure.
+	 */
+	ENGINE_API bool SeekToFrame(uint32 InFrameNum);
+
+
+	/** Pops audio from reader and copies audio into OutBuffer. It returns the number of samples copied.
+	 * Samples not written to will be set to zero.
+	 */
+	ENGINE_API int32 PopAudio(Audio::FAlignedFloatBuffer& OutBuffer);
+
+	/** Returns TRUE if the reader can produce audio (eg. has a valid decoder, it can still decode, etc.)
+	 *
+	 * @return true on success, false on failure.
+	 */
+	ENGINE_API bool CanProduceMoreAudio() const;
+
+private:
+
+	enum EDecodeResult
+	{
+		Fail,					// Decoder failed somehow.
+		MoreDataRemaining,		// Data has produced and there's more remaining
+		Finished				// The decoder has reached the end of the wave data
+	};
+
+	EDecodeResult Decode();
+
+	int32 PopAudioFromDecoderOutput(TArrayView<float> OutBufferView);
+	bool InitializeDecoder(float InStartTimeInSeconds);
+	int32 DiscardSamples(int32 InNumSamplesToDiscard);
+	float ClampLoopStartTime(float InStartTimeInSeconds);
+	float ClampLoopDuration(float InDurationInSeconds);
+	void UpdateLoopBoundaries();
+
+	FSoundWaveProxyRef WaveProxy;
+
+	FSettings Settings;
+
+	float SampleRate = 0.f;
+	int32 NumChannels = 0;
+	EDecodeResult DecodeResult = EDecodeResult::MoreDataRemaining;
+	int32 NumFramesInWave = 0;
+	int32 NumDecodeSamplesToDiscard = 0;
+	int32 CurrentFrameIndex = 0;
+	int32 LoopStartFrameIndex = 0;
+	int32 LoopEndFrameIndex = -1;
+	float DurationInSeconds = 0.f;
+	float MaxLoopStartTimeInSeconds = 0.f;
+
+	bool bIsDecoderValid = false;
+	bool bFallbackSeekMethodWarningLogged = false;
+
+private:
+	mutable TUniquePtr<class ICompressedAudioInfo> CompressedAudioInfo;
+	Audio::TCircularAudioBuffer<float> DecoderOutput;
+
+	TArray<int16> ResidualBuffer;
+	Audio::FAlignedFloatBuffer SampleConversionBuffer;
+
+	static constexpr uint32 MinNumFramesPerDecode = 1;
+
+	bool bIsFirstDecode = true;
+	bool bPreviousIsStreaming = true;
+	uint32 NumFramesPerDecode = MinNumFramesPerDecode;
+};
\ No newline at end of file
diff --git a/Engine/Source/Runtime/Engine/Private/SoundWaveProxyReader.cpp b/Engine/Source/Runtime/Engine/Private/SoundWaveProxyReader.cpp
new file mode 100644
--- /dev/null
+++ b/Engine/Source/Runtime/Engine/Private/SoundWaveProxyReader.cpp
@@ -0,0 +1,569 @@
+// Copyright Epic Games, Inc. All Rights Reserved.
+
+#include "Sound/SoundWaveProxyReader.h"
+
+#include "Audio.h"
+#include "Containers/Array.h"
+#include "Containers/ArrayView.h"
+#include "DSP/BufferVectorOperations.h"
+#include "DSP/Dsp.h"
+#include "HAL/IConsoleManager.h"
+#include "HAL/LowLevelMemTracker.h"
+#include "HAL/Platform.h"
+#include "HAL/UnrealMemory.h"
+#include "Interfaces/IAudioFormat.h"
+#include "Math/UnrealMathUtility.h"
+#include "Misc/AssertionMacros.h"
+#include "ProfilingDebugging/CpuProfilerTrace.h"
+#include "Sound/SoundWave.h"
+#include "Templates/SharedPointer.h"
+#include "Templates/UniquePtr.h"
+
+static int32 GSoundWaveProxyReaderSimulateSeekOnNonSeekable = 0;
+FAutoConsoleVariableRef CVarSoundWaveProxyReaderSimulateSeekOnNonSeekable(
+	TEXT("au.SoundWaveProxyReader.SimulateSeek"),
+	GSoundWaveProxyReaderSimulateSeekOnNonSeekable,
+	TEXT("If true, SoundWaves which are not of a seekable format will simulate seek calls by reading and discarding samples.\n")
+	TEXT("0: Do not simulate seek, !0: Simulate seek"),
+	ECVF_Default);
+
+uint32 FSoundWaveProxyReader::ConformDecodeSize(uint32 InMaxDesiredDecodeSizeInFrames)
+{
+	static_assert(DefaultMinDecodeSizeInFrames >= DecodeSizeQuantizationInFrames, "Min decode size less than decode size quantization");
+	static_assert((DefaultMinDecodeSizeInFrames % DecodeSizeQuantizationInFrames) == 0, "Min decode size must be equally divisible by decode size quantization");
+
+	const uint32 QuantizedDecodeSizeInFrames = (InMaxDesiredDecodeSizeInFrames / DecodeSizeQuantizationInFrames) * DecodeSizeQuantizationInFrames;
+	const uint32 ConformedDecodeSizeInFrames = FMath::Max(QuantizedDecodeSizeInFrames, DefaultMinDecodeSizeInFrames);
+
+	check((ConformedDecodeSizeInFrames % QuantizedDecodeSizeInFrames) == 0);
+	check(ConformedDecodeSizeInFrames > 0);
+
+	return ConformedDecodeSizeInFrames;
+}
+
+/** Construct a wave proxy reader.
+ *
+ * @param InWaveProxy - A TSharedRef of a FSoundWaveProxy which is to be played.
+ * @param InSettings - Reader settings.
+ */
+FSoundWaveProxyReader::FSoundWaveProxyReader(FSoundWaveProxyRef InWaveProxy, const FSettings& InSettings)
+	: WaveProxy(InWaveProxy)
+	, Settings(InSettings)
+{
+	// Get local copies of some values from the proxy. 
+	SampleRate = WaveProxy->GetSampleRate();
+	NumChannels = WaveProxy->GetNumChannels();
+	NumFramesInWave = WaveProxy->GetNumFrames();
+	DurationInSeconds = WaveProxy->GetDuration();
+	MaxLoopStartTimeInSeconds = FMath::Max(0, DurationInSeconds - MinLoopDurationInSeconds);
+
+	// Clamp times
+	Settings.StartTimeInSeconds = FMath::Clamp(Settings.StartTimeInSeconds, 0.f, DurationInSeconds);
+	Settings.LoopStartTimeInSeconds = ClampLoopStartTime(Settings.LoopStartTimeInSeconds);
+	Settings.LoopDurationInSeconds = ClampLoopDuration(Settings.LoopDurationInSeconds);
+
+	// Setup frame indices
+	CurrentFrameIndex = 0;
+	UpdateLoopBoundaries();
+
+	// Determine max size of decode buffer
+	Settings.MaxDecodeSizeInFrames = FMath::Max(DefaultMinDecodeSizeInFrames, Settings.MaxDecodeSizeInFrames);
+
+	// Prepare to read audio
+	bIsDecoderValid = InitializeDecoder(Settings.StartTimeInSeconds);
+}
+
+/** Create a wave proxy reader.
+ *
+ * @param InWaveProxy - A TSharedRef of a FSoundWaveProxy which is to be played.
+ * @param InSettings - Reader settings.
+ */
+TUniquePtr<FSoundWaveProxyReader> FSoundWaveProxyReader::Create(FSoundWaveProxyRef InWaveProxy, const FSettings& InSettings)
+{
+	if (InWaveProxy->GetSampleRate() <= 0.f)
+	{
+		UE_LOG(LogAudio, Warning, TEXT("Cannot create FSoundWaveProxyReader due to invalid sample rate (%f). Package: %s"), InWaveProxy->GetSampleRate(), *InWaveProxy->GetPackageName().ToString());
+		return TUniquePtr<FSoundWaveProxyReader>();
+	}
+
+	if (InWaveProxy->GetNumChannels() <= 0)
+	{
+		UE_LOG(LogAudio, Warning, TEXT("Cannot create FSoundWaveProxyReader due to invalid num channels (%d). Package: %s"), InWaveProxy->GetNumChannels(), *InWaveProxy->GetPackageName().ToString());
+		return TUniquePtr<FSoundWaveProxyReader>();
+	}
+
+	if (InWaveProxy->GetNumFrames() <= 0)
+	{
+		UE_LOG(LogAudio, Warning, TEXT("Cannot create FSoundWaveProxyReader due to invalid num frames (%d). Package: %s"), InWaveProxy->GetNumFrames(), *InWaveProxy->GetPackageName().ToString());
+		return TUniquePtr<FSoundWaveProxyReader>();
+	}
+
+	return TUniquePtr<FSoundWaveProxyReader>(new FSoundWaveProxyReader(InWaveProxy, InSettings));
+}
+
+
+/** Set whether the reader should loop the audio or not. */
+void FSoundWaveProxyReader::SetIsLooping(bool bInIsLooping)
+{
+	if (Settings.bIsLooping != bInIsLooping)
+	{
+		Settings.bIsLooping = bInIsLooping;
+		UpdateLoopBoundaries();
+	}
+}
+
+/** Sets the beginning position of the loop. */
+void FSoundWaveProxyReader::SetLoopStartTime(float InLoopStartTimeInSeconds)
+{
+	InLoopStartTimeInSeconds = ClampLoopStartTime(InLoopStartTimeInSeconds);
+	if (!FMath::IsNearlyEqual(Settings.LoopStartTimeInSeconds, InLoopStartTimeInSeconds))
+	{
+		Settings.LoopStartTimeInSeconds = InLoopStartTimeInSeconds;
+		UpdateLoopBoundaries();
+	}
+}
+
+/** Sets the duration of the loop in seconds. If the value is negative, the
+ * loop duration consists of the entire file. */
+void FSoundWaveProxyReader::SetLoopDuration(float InLoopDurationInSeconds)
+{
+	InLoopDurationInSeconds = ClampLoopDuration(InLoopDurationInSeconds);
+	if (!FMath::IsNearlyEqual(Settings.LoopDurationInSeconds, InLoopDurationInSeconds))
+	{
+		Settings.LoopDurationInSeconds = InLoopDurationInSeconds;
+		UpdateLoopBoundaries();
+	}
+}
+
+bool FSoundWaveProxyReader::SeekToTime(float InSeconds)
+{
+	int32 InFrameIndex = FMath::Clamp(static_cast<int32>(InSeconds * GetSampleRate()), 0, GetNumFramesInWave());
+	// ignore seek request if we're already at the specified time
+	if (InFrameIndex == CurrentFrameIndex)
+	{
+		return bIsDecoderValid;
+	}
+
+	if (WaveProxy->IsSeekable() && CompressedAudioInfo)
+	{
+		CompressedAudioInfo->SeekToFrame(InFrameIndex);
+		CurrentFrameIndex = InFrameIndex;
+		DecoderOutput.SetNum(0);
+		NumDecodeSamplesToDiscard = 0;
+		DecodeResult = EDecodeResult::MoreDataRemaining;
+		return bIsDecoderValid;
+	}
+	// Direct seeking is not supported. A new decoder must be created. 
+	bIsDecoderValid = InitializeDecoder(InSeconds);
+	return bIsDecoderValid;
+}
+
+bool FSoundWaveProxyReader::CanProduceMoreAudio() const
+{
+	const bool bDecoderOutputHasMoreData = DecoderOutput.Num() > 0;
+	const bool bDecoderCanDecodeMoreData = bIsDecoderValid && (EDecodeResult::MoreDataRemaining == DecodeResult);
+	return bDecoderOutputHasMoreData || bDecoderCanDecodeMoreData;
+}
+
+bool FSoundWaveProxyReader::SeekToFrame(uint32 InFrameNum)
+{
+	// ignore seek request if we're already at the specified time
+	if (InFrameNum == CurrentFrameIndex)
+	{
+		return bIsDecoderValid;
+	}
+
+	if (WaveProxy->IsSeekable() && CompressedAudioInfo)
+	{
+		CompressedAudioInfo->SeekToFrame(InFrameNum);
+		CurrentFrameIndex = InFrameNum;
+		DecoderOutput.SetNum(0);
+		NumDecodeSamplesToDiscard = 0;
+		DecodeResult = EDecodeResult::MoreDataRemaining;
+		return bIsDecoderValid;
+	}
+
+	// Direct seeking is not supported. A new decoder must be created.
+	float Seconds = GetSampleRate() * InFrameNum;
+	bIsDecoderValid = InitializeDecoder(Seconds);
+	return bIsDecoderValid;
+}
+
+/** Copies audio into OutBuffer. It returns the number of samples copied.
+ * Samples not written to will be set to zero.
+ */
+int32 FSoundWaveProxyReader::PopAudio(Audio::FAlignedFloatBuffer& OutBuffer)
+{
+	using namespace Audio;
+
+	TRACE_CPUPROFILER_EVENT_SCOPE(FSoundWaveProxyReader::PopAudio);
+
+	if ((0 == NumChannels) || !bIsDecoderValid)
+	{
+		if (OutBuffer.Num() > 0)
+		{
+			FMemory::Memset(OutBuffer.GetData(), 0, sizeof(float) * OutBuffer.Num());
+		}
+		return 0;
+	}
+
+	checkf(0 == (OutBuffer.Num() % NumChannels), TEXT("Output buffer size must be evenly divisible by the number of channels"));
+
+	TArrayView<float> OutBufferView{OutBuffer};
+	int32 NumSamplesUnset = OutBuffer.Num();
+	int32 NumSamplesCopied = 0;
+
+	bool bDecoderOutputHasMoreData = DecoderOutput.Num() > 0;
+	bool bDecoderCanDecodeMoreData = bIsDecoderValid && (EDecodeResult::MoreDataRemaining == DecodeResult);
+	bool bCanProduceMoreAudio = bDecoderOutputHasMoreData || bDecoderCanDecodeMoreData;
+
+	while ((NumSamplesUnset > 0) && bCanProduceMoreAudio)
+	{
+		int32 NumSamplesCopiedThisLoop = PopAudioFromDecoderOutput(OutBufferView);
+
+		// Update sample counters and output buffer view
+		NumSamplesCopied += NumSamplesCopiedThisLoop;
+		NumSamplesUnset -= NumSamplesCopiedThisLoop;
+		OutBufferView = OutBufferView.Slice(NumSamplesCopiedThisLoop, NumSamplesUnset);
+
+		if (Settings.bIsLooping)
+		{
+			// Seek to loop start if we have used up all our decodable samples or 
+			// are at the loop boundary.
+			if (0 == DecoderOutput.Num())
+			{
+				if ((!bDecoderCanDecodeMoreData) || (CurrentFrameIndex >= LoopEndFrameIndex))
+				{
+					SeekToTime(Settings.LoopStartTimeInSeconds);
+					bDecoderCanDecodeMoreData = bIsDecoderValid && (EDecodeResult::Fail != DecodeResult);
+				}
+			}
+		}
+
+		// Determine if we can / should decode more data. 
+		if ((NumSamplesUnset > 0) && (DecoderOutput.Num() == 0) && bDecoderCanDecodeMoreData)
+		{
+			DecodeResult = Decode();
+			bDecoderCanDecodeMoreData = EDecodeResult::MoreDataRemaining == DecodeResult;
+		}
+
+		bDecoderOutputHasMoreData = DecoderOutput.Num() > 0;
+		bCanProduceMoreAudio = bDecoderOutputHasMoreData || bDecoderCanDecodeMoreData;
+
+		if (bCanProduceMoreAudio && DecoderOutput.Num() == 0)
+		{
+			// we can produce more audio, but we were unable to
+			// this is likely due to the streaming data not being available yet
+			// let's early out to avoid a hitch and hope that it's ready on the next read
+			break;
+		}
+	}
+
+	// Zero pad any unset samples. 
+	if (OutBufferView.Num() > 0)
+	{
+		// Zero out audio that was not set.
+		FMemory::Memset(OutBufferView.GetData(), 0, sizeof(float) * OutBufferView.Num());
+
+		if (Settings.bMaintainAudioSync)
+		{
+			ensureMsgf(!Settings.bIsLooping, TEXT("Currently can't BOTH loop a wave and have it maintain sync. The code that does the bookkeeping when the decoder underruns is not robust enough to handle that situation."));
+
+			// if we were asked to maintain audio sync, then do some extra book keeping
+			// keep track of the number of samples we'll need to discard the next time we try to read from the decoder
+			NumDecodeSamplesToDiscard += OutBufferView.Num();
+			// and pretend we are advancing through the data even though we aren't...
+			CurrentFrameIndex += OutBufferView.Num() / NumChannels;
+			if (CurrentFrameIndex > GetNumFramesInWave())
+			{
+				// but don't pretend to go passed the end!...
+				CurrentFrameIndex = GetNumFramesInWave();
+			}
+			// Note: We can do the above because later we will NOT advance CurrentFrameIndex for the decoded samples that we discard!)
+		}
+	}
+
+	return NumSamplesCopied;
+}
+
+int32 FSoundWaveProxyReader::PopAudioFromDecoderOutput(TArrayView<float> OutBufferView)
+{
+	check(NumChannels > 0);
+
+	int32 NumDiscarded = DecoderOutput.Pop(NumDecodeSamplesToDiscard);
+	NumDecodeSamplesToDiscard = NumDecodeSamplesToDiscard - NumDiscarded;
+	check(NumDecodeSamplesToDiscard >= 0);
+	
+	int32 NumSamplesCopied = 0;
+	if (DecoderOutput.Num() > 0)
+	{
+		// Get samples from the decoder buffer.
+		NumSamplesCopied = DecoderOutput.Pop(OutBufferView.GetData(), OutBufferView.Num());
+		
+		int32 NumFramesCopied = NumSamplesCopied / NumChannels;
+		CurrentFrameIndex += NumFramesCopied;
+
+		// Check whether the samples copied from the decoder extend 
+		// past the end of the loop.
+		if (Settings.bIsLooping)
+		{
+			const bool bDidOvershoot = CurrentFrameIndex >= LoopEndFrameIndex;
+			if (bDidOvershoot)
+			{
+				// Rewind sample counters if the loop boundary was overshot.
+				NumFramesCopied -= (CurrentFrameIndex - LoopEndFrameIndex);
+				CurrentFrameIndex = LoopEndFrameIndex;
+				// If Settings.bIsLooping info was altered, NumFramesCopied can end up 
+				// negative if the current frame index is past the loop end frame. 
+				NumFramesCopied = FMath::Max(0, NumFramesCopied);
+				NumSamplesCopied = NumFramesCopied * NumChannels;
+
+				// Remove any remaining samples in the decoder because they 
+				// are past the end of the loop.
+				DecoderOutput.SetNum(0);
+			}
+
+		}
+	}
+
+	return NumSamplesCopied;
+}
+
+bool FSoundWaveProxyReader::InitializeDecoder(float InStartTimeInSeconds)
+{
+	using namespace Audio;
+
+	TRACE_CPUPROFILER_EVENT_SCOPE(FSoundWaveProxyReader::InitializeDecoder);
+
+	FName Format = WaveProxy->GetRuntimeFormat();
+	IAudioInfoFactory* Factory = IAudioInfoFactoryRegistry::Get().Find(Format);
+	if (!ensure(Factory))
+	{
+		UE_LOG(LogAudio, Error, TEXT("Failed to CompressedAudioInfo for wave (package: %s). Unable to find AudioInfoFactory for format: %s"), *WaveProxy->GetPackageName().ToString(), *Format.ToString());
+		return false;
+	}
+
+	TUniquePtr<ICompressedAudioInfo> InfoInstance;
+	InfoInstance.Reset(Factory->Create());
+
+	if (!ensure(InfoInstance.IsValid()))
+	{
+		UE_LOG(LogAudio, Error, TEXT("Failed to created CompressedAudioInfo for wave (package: %s). Unable to create info from factory for for format: %s"), *WaveProxy->GetPackageName().ToString(), *Format.ToString());
+		return false;
+	}
+
+	FSoundQualityInfo Info;
+	if (WaveProxy->IsStreaming())
+	{
+		if (!InfoInstance->StreamCompressedInfo(WaveProxy, &Info))
+		{
+			UE_LOG(LogAudio, Error, TEXT("Failed to created CompressedAudioInfo for wave (package: %s). Unable to stream compressed info for streaming wave"), *WaveProxy->GetPackageName().ToString());
+			return false;
+		}
+	}
+	else
+	{
+		if (!InfoInstance->ReadCompressedInfo(WaveProxy->GetResourceData(), WaveProxy->GetResourceSize(), &Info))
+		{
+			UE_LOG(LogAudio, Error, TEXT("Failed to created decoder input for wave (package: %s). Unable to read compressed info for non-streaming wave"), *WaveProxy->GetPackageName().ToString());
+			return false;
+		}
+	}
+
+	CompressedAudioInfo.Reset(InfoInstance.Release());
+
+	// Read the sample rate and number of frames from the header 
+	// Similar to refreshing the wave data in FMixerBuffer::CreateStreamingBuffer
+	// This is a runtime hack to address incorrect sample rate on soundwaves 
+	// on platforms with Resample for Device enabled (UE-183237)
+	SampleRate = Info.SampleRate;
+	uint32 NumFrames = (uint32)((float)Info.Duration * Info.SampleRate);
+	if (NumFrames > 0)
+	{
+		NumFramesInWave = NumFrames;
+	}
+	// end hack
+
+	// Seek input to start time.
+	if (!FMath::IsNearlyEqual(0.0f, InStartTimeInSeconds))
+	{
+		if (WaveProxy->IsSeekable())
+		{
+			CompressedAudioInfo->SeekToTime(InStartTimeInSeconds);
+		}
+		else
+		{
+			UE_LOG(LogAudio, Warning, TEXT("Attempt to seek on non-seekable wave: (format:%s) for wave (package:%s) to time '%.6f'"),
+				*Format.ToString(),
+				*WaveProxy->GetPackageName().ToString(),
+				InStartTimeInSeconds);
+		}
+	}
+
+	CurrentFrameIndex = FMath::Clamp(static_cast<int32>(InStartTimeInSeconds * GetSampleRate()), 0, GetNumFramesInWave());
+
+	// initialized decode buffers
+	const uint32 DecodeSize = ConformDecodeSize(Settings.MaxDecodeSizeInFrames);
+	NumFramesPerDecode = FMath::Max(DecodeSize, NumFramesPerDecode);
+	ResidualBuffer.SetNum(NumFramesPerDecode * NumChannels);
+	SampleConversionBuffer.SetNumUninitialized(NumFramesPerDecode * NumChannels);
+	DecoderOutput.Reserve(/* MinCapacity = */ NumFramesPerDecode * NumChannels * 2, /* bRetainExistingSamples = */ false);
+	NumDecodeSamplesToDiscard = 0;
+
+	// For non-seekable streaming waves, use a fallback method to seek
+	// to the start time. 
+	const bool bUseFallbackSeekMethod = (GSoundWaveProxyReaderSimulateSeekOnNonSeekable != 0) && (InStartTimeInSeconds != 0.f) && (!WaveProxy->IsSeekable());
+	if (bUseFallbackSeekMethod)
+	{
+		if (!bFallbackSeekMethodWarningLogged)
+		{
+			UE_LOG(LogAudio, Warning, TEXT("Simulating seeking in wave which is not seekable (package:%s). For better performance, set wave to a seekable format"), *WaveProxy->GetPackageName().ToString());
+			bFallbackSeekMethodWarningLogged = true;
+		}
+
+		int32 NumSamplesToDiscard = CurrentFrameIndex * NumChannels;
+		DiscardSamples(NumSamplesToDiscard);
+	}
+
+	if (!CompressedAudioInfo.IsValid())
+	{
+		UE_LOG(LogAudio, Error, TEXT("Failed to create decoder (format:%s) for wave (package:%s)"), *Format.ToString(), *WaveProxy->GetPackageName().ToString());
+		DecodeResult = EDecodeResult::Fail;
+		return false;
+	}
+	else
+	{
+		// The DecodeResult needs to be set to a valid state in case the prior
+		// decoder finished or failed. 
+		DecodeResult = EDecodeResult::MoreDataRemaining;
+	}
+
+	return true;
+}
+
+int32 FSoundWaveProxyReader::DiscardSamples(int32 InNumSamplesToDiscard)
+{
+	int32 NumSamplesDiscarded = 0;
+	while (InNumSamplesToDiscard > 0)
+	{
+		DecodeResult = Decode();
+
+		int32 NumSamplesToDiscardThisLoop = FMath::Min((int32)DecoderOutput.Num(), InNumSamplesToDiscard);
+
+		if (NumSamplesToDiscardThisLoop < 1)
+		{
+			UE_LOG(LogAudio, Error, TEXT("Failed to decode samples (package: %s)"), *WaveProxy->GetPackageName().ToString());
+			break;
+		}
+
+		int32 ActualNumSamplesDiscarded = DecoderOutput.Pop(NumSamplesToDiscardThisLoop);
+		InNumSamplesToDiscard -= ActualNumSamplesDiscarded;
+		NumSamplesDiscarded = ActualNumSamplesDiscarded;
+
+		if ((InNumSamplesToDiscard > 0) && (DecodeResult != EDecodeResult::MoreDataRemaining))
+		{
+			UE_LOG(LogAudio, Error, TEXT("Failed to simulate seek (seek to frame: %d, package: %s)"), CurrentFrameIndex, *WaveProxy->GetPackageName().ToString());
+			break;
+		}
+	}
+
+	return NumSamplesDiscarded;
+}
+
+float FSoundWaveProxyReader::ClampLoopStartTime(float InStartTimeInSeconds)
+{
+	return FMath::Clamp(InStartTimeInSeconds, 0.f, MaxLoopStartTimeInSeconds);
+}
+
+float FSoundWaveProxyReader::ClampLoopDuration(float InDurationInSeconds)
+{
+	if (InDurationInSeconds <= 0.f)
+	{
+		return MaxLoopDurationInSeconds;
+	}
+	return FMath::Max(InDurationInSeconds, MinLoopDurationInSeconds);
+}
+
+void FSoundWaveProxyReader::UpdateLoopBoundaries()
+{
+	if (Settings.bIsLooping)
+	{
+		const int32 MinLoopDurationInFrames = FMath::CeilToInt(MinLoopDurationInSeconds * SampleRate);
+		const float LoopEndTime = Settings.LoopStartTimeInSeconds + Settings.LoopDurationInSeconds;
+		const int32 MinLoopEndFrameIndex = LoopStartFrameIndex + MinLoopDurationInFrames;
+		const int32 MaxLoopEndFrameIndex = NumFramesInWave;
+
+		LoopStartFrameIndex = FMath::FloorToInt(Settings.LoopStartTimeInSeconds * SampleRate);
+		LoopEndFrameIndex = FMath::Clamp(LoopEndTime * SampleRate, MinLoopEndFrameIndex, MaxLoopEndFrameIndex);
+	}
+	else
+	{
+		LoopStartFrameIndex = 0;
+		LoopEndFrameIndex = NumFramesInWave;
+	}
+}
+
+FSoundWaveProxyReader::EDecodeResult FSoundWaveProxyReader::Decode()
+{
+	check(CompressedAudioInfo.IsValid());
+
+	bool bFinished = false;
+
+	int32 NumFramesRemaining = NumFramesPerDecode;
+	uint32 BuffSizeInBytes = ResidualBuffer.Num() * sizeof(int16);
+	uint32 BuffSizeInFrames = NumFramesPerDecode;
+	uint8* Buff = (uint8*)ResidualBuffer.GetData();
+
+	// cache the streaming flag off the wave
+	// if it has changed since the last Decode() call, bail
+	// something has probably changed in editor
+	if (bIsFirstDecode)
+	{
+		bIsFirstDecode = false;
+	}
+	else
+	{
+		if (bPreviousIsStreaming != WaveProxy->IsStreaming())
+		{
+			return EDecodeResult::Finished;
+		}
+	}
+	bPreviousIsStreaming = WaveProxy->IsStreaming();
+
+	int32 NumBytesStreamed = 0;
+	while (!bFinished && NumFramesRemaining > 0)
+	{
+		if (WaveProxy->IsStreaming())
+		{
+			NumBytesStreamed = 0;
+			bFinished = CompressedAudioInfo->StreamCompressedData(Buff, false, BuffSizeInBytes, NumBytesStreamed);
+		}
+		else
+		{
+			NumBytesStreamed = BuffSizeInBytes;
+			bFinished = CompressedAudioInfo->ReadCompressedData(Buff, false, BuffSizeInBytes);
+		}
+
+		if (NumBytesStreamed == 0)
+		{
+			break;
+		}
+
+		const int32 NumSamplesStreamed = NumBytesStreamed / sizeof(int16);
+		const int32 NumFramesStreamed = NumSamplesStreamed / NumChannels;
+
+		Audio::ArrayPcm16ToFloat(MakeArrayView(ResidualBuffer.GetData(), NumSamplesStreamed), MakeArrayView(SampleConversionBuffer.GetData(), NumSamplesStreamed));
+
+		const float* SampleData = SampleConversionBuffer.GetData();
+		DecoderOutput.Push(SampleData, NumSamplesStreamed);
+
+		NumFramesRemaining -= FMath::Min(NumFramesStreamed, NumFramesRemaining);
+	}
+
+	if (!bFinished)
+	{
+		return EDecodeResult::MoreDataRemaining;
+	}
+
+	return EDecodeResult::Finished;
+}
\ No newline at end of file
