Asset Export to DDS for RenderTarget supports cubes,volumes,etc
GetRenderTargetImage support cubes,volumes,etc
clean up virtual functions in TextureRenderTarget class hierarchy
push shared functions up to root so they can be accessed without sub-classing

#rnx

diff --git a/Engine/Source/Editor/UnrealEd/Private/Factories/EditorFactories.cpp b/Engine/Source/Editor/UnrealEd/Private/Factories/EditorFactories.cpp
--- a/Engine/Source/Editor/UnrealEd/Private/Factories/EditorFactories.cpp
+++ b/Engine/Source/Editor/UnrealEd/Private/Factories/EditorFactories.cpp
@@ -4699,7 +4699,14 @@ bool UTextureExporterDDS::SupportsObject(UObject* Object) const
 	
 	if ( !Texture->Source.IsValid() )
 	{
-		return false;
+		if ( Texture->GetTextureClass() == ETextureClass::RenderTarget )
+		{
+			return true;
+		}
+		else
+		{
+			return false;
+		}
 	}
 
 	if ( Texture->Source.GetFormat() == TSF_BGRE8 )
@@ -4720,31 +4727,50 @@ bool UTextureExporterDDS::ExportBinary( UObject* Object, const TCHAR* Type, FArc
 {
 	UTexture * Texture = GetExportTexture( Object );
 	check( Texture != nullptr );
+
+	UTextureRenderTarget * RT = Cast<UTextureRenderTarget>(Texture);
 	
-	// FileIndex for layers for VT :
+	if ( RT != nullptr )
+	{
+		TArray64<uint8> OutArray;
+		if ( ! FImageUtils::ExportRenderTargetToDDS(OutArray,RT) )
+		{
+			Warn->Logf(ELogVerbosity::Error, TEXT("ExportRenderTargetToDDS failed"));
+			return false;
+		}
+
+		Ar.Serialize(OutArray.GetData(),OutArray.Num());
+
+		return true;		
+	}
+	else
+	{
+
+		// FileIndex for layers for VT :
 	
-	int NumBlocks = Texture->Source.GetNumBlocks();
-	int NumLayers = Texture->Source.GetNumLayers();
+		int NumBlocks = Texture->Source.GetNumBlocks();
+		int NumLayers = Texture->Source.GetNumLayers();
 
-	int NumFiles = NumBlocks * NumLayers;
-	check( FileIndex < NumFiles );
-	int BlockIndex = FileIndex / NumLayers;
-	int LayerIndex = FileIndex % NumLayers; 
+		int NumFiles = NumBlocks * NumLayers;
+		check( FileIndex < NumFiles );
+		int BlockIndex = FileIndex / NumLayers;
+		int LayerIndex = FileIndex % NumLayers; 
 
-	UE_LOG(LogEditorFactories, Display, TEXT("Exporting DDS from Layer %d/%d Block %d/%d"), LayerIndex,NumLayers, BlockIndex,NumBlocks);
+		UE_LOG(LogEditorFactories, Display, TEXT("Exporting DDS from Layer %d/%d Block %d/%d"), LayerIndex,NumLayers, BlockIndex,NumBlocks);
 
-	// Type == Image extension
+		// Type == Image extension
 	
-	TArray64<uint8> OutArray;
-	if ( ! FImageUtils::ExportTextureSourceToDDS(OutArray,Texture,BlockIndex,LayerIndex) )
-	{
-		Warn->Logf(ELogVerbosity::Error, TEXT("ExportTextureSourceToDDS failed"));
-		return false;
-	}
+		TArray64<uint8> OutArray;
+		if ( ! FImageUtils::ExportTextureSourceToDDS(OutArray,Texture,BlockIndex,LayerIndex) )
+		{
+			Warn->Logf(ELogVerbosity::Error, TEXT("ExportTextureSourceToDDS failed"));
+			return false;
+		}
 
-	Ar.Serialize(OutArray.GetData(),OutArray.Num());
+		Ar.Serialize(OutArray.GetData(),OutArray.Num());
 
-	return true;
+		return true;
+	}
 }
 
 UVirtualTextureBuilderExporterDDS::UVirtualTextureBuilderExporterDDS(const FObjectInitializer& ObjectInitializer)
@@ -4856,14 +4882,22 @@ int32 UTextureExporterGeneric::GetFileCount(UObject* Object) const
 	UTexture* Texture = GetExportTexture(Object);
 	check(Texture != nullptr);
 
-	// standard textures will have NumBlocks == NumLayers == 1
-	// VT can have them > 1
-	// UDIM gives you NumBlocks
-	int NumBlocks = Texture->Source.GetNumBlocks();
-	int NumLayers = Texture->Source.GetNumLayers();
-	int FileCount = NumBlocks * NumLayers;
-	check( FileCount > 0 );
-	return FileCount;
+	if ( Texture->Source.IsValid() )
+	{
+		// standard textures will have NumBlocks == NumLayers == 1
+		// VT can have them > 1
+		// UDIM gives you NumBlocks
+		int NumBlocks = Texture->Source.GetNumBlocks();
+		int NumLayers = Texture->Source.GetNumLayers();
+		int FileCount = NumBlocks * NumLayers;
+		check( FileCount > 0 );
+		return FileCount;
+	}
+	else
+	{
+		// eg. render targets
+		return 1;
+	}
 }
 
 FString UTextureExporterGeneric::GetUniqueFilename( UObject* Object, const TCHAR* Filename, int32 FileIndex, int32 FileCount ) const
diff --git a/Engine/Source/Runtime/Engine/Classes/Engine/TextureRenderTarget.h b/Engine/Source/Runtime/Engine/Classes/Engine/TextureRenderTarget.h
--- a/Engine/Source/Runtime/Engine/Classes/Engine/TextureRenderTarget.h
+++ b/Engine/Source/Runtime/Engine/Classes/Engine/TextureRenderTarget.h
@@ -71,7 +71,17 @@ class UTextureRenderTarget : public UTexture
 	/**
 	 * Returns the UTexture class that corresponds to this render target (e.g. UTexture2D for UTextureRenderTarget2D)
 	 */
-	ENGINE_API virtual TSubclassOf<UTexture> GetTextureUClass() const PURE_VIRTUAL(UTextureRenderTarget, return nullptr;);
+	ENGINE_API virtual TSubclassOf<UTexture> GetTextureUClass() const PURE_VIRTUAL(UTextureRenderTarget, return nullptr;)
+
+	ENGINE_API virtual EPixelFormat GetFormat() const PURE_VIRTUAL(GetFormat,return PF_Unknown;)
+
+	ENGINE_API virtual bool IsSRGB() const PURE_VIRTUAL(IsSRGB,return false;)
+
+	ENGINE_API virtual float GetDisplayGamma() const PURE_VIRTUAL(GetDisplayGamma,return 0.f;)
+
+	// UTextureRenderTarget default display gamma if none is set
+	//	returns hard-coded 2.2
+	static float GetDefaultDisplayGamma();
 
 	// get the variant of ReadPixels call that should be used
 	//	either BGRA8,RGBA32F,or RGBA16F
diff --git a/Engine/Source/Runtime/Engine/Classes/Engine/TextureRenderTarget2D.h b/Engine/Source/Runtime/Engine/Classes/Engine/TextureRenderTarget2D.h
--- a/Engine/Source/Runtime/Engine/Classes/Engine/TextureRenderTarget2D.h
+++ b/Engine/Source/Runtime/Engine/Classes/Engine/TextureRenderTarget2D.h
@@ -172,6 +172,9 @@ class UTextureRenderTarget2D : public UTextureRenderTarget
 	//~ Begin UTextureRenderTarget Interface
 	virtual bool CanConvertToTexture(ETextureSourceFormat& OutTextureSourceFormat, EPixelFormat& OutPixelFormat, FText* OutErrorMessage) const override;
 	virtual TSubclassOf<UTexture> GetTextureUClass() const override;
+	virtual EPixelFormat GetFormat() const override;
+	virtual bool IsSRGB() const override;
+	virtual float GetDisplayGamma() const override;
 	//~ End UTextureRenderTarget Interface
 
 	/**
@@ -231,35 +234,6 @@ PRAGMA_ENABLE_DEPRECATION_WARNINGS
 		return NumMips;
 	}
 
-
-	FORCEINLINE EPixelFormat GetFormat() const
-	{
-		if (OverrideFormat == PF_Unknown)
-		{
-			return GetPixelFormatFromRenderTargetFormat(RenderTargetFormat);
-		}
-		else
-		{
-			return OverrideFormat;
-		}
-	}
-
-	bool IsSRGB() const
-	{
-		// ?? note: UTextureRenderTarget::TargetGamma is ignored here
-
-		// ?? note: GetDisplayGamma forces linear for some float formats, but this doesn't
-
-		if (OverrideFormat == PF_Unknown)
-		{
-			return RenderTargetFormat == RTF_RGBA8_SRGB;
-		}
-		else
-		{
-			return !bForceLinearGamma;
-		}
-	}
-
 private:
 	int32	NumMips;
 };
diff --git a/Engine/Source/Runtime/Engine/Classes/Engine/TextureRenderTarget2DArray.h b/Engine/Source/Runtime/Engine/Classes/Engine/TextureRenderTarget2DArray.h
--- a/Engine/Source/Runtime/Engine/Classes/Engine/TextureRenderTarget2DArray.h
+++ b/Engine/Source/Runtime/Engine/Classes/Engine/TextureRenderTarget2DArray.h
@@ -87,18 +87,6 @@ class UTextureRenderTarget2DArray : public UTextureRenderTarget
 	virtual EMaterialValueType GetMaterialType() const override;
 	//~ End UTexture Interface.
 
-	EPixelFormat GetFormat() const
-	{
-		if(OverrideFormat == PF_Unknown)
-		{
-			return bHDR ? PF_FloatRGBA : PF_B8G8R8A8;
-		}
-		else
-		{
-			return OverrideFormat;
-		}
-	}
-
 	FORCEINLINE int32 GetNumMips() const
 	{
 		return 1;
@@ -115,6 +103,9 @@ class UTextureRenderTarget2DArray : public UTextureRenderTarget
 	//~ Begin UTextureRenderTarget Interface
 	virtual bool CanConvertToTexture(ETextureSourceFormat& OutTextureSourceFormat, EPixelFormat& OutPixelFormat, FText* OutErrorMessage) const override;
 	virtual TSubclassOf<UTexture> GetTextureUClass() const override;
+	virtual EPixelFormat GetFormat() const override;
+	virtual bool IsSRGB() const override;
+	virtual float GetDisplayGamma() const override;
 	//~ End UTextureRenderTarget Interface
 };
 
diff --git a/Engine/Source/Runtime/Engine/Classes/Engine/TextureRenderTargetCube.h b/Engine/Source/Runtime/Engine/Classes/Engine/TextureRenderTargetCube.h
--- a/Engine/Source/Runtime/Engine/Classes/Engine/TextureRenderTargetCube.h
+++ b/Engine/Source/Runtime/Engine/Classes/Engine/TextureRenderTargetCube.h
@@ -78,19 +78,7 @@ class UTextureRenderTargetCube : public UTextureRenderTarget
 	virtual FTextureResource* CreateResource() override;
 	virtual EMaterialValueType GetMaterialType() const override;
 	//~ End UTexture Interface.
-
-	EPixelFormat GetFormat() const
-	{
-		if(OverrideFormat == PF_Unknown)
-		{
-			return bHDR ? PF_FloatRGBA : PF_B8G8R8A8;
-		}
-		else
-		{
-			return OverrideFormat;
-		}
-	}
-
+	
 	FORCEINLINE int32 GetNumMips() const
 	{
 		return 1;
@@ -107,6 +95,9 @@ class UTextureRenderTargetCube : public UTextureRenderTarget
 	//~ Begin UTextureRenderTarget Interface
 	virtual bool CanConvertToTexture(ETextureSourceFormat& OutTextureSourceFormat, EPixelFormat& OutPixelFormat, FText* OutErrorMessage) const override;
 	virtual TSubclassOf<UTexture> GetTextureUClass() const override;
+	virtual EPixelFormat GetFormat() const override;
+	virtual bool IsSRGB() const override;
+	virtual float GetDisplayGamma() const override;
 	//~ End UTextureRenderTarget Interface
 };
 
diff --git a/Engine/Source/Runtime/Engine/Classes/Engine/TextureRenderTargetVolume.h b/Engine/Source/Runtime/Engine/Classes/Engine/TextureRenderTargetVolume.h
--- a/Engine/Source/Runtime/Engine/Classes/Engine/TextureRenderTargetVolume.h
+++ b/Engine/Source/Runtime/Engine/Classes/Engine/TextureRenderTargetVolume.h
@@ -85,18 +85,6 @@ class UTextureRenderTargetVolume : public UTextureRenderTarget
 	virtual EMaterialValueType GetMaterialType() const override;
 	//~ End UTexture Interface.
 
-	EPixelFormat GetFormat() const
-	{
-		if(OverrideFormat == PF_Unknown)
-		{
-			return bHDR ? PF_FloatRGBA : PF_B8G8R8A8;
-		}
-		else
-		{
-			return OverrideFormat;
-		}
-	}
-
 	FORCEINLINE int32 GetNumMips() const
 	{
 		return 1;
@@ -113,6 +101,9 @@ class UTextureRenderTargetVolume : public UTextureRenderTarget
 	//~ Begin UTextureRenderTarget Interface
 	virtual bool CanConvertToTexture(ETextureSourceFormat& OutTextureSourceFormat, EPixelFormat& OutPixelFormat, FText* OutErrorMessage) const override;
 	virtual TSubclassOf<UTexture> GetTextureUClass() const override;
+	virtual EPixelFormat GetFormat() const override;
+	virtual bool IsSRGB() const override;
+	virtual float GetDisplayGamma() const override;
 	//~ End UTextureRenderTarget Interface
 };
 
diff --git a/Engine/Source/Runtime/Engine/Private/ImageUtils.cpp b/Engine/Source/Runtime/Engine/Private/ImageUtils.cpp
--- a/Engine/Source/Runtime/Engine/Private/ImageUtils.cpp
+++ b/Engine/Source/Runtime/Engine/Private/ImageUtils.cpp
@@ -13,6 +13,7 @@ ImageUtils.cpp: Image utility functions.
 #include "Engine/TextureCubeArray.h"
 #include "Engine/VolumeTexture.h"
 #include "Engine/TextureRenderTarget2D.h"
+#include "Engine/TextureRenderTargetCube.h"
 #include "IImageWrapper.h"
 #include "IImageWrapperModule.h"
 #include "Logging/MessageLog.h"
@@ -314,6 +315,113 @@ bool FImageUtils::ExportTextureSourceToDDS(TArray64<uint8> & OutData, UTexture *
 #endif
 }
 
+bool FImageUtils::ExportRenderTargetToDDS(TArray64<uint8> & OutData, UTextureRenderTarget * TexRT)
+{
+	FImage Image;
+	if ( ! GetRenderTargetImage(TexRT,Image) )
+	{
+		UE_LOG(LogImageUtils,Warning,TEXT("ExportRenderTargetToDDS : GetRenderTargetImage failed"));
+		return false;
+	}
+
+	// some code dupe with GetRenderTargetImage to identify 2d/cube/vol :
+	UTextureRenderTargetCube * TexRTCube = Cast<UTextureRenderTargetCube>(TexRT);
+	bool bIsCube = ( TexRTCube != nullptr );
+	// UTextureRenderTargetCubeArray does not exist at the moment
+	
+	int64 SizeX = FMath::RoundToInt64( TexRT->GetSurfaceWidth() ); // GetSurfaceWidth returns SizeX but as float
+	int64 SizeY = FMath::RoundToInt64( TexRT->GetSurfaceHeight() );
+	int64 TexRT_SizeZ = FMath::RoundToInt64( TexRT->GetSurfaceDepth() );
+	int64 TexRT_ArraySize = TexRT->GetSurfaceArraySize();
+	// TexRT_SizeZ , ArraySize are 0 if not used, not 1
+
+	int32 NumMips = 1;
+
+	int32 Dimension;
+	int32 SizeZ;
+	int32 ArraySize;
+	uint32 CreateFlags = 0;
+	if ( TexRT_SizeZ <= 1 && !bIsCube ) // 2D
+	{
+		Dimension = 2;
+		SizeZ = 1;
+		ArraySize = TexRT_ArraySize ? TexRT_ArraySize : 1;
+	}
+	else if ( bIsCube )
+	{
+		Dimension = 2;
+		SizeZ = 1;
+
+		check( TexRT_ArraySize == 6 );
+		ArraySize = 6;
+		CreateFlags = UE::DDS::FDDSFile::CREATE_FLAG_CUBEMAP;
+	}
+	else if ( TexRT_SizeZ > 1 )
+	{
+		Dimension = 3;
+		SizeZ = TexRT_SizeZ;
+		ArraySize = 1;
+	}
+	else
+	{
+		checkf(false, TEXT("unexpected TextureClass"));
+		return false;
+	}
+	
+	UE::DDS::EDXGIFormat DXGIFormat = UE::DDS::DXGIFormatFromRawFormat(Image.Format,Image.GammaSpace);
+	
+	UE_LOG(LogImageUtils,Display,TEXT("Exporting DDS Dimension=%d SizeX=%d SizeY=%d SizeZ=%d NumMips=%d ArraySize=%d"),Dimension, SizeX,SizeY,SizeZ,NumMips,ArraySize);	
+	
+	UE::DDS::EDDSError Error;
+	UE::DDS::FDDSFile * DDS = UE::DDS::FDDSFile::CreateEmpty(Dimension, SizeX,SizeY,SizeZ,NumMips,ArraySize, DXGIFormat,CreateFlags, &Error);
+	if ( DDS == nullptr || Error != UE::DDS::EDDSError::OK )
+	{
+		UE_LOG(LogImageUtils,Warning,TEXT("FDDSFile::CreateEmpty (Error=%d)"),(int)Error);			
+		return false;
+	}
+	
+	// delete DDS at scope exit :
+	TUniquePtr<UE::DDS::FDDSFile> DDSPtr(DDS);
+
+	check( DDS->Validate() == UE::DDS::EDDSError::OK );
+	
+	// blit into the mips:
+	const int32 MipIndex = 0;
+
+	if ( DDS->Dimension == 3 )
+	{
+		check( DDS->Mips.Num() == DDS->MipCount );
+		check( DDS->Mips[MipIndex].Depth == Image.NumSlices );
+
+		DDS->FillMip( Image, MipIndex );
+	}
+	else
+	{
+		// DDS->Mips[] contains both mips and arrays
+		check( DDS->Mips.Num() == DDS->MipCount * Image.NumSlices );
+
+		for(int SliceIndex=0;SliceIndex<Image.NumSlices;SliceIndex++)
+		{
+			FImageView MipSlice = Image.GetSlice(SliceIndex);
+				
+			// DDS Mips[] array has whole mip chain of each slice, then next slice
+			// we have the opposite (all slices of top mip first, then next mip)
+			int DDSMipIndex = SliceIndex * DDS->MipCount + MipIndex;
+				
+			DDS->FillMip( MipSlice, DDSMipIndex );
+		}
+	}
+
+	Error = DDS->WriteDDS(OutData);
+	if ( Error != UE::DDS::EDDSError::OK )
+	{
+		UE_LOG(LogImageUtils,Warning,TEXT("FDDSFile::WriteDDS failed (Error=%d)"),(int)Error);			
+		return false;
+	}
+
+	return true;
+}
+
 bool FImageUtils::GetRawData(UTextureRenderTarget2D* TexRT, TArray64<uint8>& RawData)
 {
 	// DEPRECATED , use GetRenderTargetImage
@@ -330,20 +438,67 @@ bool FImageUtils::GetRawData(UTextureRenderTarget2D* TexRT, TArray64<uint8>& Raw
 	return true;
 }
 
-bool FImageUtils::GetRenderTargetImage(UTextureRenderTarget2D* TexRT, FImage & Image)
+bool FImageUtils::GetRenderTargetImage(UTextureRenderTarget* TexRT, FImage & Image)
 {
 	return GetRenderTargetImage(TexRT,Image,FIntRect(0,0,0,0));
 }
 
-bool FImageUtils::GetRenderTargetImage(UTextureRenderTarget2D* TexRT, FImage & Image, const FIntRect & InRectOrZero)
+bool FImageUtils::GetRenderTargetImage(UTextureRenderTarget* TexRT, FImage & OutImage, const FIntRect & InRectOrZero)
 {
-	Image = FImage();
+	OutImage = FImage();
 	
 	FRenderTarget* RenderTarget = TexRT->GameThread_GetRenderTargetResource();
 	EPixelFormat RTFormat = TexRT->GetFormat();
 
 	ERawImageFormat::Type ReadFormat = UTextureRenderTarget::GetReadPixelsFormat(RTFormat,false);
-		
+	
+	//UTextureRenderTarget2D * TexRT2D = Cast<UTextureRenderTarget2D>(TexRT);
+	//UTextureRenderTarget2DArray * TexRT2DA = Cast<UTextureRenderTarget2DArray>(TexRT);
+	//UTextureRenderTargetVolume * TexRTVolume = Cast<UTextureRenderTargetVolume>(TexRT);
+
+	// we have to identify cubes because they are treated differently
+	// can't use TexRT->GetTextureClass() because it's always RenderTarget
+	UTextureRenderTargetCube * TexRTCube = Cast<UTextureRenderTargetCube>(TexRT);
+	bool bIsCube = ( TexRTCube != nullptr );
+	// UTextureRenderTargetCubeArray does not exist at the moment
+	
+	// UTextureRenderTarget2DArray does not derive from UTextureRenderTarget2D, so TexRT2D will be null
+
+	int64 TexRT_SizeX = FMath::RoundToInt64( TexRT->GetSurfaceWidth() ); // GetSurfaceWidth returns SizeX but as float
+	int64 TexRT_SizeY = FMath::RoundToInt64( TexRT->GetSurfaceHeight() );
+	int64 TexRT_SizeZ = FMath::RoundToInt64( TexRT->GetSurfaceDepth() );
+	int64 TexRT_ArraySize = TexRT->GetSurfaceArraySize();
+
+	if ( TexRT_SizeX <= 0 || TexRT_SizeY <= 0 )
+	{
+		return false;
+	}
+
+	int64 NumSlices;
+	if ( bIsCube )
+	{
+		check( TexRT_SizeZ == 0 );
+		check( TexRT_ArraySize == 6 );
+		NumSlices = 6;
+	}
+	else if ( TexRT_SizeZ != 0 )
+	{
+		check( TexRT_ArraySize == 0 );
+		NumSlices = TexRT_SizeZ;
+	}
+	else if ( TexRT_ArraySize != 0 )
+	{
+		check( TexRT_SizeZ == 0 );
+		NumSlices = TexRT_ArraySize;
+	}
+	else
+	{
+		// 2D
+		NumSlices = 1;
+	}
+	
+	// UTextureRenderTarget2D returns 0 for Depth and ArraySize, not 1
+
 	// RCM_MinMax means don't renormalize, just read the pixels as they are
 	//	default RCM_UNorm does funny scalings
 	FReadSurfaceDataFlags ReadFlags(RCM_MinMax, CubeFace_MAX);
@@ -351,13 +506,13 @@ bool FImageUtils::GetRenderTargetImage(UTextureRenderTarget2D* TexRT, FImage & I
 	FIntRect Rect = InRectOrZero;
 	if (InRectOrZero == FIntRect(0, 0, 0, 0))
 	{
-		Rect = FIntRect(0, 0, TexRT->SizeX, TexRT->SizeY);
+		Rect = FIntRect(0, 0, TexRT_SizeX, TexRT_SizeY);
 	}
 
 	int64 RectSizeX = Rect.Width();
 	int64 RectSizeY = Rect.Height();
 	if ( ! ensure( Rect.Min.X >= 0 && Rect.Min.Y >= 0 && 
-		Rect.Max.X <= TexRT->SizeX && Rect.Max.Y <= TexRT->SizeY &&
+		Rect.Max.X <= TexRT_SizeX && Rect.Max.Y <= TexRT_SizeY &&
 		RectSizeX >= 0 && RectSizeY >= 0 ) )
 	{
 		return false;
@@ -372,16 +527,24 @@ bool FImageUtils::GetRenderTargetImage(UTextureRenderTarget2D* TexRT, FImage & I
 		// ReadFloat16Pixels does no conversions
 		//	must be used only exactly with FloatRGBA type
 
-		Image.Init(RectSizeX,RectSizeY,ERawImageFormat::RGBA16F,EGammaSpace::Linear);
+		OutImage.Init(RectSizeX,RectSizeY,NumSlices,ERawImageFormat::RGBA16F,EGammaSpace::Linear);
 		
 		TArray<FFloat16Color> Colors;
-		if ( ! RenderTarget->ReadFloat16Pixels(Colors,ReadFlags,Rect) )
+		
+		for (int32 SliceIndex = 0; SliceIndex < NumSlices; ++SliceIndex)
 		{
-			return false;
-		}
+			ReadFlags.SetCubeFace(  bIsCube ? (ECubeFace)(SliceIndex % 6) : CubeFace_MAX);
+			ReadFlags.SetArrayIndex(bIsCube ? (SliceIndex/6) : SliceIndex);
 
-		check( Image.GetImageSizeBytes() == Colors.Num() * sizeof(Colors[0]) );
-		memcpy( &Image.RawData[0], &Colors[0], Image.GetImageSizeBytes() );
+			if ( ! RenderTarget->ReadFloat16Pixels(Colors,ReadFlags,Rect) )
+			{
+				return false;
+			}
+
+			FImageView ImageSlice = OutImage.GetSlice(SliceIndex);
+			check( ImageSlice.GetImageSizeBytes() == Colors.Num() * sizeof(Colors[0]) );
+			memcpy( ImageSlice.RawData, Colors.GetData(), ImageSlice.GetImageSizeBytes() );
+		}
 	}
 	else if ( ReadFormat == ERawImageFormat::BGRA8 )
 	{
@@ -389,7 +552,7 @@ bool FImageUtils::GetRenderTargetImage(UTextureRenderTarget2D* TexRT, FImage & I
 		//	mainly we are trying to catch the check for whether the _SRGB or non _SRGB BGRA8 format as chosen
 		EGammaSpace GammaSpace = TexRT->IsSRGB() ? EGammaSpace::sRGB : EGammaSpace::Linear;
 		
-		Image.Init(RectSizeX,RectSizeY,ERawImageFormat::BGRA8,GammaSpace);
+		OutImage.Init(RectSizeX,RectSizeY,NumSlices,ERawImageFormat::BGRA8,GammaSpace);
 		
 		// "LinearToGamma" is basically moot; that would only be used if we were reading float pixels to FColor
 		//	but in that case the ReadFormat should have been float, so we won't be here
@@ -397,26 +560,42 @@ bool FImageUtils::GetRenderTargetImage(UTextureRenderTarget2D* TexRT, FImage & I
 		ReadFlags.SetLinearToGamma( GammaSpace == EGammaSpace::sRGB );
 
 		TArray<FColor> Colors;
-		if ( ! RenderTarget->ReadPixels(Colors,ReadFlags,Rect) )
+		
+		for (int32 SliceIndex = 0; SliceIndex < NumSlices; ++SliceIndex)
 		{
-			return false;
-		}
+			ReadFlags.SetCubeFace(  bIsCube ? (ECubeFace)(SliceIndex % 6) : CubeFace_MAX);
+			ReadFlags.SetArrayIndex(bIsCube ? (SliceIndex/6) : SliceIndex);
 
-		check( Image.GetImageSizeBytes() == Colors.Num() * sizeof(Colors[0]) );
-		memcpy( &Image.RawData[0], &Colors[0], Image.GetImageSizeBytes() );
+			if ( ! RenderTarget->ReadPixels(Colors,ReadFlags,Rect) )
+			{
+				return false;
+			}
+		
+			FImageView ImageSlice = OutImage.GetSlice(SliceIndex);
+			check( ImageSlice.GetImageSizeBytes() == Colors.Num() * sizeof(Colors[0]) );
+			memcpy( ImageSlice.RawData, Colors.GetData(), ImageSlice.GetImageSizeBytes() );
+		}
 	}
 	else if ( ReadFormat == ERawImageFormat::RGBA32F )
 	{
-		Image.Init(RectSizeX,RectSizeY,ERawImageFormat::RGBA32F,EGammaSpace::Linear);
+		OutImage.Init(RectSizeX,RectSizeY,NumSlices,ERawImageFormat::RGBA32F,EGammaSpace::Linear);
 		
 		TArray<FLinearColor> Colors;
-		if ( ! RenderTarget->ReadLinearColorPixels(Colors,ReadFlags,Rect) )
+		
+		for (int32 SliceIndex = 0; SliceIndex < NumSlices; ++SliceIndex)
 		{
-			return false;
-		}
+			ReadFlags.SetCubeFace(  bIsCube ? (ECubeFace)(SliceIndex % 6) : CubeFace_MAX);
+			ReadFlags.SetArrayIndex(bIsCube ? (SliceIndex/6) : SliceIndex);
 
-		check( Image.GetImageSizeBytes() == Colors.Num() * sizeof(Colors[0]) );
-		memcpy( &Image.RawData[0], &Colors[0], Image.GetImageSizeBytes() );
+			if ( ! RenderTarget->ReadLinearColorPixels(Colors,ReadFlags,Rect) )
+			{
+				return false;
+			}
+			
+			FImageView ImageSlice = OutImage.GetSlice(SliceIndex);
+			check( ImageSlice.GetImageSizeBytes() == Colors.Num() * sizeof(Colors[0]) );
+			memcpy( ImageSlice.RawData, Colors.GetData(), ImageSlice.GetImageSizeBytes() );
+		}
 	}
 	else
 	{
diff --git a/Engine/Source/Runtime/Engine/Private/Slate/SlateTextures.cpp b/Engine/Source/Runtime/Engine/Private/Slate/SlateTextures.cpp
--- a/Engine/Source/Runtime/Engine/Private/Slate/SlateTextures.cpp
+++ b/Engine/Source/Runtime/Engine/Private/Slate/SlateTextures.cpp
@@ -388,6 +388,9 @@ FIntPoint FSlateTextureRenderTarget2DResource::GetSizeXY() const
 
 float FSlateTextureRenderTarget2DResource::GetDisplayGamma() const
 {
+	// FSlateTextureRenderTarget2DResource doesn't have Owner
+	//return Owner->GetDisplayGamma();
+
 	if (TargetGamma > UE_KINDA_SMALL_NUMBER * 10.0f)
 	{
 		return TargetGamma;
diff --git a/Engine/Source/Runtime/Engine/Private/TextureRenderTarget.cpp b/Engine/Source/Runtime/Engine/Private/TextureRenderTarget.cpp
--- a/Engine/Source/Runtime/Engine/Private/TextureRenderTarget.cpp
+++ b/Engine/Source/Runtime/Engine/Private/TextureRenderTarget.cpp
@@ -225,7 +225,7 @@ bool UTextureRenderTarget::UpdateTexture(UTexture* InTexture, EConstructTextureF
 		ReadFormat == ERawImageFormat::RGBA16F );
 
 	EGammaSpace ReadGammaSpace = EGammaSpace::Linear;
-	bool bSRGB = RenderTarget->GetDisplayGamma() > 1.5f; // @@ no way to call IsSRGB() , no uniform query on RenderTarget
+	bool bSRGB = IsSRGB();
 	if ( ERawImageFormat::GetFormatNeedsGammaSpace(ReadFormat) && bSRGB )
 	{
 		ReadGammaSpace = EGammaSpace::sRGB;
@@ -531,6 +531,16 @@ const FTextureRHIRef& FTextureRenderTargetResource::GetShaderResourceTexture() c
 }
 
 
+/*static*/ float UTextureRenderTarget::GetDefaultDisplayGamma()
+{
+	// TextureRenderTarget default gamma does not respond to Engine->DisplayGamma setting
+	//	?? was that intentional or a bug ?
+	//return FRenderTarget::GetEngineDisplayGamma();
+
+	// return hard-coded gamma 2.2 which corresponds to SRGB
+	return 2.2f;
+}
+
 /** 
 * Render target resource should be sampled in linear color space
 *
@@ -538,8 +548,10 @@ const FTextureRHIRef& FTextureRenderTargetResource::GetShaderResourceTexture() c
 */
 float FTextureRenderTargetResource::GetDisplayGamma() const
 {
-	// when we say we want a 2.2 gamma, what we actually mean is that we want SRGB conversion in most cases
-	return 2.2f;  
+	// we'd like to just do this, but Owner is in the derived classes
+	//return Owner->GetDisplayGamma();
+
+	return UTextureRenderTarget::GetDefaultDisplayGamma();
 }
 
 /*-----------------------------------------------------------------------------
diff --git a/Engine/Source/Runtime/Engine/Private/TextureRenderTarget2D.cpp b/Engine/Source/Runtime/Engine/Private/TextureRenderTarget2D.cpp
--- a/Engine/Source/Runtime/Engine/Private/TextureRenderTarget2D.cpp
+++ b/Engine/Source/Runtime/Engine/Private/TextureRenderTarget2D.cpp
@@ -54,6 +54,35 @@ UTextureRenderTarget2D::UTextureRenderTarget2D(const FObjectInitializer& ObjectI
 	// see also UTextureRenderTarget::TargetGamma
 }
 
+EPixelFormat UTextureRenderTarget2D::GetFormat() const
+{
+	if (OverrideFormat == PF_Unknown)
+	{
+		return GetPixelFormatFromRenderTargetFormat(RenderTargetFormat);
+	}
+	else
+	{
+		return OverrideFormat;
+	}
+}
+
+bool UTextureRenderTarget2D::IsSRGB() const
+{
+	// in theory you'd like the "bool SRGB" variable to == this, but it does not
+
+	// ?? note: UTextureRenderTarget::TargetGamma is ignored here
+	// ?? note: GetDisplayGamma forces linear for some float formats, but this doesn't
+
+	if (OverrideFormat == PF_Unknown)
+	{
+		return RenderTargetFormat == RTF_RGBA8_SRGB;
+	}
+	else
+	{
+		return !bForceLinearGamma;
+	}
+}
+
 FTextureResource* UTextureRenderTarget2D::CreateResource()
 {
 	UWorld* World = GetWorld();
@@ -646,21 +675,27 @@ FIntPoint FTextureRenderTarget2DResource::GetSizeXY() const
 *
 * @return display gamma expected for rendering to this render target 
 */
-float FTextureRenderTarget2DResource::GetDisplayGamma() const
+float UTextureRenderTarget2D::GetDisplayGamma() const
 {
 	// if TargetGamma is set (not zero), it overrides everything else
-	if (Owner->TargetGamma > UE_KINDA_SMALL_NUMBER * 10.0f)
+	if (TargetGamma > UE_KINDA_SMALL_NUMBER * 10.0f)
 	{
-		return Owner->TargetGamma;
+		return TargetGamma;
 	}
 
 	// ?? special casing just two of the float PixelFormats to force 1.0 gamma here is inconsistent
 	//		(there are lots of other float formats)
 	// ignores Owner->IsSRGB() ? it's similar but not quite the same
-	if (Format == PF_FloatRGB || Format == PF_FloatRGBA || Owner->bForceLinearGamma )
+	EPixelFormat Format = GetFormat();
+	if (Format == PF_FloatRGB || Format == PF_FloatRGBA || bForceLinearGamma )
 	{
 		return 1.0f;
 	}
 
-	return FTextureRenderTargetResource::GetDisplayGamma(); // hard-coded 2.2 , actually means SRGB
+	return UTextureRenderTarget::GetDefaultDisplayGamma(); // hard-coded 2.2 , actually means SRGB
 }
+
+float FTextureRenderTarget2DResource::GetDisplayGamma() const
+{
+	return Owner->GetDisplayGamma();
+}
\ No newline at end of file
diff --git a/Engine/Source/Runtime/Engine/Private/TextureRenderTarget2DArray.cpp b/Engine/Source/Runtime/Engine/Private/TextureRenderTarget2DArray.cpp
--- a/Engine/Source/Runtime/Engine/Private/TextureRenderTarget2DArray.cpp
+++ b/Engine/Source/Runtime/Engine/Private/TextureRenderTarget2DArray.cpp
@@ -31,6 +31,31 @@ UTextureRenderTarget2DArray::UTextureRenderTarget2DArray(const FObjectInitialize
 	bForceLinearGamma = true;
 }
 
+EPixelFormat UTextureRenderTarget2DArray::GetFormat() const
+{
+	if(OverrideFormat == PF_Unknown)
+	{
+		return bHDR ? PF_FloatRGBA : PF_B8G8R8A8;
+	}
+	else
+	{
+		return OverrideFormat;
+	}
+}
+
+bool UTextureRenderTarget2DArray::IsSRGB() const
+{
+	bool bIsSRGB = true;
+
+	// if render target gamma used was 1.0 then disable SRGB for the static texture
+	if(FMath::Abs(GetDisplayGamma() - 1.0f) < UE_KINDA_SMALL_NUMBER)
+	{
+		bIsSRGB = false;
+	}	
+
+	return bIsSRGB;
+}
+
 void UTextureRenderTarget2DArray::Init(uint32 InSizeX, uint32 InSizeY, uint32 InSlices, EPixelFormat InFormat)
 {
 	check((InSizeX > 0) && (InSizeY > 0) && (InSlices > 0));
@@ -201,16 +226,8 @@ void FTextureRenderTarget2DArrayResource::InitRHI(FRHICommandListBase& RHICmdLis
 
 	if((Owner->SizeX > 0) && (Owner->SizeY > 0) && (Owner->Slices > 0))
 	{
-		bool bIsSRGB = true;
-
-		// if render target gamma used was 1.0 then disable SRGB for the static texture
-		if(FMath::Abs(GetDisplayGamma() - 1.0f) < UE_KINDA_SMALL_NUMBER)
-		{
-			bIsSRGB = false;
-		}
-
 		// Create the RHI texture. Only one mip is used and the texture is targetable for resolve.
-		ETextureCreateFlags TexCreateFlags = bIsSRGB ? ETextureCreateFlags::SRGB : ETextureCreateFlags::None;
+		ETextureCreateFlags TexCreateFlags = Owner->IsSRGB() ? ETextureCreateFlags::SRGB : ETextureCreateFlags::None;
 		if (Owner->bCanCreateUAV)
 		{
 			TexCreateFlags |= ETextureCreateFlags::UAV;
@@ -318,18 +335,23 @@ FIntPoint FTextureRenderTarget2DArrayResource::GetSizeXY() const
 	return FIntPoint(Owner->SizeX, Owner->SizeX);
 }
 
-float FTextureRenderTarget2DArrayResource::GetDisplayGamma() const
+float UTextureRenderTarget2DArray::GetDisplayGamma() const
 {
-	if(Owner->TargetGamma > UE_KINDA_SMALL_NUMBER * 10.0f)
+	if(TargetGamma > UE_KINDA_SMALL_NUMBER * 10.0f)
 	{
-		return Owner->TargetGamma;
+		return TargetGamma;
 	}
-	EPixelFormat Format = Owner->GetFormat();
-	if(Format == PF_FloatRGB || Format == PF_FloatRGBA || Owner->bForceLinearGamma)
+	EPixelFormat Format = GetFormat();
+	if(Format == PF_FloatRGB || Format == PF_FloatRGBA || bForceLinearGamma)
 	{
 		return 1.0f;
 	}
-	return FTextureRenderTargetResource::GetDisplayGamma();
+	return UTextureRenderTarget::GetDefaultDisplayGamma();
+}
+
+float FTextureRenderTarget2DArrayResource::GetDisplayGamma() const
+{
+	return Owner->GetDisplayGamma();
 }
 
 PRAGMA_DISABLE_DEPRECATION_WARNINGS
diff --git a/Engine/Source/Runtime/Engine/Private/TextureRenderTargetCube.cpp b/Engine/Source/Runtime/Engine/Private/TextureRenderTargetCube.cpp
--- a/Engine/Source/Runtime/Engine/Private/TextureRenderTargetCube.cpp
+++ b/Engine/Source/Runtime/Engine/Private/TextureRenderTargetCube.cpp
@@ -35,8 +35,35 @@ UTextureRenderTargetCube::UTextureRenderTargetCube(const FObjectInitializer& Obj
 	ClearColor = FLinearColor(0.0f, 1.0f, 0.0f, 1.0f);
 	OverrideFormat = PF_Unknown;
 	bForceLinearGamma = true;
+	// note bool SRGB not set
 }
 
+EPixelFormat UTextureRenderTargetCube::GetFormat() const
+{
+	if(OverrideFormat == PF_Unknown)
+	{
+		return bHDR ? PF_FloatRGBA : PF_B8G8R8A8;
+	}
+	else
+	{
+		return OverrideFormat;
+	}
+}
+
+bool UTextureRenderTargetCube::IsSRGB() const
+{
+	// compare with IsSRGB() call on RenderTarget2D which uses different logic
+	//	also "bool SRGB" should be set = bIsSRGB, but isn't
+
+	bool bIsSRGB = true;
+	// if render target gamma used was 1.0 then disable SRGB for the static texture
+	if(FMath::Abs(GetDisplayGamma() - 1.0f) < UE_KINDA_SMALL_NUMBER)
+	{
+		bIsSRGB = false;
+	}
+
+	return bIsSRGB;
+}
 
 void UTextureRenderTargetCube::Init(uint32 InSizeX, EPixelFormat InFormat)
 {
@@ -224,17 +251,11 @@ void FTextureRenderTargetCubeResource::InitRHI(FRHICommandListBase& RHICmdList)
 
 	if(Owner->SizeX > 0)
 	{
-		// compare with IsSRGB() call on RenderTarget2D
-		//	also "bool SRGB" should be set = bIsSRGB, but isn't
-		bool bIsSRGB = true;
-		// if render target gamma used was 1.0 then disable SRGB for the static texture
-		if(FMath::Abs(GetDisplayGamma() - 1.0f) < UE_KINDA_SMALL_NUMBER)
-		{
-			bIsSRGB = false;
-		}
+		// note Resource has a bSRGB and Owner has a bool SRGB , neither are set right
+		//	instead call Owner->IsSRGB()
 
 		// Create the RHI texture. Only one mip is used and the texture is targetable for resolve.
-		ETextureCreateFlags TexCreateFlags = bIsSRGB ? ETextureCreateFlags::SRGB : ETextureCreateFlags::None;
+		ETextureCreateFlags TexCreateFlags = Owner->IsSRGB() ? ETextureCreateFlags::SRGB : ETextureCreateFlags::None;
 		if (Owner->bCanCreateUAV)
 		{
 			TexCreateFlags |= ETextureCreateFlags::UAV;
@@ -360,7 +381,7 @@ uint32 FTextureRenderTargetCubeResource::GetSizeX() const
  */
 uint32 FTextureRenderTargetCubeResource::GetSizeY() const
 {
-	return Owner->SizeX;
+	return Owner->SizeX; // there is no "SizeY" , cubes must be square
 }
 
 /** 
@@ -371,20 +392,26 @@ FIntPoint FTextureRenderTargetCubeResource::GetSizeXY() const
 	return FIntPoint(Owner->SizeX, Owner->SizeX);
 }
 
-float FTextureRenderTargetCubeResource::GetDisplayGamma() const
+float UTextureRenderTargetCube::GetDisplayGamma() const
 {
 	// code dupe of RenderTarget2D
 
-	if(Owner->TargetGamma > UE_KINDA_SMALL_NUMBER * 10.0f)
+	if(TargetGamma > UE_KINDA_SMALL_NUMBER * 10.0f)
 	{
-		return Owner->TargetGamma;
+		return TargetGamma;
 	}
-	EPixelFormat Format = Owner->GetFormat();
-	if(Format == PF_FloatRGB || Format == PF_FloatRGBA || Owner->bForceLinearGamma)
+	EPixelFormat Format = GetFormat();
+	if(Format == PF_FloatRGB || Format == PF_FloatRGBA || bForceLinearGamma)
 	{
 		return 1.0f;
 	}
-	return FTextureRenderTargetResource::GetDisplayGamma();
+
+	return UTextureRenderTarget::GetDefaultDisplayGamma(); // eg. 2.2
+}
+
+float FTextureRenderTargetCubeResource::GetDisplayGamma() const
+{
+	return Owner->GetDisplayGamma();
 }
 
 PRAGMA_DISABLE_DEPRECATION_WARNINGS
diff --git a/Engine/Source/Runtime/Engine/Private/TextureRenderTargetVolume.cpp b/Engine/Source/Runtime/Engine/Private/TextureRenderTargetVolume.cpp
--- a/Engine/Source/Runtime/Engine/Private/TextureRenderTargetVolume.cpp
+++ b/Engine/Source/Runtime/Engine/Private/TextureRenderTargetVolume.cpp
@@ -31,6 +31,30 @@ UTextureRenderTargetVolume::UTextureRenderTargetVolume(const FObjectInitializer&
 	bForceLinearGamma = true;
 }
 
+EPixelFormat UTextureRenderTargetVolume::GetFormat() const
+{
+	if(OverrideFormat == PF_Unknown)
+	{
+		return bHDR ? PF_FloatRGBA : PF_B8G8R8A8;
+	}
+	else
+	{
+		return OverrideFormat;
+	}
+}
+
+bool UTextureRenderTargetVolume::IsSRGB() const
+{
+	bool bIsSRGB = true;
+	// if render target gamma used was 1.0 then disable SRGB for the static texture
+	if(FMath::Abs(GetDisplayGamma() - 1.0f) < UE_KINDA_SMALL_NUMBER)
+	{
+		bIsSRGB = false;
+	}
+	
+	return bIsSRGB;
+}
+
 void UTextureRenderTargetVolume::Init(uint32 InSizeX, uint32 InSizeY, uint32 InSizeZ, EPixelFormat InFormat)
 {
 	check((InSizeX > 0) && (InSizeY > 0) && (InSizeZ > 0));
@@ -211,16 +235,8 @@ void FTextureRenderTargetVolumeResource::InitRHI(FRHICommandListBase& RHICmdList
 
 	if((Owner->SizeX > 0) && (Owner->SizeY > 0) && (Owner->SizeZ > 0))
 	{
-		// ?? Volume's version of IsSRGB() here :
-		bool bIsSRGB = true;
-		// if render target gamma used was 1.0 then disable SRGB for the static texture
-		if(FMath::Abs(GetDisplayGamma() - 1.0f) < UE_KINDA_SMALL_NUMBER)
-		{
-			bIsSRGB = false;
-		}
-
 		// Create the RHI texture. Only one mip is used and the texture is targetable for resolve.
-		ETextureCreateFlags TexCreateFlags = bIsSRGB ? TexCreate_SRGB : TexCreate_None;
+		ETextureCreateFlags TexCreateFlags = Owner->IsSRGB() ? TexCreate_SRGB : TexCreate_None;
 
 		if (Owner->bCanCreateUAV)
 		{
@@ -298,45 +314,47 @@ void FTextureRenderTargetVolumeResource::UpdateDeferredResource(FRHICommandListI
 	}
 }
 
-/** 
- * @return width of target
- */
 uint32 FTextureRenderTargetVolumeResource::GetSizeX() const
 {
 	return Owner->SizeX;
 }
 
-/** 
- * @return height of target
- */
 uint32 FTextureRenderTargetVolumeResource::GetSizeY() const
 {
-	return Owner->SizeX;
+	return Owner->SizeY;
+}
+
+uint32 FTextureRenderTargetVolumeResource::GetSizeZ() const
+{
+	return Owner->SizeZ;
 }
 
-/** 
- * @return dimensions of target surface
- */
 FIntPoint FTextureRenderTargetVolumeResource::GetSizeXY() const
 {
-	return FIntPoint(Owner->SizeX, Owner->SizeX);
+	return FIntPoint(Owner->SizeX, Owner->SizeY);
 }
 
-float FTextureRenderTargetVolumeResource::GetDisplayGamma() const
+float UTextureRenderTargetVolume::GetDisplayGamma() const
 {
 	// code dupe ; move this up to the top level, it's duped everywhere
 
-	if(Owner->TargetGamma > UE_KINDA_SMALL_NUMBER * 10.0f)
+	if(TargetGamma > UE_KINDA_SMALL_NUMBER * 10.0f)
 	{
-		return Owner->TargetGamma;
+		return TargetGamma;
 	}
-	EPixelFormat Format = Owner->GetFormat();
+	EPixelFormat Format = GetFormat();
 	// ?? hard-coding a few formats but not others, likely wrong
-	if(Format == PF_R32_FLOAT || Format == PF_FloatRGBA || Owner->bForceLinearGamma)
+	if(Format == PF_R32_FLOAT || Format == PF_FloatRGBA || bForceLinearGamma)
 	{
 		return 1.0f;
 	}
-	return FTextureRenderTargetResource::GetDisplayGamma();
+
+	return UTextureRenderTarget::GetDefaultDisplayGamma();
+}
+
+float FTextureRenderTargetVolumeResource::GetDisplayGamma() const
+{
+	return Owner->GetDisplayGamma();
 }
 
 bool FTextureRenderTargetVolumeResource::ReadPixels(TArray<FColor>& OutImageData, FReadSurfaceDataFlags InFlags, FIntRect InSrcRect)
@@ -357,15 +375,7 @@ bool FTextureRenderTargetVolumeResource::ReadFloat16Pixels(TArray<FFloat16Color>
 
 	if ( PF != PF_FloatRGBA && PF != PF_R16F && PF != PF_R32_FLOAT )
 	{
-		/**
-
-		deep inside RHI this will crash unless the format is one of these :
-
-		bool bIsRGBAFmt = TextureDesc.Format == GPixelFormats[PF_FloatRGBA].PlatformFormat;
-		bool bIsR16FFmt = TextureDesc.Format == GPixelFormats[PF_R16F].PlatformFormat;
-		bool bIsR32FFmt = TextureDesc.Format == GPixelFormats[PF_R32_FLOAT].PlatformFormat;
-		check(bIsRGBAFmt || bIsR16FFmt || bIsR32FFmt);
-		**/
+		// limitation of RHIRead3DSurfaceFloatData
 
 		return false;
 	}
diff --git a/Engine/Source/Runtime/Engine/Private/UnrealClient.cpp b/Engine/Source/Runtime/Engine/Private/UnrealClient.cpp
--- a/Engine/Source/Runtime/Engine/Private/UnrealClient.cpp
+++ b/Engine/Source/Runtime/Engine/Private/UnrealClient.cpp
@@ -150,6 +150,13 @@ bool FRenderTarget::ReadLinearColorPixelsPtr(FLinearColor* OutImageBytes, FReadS
 
 float FRenderTarget::GetDisplayGamma() const
 {
+	return GetEngineDisplayGamma();
+}
+
+float FRenderTarget::GetEngineDisplayGamma() // static
+{
+	// when we say we want a 2.2 gamma, what we actually mean is that we want SRGB conversion in most cases
+
 	if (GEngine == NULL)
 	{
 		return 2.2f;
diff --git a/Engine/Source/Runtime/Engine/Private/UserInterface/Canvas.cpp b/Engine/Source/Runtime/Engine/Private/UserInterface/Canvas.cpp
--- a/Engine/Source/Runtime/Engine/Private/UserInterface/Canvas.cpp
+++ b/Engine/Source/Runtime/Engine/Private/UserInterface/Canvas.cpp
@@ -467,9 +467,9 @@ bool FCanvasBatchedElementRenderItem::Render_GameThread(const FCanvas* Canvas, F
 	{
 		bDirty = true;
 
-		// current render target set for the canvas
+		// current render target set for the canvas  (eg. an FSlateTextureRenderTarget2DResource)
 		const FRenderTarget* CanvasRenderTarget = Canvas->GetRenderTarget();
-		float Gamma = 1.0f / CanvasRenderTarget->GetDisplayGamma();
+		float Gamma = 1.0f / CanvasRenderTarget->GetDisplayGamma(); // GetDisplayGamma typically == 2.2
 		
 		// Render the batched elements.
 		struct FBatchedDrawParameters
diff --git a/Engine/Source/Runtime/Engine/Public/ImageUtils.h b/Engine/Source/Runtime/Engine/Public/ImageUtils.h
--- a/Engine/Source/Runtime/Engine/Public/ImageUtils.h
+++ b/Engine/Source/Runtime/Engine/Public/ImageUtils.h
@@ -29,6 +29,7 @@ class UTexture2DArray;
 class UTextureCube;
 class UTextureCubeArray;
 class UVolumeTexture;
+class UTextureRenderTarget;
 class UTextureRenderTarget2D;
 class UTextureRenderTargetCube;
 class IImageWrapperModule;
@@ -149,6 +150,16 @@ public:
 	* 
 	*/
 	ENGINE_API static bool ExportTextureSourceToDDS(TArray64<uint8> & OutData, UTexture * Texture, int BlockIndex=0, int LayerIndex=0);
+	
+	/**
+	*  Export Texture RenderTarget (2D,Array,Cube,CubeArray,Volume) to DDS
+	*  reads from texture rendertarget RHI texture
+	* 
+	* @param OutData    Filled with DDS file format
+	* @param TexRT    Texture RenderTarget (2D,Array,Cube,CubeArray,Volume) to export
+	* 
+	*/
+	ENGINE_API static bool ExportRenderTargetToDDS(TArray64<uint8> & OutData, UTextureRenderTarget * TexRT);
 
 	/**
 	 * if Texture source is available, get it as an FImage
@@ -437,8 +448,10 @@ public:
 	* 
 	* @param TexRT		The texture rendertarget to copy from
 	* @param OutImage	Filled with the image, allocated as needed
+	*
+	* Works for Cubes, Volumes, etc.  (fills Image slices)
 	*/
-	ENGINE_API static bool GetRenderTargetImage(UTextureRenderTarget2D* TexRT, FImage & OutImage);
-	ENGINE_API static bool GetRenderTargetImage(UTextureRenderTarget2D* TexRT, FImage & OutImage, const FIntRect & Rect);
+	ENGINE_API static bool GetRenderTargetImage(UTextureRenderTarget* TexRT, FImage & OutImage);
+	ENGINE_API static bool GetRenderTargetImage(UTextureRenderTarget* TexRT, FImage & OutImage, const FIntRect & Rect);
 
 };
diff --git a/Engine/Source/Runtime/Engine/Public/TextureRenderTargetVolumeResource.h b/Engine/Source/Runtime/Engine/Public/TextureRenderTargetVolumeResource.h
--- a/Engine/Source/Runtime/Engine/Public/TextureRenderTargetVolumeResource.h
+++ b/Engine/Source/Runtime/Engine/Public/TextureRenderTargetVolumeResource.h
@@ -48,19 +48,12 @@ public:
 
 	// FRenderTarget interface.
 
-	/**
-	 * @return width of the target
-	 */
+	// FTexture interface :
 	virtual uint32 GetSizeX() const override;
-
-	/**
-	 * @return height of the target
-	 */
 	virtual uint32 GetSizeY() const override;
+	virtual uint32 GetSizeZ() const override;
 
-	/**
-	 * @return dimensions of the target
-	 */
+	// FRenderTarget interface:	
 	virtual FIntPoint GetSizeXY() const override;
 
 	/**
diff --git a/Engine/Source/Runtime/Engine/Public/TextureResource.h b/Engine/Source/Runtime/Engine/Public/TextureResource.h
--- a/Engine/Source/Runtime/Engine/Public/TextureResource.h
+++ b/Engine/Source/Runtime/Engine/Public/TextureResource.h
@@ -374,17 +374,20 @@ public:
 
 	virtual void ClampSize(int32 SizeX,int32 SizeY) {}
 
-	// FRenderTarget interface.
-	virtual uint32 GetSizeX() const = 0;
-	virtual uint32 GetSizeY() const = 0;
-	virtual FIntPoint GetSizeXY() const = 0;
+	// GetSizeX/GetSizeY are from "FTexture" interface.
+	//   the FTexture implementations return zero; force them to be implemented:
+	virtual uint32 GetSizeX() const override = 0;
+	virtual uint32 GetSizeY() const override = 0;
+	// also GetSizeZ()
+	// GetSizeXY from "FRenderTarget"
+	//virtual FIntPoint GetSizeXY() const override = 0;
 
 	/** 
 	 * Render target resource should be sampled in linear color space
 	 *
 	 * @return display gamma expected for rendering to this render target 
 	 */
-	virtual float GetDisplayGamma() const;
+	virtual float GetDisplayGamma() const override;
 
 	virtual FRHIGPUMask GetGPUMask(FRHICommandListImmediate& RHICmdList) const final override
 	{
diff --git a/Engine/Source/Runtime/Engine/Public/UnrealClient.h b/Engine/Source/Runtime/Engine/Public/UnrealClient.h
--- a/Engine/Source/Runtime/Engine/Public/UnrealClient.h
+++ b/Engine/Source/Runtime/Engine/Public/UnrealClient.h
@@ -73,6 +73,9 @@ public:
 	*/
 	ENGINE_API virtual float GetDisplayGamma() const;
 
+	// return global Engine default gamma (GetDisplayGamma returns this if not overriden)
+	static float GetEngineDisplayGamma();
+
 	virtual EDisplayColorGamut GetDisplayColorGamut() const { return EDisplayColorGamut::sRGB_D65; }
 	virtual EDisplayOutputFormat GetDisplayOutputFormat() const { return EDisplayOutputFormat::SDR_sRGB; }
 	virtual bool GetSceneHDREnabled() const { return false; }
diff --git a/Engine/Source/Runtime/RHI/Public/RHITypes.h b/Engine/Source/Runtime/RHI/Public/RHITypes.h
--- a/Engine/Source/Runtime/RHI/Public/RHITypes.h
+++ b/Engine/Source/Runtime/RHI/Public/RHITypes.h
@@ -29,6 +29,10 @@ public:
 	{
 		return CubeFace;
 	}
+	void SetCubeFace(ECubeFace InCubeFace)
+	{
+		CubeFace = InCubeFace;
+	}
 
 	ERangeCompressionMode GetCompressionMode() const
 	{
