Net.DebugDraw 1 now showing NetGUIDs

#rnx

diff --git a/Engine/Source/Runtime/Engine/Private/NetDriver.cpp b/Engine/Source/Runtime/Engine/Private/NetDriver.cpp
--- a/Engine/Source/Runtime/Engine/Private/NetDriver.cpp
+++ b/Engine/Source/Runtime/Engine/Private/NetDriver.cpp
@@ -5775,7 +5775,8 @@ void UNetDriver::DrawNetDriverDebug()
 
 		// Draw NetGUID or NetHandle
 		FVector AboveActor(0.f, 0.f, 32.f);
-#if UE_WITH_IRIS && !UE_BUILD_SHIPPING
+#if !UE_BUILD_SHIPPING
+#if UE_WITH_IRIS 
 		if (Bridge)
 		{
 			AActor* Actor = *It;
@@ -5791,7 +5792,7 @@ void UNetDriver::DrawNetDriverDebug()
 			}
 		}
 		else
-#else
+#endif
 		{
 			FNetworkGUID NetGUID = Connection->PackageMap->GetNetGUIDFromObject(*It);
 			if (NetGUID.IsValid())
