ISM: New flag to allow editing properties of Packed Components but prevent editing of PerInstanceData

#rnx

diff --git a/Engine/Source/Runtime/Engine/Classes/Components/InstancedStaticMeshComponent.h b/Engine/Source/Runtime/Engine/Classes/Components/InstancedStaticMeshComponent.h
--- a/Engine/Source/Runtime/Engine/Classes/Components/InstancedStaticMeshComponent.h
+++ b/Engine/Source/Runtime/Engine/Classes/Components/InstancedStaticMeshComponent.h
@@ -153,6 +153,10 @@ class UInstancedStaticMeshComponent : public UStaticMeshComponent, public ISMIns
 	UPROPERTY(EditAnywhere, BlueprintReadOnly, Category = Culling)
 	uint8 bUseGpuLodSelection : 1;
 
+	/** If true, this component will avoid serializing its per instance data / those properties will also not be editable */
+	UPROPERTY()
+	uint8 bInheritPerInstanceData : 1;
+
 	/** Mapping from PerInstanceSMData order to instance render buffer order. If empty, the PerInstanceSMData order is used. */
 	UPROPERTY()
 	TArray<int32> InstanceReorderTable;
@@ -394,6 +398,8 @@ class UInstancedStaticMeshComponent : public UStaticMeshComponent, public ISMIns
 
 	ENGINE_API virtual bool IsInstanceTouchingSelectionBox(int32 InstanceIndex, const FBox& InBox, const bool bMustEncompassEntireInstance) const;
 	ENGINE_API virtual bool IsInstanceTouchingSelectionFrustum(int32 InstanceIndex, const FConvexVolume& InFrustum, const bool bMustEncompassEntireInstance) const;
+
+	ENGINE_API virtual bool CanEditChange(const FProperty* InProperty) const override;
 #endif
 
 	// Helper function to construct a base-set of instance data flags that in
@@ -590,6 +596,10 @@ public:
 	ENGINE_API void InvalidateInstanceDataTracking();
 
 private:
+	ENGINE_API void ApplyInheritedPerInstanceData(const UInstancedStaticMeshComponent* InArchetype);
+	ENGINE_API bool ShouldInheritPerInstanceData(const UInstancedStaticMeshComponent* InArchetype) const;
+	ENGINE_API bool ShouldInheritPerInstanceData() const;
+
 	void CalcAndCacheNavigationBounds();
 
 	/** Sets up new instance data to sensible defaults, creates physics counterparts if possible. */
diff --git a/Engine/Source/Runtime/Engine/Private/InstancedStaticMesh.cpp b/Engine/Source/Runtime/Engine/Private/InstancedStaticMesh.cpp
--- a/Engine/Source/Runtime/Engine/Private/InstancedStaticMesh.cpp
+++ b/Engine/Source/Runtime/Engine/Private/InstancedStaticMesh.cpp
@@ -1999,6 +1999,7 @@ UInstancedStaticMeshComponent::UInstancedStaticMeshComponent(const FObjectInitia
 	bMultiBodyOverlap = true;
 
 	bUseGpuLodSelection = true;
+	bInheritPerInstanceData = false;
 
 #if STATS
 	{
@@ -2087,6 +2088,13 @@ void UInstancedStaticMeshComponent::ApplyComponentInstanceData(FInstancedStaticM
 
 	bool bMatch = false;
 
+	// If we should inherit from archetype do it here after data was applied and before comparing (RerunConstructionScript will serialize SkipSerialization properties and reapply them even if we want to inherit them)
+	const UInstancedStaticMeshComponent* Archetype = Cast<UInstancedStaticMeshComponent>(GetArchetype());
+	if (ShouldInheritPerInstanceData(Archetype))
+	{
+		ApplyInheritedPerInstanceData(Archetype);
+	}
+
 	// Check for any instance having moved as that would invalidate static lighting
 	if (PerInstanceSMData.Num() == InstancedMeshData->PerInstanceSMData.Num() &&
 		InstancedMeshData->CachedStaticLighting.Transform.Equals(GetComponentTransform()))
@@ -2983,6 +2991,24 @@ void UInstancedStaticMeshComponent::SerializeRenderData(FArchive& Ar)
 	}
 }
 
+void UInstancedStaticMeshComponent::ApplyInheritedPerInstanceData(const UInstancedStaticMeshComponent* InArchetype)
+{
+	check(InArchetype);
+	PerInstanceSMData = InArchetype->PerInstanceSMData;
+	PerInstanceSMCustomData = InArchetype->PerInstanceSMCustomData;
+	NumCustomDataFloats = InArchetype->NumCustomDataFloats;
+}
+
+bool UInstancedStaticMeshComponent::ShouldInheritPerInstanceData() const
+{
+	return ShouldInheritPerInstanceData(Cast<UInstancedStaticMeshComponent>(GetArchetype()));
+}
+
+bool UInstancedStaticMeshComponent::ShouldInheritPerInstanceData(const UInstancedStaticMeshComponent* InArchetype) const
+{
+	return (bInheritPerInstanceData || !bEditableWhenInherited) && InArchetype && InArchetype->IsInBlueprint() && !IsTemplate();
+}
+
 void UInstancedStaticMeshComponent::Serialize(FArchive& Ar)
 {
 	LLM_SCOPE(ELLMTag::InstancedMesh);
@@ -2999,11 +3025,10 @@ void UInstancedStaticMeshComponent::Serialize(FArchive& Ar)
 		Ar << bCooked;
 	}
 
-	// Inherit properties when bEditableWhenInherited == true (when the component isn't a template and we are persisting data)
+	// Inherit properties when bEditableWhenInherited == false || bInheritPerInstanceData == true (when the component isn't a template and we are persisting data)
 	const UInstancedStaticMeshComponent* Archetype = Cast<UInstancedStaticMeshComponent>(GetArchetype());
-	const bool bInheritSkipSerializationProperties = !bEditableWhenInherited && Archetype && Archetype->IsInBlueprint() && Ar.IsPersistent() && !IsTemplate();
+	const bool bInheritSkipSerializationProperties = ShouldInheritPerInstanceData(Archetype) && Ar.IsPersistent();
 	
-
 	// Check if we need have SkipSerialization property data to load/save
 	bool bHasSkipSerializationPropertiesData = !bInheritSkipSerializationProperties;
 	if (Ar.IsLoading())
@@ -3053,11 +3078,7 @@ void UInstancedStaticMeshComponent::Serialize(FArchive& Ar)
 		// If we should inherit use Archetype Data
 		if (bInheritSkipSerializationProperties)
 		{
-			PerInstanceSMData = Archetype->PerInstanceSMData;
-			PerInstanceSMCustomData = Archetype->PerInstanceSMCustomData;
-
-			// Make sure that if we inherit the PerInstanceSMCustomData we also do inherit this value
-			NumCustomDataFloats = Archetype->NumCustomDataFloats;
+			ApplyInheritedPerInstanceData(Archetype);
 		} 
 		// It is possible for a component to lose its BP archetype between a save / load so in this case we have no per instance data (usually this component gets deleted through construction script)
 		else if(bHasSkipSerializationPropertiesData)
@@ -4542,6 +4563,26 @@ void UInstancedStaticMeshComponent::OnRegister()
 		
 #if WITH_EDITOR
 
+bool UInstancedStaticMeshComponent::CanEditChange(const FProperty* InProperty) const
+{
+	if (!Super::CanEditChange(InProperty))
+	{
+		return false;
+	}
+
+	if (InProperty)
+	{
+		if (InProperty->GetFName() == GET_MEMBER_NAME_CHECKED(UInstancedStaticMeshComponent, PerInstanceSMData) ||
+			InProperty->GetFName() == GET_MEMBER_NAME_CHECKED(UInstancedStaticMeshComponent, PerInstanceSMCustomData) ||
+			InProperty->GetFName() == GET_MEMBER_NAME_CHECKED(UInstancedStaticMeshComponent, NumCustomDataFloats))
+		{
+			return !ShouldInheritPerInstanceData();
+		}
+	}
+
+	return true;
+}
+
 bool UInstancedStaticMeshComponent::ComponentIsTouchingSelectionBox(const FBox& InSelBBox, const bool bConsiderOnlyBSP, const bool bMustEncompassEntireComponent) const
 {
 	if (bConsiderOnlyBSP)
diff --git a/Engine/Source/Runtime/Engine/Private/PackedLevelActor/PackedLevelActorBuilder.cpp b/Engine/Source/Runtime/Engine/Private/PackedLevelActor/PackedLevelActorBuilder.cpp
--- a/Engine/Source/Runtime/Engine/Private/PackedLevelActor/PackedLevelActorBuilder.cpp
+++ b/Engine/Source/Runtime/Engine/Private/PackedLevelActor/PackedLevelActorBuilder.cpp
@@ -34,6 +34,8 @@
 #include "WorldPartition/Filter/WorldPartitionActorFilter.h"
 #include "Serialization/ArchiveCrc32.h"
 
+#include "Components/InstancedStaticMeshComponent.h"
+
 
 #define LOCTEXT_NAMESPACE "FPackedLevelActorBuilder"
 
@@ -204,6 +206,10 @@ bool FPackedLevelActorBuilder::PackActor(FPackedLevelActorBuilderContext& InCont
 	}
 	HashArray.Sort();
 	FArchiveCrc32 Ar;
+
+	// Change this Guid if existing hashes need to be invalidated (code changes that aren't part of the computed hash)
+	static FGuid CodeVersionGuid = FGuid(0xE023D897, 0x4F404333, 0xAD604BB9, 0x8CC6B959);
+	Ar << CodeVersionGuid;
 	Ar << HashArray;
 	InContext.GetPackedLevelActor()->SetPackedHash(Ar.GetCrc());
 	InContext.Report(Log);
@@ -454,6 +460,8 @@ bool FPackedLevelActorBuilder::CreateOrUpdateBlueprintFromPacked(APackedLevelAct
 		
 	// Avoid running construction script while dragging an instance of that BP for performance reasons
 	BP->bRunConstructionScriptOnDrag = false;
+
+	// New Guid everytime we pack a BP so that we can avoid running construction scripts in PIE for Instances that have a matching version
 	FGuid NewVersion = FGuid::NewGuid();
 	APackedLevelActor* CDO = CastChecked<APackedLevelActor>(BP->GeneratedClass->GetDefaultObject());
 			
@@ -466,10 +474,13 @@ bool FPackedLevelActorBuilder::CreateOrUpdateBlueprintFromPacked(APackedLevelAct
 	TArray<UActorComponent*> PackedComponents;
 	InActor->GetPackedComponents(PackedComponents);
 	
-	// To avoid any delta serialization happening on those generated components, we make them non editable.
 	for (UActorComponent* PackedComponent : PackedComponents)
 	{
-		PackedComponent->bEditableWhenInherited = false;
+		// To avoid any custom serialization on PerInstance data we make PerInstance data inherited always
+		if (UInstancedStaticMeshComponent* ISMComponent = Cast<UInstancedStaticMeshComponent>(PackedComponent))
+		{
+			ISMComponent->bInheritPerInstanceData = true;
+		}
 	}
 
 	FKismetEditorUtilities::AddComponentsToBlueprint(BP, PackedComponents, AddCompToBPParams);
