GAllowActorScriptExecutionInEditor must be false when we call BeginPlay
Replace previous fix in
#rnx

diff --git a/Engine/Source/Editor/UnrealEd/Private/Kismet2/KismetReinstanceUtilities.cpp b/Engine/Source/Editor/UnrealEd/Private/Kismet2/KismetReinstanceUtilities.cpp
--- a/Engine/Source/Editor/UnrealEd/Private/Kismet2/KismetReinstanceUtilities.cpp
+++ b/Engine/Source/Editor/UnrealEd/Private/Kismet2/KismetReinstanceUtilities.cpp
@@ -1455,6 +1455,8 @@ void FActorReplacementHelper::Finalize(const TMap<UObject*, UObject*>& OldToNewI
 	{
 		if (World->HasBegunPlay())
 		{
+			//GAllowActorScriptExecutionInEditor must be false when we call BeginPlay
+			TGuardValue AutoRestore(GAllowActorScriptExecutionInEditor, false);
 			NewActor->PostActorConstruction();
 		}
 	}
diff --git a/Engine/Source/Runtime/Engine/Private/CheatManager.cpp b/Engine/Source/Runtime/Engine/Private/CheatManager.cpp
--- a/Engine/Source/Runtime/Engine/Private/CheatManager.cpp
+++ b/Engine/Source/Runtime/Engine/Private/CheatManager.cpp
@@ -123,8 +123,10 @@ bool UCheatManager::ProcessConsoleExec(const TCHAR* Cmd, FOutputDevice& Ar, UObj
 
 				if ((Function != nullptr) && Function->HasAnyFunctionFlags(FUNC_BlueprintAuthorityOnly))
 				{
-					TGuardValue AutoRestore(GAllowActorScriptExecutionInEditor, false);
-					MyPC->ServerExec(Cmd);
+					if(ensureMsgf(GAllowActorScriptExecutionInEditor == false, TEXT("GAllowActorScriptExecutionInEditor must be false when executing commands.")))
+					{
+						MyPC->ServerExec(Cmd);
+					}
 					return true;
 				}
 			}
