Material Editor - Fix build error (regarding undeclared SerializeTime).

diff --git a/Engine/Source/Runtime/Engine/Private/Materials/HLSLMaterialTranslator.cpp b/Engine/Source/Runtime/Engine/Private/Materials/HLSLMaterialTranslator.cpp
--- a/Engine/Source/Runtime/Engine/Private/Materials/HLSLMaterialTranslator.cpp
+++ b/Engine/Source/Runtime/Engine/Private/Materials/HLSLMaterialTranslator.cpp
@@ -25,7 +25,7 @@
 #include "Materials/MaterialExpressionAbsorptionMediumMaterialOutput.h"
 #include "Materials/MaterialExpressionCustomOutput.h"
 #include "Materials/MaterialExpressionFunctionInput.h"
-#include "Materials/MaterialExpressionFunctionOutput.h"
+#include "Materials/MaterialExpressionFunctionOutput.h" 
 #include "Materials/MaterialExpressionMaterialFunctionCall.h"
 #include "Materials/MaterialExpressionNoise.h"
 #include "Materials/MaterialExpressionSingleLayerWaterMaterialOutput.h"
@@ -208,15 +208,19 @@ static const bool GUseMaterialTranslationResultsGrouping = true;
 /* Controls whether DDC caching of material translation results should be forcefully disabled. */
 #define FORCE_DISABLE_MATERIAL_TRANSLATION_DDC true
 
+/* Helper macro to check whether the DDC material translation data has arrived and
+ * therefore quit material translation. */
+#if FORCE_DISABLE_MATERIAL_TRANSLATION_DDC
+#define CHECK_DDC_QUERY_FINISHED_ELSE_RETURN()
+#else
+#define CHECK_DDC_QUERY_FINISHED_ELSE_RETURN() if (DDCQueryHit) { return; }
+#endif
+
 UE::DerivedData::FCacheBucket MaterialTranslationDDCBucket = UE::DerivedData::FCacheBucket(TEXT("MaterialTranslation"));
 UE::DerivedData::FValueId MaterialCompilationOutputId = UE::DerivedData::FValueId::FromName("FHLSLMaterialTranslator_MaterialCompilationOutput");
 UE::DerivedData::FValueId MaterialResultsOutputId = UE::DerivedData::FValueId::FromName("FHLSLMaterialTranslator_Results");
 UE::DerivedData::FValueId EnvironmentDefinesId = UE::DerivedData::FValueId::FromName("FHLSLMaterialTranslator_EnvironmentDefines");
 
-/* Helper macro to check whether the DDC material translation data has arrived and
- * therefore quit material translation. */
-#define CHECK_DDC_QUERY_FINISHED_ELSE_RETURN() if (DDCQueryHit) { return; }
-
 /** Data structure used to cache a part of material translation results. It contains all the generated
  *  defines that will be declared during the compilation of the generated material shader.
  *  Note: if you make a change to this structure, remember to FDevSystemGuids::MaterialTranslationDDCVersion.
@@ -1128,7 +1132,6 @@ bool FHLSLMaterialTranslator::Translate()
 {
 	STAT(FDateTime TranslationDateTime = FDateTime::Now());
 	STAT(double TotalTime = FPlatformTime::Seconds());
-	STAT(double SerializeTime = 0);
 	STAT(double TranslationOnlyTime = 0);
 
 #if CPUPROFILERTRACE_ENABLED
@@ -1177,7 +1180,7 @@ bool FHLSLMaterialTranslator::Translate()
 	// Asynchronously query the DDC for results
 	if (!bDisableTranslationDDC)
 	{
-		AsyncQueryDDC(DDCRequestOwner, EnvironmentDefinesBuffer, SerializeTime);
+		AsyncQueryDDC(DDCRequestOwner, EnvironmentDefinesBuffer);
 	}
 
 	// Synchronously begin translating the material
@@ -1197,8 +1200,6 @@ bool FHLSLMaterialTranslator::Translate()
 			UE_LOG(LogMaterial, Display, TEXT("Material '%s' translation results were in the DDC"), *Material->GetMaterialInterface()->GetFullName());
 		}
 
-		STAT(double SerializeBegin = FPlatformTime::Seconds());
-		
 		// Serialize the environment defines from the buffer retrieved from the DDC.
 		FMemoryReaderView EnvironmentDefinesBufferReader{ TArrayView<uint8>{ (uint8*)EnvironmentDefinesBuffer.GetData(), (int)EnvironmentDefinesBuffer.GetSize() } };
 		FObjectAndNameAsStringProxyArchive EnvironmentDefinesBufferReaderProxy{ EnvironmentDefinesBufferReader, true };
@@ -1206,7 +1207,6 @@ bool FHLSLMaterialTranslator::Translate()
 
 		MaterialCompilationOutput = DDCMaterialCompilationOutput;
 
-		STAT(SerializeTime += FPlatformTime::Seconds() - SerializeBegin);
 		STAT(GShaderCompilerStats->IncrementMaterialCacheHit());
 		bSuccess = true;
 	}
@@ -1232,7 +1232,7 @@ bool FHLSLMaterialTranslator::Translate()
 	// Report timings to Material Cook Stats
 #if STATS
 	TotalTime = FPlatformTime::Seconds() - TotalTime;
-	GShaderCompilerStats->IncrementMaterialTranslated(TotalTime, TranslationOnlyTime, SerializeTime);
+	GShaderCompilerStats->IncrementMaterialTranslated(TotalTime, TranslationOnlyTime, 0 /* todo */);
 #endif
 
 	INC_FLOAT_STAT_BY(STAT_ShaderCompiling_HLSLTranslation, (float)TotalTime);
@@ -15411,8 +15411,7 @@ void FHLSLMaterialTranslator::PrepareEnvironmentDefines()
 
 void FHLSLMaterialTranslator::AsyncQueryDDC(
 	UE::DerivedData::FRequestOwner& DDCRequestOwner,
-	FSharedBuffer& EnvironmentDefinesBuffer,
-	double& SerializeTime)
+	FSharedBuffer& EnvironmentDefinesBuffer)
 {
 	UE::DerivedData::FCacheKey CacheKey{ MaterialTranslationDDCBucket, DDCKeyHash };
 	UE::DerivedData::FCacheGetRequest Request;
@@ -15435,7 +15434,6 @@ void FHLSLMaterialTranslator::AsyncQueryDDC(
 				// Load the results if we hit the cache.
 				if (!MaterialCompilationOutputBuffer.IsNull() && !TranslationResultsBuffer.IsNull() && !EnvironmentDefinesBuffer.IsNull())
 				{
-					STAT(double SerializeBegin = FPlatformTime::Seconds());
 					// Read the material compilation output
 					FShaderMapPointerTable PointerTable;
 					FPlatformTypeLayoutParameters LayoutParams;
@@ -15448,7 +15446,6 @@ void FHLSLMaterialTranslator::AsyncQueryDDC(
 					FMemoryReaderView ResultsMemoryReader{ TArrayView<uint8>{ (uint8*)TranslationResultsBuffer.GetData(), (int)TranslationResultsBuffer.GetSize() } };
 					ResultsMemoryReader << MaterialSourceTemplateParams;
 
-					STAT(SerializeTime += FPlatformTime::Seconds() - SerializeBegin);
 					DDCQueryHit = true;
 				}
 			}
diff --git a/Engine/Source/Runtime/Engine/Private/Materials/HLSLMaterialTranslator.h b/Engine/Source/Runtime/Engine/Private/Materials/HLSLMaterialTranslator.h
--- a/Engine/Source/Runtime/Engine/Private/Materials/HLSLMaterialTranslator.h
+++ b/Engine/Source/Runtime/Engine/Private/Materials/HLSLMaterialTranslator.h
@@ -1389,7 +1389,7 @@ protected:
 	 * Queries the DDC cache for a cached translation.
 	 * @return Whether the speecified key is in the DDC.
 	 */
-	void AsyncQueryDDC(UE::DerivedData::FRequestOwner& DDCRequestOwner, FSharedBuffer& EnvironmentDefinesBuffer, double& SerializeTime);
+	void AsyncQueryDDC(UE::DerivedData::FRequestOwner& DDCRequestOwner, FSharedBuffer& EnvironmentDefinesBuffer);
 
 	/**
 	 * Pushes the final results to the DDC cache.
