Iris
- Added GetNetDebugName method to ReplicationSystem
- Added name to ensure when failing to write creation header

diff --git a/Engine/Source/Runtime/Engine/Private/Net/Iris/ReplicationSystem/ActorReplicationBridge.cpp b/Engine/Source/Runtime/Engine/Private/Net/Iris/ReplicationSystem/ActorReplicationBridge.cpp
--- a/Engine/Source/Runtime/Engine/Private/Net/Iris/ReplicationSystem/ActorReplicationBridge.cpp
+++ b/Engine/Source/Runtime/Engine/Private/Net/Iris/ReplicationSystem/ActorReplicationBridge.cpp
@@ -35,6 +35,7 @@
 #include "Net/DataChannel.h"
 #include "Net/Core/Connection/NetCloseResult.h"
 #include "Net/Core/Misc/NetSubObjectRegistry.h"
+#include "Net/Core/Trace/NetDebugName.h"
 #include "Net/NetSubObjectRegistryGetter.h"
 #include "ProfilingDebugging/AssetMetadataTrace.h"
 #include "Templates/Casts.h"
@@ -508,7 +509,7 @@ bool UActorReplicationBridge::WriteCreationHeader(UE::Net::FNetSerializationCont
 		return !Writer->IsOverflown();
 	}
 
-	ensureMsgf(false, TEXT("UActorReplicationBridge::WriteCreationHeader Failed to write creationHeader for NetRefHandle (Id=%u)"), Handle.GetId());
+	ensureMsgf(false, TEXT("UActorReplicationBridge::WriteCreationHeader Failed to write creationHeader for NetRefHandle (Id=%u) %s"), Handle.GetId(), ToCStr(GetReplicationSystem()->GetDebugName(Handle)));
 
 	return false;
 }
diff --git a/Engine/Source/Runtime/Experimental/Iris/Core/Private/Iris/ReplicationSystem/ReplicationSystem.cpp b/Engine/Source/Runtime/Experimental/Iris/Core/Private/Iris/ReplicationSystem/ReplicationSystem.cpp
--- a/Engine/Source/Runtime/Experimental/Iris/Core/Private/Iris/ReplicationSystem/ReplicationSystem.cpp
+++ b/Engine/Source/Runtime/Experimental/Iris/Core/Private/Iris/ReplicationSystem/ReplicationSystem.cpp
@@ -1025,6 +1025,14 @@ const UE::Net::FReplicationProtocol* UReplicationSystem::GetReplicationProtocol(
 	return NetRefHandleManager.GetReplicatedObjectDataNoCheck(ObjectInternalIndex).Protocol;
 }
 
+const UE::Net::FNetDebugName* UReplicationSystem::GetDebugName(FNetRefHandle Handle) const
+{
+	using namespace UE::Net;
+	
+	const FReplicationProtocol* Protocol = GetReplicationProtocol(Handle);
+	return Protocol ? Protocol->DebugName : nullptr;
+}
+
 void UReplicationSystem::SetOwningNetConnection(FNetRefHandle Handle, uint32 ConnectionId)
 {
 	using namespace UE::Net::Private;
diff --git a/Engine/Source/Runtime/Experimental/Iris/Core/Public/Iris/ReplicationSystem/ReplicationSystem.h b/Engine/Source/Runtime/Experimental/Iris/Core/Public/Iris/ReplicationSystem/ReplicationSystem.h
--- a/Engine/Source/Runtime/Experimental/Iris/Core/Public/Iris/ReplicationSystem/ReplicationSystem.h
+++ b/Engine/Source/Runtime/Experimental/Iris/Core/Public/Iris/ReplicationSystem/ReplicationSystem.h
@@ -39,6 +39,7 @@ namespace UE::Net
 	struct FReplicationView;
 	class FStringTokenStore;
 	class FWorldLocations;
+	struct FNetDebugName;
 	namespace Private
 	{
 		class FReplicationSystemImpl;
@@ -289,6 +290,13 @@ public:
 	 */
 	IRISCORE_API const UE::Net::FReplicationProtocol* GetReplicationProtocol(FNetRefHandle Handle) const;
 
+	/**
+	 * Get the DebugName associated with a handle.
+	 * @param Handle The handle to retrieve the DebugName for
+	 * @return DebugName if the handle is valid, otherwise nullptr.
+	 */
+	IRISCORE_API const UE::Net::FNetDebugName* GetDebugName(FNetRefHandle Handle) const;
+
 	// Groups
 
 	/**
