Fix Modular game build CIS error.

diff --git a/Engine/Source/Runtime/Engine/Public/PSOPrecacheValidation.h b/Engine/Source/Runtime/Engine/Public/PSOPrecacheValidation.h
--- a/Engine/Source/Runtime/Engine/Public/PSOPrecacheValidation.h
+++ b/Engine/Source/Runtime/Engine/Public/PSOPrecacheValidation.h
@@ -249,7 +249,7 @@ namespace PSOCollectorStats
 			return UpdatePrecacheStats(PrecacheStateHash, PSOCollectorIndex, VertexFactoryType, bTracked, PrecacheResult);
 		}
 
-		bool IsPrecached(uint64 PrecacheStateHash);
+		ENGINE_API bool IsPrecached(uint64 PrecacheStateHash);
 #if PSO_PRECACHING_TRACKING
 		bool GetPrecacheData(uint64 PrecacheStateHash, FString& OutMaterialName, int32& OutPSOCollectorIndex, const FVertexFactoryType*& OutVertexFactoryType);
 #endif
