Iris - fixed cis error

#rnx

diff --git a/Engine/Source/Runtime/Engine/Private/Net/Iris/ReplicationSystem/ActorReplicationBridge.cpp b/Engine/Source/Runtime/Engine/Private/Net/Iris/ReplicationSystem/ActorReplicationBridge.cpp
--- a/Engine/Source/Runtime/Engine/Private/Net/Iris/ReplicationSystem/ActorReplicationBridge.cpp
+++ b/Engine/Source/Runtime/Engine/Private/Net/Iris/ReplicationSystem/ActorReplicationBridge.cpp
@@ -435,7 +435,7 @@ void UActorReplicationBridge::EndReplication(AActor* Actor, EEndPlayReason::Type
 		}
 
 		// If we are shutting down we do not need to validate that we are not detaching remote instances by accident.
-		if (EndPlayReason == EEndPlayReason::EndPlayInEditor || EndPlayReason != EEndPlayReason::Quit)
+		if ((EndPlayReason == EEndPlayReason::EndPlayInEditor) || (EndPlayReason == EEndPlayReason::Quit))
 		{
 			Flags |= EEndReplicationFlags::SkipPendingEndReplicationValidation;
 		}
