Lumen - fixed Lumen GI screen space category in PPV

diff --git a/Engine/Source/Runtime/Engine/Classes/Engine/Scene.h b/Engine/Source/Runtime/Engine/Classes/Engine/Scene.h
--- a/Engine/Source/Runtime/Engine/Classes/Engine/Scene.h
+++ b/Engine/Source/Runtime/Engine/Classes/Engine/Scene.h
@@ -1638,7 +1638,7 @@ struct FPostProcessSettings
 	float LumenFinalGatherLightingUpdateSpeed;
 
 	/** Whether to use screen space traces for Lumen Global Illumination. Screen space traces bypass Lumen Scene and instead sample Scene Depth and Color. This improves quality, but at the same time prevents from Lumen Scene only changes like adding emissive objects, which are visible only in Global Illumination. */
-	UPROPERTY(interp, EditAnywhere, BlueprintReadWrite, Category = "Reflections|Lumen Global Illumination", meta = (editcondition = "bOverride_LumenFinalGatherScreenTraces", DisplayName = "Screen Traces"))
+	UPROPERTY(interp, EditAnywhere, BlueprintReadWrite, Category = "Global Illumination|Lumen Global Illumination", meta = (editcondition = "bOverride_LumenFinalGatherScreenTraces", DisplayName = "Screen Traces"))
 	uint8 LumenFinalGatherScreenTraces : 1;
 
 	/** Controls the maximum distance that Lumen should trace while solving lighting.  Values that are too small will cause lighting to leak into large caves, while values that are large will increase GPU cost. */
