Commenting out check that relies on PSOCollectorStats::IsFullPrecachingValidationEnabled() being enabled. This was firing 100% of the time in cooked Lyra builds.

diff --git a/Engine/Source/Runtime/Engine/Private/Materials/MaterialShader.cpp b/Engine/Source/Runtime/Engine/Private/Materials/MaterialShader.cpp
--- a/Engine/Source/Runtime/Engine/Private/Materials/MaterialShader.cpp
+++ b/Engine/Source/Runtime/Engine/Private/Materials/MaterialShader.cpp
@@ -2758,7 +2758,7 @@ FPSOPrecacheDataArray FMaterialShaderMap::CollectPSOPrecacheData(const FMaterial
 			IPSOCollector* PSOCollector = CreateFunction(PrecacheParams.FeatureLevel);
 			if (PSOCollector != nullptr)
 			{
-				check(PSOCollector->PSOCollectorIndex == Index);
+				//check(PSOCollector->PSOCollectorIndex == Index);
 
 #if PSO_PRECACHING_VALIDATE
 				ConditionalBreakOnPSOPrecacheMaterial(*PrecacheParams.Material, Index);
