Fixing include issues in PSOPrecache

diff --git a/Engine/Source/Runtime/Engine/Private/PSOPrecache.cpp b/Engine/Source/Runtime/Engine/Private/PSOPrecache.cpp
--- a/Engine/Source/Runtime/Engine/Private/PSOPrecache.cpp
+++ b/Engine/Source/Runtime/Engine/Private/PSOPrecache.cpp
@@ -6,6 +6,7 @@
 
 #include "PSOPrecache.h"
 #include "Misc/App.h"
+#include "HAL/IConsoleManager.h"
 
 static TAutoConsoleVariable<int32> CVarPrecacheGlobalComputeShaders(
 	TEXT("r.PSOPrecache.GlobalComputeShaders"),
diff --git a/Engine/Source/Runtime/Engine/Public/PSOPrecache.h b/Engine/Source/Runtime/Engine/Public/PSOPrecache.h
--- a/Engine/Source/Runtime/Engine/Public/PSOPrecache.h
+++ b/Engine/Source/Runtime/Engine/Public/PSOPrecache.h
@@ -6,9 +6,8 @@
 
 #pragma once
 
-#include "PSOPrecache.h"
-
 #include "RHIDefinitions.h"
+#include "RHIFeatureLevel.h"
 #include "RHIResources.h"
 #include "Engine/EngineTypes.h"
 #include "PipelineStateCache.h"
@@ -273,4 +272,4 @@ extern ENGINE_API EPSOPrecacheProxyCreationStrategy GetPSOPrecacheProxyCreationS
 extern ENGINE_API bool ProxyCreationWhenPSOReady();
 
 // Unique request ID of MaterialPSOPrecache which can be used to boost the priority of a PSO precache requests if it's needed for rendering
-using FMaterialPSOPrecacheRequestID = uint32;
\ No newline at end of file
+using FMaterialPSOPrecacheRequestID = uint32;
