GetAssetRegistryTags deprecation: Make the old GetAssetRegistryTags and related functions deprecated.

diff --git a/Engine/Source/Runtime/CoreUObject/Public/UObject/Class.h b/Engine/Source/Runtime/CoreUObject/Public/UObject/Class.h
--- a/Engine/Source/Runtime/CoreUObject/Public/UObject/Class.h
+++ b/Engine/Source/Runtime/CoreUObject/Public/UObject/Class.h
@@ -3037,6 +3037,8 @@ public:
 	static COREUOBJECT_API void AddReferencedObjects(UObject* InThis, FReferenceCollector& Collector);
 	COREUOBJECT_API virtual FRestoreForUObjectOverwrite* GetRestoreForUObjectOverwrite() override;
 	COREUOBJECT_API virtual FString GetDesc() override;
+	COREUOBJECT_API virtual void GetAssetRegistryTags(FAssetRegistryTagsContext Context) const override;
+	UE_DEPRECATED(5.4, "Implement the version that takes FAssetRegistryTagsContext instead.")
 	COREUOBJECT_API virtual void GetAssetRegistryTags(TArray<FAssetRegistryTag>& OutTags) const override;
 #if WITH_EDITOR
 	COREUOBJECT_API virtual void PostLoadAssetRegistryTags(const FAssetData& InAssetData, TArray<FAssetRegistryTag>& OutTagsAndValuesToUpdate) const;
diff --git a/Engine/Source/Runtime/CoreUObject/Public/UObject/Object.h b/Engine/Source/Runtime/CoreUObject/Public/UObject/Object.h
--- a/Engine/Source/Runtime/CoreUObject/Public/UObject/Object.h
+++ b/Engine/Source/Runtime/CoreUObject/Public/UObject/Object.h
@@ -857,12 +857,14 @@ public:
 
 #if WITH_EDITOR
 		/** Event for listeners who want to add tags to some UObjects' GetAssetRegistryTags. */
-		DECLARE_MULTICAST_DELEGATE_OneParam(FOnGetObjectAssetRegistryTagsWithContext, FAssetRegistryTagsContext&);
+		DECLARE_MULTICAST_DELEGATE_OneParam(FOnGetObjectAssetRegistryTagsWithContext, FAssetRegistryTagsContext);
 		COREUOBJECT_API static FOnGetObjectAssetRegistryTagsWithContext OnGetExtraObjectTagsWithContext;
 
 		DECLARE_MULTICAST_DELEGATE_TwoParams(FOnGetObjectAssetRegistryTags, const UObject* /*Object*/, TArray<FAssetRegistryTag>& /*InOutTags*/);
+		UE_DEPRECATED(5.4, "Subscribe to OnGetExtraObjectTagsWithContext instead")
 		COREUOBJECT_API static FOnGetObjectAssetRegistryTags OnGetExtraObjectTags;
 		DECLARE_MULTICAST_DELEGATE_ThreeParams(FOnGetExtendedAssetRegistryTagsForSave, const UObject* /*Object*/, const ITargetPlatform* TargetPlatform, TArray<FAssetRegistryTag>& /*InOutTags*/);
+		UE_DEPRECATED(5.4, "Subscribe to OnGetExtraObjectTagsWithContext instead, and early exit if !Context.IsSaving")
 		COREUOBJECT_API static FOnGetExtendedAssetRegistryTagsForSave OnGetExtendedAssetRegistryTagsForSave;
 #endif // WITH_EDITOR
 	};
@@ -873,8 +875,9 @@ public:
 	 *
 	 * @param	OutTags		A list of key-value pairs associated with this object and their types
 	 */
-	COREUOBJECT_API virtual void GetAssetRegistryTags(TArray<FAssetRegistryTag>& OutTags) const;
 	COREUOBJECT_API virtual void GetAssetRegistryTags(FAssetRegistryTagsContext Context) const;
+	UE_DEPRECATED(5.4, "Implement the version that takes FAssetRegistryTagsContext instead.")
+	COREUOBJECT_API virtual void GetAssetRegistryTags(TArray<FAssetRegistryTag>& OutTags) const;
 
 #if WITH_EDITOR
 	/**
diff --git a/Engine/Source/Runtime/Engine/Classes/Animation/AnimSequence.h b/Engine/Source/Runtime/Engine/Classes/Animation/AnimSequence.h
--- a/Engine/Source/Runtime/Engine/Classes/Animation/AnimSequence.h
+++ b/Engine/Source/Runtime/Engine/Classes/Animation/AnimSequence.h
@@ -410,6 +410,8 @@ public:
 #endif // WITH_EDITOR
 	ENGINE_API virtual void BeginDestroy() override;
 	ENGINE_API virtual bool IsReadyForFinishDestroy() override;
+	ENGINE_API virtual void GetAssetRegistryTags(FAssetRegistryTagsContext Context) const override;
+	UE_DEPRECATED(5.4, "Implement the version that takes FAssetRegistryTagsContext instead.")
 	ENGINE_API virtual void GetAssetRegistryTags(TArray<FAssetRegistryTag>& OutTags) const override;
 	static ENGINE_API void AddReferencedObjects(UObject* This, FReferenceCollector& Collector);
 	//~ End UObject Interface
diff --git a/Engine/Source/Runtime/Engine/Classes/Animation/AnimSequenceBase.h b/Engine/Source/Runtime/Engine/Classes/Animation/AnimSequenceBase.h
--- a/Engine/Source/Runtime/Engine/Classes/Animation/AnimSequenceBase.h
+++ b/Engine/Source/Runtime/Engine/Classes/Animation/AnimSequenceBase.h
@@ -179,8 +179,10 @@ public:
 	/** Calculates what (if any) offset should be applied to the trigger time of a notify given its display time */ 
 	ENGINE_API virtual EAnimEventTriggerOffsets::Type CalculateOffsetForNotify(float NotifyDisplayTime) const;
 
+	ENGINE_API virtual void GetAssetRegistryTags(FAssetRegistryTagsContext Context) const override;
+	UE_DEPRECATED(5.4, "Implement the version that takes FAssetRegistryTagsContext instead.")
 	ENGINE_API virtual void GetAssetRegistryTags(TArray<FAssetRegistryTag>& OutTags) const override;
-	
+
 	// Get a pointer to the data for a given Anim Notify
 	ENGINE_API uint8* FindNotifyPropertyData(int32 NotifyIndex, FArrayProperty*& ArrayProperty);
 
diff --git a/Engine/Source/Runtime/Engine/Classes/Animation/AnimationAsset.h b/Engine/Source/Runtime/Engine/Classes/Animation/AnimationAsset.h
--- a/Engine/Source/Runtime/Engine/Classes/Animation/AnimationAsset.h
+++ b/Engine/Source/Runtime/Engine/Classes/Animation/AnimationAsset.h
@@ -1155,6 +1155,8 @@ public:
 	//~ Begin UObject Interface.
 #if WITH_EDITOR
 	ENGINE_API virtual void PostEditChangeProperty(FPropertyChangedEvent& PropertyChangedEvent) override;
+	ENGINE_API virtual void GetAssetRegistryTags(FAssetRegistryTagsContext Context) const override;
+	UE_DEPRECATED(5.4, "Implement the version that takes FAssetRegistryTagsContext instead.")
 	ENGINE_API virtual void GetAssetRegistryTags(TArray<FAssetRegistryTag>& OutTags) const override;
 	ENGINE_API virtual EDataValidationResult IsDataValid(class FDataValidationContext& Context) const override;
 #endif // WITH_EDITOR
diff --git a/Engine/Source/Runtime/Engine/Classes/Animation/PoseAsset.h b/Engine/Source/Runtime/Engine/Classes/Animation/PoseAsset.h
--- a/Engine/Source/Runtime/Engine/Classes/Animation/PoseAsset.h
+++ b/Engine/Source/Runtime/Engine/Classes/Animation/PoseAsset.h
@@ -255,8 +255,10 @@ public:
 	virtual void PreSave(const class ITargetPlatform* TargetPlatform) override;
 	PRAGMA_ENABLE_DEPRECATION_WARNINGS
 	virtual void PreSave(FObjectPreSaveContext ObjectSaveContext) override;
+	virtual void GetAssetRegistryTags(FAssetRegistryTagsContext Context) const override;
+	UE_DEPRECATED(5.4, "Implement the version that takes FAssetRegistryTagsContext instead.")
 	virtual void GetAssetRegistryTags(TArray<FAssetRegistryTag>& OutTags) const override;
-	//End UObject Interface
+//End UObject Interface
 
 public:
 	ENGINE_API int32 GetNumPoses() const;
diff --git a/Engine/Source/Runtime/Engine/Classes/EditorFramework/AssetImportData.h b/Engine/Source/Runtime/Engine/Classes/EditorFramework/AssetImportData.h
--- a/Engine/Source/Runtime/Engine/Classes/EditorFramework/AssetImportData.h
+++ b/Engine/Source/Runtime/Engine/Classes/EditorFramework/AssetImportData.h
@@ -3,6 +3,7 @@
 #pragma once
 
 #include "CoreMinimal.h"
+#include "UObject/AssetRegistryTagsContext.h"
 #include "UObject/ObjectMacros.h"
 #include "UObject/Object.h"
 #include "Misc/SecureHash.h"
@@ -129,9 +130,11 @@ public:
 	ENGINE_API void SetSourceFiles(TArray<FAssetImportInfo::FSourceFile>&& SourceFiles);
 
 #if WITH_EDITOR
+	UE_DEPRECATED(5.4, "Implement the version that takes FAssetRegistryTagsContext instead.")
+	virtual void AppendAssetRegistryTags(TArray<FAssetRegistryTag>& OutTags) {}
 	/** If your asset import data flavor need to add some asset registry tag, override this function. */
-	virtual void AppendAssetRegistryTags(TArray<FAssetRegistryTag>& OutTags){}
-	
+	virtual void AppendAssetRegistryTags(FAssetRegistryTagsContext Context) {}
+
 	/** Helper function to return the first filename stored in this data. The resulting filename will be absolute (ie, not relative to the asset).  */
 	UFUNCTION(BlueprintCallable, meta=(DisplayName="GetFirstFilename", ScriptName="GetFirstFilename"), Category="AssetImportData")
 	ENGINE_API FString K2_GetFirstFilename() const;
diff --git a/Engine/Source/Runtime/Engine/Classes/Engine/Blueprint.h b/Engine/Source/Runtime/Engine/Classes/Engine/Blueprint.h
--- a/Engine/Source/Runtime/Engine/Classes/Engine/Blueprint.h
+++ b/Engine/Source/Runtime/Engine/Classes/Engine/Blueprint.h
@@ -927,10 +927,9 @@ public:
 	static ENGINE_API void DeclareConstructClasses(TArray<FTopLevelAssetPath>& OutConstructClasses, const UClass* SpecificSubclass);
 #endif
 	ENGINE_API virtual bool Modify(bool bAlwaysMarkDirty = true) override;
+	ENGINE_API virtual void GetAssetRegistryTags(FAssetRegistryTagsContext Context) const override;
+	UE_DEPRECATED(5.4, "Implement the version that takes FAssetRegistryTagsContext instead.")
 	ENGINE_API virtual void GetAssetRegistryTags(TArray<FAssetRegistryTag>& OutTags) const override;
-#if WITH_EDITOR
-	ENGINE_API virtual void GetExtendedAssetRegistryTagsForSave(const ITargetPlatform* TargetPlatform, TArray<FAssetRegistryTag>& OutTags) const override;
-#endif
 	ENGINE_API virtual void PostLoadAssetRegistryTags(const FAssetData& InAssetData, TArray<FAssetRegistryTag>& OutTagsAndValuesToUpdate) const;
 	static ENGINE_API void PostLoadBlueprintAssetRegistryTags(const FAssetData& InAssetData, TArray<FAssetRegistryTag>& OutTagsAndValuesToUpdate);
 	ENGINE_API virtual FPrimaryAssetId GetPrimaryAssetId() const override;
diff --git a/Engine/Source/Runtime/Engine/Classes/Engine/Texture2DArray.cpp b/Engine/Source/Runtime/Engine/Classes/Engine/Texture2DArray.cpp
--- a/Engine/Source/Runtime/Engine/Classes/Engine/Texture2DArray.cpp
+++ b/Engine/Source/Runtime/Engine/Classes/Engine/Texture2DArray.cpp
@@ -21,6 +21,7 @@
 #include "Streaming/Texture2DMipDataProvider_IO.h"
 #include "Streaming/TextureStreamOut.h"
 #include "TextureCompiler.h"
+#include "UObject/AssetRegistryTagsContext.h"
 #include "UObject/Package.h"
 #include "UObject/StrongObjectPtr.h"
 #include "ImageCoreUtils.h"
@@ -542,6 +543,13 @@ void UTexture2DArray::PostLoad()
 };
 
 void UTexture2DArray::GetAssetRegistryTags(TArray<FAssetRegistryTag>& OutTags) const
+{
+	PRAGMA_DISABLE_DEPRECATION_WARNINGS;
+	Super::GetAssetRegistryTags(OutTags);
+	PRAGMA_ENABLE_DEPRECATION_WARNINGS;
+}
+
+void UTexture2DArray::GetAssetRegistryTags(FAssetRegistryTagsContext Context) const
 {
 #if WITH_EDITOR
 	int32 SizeX = Source.GetSizeX();
@@ -553,10 +561,10 @@ void UTexture2DArray::GetAssetRegistryTags(TArray<FAssetRegistryTag>& OutTags) c
 	int32 ArraySize = 0;
 #endif
 	const FString Dimensions = FString::Printf(TEXT("%dx%d*%d"), SizeX, SizeY, ArraySize);
-	OutTags.Add(FAssetRegistryTag("Dimensions", Dimensions, FAssetRegistryTag::TT_Dimensional));
-	OutTags.Add(FAssetRegistryTag("Format", GPixelFormats[GetPixelFormat()].Name, FAssetRegistryTag::TT_Alphabetical));
+	Context.AddTag(FAssetRegistryTag("Dimensions", Dimensions, FAssetRegistryTag::TT_Dimensional));
+	Context.AddTag(FAssetRegistryTag("Format", GPixelFormats[GetPixelFormat()].Name, FAssetRegistryTag::TT_Alphabetical));
 
-	Super::GetAssetRegistryTags(OutTags);
+	Super::GetAssetRegistryTags(Context);
 }
 
 FString UTexture2DArray::GetDesc() 
diff --git a/Engine/Source/Runtime/Engine/Classes/Engine/World.h b/Engine/Source/Runtime/Engine/Classes/Engine/World.h
--- a/Engine/Source/Runtime/Engine/Classes/Engine/World.h
+++ b/Engine/Source/Runtime/Engine/Classes/Engine/World.h
@@ -2898,6 +2898,8 @@ public:
 #if WITH_EDITOR
 	virtual EDataValidationResult IsDataValid(FDataValidationContext& Context) const override;
 	virtual bool Rename(const TCHAR* NewName = NULL, UObject* NewOuter = NULL, ERenameFlags Flags = REN_None) override;
+	virtual void GetAssetRegistryTags(FAssetRegistryTagsContext Context) const override;
+	UE_DEPRECATED(5.4, "Implement the version that takes FAssetRegistryTagsContext instead.")
 	virtual void GetAssetRegistryTags(TArray<FAssetRegistryTag>& OutTags) const override;
 	virtual void PostLoadAssetRegistryTags(const FAssetData& InAssetData, TArray<FAssetRegistryTag>& OutTagsAndValuesToUpdate) const;
 	virtual bool IsNameStableForNetworking() const override;
@@ -4148,6 +4150,7 @@ public:
 	DECLARE_MULTICAST_DELEGATE_TwoParams(FOnLevelChanged, ULevel*, UWorld*);
 
 	// delegate for generating world asset registry tags so project/game scope can add additional tags for filtering levels in their UI, etc
+	DECLARE_MULTICAST_DELEGATE_TwoParams(FWorldGetAssetTagsWithContext, const UWorld*, FAssetRegistryTagsContext);
 	DECLARE_MULTICAST_DELEGATE_TwoParams(FWorldGetAssetTags, const UWorld*, TArray<UObject::FAssetRegistryTag>&);
 
 	DECLARE_MULTICAST_DELEGATE_ThreeParams(FOnWorldTickStart, UWorld*, ELevelTick, float);
@@ -4235,6 +4238,8 @@ public:
 	static FLevelTransformEvent		PostApplyLevelTransform;
 
 	// called by UWorld::GetAssetRegistryTags()
+	static FWorldGetAssetTagsWithContext GetAssetTagsWithContext;
+	UE_DEPRECATED(5.4, "Subscribe to GetAssetTagsWithContext instead")
 	static FWorldGetAssetTags GetAssetTags;
 
 #if WITH_EDITOR
diff --git a/Engine/Source/Runtime/Engine/Classes/GameFramework/Actor.h b/Engine/Source/Runtime/Engine/Classes/GameFramework/Actor.h
--- a/Engine/Source/Runtime/Engine/Classes/GameFramework/Actor.h
+++ b/Engine/Source/Runtime/Engine/Classes/GameFramework/Actor.h
@@ -2188,8 +2188,9 @@ public:
 
 #if WITH_EDITOR
 	ENGINE_API virtual bool Modify(bool bAlwaysMarkDirty = true) override;
+	ENGINE_API virtual void GetAssetRegistryTags(FAssetRegistryTagsContext Context) const override;
+	UE_DEPRECATED(5.4, "Implement the version that takes FAssetRegistryTagsContext instead.")
 	ENGINE_API virtual void GetAssetRegistryTags(TArray<FAssetRegistryTag>& OutTags) const override;
-	ENGINE_API virtual void GetExtendedAssetRegistryTagsForSave(const ITargetPlatform* TargetPlatform, TArray<FAssetRegistryTag>& OutTags) const override;
 	ENGINE_API virtual bool NeedsLoadForTargetPlatform(const ITargetPlatform* TargetPlatform) const;
 	ENGINE_API virtual void PreEditChange(FProperty* PropertyThatWillChange) override;
 	ENGINE_API virtual bool CanEditChange(const FProperty* InProperty) const;
diff --git a/Engine/Source/Runtime/Engine/Private/Actor.cpp b/Engine/Source/Runtime/Engine/Private/Actor.cpp
--- a/Engine/Source/Runtime/Engine/Private/Actor.cpp
+++ b/Engine/Source/Runtime/Engine/Private/Actor.cpp
@@ -16,6 +16,7 @@
 #include "Components/SkeletalMeshComponent.h"
 #include "Engine/Engine.h"
 #include "Net/Core/PropertyConditions/RepChangedPropertyTracker.h"
+#include "UObject/AssetRegistryTagsContext.h"
 #include "UObject/ObjectSaveContext.h"
 #include "Physics/Experimental/PhysScene_Chaos.h"
 #include "UObject/UE5MainStreamObjectVersion.h"
@@ -495,25 +496,30 @@ void AActor::PreSave(FObjectPreSaveContext ObjectSaveContext)
 #if WITH_EDITOR
 void AActor::GetAssetRegistryTags(TArray<FAssetRegistryTag>& OutTags) const
 {
+	PRAGMA_DISABLE_DEPRECATION_WARNINGS;
 	Super::GetAssetRegistryTags(OutTags);
+	PRAGMA_ENABLE_DEPRECATION_WARNINGS;
+}
+
+void AActor::GetAssetRegistryTags(FAssetRegistryTagsContext Context) const
+{
+	Super::GetAssetRegistryTags(Context);
 
 	if (IsPackageExternal() && !IsChildActor())
 	{
 		if (ActorLabel.Len() > 0)
 		{
 			static FName NAME_ActorLabel(TEXT("ActorLabel"));
-			OutTags.Add(UObject::FAssetRegistryTag(NAME_ActorLabel, ActorLabel, UObject::FAssetRegistryTag::TT_Hidden));
+			Context.AddTag(FAssetRegistryTag(NAME_ActorLabel, ActorLabel, UObject::FAssetRegistryTag::TT_Hidden));
 		}
 	}
-}
 
-void AActor::GetExtendedAssetRegistryTagsForSave(const ITargetPlatform* TargetPlatform, TArray<FAssetRegistryTag>& OutTags) const
-{
-	Super::GetExtendedAssetRegistryTagsForSave(TargetPlatform, OutTags);
-	
-	if (IsPackageExternal() && !IsChildActor())
+	if (Context.IsSaving())
 	{
-		FWorldPartitionActorDescUtils::AppendAssetDataTagsFromActor(this, OutTags);
+		if (IsPackageExternal() && !IsChildActor())
+		{
+			FWorldPartitionActorDescUtils::AppendAssetDataTagsFromActor(this, Context);
+		}
 	}
 }
 
diff --git a/Engine/Source/Runtime/Engine/Private/ActorFolder.cpp b/Engine/Source/Runtime/Engine/Private/ActorFolder.cpp
--- a/Engine/Source/Runtime/Engine/Private/ActorFolder.cpp
+++ b/Engine/Source/Runtime/Engine/Private/ActorFolder.cpp
@@ -14,6 +14,7 @@
 #include "Engine/Level.h"
 #include "ExternalPackageHelper.h"
 #include "ActorFolderDesc.h"
+#include "UObject/AssetRegistryTagsContext.h"
 
 UActorFolder* UActorFolder::Create(ULevel* InLevel, const FString& InFolderLabel, UActorFolder* InParent)
 {
@@ -63,12 +64,21 @@ namespace ActorFolder
 
 void UActorFolder::GetAssetRegistryTags(TArray<FAssetRegistryTag>& OutTags) const
 {
-	OutTags.Add(FAssetRegistryTag(ActorFolder::NAME_FolderGuid, *FolderGuid.ToString(), FAssetRegistryTag::TT_Hidden));
-	OutTags.Add(FAssetRegistryTag(ActorFolder::NAME_ParentFolderGuid, *ParentFolderGuid.ToString(), FAssetRegistryTag::TT_Hidden));
-	OutTags.Add(FAssetRegistryTag(ActorFolder::NAME_FolderLabel, *FolderLabel, FAssetRegistryTag::TT_Hidden));
-	OutTags.Add(FAssetRegistryTag(ActorFolder::NAME_FolderInitiallyExpanded, bFolderInitiallyExpanded ? TEXT("1") : TEXT("0"), FAssetRegistryTag::TT_Hidden));
-	OutTags.Add(FAssetRegistryTag(ActorFolder::NAME_FolderIsDeleted, bIsDeleted ? TEXT("1") : TEXT("0"), FAssetRegistryTag::TT_Hidden));
-	OutTags.Add(FAssetRegistryTag(ActorFolder::NAME_OuterPackageName, *GetOuterULevel()->GetPackage()->GetName(), FAssetRegistryTag::TT_Hidden));
+	PRAGMA_DISABLE_DEPRECATION_WARNINGS;
+	Super::GetAssetRegistryTags(OutTags);
+	PRAGMA_ENABLE_DEPRECATION_WARNINGS;
+}
+
+void UActorFolder::GetAssetRegistryTags(FAssetRegistryTagsContext Context) const
+{
+	Super::GetAssetRegistryTags(Context);
+
+	Context.AddTag(FAssetRegistryTag(ActorFolder::NAME_FolderGuid, *FolderGuid.ToString(), FAssetRegistryTag::TT_Hidden));
+	Context.AddTag(FAssetRegistryTag(ActorFolder::NAME_ParentFolderGuid, *ParentFolderGuid.ToString(), FAssetRegistryTag::TT_Hidden));
+	Context.AddTag(FAssetRegistryTag(ActorFolder::NAME_FolderLabel, *FolderLabel, FAssetRegistryTag::TT_Hidden));
+	Context.AddTag(FAssetRegistryTag(ActorFolder::NAME_FolderInitiallyExpanded, bFolderInitiallyExpanded ? TEXT("1") : TEXT("0"), FAssetRegistryTag::TT_Hidden));
+	Context.AddTag(FAssetRegistryTag(ActorFolder::NAME_FolderIsDeleted, bIsDeleted ? TEXT("1") : TEXT("0"), FAssetRegistryTag::TT_Hidden));
+	Context.AddTag(FAssetRegistryTag(ActorFolder::NAME_OuterPackageName, *GetOuterULevel()->GetPackage()->GetName(), FAssetRegistryTag::TT_Hidden));
 }
 
 FActorFolderDesc UActorFolder::GetAssetRegistryInfoFromPackage(FName ActorFolderPackageName)
diff --git a/Engine/Source/Runtime/Engine/Private/Animation/AnimCurveMetaData.cpp b/Engine/Source/Runtime/Engine/Private/Animation/AnimCurveMetaData.cpp
--- a/Engine/Source/Runtime/Engine/Private/Animation/AnimCurveMetaData.cpp
+++ b/Engine/Source/Runtime/Engine/Private/Animation/AnimCurveMetaData.cpp
@@ -4,6 +4,7 @@
 #include "Animation/Skeleton.h"
 #if WITH_EDITOR
 #include "ScopedTransaction.h"
+#include "UObject/AssetRegistryTagsContext.h"
 #endif
 
 #define LOCTEXT_NAMESPACE "AnimCurveMetaData" 
@@ -39,6 +40,14 @@ FCurveMetaData::FCurveMetaData()
 
 void UAnimCurveMetaData::GetAssetRegistryTags(TArray<FAssetRegistryTag>& OutTags) const
 {
+	PRAGMA_DISABLE_DEPRECATION_WARNINGS;
+	Super::GetAssetRegistryTags(OutTags);
+	PRAGMA_ENABLE_DEPRECATION_WARNINGS;
+}
+
+void UAnimCurveMetaData::GetAssetRegistryTags(FAssetRegistryTagsContext Context) const
+{
+	Super::GetAssetRegistryTags(Context);
 	// Add curve metadata IDs to a tag list, or a delimiter if we have no curve metadata.
 	// The delimiter is necessary so we can distinguish between data with no curves and old data, as the asset registry
 	// strips tags that have empty values 
@@ -51,7 +60,7 @@ void UAnimCurveMetaData::GetAssetRegistryTags(TArray<FAssetRegistryTag>& OutTags
 		CurvesBuilder << USkeleton::CurveTagDelimiter;
 	});
 
-	OutTags.Add(FAssetRegistryTag(USkeleton::CurveNameTag, CurvesBuilder.ToString(), FAssetRegistryTag::TT_Hidden));
+	Context.AddTag(FAssetRegistryTag(USkeleton::CurveNameTag, CurvesBuilder.ToString(), FAssetRegistryTag::TT_Hidden));
 }
 
 #endif //WITH_EDITOR
diff --git a/Engine/Source/Runtime/Engine/Private/Animation/AnimSequence.cpp b/Engine/Source/Runtime/Engine/Private/Animation/AnimSequence.cpp
--- a/Engine/Source/Runtime/Engine/Private/Animation/AnimSequence.cpp
+++ b/Engine/Source/Runtime/Engine/Private/Animation/AnimSequence.cpp
@@ -8,6 +8,7 @@
 
 #include "Algo/Find.h"
 #include "Animation/AnimData/CurveIdentifier.h"
+#include "UObject/AssetRegistryTagsContext.h"
 #include "UObject/FrameworkObjectVersion.h"
 #include "Animation/AnimData/IAnimationDataController.h"
 #include "UObject/ObjectSaveContext.h"
@@ -414,25 +415,32 @@ void UAnimSequence::PostInitProperties()
 }
 
 void UAnimSequence::GetAssetRegistryTags(TArray<FAssetRegistryTag>& OutTags) const
+{
+	PRAGMA_DISABLE_DEPRECATION_WARNINGS;
+	Super::GetAssetRegistryTags(OutTags);
+	PRAGMA_ENABLE_DEPRECATION_WARNINGS;
+}
+
+void UAnimSequence::GetAssetRegistryTags(FAssetRegistryTagsContext Context) const
 {
 #if WITH_EDITORONLY_DATA
 	if (AssetImportData)
 	{
-		OutTags.Add( FAssetRegistryTag(SourceFileTagName(), AssetImportData->GetSourceData().ToJson(), FAssetRegistryTag::TT_Hidden) );
+		Context.AddTag( FAssetRegistryTag(SourceFileTagName(), AssetImportData->GetSourceData().ToJson(), FAssetRegistryTag::TT_Hidden) );
 	}
 
 	if (DataModelInterface.GetObject() && DataModelInterface->HasBeenPopulated())
 	{
-		OutTags.Add(FAssetRegistryTag(TEXT("Compression Ratio"), FString::Printf(TEXT("%.03f"), (float)GetApproxCompressedSize() / (float)GetUncompressedRawSize()), FAssetRegistryTag::TT_Numerical));
-		OutTags.Add(FAssetRegistryTag(TEXT("Source Frame Rate"), FString::Printf(TEXT("%.2f"), DataModelInterface->GetFrameRate().AsDecimal()), FAssetRegistryTag::TT_Numerical));
-		OutTags.Add(FAssetRegistryTag(TEXT("Number of Frames"), FString::Printf(TEXT("%.i"), DataModelInterface->GetNumberOfFrames()), FAssetRegistryTag::TT_Numerical));
-		OutTags.Add(FAssetRegistryTag(TEXT("Number of Keys"), FString::Printf(TEXT("%i"), DataModelInterface->GetNumberOfKeys()), FAssetRegistryTag::TT_Numerical));
+		Context.AddTag(FAssetRegistryTag(TEXT("Compression Ratio"), FString::Printf(TEXT("%.03f"), (float)GetApproxCompressedSize() / (float)GetUncompressedRawSize()), FAssetRegistryTag::TT_Numerical));
+		Context.AddTag(FAssetRegistryTag(TEXT("Source Frame Rate"), FString::Printf(TEXT("%.2f"), DataModelInterface->GetFrameRate().AsDecimal()), FAssetRegistryTag::TT_Numerical));
+		Context.AddTag(FAssetRegistryTag(TEXT("Number of Frames"), FString::Printf(TEXT("%.i"), DataModelInterface->GetNumberOfFrames()), FAssetRegistryTag::TT_Numerical));
+		Context.AddTag(FAssetRegistryTag(TEXT("Number of Keys"), FString::Printf(TEXT("%i"), DataModelInterface->GetNumberOfKeys()), FAssetRegistryTag::TT_Numerical));
 	}
 	else
 	{
-		OutTags.Add(FAssetRegistryTag(TEXT("Source Frame Rate"), TEXT("0.0"), FAssetRegistryTag::TT_Numerical));
-		OutTags.Add(FAssetRegistryTag(TEXT("Number of Frames"), TEXT("0"), FAssetRegistryTag::TT_Numerical));
-		OutTags.Add(FAssetRegistryTag(TEXT("Number of Keys"), TEXT("0"), FAssetRegistryTag::TT_Numerical));
+		Context.AddTag(FAssetRegistryTag(TEXT("Source Frame Rate"), TEXT("0.0"), FAssetRegistryTag::TT_Numerical));
+		Context.AddTag(FAssetRegistryTag(TEXT("Number of Frames"), TEXT("0"), FAssetRegistryTag::TT_Numerical));
+		Context.AddTag(FAssetRegistryTag(TEXT("Number of Keys"), TEXT("0"), FAssetRegistryTag::TT_Numerical));
 	}
 
 	// Output unique sync marker names we use
@@ -445,12 +453,12 @@ void UAnimSequence::GetAssetRegistryTags(TArray<FAssetRegistryTag>& OutTags) con
 		SyncMarkersBuilder.Append(USkeleton::AnimSyncMarkerTagDelimiter);
 	}
 
-	OutTags.Add(FAssetRegistryTag(USkeleton::AnimSyncMarkerTag, SyncMarkersBuilder.ToString(), FAssetRegistryTag::TT_Hidden));
+	Context.AddTag(FAssetRegistryTag(USkeleton::AnimSyncMarkerTag, SyncMarkersBuilder.ToString(), FAssetRegistryTag::TT_Hidden));
 #endif
 
-	OutTags.Add(FAssetRegistryTag(TEXT("Compressed Size (KB)"), FString::Printf(TEXT("%.02f"), (float)GetApproxCompressedSize() / 1024.0f), FAssetRegistryTag::TT_Numerical));
-	OutTags.Add(FAssetRegistryTag(TEXT("Target Frame Rate"), FString::Printf(TEXT("%.2f"), PlatformTargetFrameRate.Default.AsDecimal()), FAssetRegistryTag::TT_Numerical));
-	Super::GetAssetRegistryTags(OutTags);
+	Context.AddTag(FAssetRegistryTag(TEXT("Compressed Size (KB)"), FString::Printf(TEXT("%.02f"), (float)GetApproxCompressedSize() / 1024.0f), FAssetRegistryTag::TT_Numerical));
+	Context.AddTag(FAssetRegistryTag(TEXT("Target Frame Rate"), FString::Printf(TEXT("%.2f"), PlatformTargetFrameRate.Default.AsDecimal()), FAssetRegistryTag::TT_Numerical));
+	Super::GetAssetRegistryTags(Context);
 }
 
 void UAnimSequence::AddReferencedObjects(UObject* This, FReferenceCollector& Collector)
diff --git a/Engine/Source/Runtime/Engine/Private/Animation/AnimSequenceBase.cpp b/Engine/Source/Runtime/Engine/Private/Animation/AnimSequenceBase.cpp
--- a/Engine/Source/Runtime/Engine/Private/Animation/AnimSequenceBase.cpp
+++ b/Engine/Source/Runtime/Engine/Private/Animation/AnimSequenceBase.cpp
@@ -12,6 +12,7 @@
 #include "Animation/AnimNotifyEndDataContext.h"
 #include "Animation/Skeleton.h"
 #include "Logging/MessageLog.h"
+#include "UObject/AssetRegistryTagsContext.h"
 #include "UObject/FrameworkObjectVersion.h"
 #include "UObject/FortniteMainBranchObjectVersion.h"
 #include "Animation/AnimationPoseData.h"
@@ -876,7 +877,14 @@ EAnimEventTriggerOffsets::Type UAnimSequenceBase::CalculateOffsetForNotify(float
 
 void UAnimSequenceBase::GetAssetRegistryTags(TArray<FAssetRegistryTag>& OutTags) const
 {
+	PRAGMA_DISABLE_DEPRECATION_WARNINGS;
 	Super::GetAssetRegistryTags(OutTags);
+	PRAGMA_ENABLE_DEPRECATION_WARNINGS;
+}
+
+void UAnimSequenceBase::GetAssetRegistryTags(FAssetRegistryTagsContext Context) const
+{
+	Super::GetAssetRegistryTags(Context);
 
 	// Add notify IDs to a tag list, or a delimiter if we have no notifies.
 	// The delimiter is necessary so we can distinguish between data with no curves and old data, as the asset registry
@@ -891,7 +899,7 @@ void UAnimSequenceBase::GetAssetRegistryTags(TArray<FAssetRegistryTag>& OutTags)
 		}
 	}
 	
-	OutTags.Add(FAssetRegistryTag(USkeleton::AnimNotifyTag, NotifyList, FAssetRegistryTag::TT_Hidden));
+	Context.AddTag(FAssetRegistryTag(USkeleton::AnimNotifyTag, NotifyList, FAssetRegistryTag::TT_Hidden));
 
 	// Add curve IDs to a tag list, or a delimiter if we have no curves.
 	// The delimiter is necessary so we can distinguish between data with no curves and old data, as the asset registry
@@ -904,7 +912,7 @@ void UAnimSequenceBase::GetAssetRegistryTags(TArray<FAssetRegistryTag>& OutTags)
 			CurveNameList += FString::Printf(TEXT("%s%s"), *Curve.GetName().ToString(), *USkeleton::CurveTagDelimiter);
 		}
 	}
-	OutTags.Add(FAssetRegistryTag(USkeleton::CurveNameTag, CurveNameList, FAssetRegistryTag::TT_Hidden));
+	Context.AddTag(FAssetRegistryTag(USkeleton::CurveNameTag, CurveNameList, FAssetRegistryTag::TT_Hidden));
 }
 
 uint8* UAnimSequenceBase::FindNotifyPropertyData(int32 NotifyIndex, FArrayProperty*& ArrayProperty)
diff --git a/Engine/Source/Runtime/Engine/Private/Animation/PoseAsset.cpp b/Engine/Source/Runtime/Engine/Private/Animation/PoseAsset.cpp
--- a/Engine/Source/Runtime/Engine/Private/Animation/PoseAsset.cpp
+++ b/Engine/Source/Runtime/Engine/Private/Animation/PoseAsset.cpp
@@ -18,6 +18,7 @@
 #include "Animation/SkeletonRemapping.h"
 #include "Animation/SkeletonRemappingRegistry.h"
 #include "Engine/SkeletalMesh.h"
+#include "UObject/AssetRegistryTagsContext.h"
 #include "UObject/LinkerLoad.h"
 #include "UObject/ObjectSaveContext.h"
 #include "UObject/Package.h"
@@ -1024,10 +1025,17 @@ void UPoseAsset::PreSave(FObjectPreSaveContext ObjectSaveContext)
 
 void UPoseAsset::GetAssetRegistryTags(TArray<FAssetRegistryTag>& OutTags) const
 {
+	PRAGMA_DISABLE_DEPRECATION_WARNINGS;
 	Super::GetAssetRegistryTags(OutTags);
+	PRAGMA_ENABLE_DEPRECATION_WARNINGS;
+}
+
+void UPoseAsset::GetAssetRegistryTags(FAssetRegistryTagsContext Context) const
+{
+	Super::GetAssetRegistryTags(Context);
 
 	// Number of poses
-	OutTags.Add(FAssetRegistryTag("Poses", FString::FromInt(GetNumPoses()), FAssetRegistryTag::TT_Numerical));
+	Context.AddTag(FAssetRegistryTag("Poses", FString::FromInt(GetNumPoses()), FAssetRegistryTag::TT_Numerical));
 #if WITH_EDITOR
 	TArray<FName> Names;
 	Names.Reserve(PoseContainer.PoseFNames.Num() + PoseContainer.Curves.Num());
@@ -1050,7 +1058,7 @@ void UPoseAsset::GetAssetRegistryTags(TArray<FAssetRegistryTag>& OutTags) const
 	{
 		PoseNameList += FString::Printf(TEXT("%s%s"), *Name.ToString(), *USkeleton::CurveTagDelimiter);
 	}
-	OutTags.Add(FAssetRegistryTag(USkeleton::CurveNameTag, PoseNameList, FAssetRegistryTag::TT_Hidden)); //write pose names as curve tag as they use 
+	Context.AddTag(FAssetRegistryTag(USkeleton::CurveNameTag, PoseNameList, FAssetRegistryTag::TT_Hidden)); //write pose names as curve tag as they use 
 #endif
 }
 
diff --git a/Engine/Source/Runtime/Engine/Private/Animation/Skeleton.cpp b/Engine/Source/Runtime/Engine/Private/Animation/Skeleton.cpp
--- a/Engine/Source/Runtime/Engine/Private/Animation/Skeleton.cpp
+++ b/Engine/Source/Runtime/Engine/Private/Animation/Skeleton.cpp
@@ -8,6 +8,7 @@
 
 #include "AnimationSequenceCompiler.h"
 #include "Animation/AnimData/AnimDataModel.h"
+#include "UObject/AssetRegistryTagsContext.h"
 #include "UObject/LinkerLoad.h"
 #include "Engine/AssetUserData.h"
 #include "Modules/ModuleManager.h"
@@ -1973,7 +1974,14 @@ void USkeleton::HandleVirtualBoneChanges()
 
 void USkeleton::GetAssetRegistryTags(TArray<FAssetRegistryTag>& OutTags) const
 {
+	PRAGMA_DISABLE_DEPRECATION_WARNINGS;
 	Super::GetAssetRegistryTags(OutTags);
+	PRAGMA_ENABLE_DEPRECATION_WARNINGS;
+}
+
+void USkeleton::GetAssetRegistryTags(FAssetRegistryTagsContext Context) const
+{
+	Super::GetAssetRegistryTags(Context);
 
 	TStringBuilder<256> CompatibleSkeletonsBuilder;
 	
@@ -1984,7 +1992,7 @@ void USkeleton::GetAssetRegistryTags(TArray<FAssetRegistryTag>& OutTags) const
 		CompatibleSkeletonsBuilder.Append(USkeleton::CompatibleSkeletonsTagDelimiter);
 	}
 
-	OutTags.Add(FAssetRegistryTag(USkeleton::CompatibleSkeletonsNameTag, CompatibleSkeletonsBuilder.ToString(), FAssetRegistryTag::TT_Hidden));
+	Context.AddTag(FAssetRegistryTag(USkeleton::CompatibleSkeletonsNameTag, CompatibleSkeletonsBuilder.ToString(), FAssetRegistryTag::TT_Hidden));
 
 	// Output sync notify names we use
 	TStringBuilder<256> NotifiesBuilder;
@@ -1996,7 +2004,7 @@ void USkeleton::GetAssetRegistryTags(TArray<FAssetRegistryTag>& OutTags) const
 		NotifiesBuilder.Append(USkeleton::AnimNotifyTagDelimiter);
 	}
 
-	OutTags.Add(FAssetRegistryTag(USkeleton::AnimNotifyTag, NotifiesBuilder.ToString(), FAssetRegistryTag::TT_Hidden));
+	Context.AddTag(FAssetRegistryTag(USkeleton::AnimNotifyTag, NotifiesBuilder.ToString(), FAssetRegistryTag::TT_Hidden));
 	
 	// Output sync marker names we use
 	TStringBuilder<256> SyncMarkersBuilder;
@@ -2008,14 +2016,14 @@ void USkeleton::GetAssetRegistryTags(TArray<FAssetRegistryTag>& OutTags) const
 		SyncMarkersBuilder.Append(USkeleton::AnimSyncMarkerTagDelimiter);
 	}
 
-	OutTags.Add(FAssetRegistryTag(USkeleton::AnimSyncMarkerTag, SyncMarkersBuilder.ToString(), FAssetRegistryTag::TT_Hidden));
+	Context.AddTag(FAssetRegistryTag(USkeleton::AnimSyncMarkerTag, SyncMarkersBuilder.ToString(), FAssetRegistryTag::TT_Hidden));
 	
 	// Allow asset user data to output tags
 	for(const UAssetUserData* AssetUserDataItem : *GetAssetUserDataArray())
 	{
 		if (AssetUserDataItem)
 		{
-			AssetUserDataItem->GetAssetRegistryTags(OutTags);
+			AssetUserDataItem->GetAssetRegistryTags(Context);
 		}
 	}
 }
diff --git a/Engine/Source/Runtime/Engine/Private/AnimationAsset.cpp b/Engine/Source/Runtime/Engine/Private/AnimationAsset.cpp
--- a/Engine/Source/Runtime/Engine/Private/AnimationAsset.cpp
+++ b/Engine/Source/Runtime/Engine/Private/AnimationAsset.cpp
@@ -9,6 +9,7 @@
 #include "Animation/AnimMetaData.h"
 #include "Animation/AnimSequence.h"
 #include "AnimationUtils.h"
+#include "UObject/AssetRegistryTagsContext.h"
 #include "UObject/LinkerLoad.h"
 #include "Animation/BlendSpace.h"
 #include "Animation/PoseAsset.h"
@@ -767,17 +768,24 @@ void UAnimationAsset::PostEditChangeProperty(FPropertyChangedEvent& PropertyChan
 
 void UAnimationAsset::GetAssetRegistryTags(TArray<FAssetRegistryTag>& OutTags) const
 {
+	PRAGMA_DISABLE_DEPRECATION_WARNINGS;
 	Super::GetAssetRegistryTags(OutTags);
+	PRAGMA_ENABLE_DEPRECATION_WARNINGS;
+}
+
+void UAnimationAsset::GetAssetRegistryTags(FAssetRegistryTagsContext Context) const
+{
+	Super::GetAssetRegistryTags(Context);
 	
 	for (const UAssetUserData* UserData : AssetUserData)
 	{
 		if (UserData)
 		{
-			UserData->GetAssetRegistryTags(OutTags);	
+			UserData->GetAssetRegistryTags(Context);
 		}
 	}
 	
-	OutTags.Add( FAssetRegistryTag("HasParentAsset", HasParentAsset() ? TEXT("True") : TEXT("False"), FAssetRegistryTag::TT_Hidden) );
+	Context.AddTag( FAssetRegistryTag("HasParentAsset", HasParentAsset() ? TEXT("True") : TEXT("False"), FAssetRegistryTag::TT_Hidden) );
 }
 
 EDataValidationResult UAnimationAsset::IsDataValid(FDataValidationContext& Context) const
diff --git a/Engine/Source/Runtime/Engine/Private/Blueprint.cpp b/Engine/Source/Runtime/Engine/Private/Blueprint.cpp
--- a/Engine/Source/Runtime/Engine/Private/Blueprint.cpp
+++ b/Engine/Source/Runtime/Engine/Private/Blueprint.cpp
@@ -3,6 +3,7 @@
 #include "Engine/Blueprint.h"
 
 #include "Blueprint/BlueprintSupport.h"
+#include "UObject/AssetRegistryTagsContext.h"
 #include "UObject/BlueprintsObjectVersion.h"
 #include "EngineLogs.h"
 #include "UObject/FrameworkObjectVersion.h"
@@ -935,6 +936,13 @@ bool UBlueprint::CanAlwaysRecompileWhilePlayingInEditor() const
 }
 
 void UBlueprint::GetAssetRegistryTags(TArray<FAssetRegistryTag>& OutTags) const
+{
+	PRAGMA_DISABLE_DEPRECATION_WARNINGS;
+	Super::GetAssetRegistryTags(OutTags);
+	PRAGMA_ENABLE_DEPRECATION_WARNINGS;
+}
+
+void UBlueprint::GetAssetRegistryTags(FAssetRegistryTagsContext Context) const
 {
 	// We use Generated instead of Skeleton because the CDO data is more accurate on Generated
 	UObject* BlueprintCDO = nullptr;
@@ -943,11 +951,11 @@ void UBlueprint::GetAssetRegistryTags(TArray<FAssetRegistryTag>& OutTags) const
 		BlueprintCDO = GeneratedClass->GetDefaultObject(/*bCreateIfNeeded*/false);
 		if (BlueprintCDO)
 		{
-			BlueprintCDO->GetAssetRegistryTags(OutTags);
+			BlueprintCDO->GetAssetRegistryTags(Context);
 		}
 	}
 
-	Super::GetAssetRegistryTags(OutTags);
+	Super::GetAssetRegistryTags(Context);
 
 	// Output native parent class and generated class as if they were AssetRegistrySearchable
 	FString GeneratedClassVal;
@@ -979,10 +987,10 @@ void UBlueprint::GetAssetRegistryTags(TArray<FAssetRegistryTag>& OutTags) const
 	}
 
 
-	OutTags.Add(FAssetRegistryTag(FBlueprintTags::BlueprintPathWithinPackage, GetPathName(GetOutermost()), FAssetRegistryTag::TT_Hidden));
-	OutTags.Add(FAssetRegistryTag(FBlueprintTags::GeneratedClassPath, GeneratedClassVal, FAssetRegistryTag::TT_Hidden));
-	OutTags.Add(FAssetRegistryTag(FBlueprintTags::ParentClassPath, ParentClassName, FAssetRegistryTag::TT_Alphabetical));
-	OutTags.Add(FAssetRegistryTag(FBlueprintTags::NativeParentClassPath, NativeParentClassName, FAssetRegistryTag::TT_Alphabetical));
+	Context.AddTag(FAssetRegistryTag(FBlueprintTags::BlueprintPathWithinPackage, GetPathName(GetOutermost()), FAssetRegistryTag::TT_Hidden));
+	Context.AddTag(FAssetRegistryTag(FBlueprintTags::GeneratedClassPath, GeneratedClassVal, FAssetRegistryTag::TT_Hidden));
+	Context.AddTag(FAssetRegistryTag(FBlueprintTags::ParentClassPath, ParentClassName, FAssetRegistryTag::TT_Alphabetical));
+	Context.AddTag(FAssetRegistryTag(FBlueprintTags::NativeParentClassPath, NativeParentClassName, FAssetRegistryTag::TT_Alphabetical));
 
 	// BlueprintGeneratedClass is not automatically traversed so we have to manually add NumReplicatedProperties
 	int32 NumReplicatedProperties = 0;
@@ -991,10 +999,10 @@ void UBlueprint::GetAssetRegistryTags(TArray<FAssetRegistryTag>& OutTags) const
 	{
 		NumReplicatedProperties = BlueprintClass->NumReplicatedProperties;
 	}
-	OutTags.Add(FAssetRegistryTag(FBlueprintTags::NumReplicatedProperties, FString::FromInt(NumReplicatedProperties), FAssetRegistryTag::TT_Numerical));
-	OutTags.Add(FAssetRegistryTag(FBlueprintTags::BlueprintDescription, BlueprintDescription, FAssetRegistryTag::TT_Hidden));
-	OutTags.Add(FAssetRegistryTag(FBlueprintTags::BlueprintCategory, BlueprintCategory, FAssetRegistryTag::TT_Hidden));
-	OutTags.Add(FAssetRegistryTag(FBlueprintTags::BlueprintDisplayName, BlueprintDisplayName, FAssetRegistryTag::TT_Hidden));
+	Context.AddTag(FAssetRegistryTag(FBlueprintTags::NumReplicatedProperties, FString::FromInt(NumReplicatedProperties), FAssetRegistryTag::TT_Numerical));
+	Context.AddTag(FAssetRegistryTag(FBlueprintTags::BlueprintDescription, BlueprintDescription, FAssetRegistryTag::TT_Hidden));
+	Context.AddTag(FAssetRegistryTag(FBlueprintTags::BlueprintCategory, BlueprintCategory, FAssetRegistryTag::TT_Hidden));
+	Context.AddTag(FAssetRegistryTag(FBlueprintTags::BlueprintDisplayName, BlueprintDisplayName, FAssetRegistryTag::TT_Hidden));
 
 	uint32 ClassFlagsTagged = 0;
 	if (BlueprintClass)
@@ -1005,9 +1013,9 @@ void UBlueprint::GetAssetRegistryTags(TArray<FAssetRegistryTag>& OutTags) const
 	{
 		ClassFlagsTagged = GetClass()->GetClassFlags();
 	}
-	OutTags.Add( FAssetRegistryTag(FBlueprintTags::ClassFlags, FString::FromInt(ClassFlagsTagged), FAssetRegistryTag::TT_Hidden) );
+	Context.AddTag( FAssetRegistryTag(FBlueprintTags::ClassFlags, FString::FromInt(ClassFlagsTagged), FAssetRegistryTag::TT_Hidden) );
 
-	OutTags.Add( FAssetRegistryTag(FBlueprintTags::IsDataOnly,
+	Context.AddTag( FAssetRegistryTag(FBlueprintTags::IsDataOnly,
 			FBlueprintEditorUtils::IsDataOnlyBlueprint(this) ? TEXT("True") : TEXT("False"),
 			FAssetRegistryTag::TT_Alphabetical ) );
 
@@ -1022,7 +1030,7 @@ void UBlueprint::GetAssetRegistryTags(TArray<FAssetRegistryTag>& OutTags) const
 			Value = SearchData.Value;
 		}
 		
-		OutTags.Add( FAssetRegistryTag(FBlueprintTags::FindInBlueprintsData, Value, FAssetRegistryTag::TT_Hidden) );
+		Context.AddTag( FAssetRegistryTag(FBlueprintTags::FindInBlueprintsData, Value, FAssetRegistryTag::TT_Hidden) );
 	}
 
 	// Only show for strict blueprints (not animation or widget blueprints)
@@ -1044,7 +1052,7 @@ void UBlueprint::GetAssetRegistryTags(TArray<FAssetRegistryTag>& OutTags) const
 				}
 			}
 		}
-		OutTags.Add(FAssetRegistryTag(FBlueprintTags::NumNativeComponents, FString::FromInt(NumNativeComponents), UObject::FAssetRegistryTag::TT_Numerical));
+		Context.AddTag(FAssetRegistryTag(FBlueprintTags::NumNativeComponents, FString::FromInt(NumNativeComponents), UObject::FAssetRegistryTag::TT_Numerical));
 
 		// Determine how many components are added via a SimpleConstructionScript (both newly introduced and inherited from parent BPs)
 		int32 NumAddedComponents = 0;
@@ -1056,45 +1064,44 @@ void UBlueprint::GetAssetRegistryTags(TArray<FAssetRegistryTag>& OutTags) const
 				NumAddedComponents += AssociatedBP->SimpleConstructionScript->GetAllNodes().Num();
 			}
 		}
-		OutTags.Add(FAssetRegistryTag(FBlueprintTags::NumBlueprintComponents, FString::FromInt(NumAddedComponents), UObject::FAssetRegistryTag::TT_Numerical));
+		Context.AddTag(FAssetRegistryTag(FBlueprintTags::NumBlueprintComponents, FString::FromInt(NumAddedComponents), UObject::FAssetRegistryTag::TT_Numerical));
 	}
-}
 
 #if WITH_EDITOR
-void UBlueprint::GetExtendedAssetRegistryTagsForSave(const ITargetPlatform* TargetPlatform, TArray<FAssetRegistryTag>& OutTags) const
-{
-	Super::GetExtendedAssetRegistryTagsForSave(TargetPlatform, OutTags);
-
-	if ( ParentClass && GIsEditor && !GetOutermost()->HasAnyPackageFlags(PKG_ForDiffing) && IsRunningCookCommandlet())
+	if (Context.IsSaving())
 	{
-		if (!TargetPlatform || TargetPlatform->HasEditorOnlyData())
+		if (ParentClass && GIsEditor && !GetOutermost()->HasAnyPackageFlags(PKG_ForDiffing) && IsRunningCookCommandlet())
 		{
-			FString Value;
-			const bool bRebuildSearchData = false;
-			FSearchData SearchData = FFindInBlueprintSearchManager::Get().QuerySingleBlueprint((UBlueprint*)this, bRebuildSearchData);
-			if (SearchData.IsValid())
+			if (!Context.GetTargetPlatform() || Context.GetTargetPlatform()->HasEditorOnlyData())
 			{
-				Value = SearchData.Value;
+				FString Value;
+				const bool bRebuildSearchData = false;
+				FSearchData SearchData = FFindInBlueprintSearchManager::Get().QuerySingleBlueprint((UBlueprint*)this, bRebuildSearchData);
+				if (SearchData.IsValid())
+				{
+					Value = SearchData.Value;
+				}
+
+				Context.AddTag(FAssetRegistryTag(FBlueprintTags::FindInBlueprintsData, Value, FAssetRegistryTag::TT_Hidden));
 			}
-			
-			OutTags.Add( FAssetRegistryTag(FBlueprintTags::FindInBlueprintsData, Value, FAssetRegistryTag::TT_Hidden) );
 		}
-	}
 
-	/*
-	 * Using GetExtendedAssetRegistryTagsForSave here because:
-	 *	- UBlueprint is an asset in editor builds only.
-	 *	- UBlueprintGeneratedClass is an asset in cooked builds only.
-	 *	- Extended tags are only present in editor builds.
-	 *
-	 * See UBlueprintGeneratedClass::GetAssetRegistryTags.
-	 */
-	if (AActor* BlueprintCDO = GeneratedClass ? Cast<AActor>(GeneratedClass->ClassDefaultObject) : nullptr)
-	{
-		FWorldPartitionActorDescUtils::AppendAssetDataTagsFromActor(BlueprintCDO, OutTags);
+		/*
+		 * Restricting these tags to IsSaving because:
+		 *	- UBlueprint is an asset in editor builds only.
+		 *	- UBlueprintGeneratedClass is an asset in cooked builds only.
+		 *	- Extended tags are only present in editor builds.
+		 *
+		 * See UBlueprintGeneratedClass::GetAssetRegistryTags.
+		 */
+		AActor* BlueprintCDOAsActor = Cast<AActor>(BlueprintCDO);
+		if (BlueprintCDOAsActor)
+		{
+			FWorldPartitionActorDescUtils::AppendAssetDataTagsFromActor(BlueprintCDOAsActor, Context);
+		}
 	}
-}
 #endif
+}
 
 void UBlueprint::PostLoadAssetRegistryTags(const FAssetData& InAssetData, TArray<FAssetRegistryTag>& OutTagsAndValuesToUpdate) const
 {
diff --git a/Engine/Source/Runtime/Engine/Private/BlueprintGeneratedClass.cpp b/Engine/Source/Runtime/Engine/Private/BlueprintGeneratedClass.cpp
--- a/Engine/Source/Runtime/Engine/Private/BlueprintGeneratedClass.cpp
+++ b/Engine/Source/Runtime/Engine/Private/BlueprintGeneratedClass.cpp
@@ -7,6 +7,7 @@
 #include "HAL/IConsoleManager.h"
 #include "EngineLogs.h"
 #include "Stats/StatsMisc.h"
+#include "UObject/AssetRegistryTagsContext.h"
 #include "UObject/CoreNet.h"
 #include "UObject/CoreRedirects.h"
 #include "UObject/ObjectSaveContext.h"
@@ -523,7 +524,14 @@ void UBlueprintGeneratedClass::InitializeFieldNotifies()
 
 void UBlueprintGeneratedClass::GetAssetRegistryTags(TArray<FAssetRegistryTag>& OutTags) const
 {
+	PRAGMA_DISABLE_DEPRECATION_WARNINGS;
 	Super::GetAssetRegistryTags(OutTags);
+	PRAGMA_ENABLE_DEPRECATION_WARNINGS;
+}
+
+void UBlueprintGeneratedClass::GetAssetRegistryTags(FAssetRegistryTagsContext Context) const
+{
+	Super::GetAssetRegistryTags(Context);
 
 	FString NativeParentClassName;
 	FString ParentClassName;
@@ -545,9 +553,9 @@ void UBlueprintGeneratedClass::GetAssetRegistryTags(TArray<FAssetRegistryTag>& O
 		NativeParentClassName = ParentClassName = TEXT("None");
 	}
 
-	OutTags.Add(FAssetRegistryTag(FBlueprintTags::ParentClassPath, ParentClassName, FAssetRegistryTag::TT_Alphabetical));
-	OutTags.Add(FAssetRegistryTag(FBlueprintTags::NativeParentClassPath, NativeParentClassName, FAssetRegistryTag::TT_Alphabetical));
-	OutTags.Add(FAssetRegistryTag(FBlueprintTags::ClassFlags, FString::FromInt((uint32)GetClassFlags()), FAssetRegistryTag::TT_Hidden));
+	Context.AddTag(FAssetRegistryTag(FBlueprintTags::ParentClassPath, ParentClassName, FAssetRegistryTag::TT_Alphabetical));
+	Context.AddTag(FAssetRegistryTag(FBlueprintTags::NativeParentClassPath, NativeParentClassName, FAssetRegistryTag::TT_Alphabetical));
+	Context.AddTag(FAssetRegistryTag(FBlueprintTags::ClassFlags, FString::FromInt((uint32)GetClassFlags()), FAssetRegistryTag::TT_Hidden));
 
 #if WITH_EDITORONLY_DATA
 	// Get editor-only tags; on a cooked BPGC, those tags are deserialized into CookedEditorTags, otherwise generate them for the BP
@@ -561,7 +569,7 @@ void UBlueprintGeneratedClass::GetAssetRegistryTags(TArray<FAssetRegistryTag>& O
 
 	for (const auto& EditorTag : *EditorTagsToAdd)
 	{
-		OutTags.Add(FAssetRegistryTag(EditorTag.Key, EditorTag.Value, FAssetRegistryTag::TT_Hidden));
+		Context.AddTag(FAssetRegistryTag(EditorTag.Key, EditorTag.Value, FAssetRegistryTag::TT_Hidden));
 	}
 
 	if (const UObject* CDO = GetDefaultObject())
@@ -570,7 +578,7 @@ void UBlueprintGeneratedClass::GetAssetRegistryTags(TArray<FAssetRegistryTag>& O
 		{
 			if (AssetRegistryProvider->ShouldAddCDOTagsToBlueprintClass())
 			{
-				CDO->GetAssetRegistryTags(OutTags);
+				CDO->GetAssetRegistryTags(Context);
 			}
 		}
 	}
@@ -593,7 +601,7 @@ void UBlueprintGeneratedClass::GetAssetRegistryTags(TArray<FAssetRegistryTag>& O
 			{
 				if (!FPackageName::IsTempPackage(BlueprintCDOPackage->GetName()) && !BlueprintCDOPackage->HasAnyPackageFlags(PKG_PlayInEditor))
 				{
-					FWorldPartitionActorDescUtils::AppendAssetDataTagsFromActor(BlueprintCDO, OutTags);
+					FWorldPartitionActorDescUtils::AppendAssetDataTagsFromActor(BlueprintCDO, Context);
 				}
 			}
 		}
diff --git a/Engine/Source/Runtime/Engine/Private/CurveTable.cpp b/Engine/Source/Runtime/Engine/Private/CurveTable.cpp
--- a/Engine/Source/Runtime/Engine/Private/CurveTable.cpp
+++ b/Engine/Source/Runtime/Engine/Private/CurveTable.cpp
@@ -5,6 +5,7 @@
 #include "Serialization/Csv/CsvParser.h"
 #include "HAL/IConsoleManager.h"
 #include "Stats/Stats.h"
+#include "UObject/AssetRegistryTagsContext.h"
 #include "UObject/FortniteMainBranchObjectVersion.h"
 
 #include "EditorFramework/AssetImportData.h"
@@ -272,13 +273,20 @@ void UCurveTable::FinishDestroy()
 
 #if WITH_EDITORONLY_DATA
 void UCurveTable::GetAssetRegistryTags(TArray<FAssetRegistryTag>& OutTags) const
+{
+	PRAGMA_DISABLE_DEPRECATION_WARNINGS;
+	Super::GetAssetRegistryTags(OutTags);
+	PRAGMA_ENABLE_DEPRECATION_WARNINGS;
+}
+
+void UCurveTable::GetAssetRegistryTags(FAssetRegistryTagsContext Context) const
 {
 	if (AssetImportData)
 	{
-		OutTags.Add(FAssetRegistryTag(SourceFileTagName(), AssetImportData->GetSourceData().ToJson(), FAssetRegistryTag::TT_Hidden));
+		Context.AddTag(FAssetRegistryTag(SourceFileTagName(), AssetImportData->GetSourceData().ToJson(), FAssetRegistryTag::TT_Hidden));
 	}
 
-	Super::GetAssetRegistryTags(OutTags);
+	Super::GetAssetRegistryTags(Context);
 }
 
 void UCurveTable::PostInitProperties()
diff --git a/Engine/Source/Runtime/Engine/Private/Curves/CurveBase.cpp b/Engine/Source/Runtime/Engine/Private/Curves/CurveBase.cpp
--- a/Engine/Source/Runtime/Engine/Private/Curves/CurveBase.cpp
+++ b/Engine/Source/Runtime/Engine/Private/Curves/CurveBase.cpp
@@ -5,6 +5,7 @@
 #include "JsonObjectConverter.h"
 #include "Serialization/Csv/CsvParser.h"
 #include "EditorFramework/AssetImportData.h"
+#include "UObject/AssetRegistryTagsContext.h"
 
 #include UE_INLINE_GENERATED_CPP_BY_NAME(CurveBase)
 
@@ -226,13 +227,20 @@ FString UCurveBase::ExportAsJSONString() const
 #if WITH_EDITORONLY_DATA
 
 void UCurveBase::GetAssetRegistryTags(TArray<FAssetRegistryTag>& OutTags) const
+{
+	PRAGMA_DISABLE_DEPRECATION_WARNINGS;
+	Super::GetAssetRegistryTags(OutTags);
+	PRAGMA_ENABLE_DEPRECATION_WARNINGS;
+}
+
+void UCurveBase::GetAssetRegistryTags(FAssetRegistryTagsContext Context) const
 {
 	if (AssetImportData)
 	{
-		OutTags.Add(FAssetRegistryTag(SourceFileTagName(), AssetImportData->GetSourceData().ToJson(), FAssetRegistryTag::TT_Hidden));
+		Context.AddTag(FAssetRegistryTag(SourceFileTagName(), AssetImportData->GetSourceData().ToJson(), FAssetRegistryTag::TT_Hidden));
 	}
 
-	Super::GetAssetRegistryTags(OutTags);
+	Super::GetAssetRegistryTags(Context);
 }
 
 void UCurveBase::PostInitProperties()
diff --git a/Engine/Source/Runtime/Engine/Private/DataTable.cpp b/Engine/Source/Runtime/Engine/Private/DataTable.cpp
--- a/Engine/Source/Runtime/Engine/Private/DataTable.cpp
+++ b/Engine/Source/Runtime/Engine/Private/DataTable.cpp
@@ -7,6 +7,7 @@
 #include "Serialization/PropertyLocalizationDataGathering.h"
 #include "Serialization/ObjectWriter.h"
 #include "Serialization/ObjectReader.h"
+#include "UObject/AssetRegistryTagsContext.h"
 #include "UObject/LinkerLoad.h"
 #include "DataTableCSV.h"
 #include "DataTableJSON.h"
@@ -318,19 +319,26 @@ FTopLevelAssetPath UDataTable::GetRowStructPathName() const
 }
 
 void UDataTable::GetAssetRegistryTags(TArray<FAssetRegistryTag>& OutTags) const
+{
+	PRAGMA_DISABLE_DEPRECATION_WARNINGS;
+	Super::GetAssetRegistryTags(OutTags);
+	PRAGMA_ENABLE_DEPRECATION_WARNINGS;
+}
+
+void UDataTable::GetAssetRegistryTags(FAssetRegistryTagsContext Context) const
 {
 	if (AssetImportData)
 	{
-		OutTags.Add( FAssetRegistryTag(SourceFileTagName(), AssetImportData->GetSourceData().ToJson(), FAssetRegistryTag::TT_Hidden) );
+		Context.AddTag( FAssetRegistryTag(SourceFileTagName(), AssetImportData->GetSourceData().ToJson(), FAssetRegistryTag::TT_Hidden) );
 	}
 
 	// Add the row structure tag
 	{
 		static const FName RowStructureTag = "RowStructure";
-		OutTags.Add( FAssetRegistryTag(RowStructureTag, GetRowStructPathName().ToString(), FAssetRegistryTag::TT_Alphabetical) );
+		Context.AddTag( FAssetRegistryTag(RowStructureTag, GetRowStructPathName().ToString(), FAssetRegistryTag::TT_Alphabetical) );
 	}
 
-	Super::GetAssetRegistryTags(OutTags);
+	Super::GetAssetRegistryTags(Context);
 }
 
 void UDataTable::PostInitProperties()
diff --git a/Engine/Source/Runtime/Engine/Private/DialogueVoice.cpp b/Engine/Source/Runtime/Engine/Private/DialogueVoice.cpp
--- a/Engine/Source/Runtime/Engine/Private/DialogueVoice.cpp
+++ b/Engine/Source/Runtime/Engine/Private/DialogueVoice.cpp
@@ -2,6 +2,7 @@
 
 #include "Sound/DialogueVoice.h"
 #include "Sound/DialogueTypes.h"
+#include "UObject/AssetRegistryTagsContext.h"
 #include "UObject/UnrealType.h"
 
 #include UE_INLINE_GENERATED_CPP_BY_NAME(DialogueVoice)
@@ -44,7 +45,14 @@ FString UDialogueVoice::GetDesc()
 
 void UDialogueVoice::GetAssetRegistryTags(TArray<FAssetRegistryTag>& OutTags) const
 {
+	PRAGMA_DISABLE_DEPRECATION_WARNINGS;
 	Super::GetAssetRegistryTags(OutTags);
+	PRAGMA_ENABLE_DEPRECATION_WARNINGS;
+}
+
+void UDialogueVoice::GetAssetRegistryTags(FAssetRegistryTagsContext Context) const
+{
+	Super::GetAssetRegistryTags(Context);
 }
 
 void UDialogueVoice::PostDuplicate(bool bDuplicateForPIE)
diff --git a/Engine/Source/Runtime/Engine/Private/DialogueWave.cpp b/Engine/Source/Runtime/Engine/Private/DialogueWave.cpp
--- a/Engine/Source/Runtime/Engine/Private/DialogueWave.cpp
+++ b/Engine/Source/Runtime/Engine/Private/DialogueWave.cpp
@@ -9,6 +9,7 @@
 #include "Sound/DialogueSoundWaveProxy.h"
 #include "Sound/DialogueVoice.h"
 #include "SubtitleManager.h"
+#include "UObject/AssetRegistryTagsContext.h"
 
 #include UE_INLINE_GENERATED_CPP_BY_NAME(DialogueWave)
 
@@ -539,7 +540,14 @@ FString UDialogueWave::GetDesc()
 
 void UDialogueWave::GetAssetRegistryTags(TArray<FAssetRegistryTag>& OutTags) const
 {
+	PRAGMA_DISABLE_DEPRECATION_WARNINGS;
 	Super::GetAssetRegistryTags(OutTags);
+	PRAGMA_ENABLE_DEPRECATION_WARNINGS;
+}
+
+void UDialogueWave::GetAssetRegistryTags(FAssetRegistryTagsContext Context) const
+{
+	Super::GetAssetRegistryTags(Context);
 }
 
 void UDialogueWave::PostDuplicate(bool bDuplicateForPIE)
diff --git a/Engine/Source/Runtime/Engine/Private/FontFace.cpp b/Engine/Source/Runtime/Engine/Private/FontFace.cpp
--- a/Engine/Source/Runtime/Engine/Private/FontFace.cpp
+++ b/Engine/Source/Runtime/Engine/Private/FontFace.cpp
@@ -7,6 +7,7 @@
 #include "Framework/Application/SlateApplication.h"
 #include "Fonts/FontCache.h"
 #include "Rendering/SlateRenderer.h"
+#include "UObject/AssetRegistryTagsContext.h"
 #include "UObject/Package.h"
 #include "Misc/PackageName.h"
 #include "UObject/EditorObjectVersion.h"
@@ -130,11 +131,18 @@ void UFontFace::PostEditUndo()
 
 void UFontFace::GetAssetRegistryTags(TArray<FAssetRegistryTag>& OutTags) const
 {
+	PRAGMA_DISABLE_DEPRECATION_WARNINGS;
 	Super::GetAssetRegistryTags(OutTags);
+	PRAGMA_ENABLE_DEPRECATION_WARNINGS;
+}
+
+void UFontFace::GetAssetRegistryTags(FAssetRegistryTagsContext Context) const
+{
+	Super::GetAssetRegistryTags(Context);
 
 	FAssetImportInfo ImportInfo;
 	ImportInfo.Insert(FAssetImportInfo::FSourceFile(SourceFilename));
-	OutTags.Add(FAssetRegistryTag(SourceFileTagName(), ImportInfo.ToJson(), FAssetRegistryTag::TT_Hidden));
+	Context.AddTag(FAssetRegistryTag(SourceFileTagName(), ImportInfo.ToJson(), FAssetRegistryTag::TT_Hidden));
 }
 
 void UFontFace::CookAdditionalFilesOverride(const TCHAR* PackageFilename, const ITargetPlatform* TargetPlatform,
diff --git a/Engine/Source/Runtime/Engine/Private/Materials/Material.cpp b/Engine/Source/Runtime/Engine/Private/Materials/Material.cpp
--- a/Engine/Source/Runtime/Engine/Private/Materials/Material.cpp
+++ b/Engine/Source/Runtime/Engine/Private/Materials/Material.cpp
@@ -9,6 +9,7 @@
 #include "Stats/StatsMisc.h"
 #include "Misc/FeedbackContext.h"
 #include "Stats/StatsTrace.h"
+#include "UObject/AssetRegistryTagsContext.h"
 #include "UObject/ObjectSaveContext.h"
 #include "UObject/FortniteMainBranchObjectVersion.h"
 #include "UObject/UObjectIterator.h"
@@ -5439,7 +5440,14 @@ bool UMaterial::CanBeClusterRoot() const
 
 void UMaterial::GetAssetRegistryTags(TArray<FAssetRegistryTag>& OutTags) const
 {
+	PRAGMA_DISABLE_DEPRECATION_WARNINGS;
 	Super::GetAssetRegistryTags(OutTags);
+	PRAGMA_ENABLE_DEPRECATION_WARNINGS;
+}
+
+void UMaterial::GetAssetRegistryTags(FAssetRegistryTagsContext Context) const
+{
+	Super::GetAssetRegistryTags(Context);
 }
 
 #if WITH_EDITOR
diff --git a/Engine/Source/Runtime/Engine/Private/Materials/MaterialExpressions.cpp b/Engine/Source/Runtime/Engine/Private/Materials/MaterialExpressions.cpp
--- a/Engine/Source/Runtime/Engine/Private/Materials/MaterialExpressions.cpp
+++ b/Engine/Source/Runtime/Engine/Private/Materials/MaterialExpressions.cpp
@@ -7,6 +7,7 @@
 #include "Field/FieldSystemTypes.h"
 #include "Misc/MessageDialog.h"
 #include "HLSLTree/HLSLTreeTypes.h"
+#include "UObject/AssetRegistryTagsContext.h"
 #include "UObject/UE5MainStreamObjectVersion.h"
 #include "Internationalization/LocKeyFuncs.h"
 #include "UObject/UE5LWCRenderingStreamObjectVersion.h"
@@ -14479,16 +14480,23 @@ void UMaterialFunctionInterface::DeclareConstructClasses(TArray<FTopLevelAssetPa
 
 void UMaterialFunctionInterface::GetAssetRegistryTags(TArray<FAssetRegistryTag>& OutTags) const
 {
+	PRAGMA_DISABLE_DEPRECATION_WARNINGS;
 	Super::GetAssetRegistryTags(OutTags);
+	PRAGMA_ENABLE_DEPRECATION_WARNINGS;
+}
+
+void UMaterialFunctionInterface::GetAssetRegistryTags(FAssetRegistryTagsContext Context) const
+{
+	Super::GetAssetRegistryTags(Context);
 
 #if WITH_EDITORONLY_DATA
-	for (FAssetRegistryTag& AssetTag : OutTags)
+	for (FName TagName : { GET_MEMBER_NAME_CHECKED(UMaterialFunctionInterface, CombinedInputTypes),
+		GET_MEMBER_NAME_CHECKED(UMaterialFunctionInterface, CombinedOutputTypes)})
 	{
 		// Hide the combined input/output types as they are only needed in code
-		if (AssetTag.Name == GET_MEMBER_NAME_CHECKED(UMaterialFunctionInterface, CombinedInputTypes)
-		|| AssetTag.Name == GET_MEMBER_NAME_CHECKED(UMaterialFunctionInterface, CombinedOutputTypes))
+		if (FAssetRegistryTag* AssetTag = Context.FindTag(TagName); AssetTag)
 		{
-			AssetTag.Type = UObject::FAssetRegistryTag::TT_Hidden;
+			AssetTag->Type = UObject::FAssetRegistryTag::TT_Hidden;
 		}
 	}
 #endif
diff --git a/Engine/Source/Runtime/Engine/Private/Materials/MaterialInterface.cpp b/Engine/Source/Runtime/Engine/Private/Materials/MaterialInterface.cpp
--- a/Engine/Source/Runtime/Engine/Private/Materials/MaterialInterface.cpp
+++ b/Engine/Source/Runtime/Engine/Private/Materials/MaterialInterface.cpp
@@ -12,6 +12,7 @@
 #include "PrimitiveViewRelevance.h"
 #include "MaterialShared.h"
 #include "Materials/Material.h"
+#include "UObject/AssetRegistryTagsContext.h"
 #include "UObject/ObjectSaveContext.h"
 #include "UObject/PropertyPortFlags.h"
 #include "UObject/UObjectIterator.h"
@@ -765,21 +766,28 @@ void UMaterialInterface::PostEditChangeProperty(FPropertyChangedEvent& PropertyC
 }
 
 void UMaterialInterface::GetAssetRegistryTags(TArray<FAssetRegistryTag>& OutTags) const
+{
+	PRAGMA_DISABLE_DEPRECATION_WARNINGS;
+	Super::GetAssetRegistryTags(OutTags);
+	PRAGMA_ENABLE_DEPRECATION_WARNINGS;
+}
+
+void UMaterialInterface::GetAssetRegistryTags(FAssetRegistryTagsContext Context) const
 {
 	if (AssetImportData)
 	{
-		OutTags.Add( FAssetRegistryTag(SourceFileTagName(), AssetImportData->GetSourceData().ToJson(), FAssetRegistryTag::TT_Hidden) );
+		Context.AddTag( FAssetRegistryTag(SourceFileTagName(), AssetImportData->GetSourceData().ToJson(), FAssetRegistryTag::TT_Hidden) );
 	}
 
 	{
 		const FMaterialCachedExpressionData& CachedData = GetCachedExpressionData();
-		OutTags.Add(FAssetRegistryTag("HasSceneColor", CachedData.bHasSceneColor ? TEXT("True") : TEXT("False"), FAssetRegistryTag::TT_Alphabetical));
-		OutTags.Add(FAssetRegistryTag("HasPerInstanceRandom", CachedData.bHasPerInstanceRandom ? TEXT("True") : TEXT("False"), FAssetRegistryTag::TT_Alphabetical));
-		OutTags.Add(FAssetRegistryTag("HasPerInstanceCustomData", CachedData.bHasPerInstanceCustomData ? TEXT("True") : TEXT("False"), FAssetRegistryTag::TT_Alphabetical));
-		OutTags.Add(FAssetRegistryTag("HasVertexInterpolator", CachedData.bHasVertexInterpolator ? TEXT("True") : TEXT("False"), FAssetRegistryTag::TT_Alphabetical));
+		Context.AddTag(FAssetRegistryTag("HasSceneColor", CachedData.bHasSceneColor ? TEXT("True") : TEXT("False"), FAssetRegistryTag::TT_Alphabetical));
+		Context.AddTag(FAssetRegistryTag("HasPerInstanceRandom", CachedData.bHasPerInstanceRandom ? TEXT("True") : TEXT("False"), FAssetRegistryTag::TT_Alphabetical));
+		Context.AddTag(FAssetRegistryTag("HasPerInstanceCustomData", CachedData.bHasPerInstanceCustomData ? TEXT("True") : TEXT("False"), FAssetRegistryTag::TT_Alphabetical));
+		Context.AddTag(FAssetRegistryTag("HasVertexInterpolator", CachedData.bHasVertexInterpolator ? TEXT("True") : TEXT("False"), FAssetRegistryTag::TT_Alphabetical));
 	}
 
-	Super::GetAssetRegistryTags(OutTags);
+	Super::GetAssetRegistryTags(Context);
 }
 #endif // WITH_EDITOR
 
diff --git a/Engine/Source/Runtime/Engine/Private/Particles/ParticleComponents.cpp b/Engine/Source/Runtime/Engine/Private/Particles/ParticleComponents.cpp
--- a/Engine/Source/Runtime/Engine/Private/Particles/ParticleComponents.cpp
+++ b/Engine/Source/Runtime/Engine/Private/Particles/ParticleComponents.cpp
@@ -11,6 +11,7 @@
 #include "Engine/Level.h"
 #include "Engine/CollisionProfile.h"
 #include "GameFramework/PlayerController.h"
+#include "UObject/AssetRegistryTagsContext.h"
 #include "UObject/ObjectSaveContext.h"
 #include "Materials/MaterialRelevance.h"
 #include "UObject/UObjectIterator.h"
@@ -2842,16 +2843,23 @@ void UParticleSystem::UpdateColorModuleClampAlpha(UParticleModuleColorBase* Colo
 
 void UParticleSystem::GetAssetRegistryTags(TArray<FAssetRegistryTag>& OutTags) const
 {
-	OutTags.Add( FAssetRegistryTag("HasGPUEmitter", HasGPUEmitter() ? TEXT("True") : TEXT("False"), FAssetRegistryTag::TT_Alphabetical) );
+	PRAGMA_DISABLE_DEPRECATION_WARNINGS;
+	Super::GetAssetRegistryTags(OutTags);
+	PRAGMA_ENABLE_DEPRECATION_WARNINGS;
+}
+
+void UParticleSystem::GetAssetRegistryTags(FAssetRegistryTagsContext Context) const
+{
+	Context.AddTag( FAssetRegistryTag("HasGPUEmitter", HasGPUEmitter() ? TEXT("True") : TEXT("False"), FAssetRegistryTag::TT_Alphabetical) );
 
 	const float BoundsSize = FixedRelativeBoundingBox.GetSize().GetMax();
-	OutTags.Add(FAssetRegistryTag("FixedBoundsSize", bUseFixedRelativeBoundingBox ? FString::Printf(TEXT("%.2f"), BoundsSize) : FString(TEXT("None")), FAssetRegistryTag::TT_Numerical));
+	Context.AddTag(FAssetRegistryTag("FixedBoundsSize", bUseFixedRelativeBoundingBox ? FString::Printf(TEXT("%.2f"), BoundsSize) : FString(TEXT("None")), FAssetRegistryTag::TT_Numerical));
 
-	OutTags.Add(FAssetRegistryTag("NumEmitters", LexToString(Emitters.Num()), FAssetRegistryTag::TT_Numerical));
+	Context.AddTag(FAssetRegistryTag("NumEmitters", LexToString(Emitters.Num()), FAssetRegistryTag::TT_Numerical));
 
-	OutTags.Add(FAssetRegistryTag("NumLODs", LexToString(LODDistances.Num()), FAssetRegistryTag::TT_Numerical));
+	Context.AddTag(FAssetRegistryTag("NumLODs", LexToString(LODDistances.Num()), FAssetRegistryTag::TT_Numerical));
 
-	OutTags.Add(FAssetRegistryTag("WarmupTime", LexToString(WarmupTime), FAssetRegistryTag::TT_Numerical));
+	Context.AddTag(FAssetRegistryTag("WarmupTime", LexToString(WarmupTime), FAssetRegistryTag::TT_Numerical));
 
 	// Done here instead of as an AssetRegistrySearchable string to avoid the long prefix on the enum value string
 	FString LODMethodString = TEXT("Unknown");
@@ -2870,13 +2878,13 @@ void UParticleSystem::GetAssetRegistryTags(TArray<FAssetRegistryTag>& OutTags) c
 		check(false); // Missing enum entry
 		break;
 	}
-	OutTags.Add(FAssetRegistryTag("LODMethod", LODMethodString, FAssetRegistryTag::TT_Alphabetical));
+	Context.AddTag(FAssetRegistryTag("LODMethod", LODMethodString, FAssetRegistryTag::TT_Alphabetical));
 
-	OutTags.Add(FAssetRegistryTag("CPUCollision", UsesCPUCollision() ? TEXT("True") : TEXT("False"), FAssetRegistryTag::TT_Alphabetical));
-	OutTags.Add(FAssetRegistryTag("Looping", bAnyEmitterLoopsForever ? TEXT("True") : TEXT("False"), FAssetRegistryTag::TT_Alphabetical));
-	OutTags.Add(FAssetRegistryTag("Immortal", IsImmortal() ? TEXT("True") : TEXT("False"), FAssetRegistryTag::TT_Alphabetical));
-	OutTags.Add(FAssetRegistryTag("Becomes Zombie", WillBecomeZombie() ? TEXT("True") : TEXT("False"), FAssetRegistryTag::TT_Alphabetical));
-	OutTags.Add(FAssetRegistryTag("CanBeOccluded", OcclusionBoundsMethod == EParticleSystemOcclusionBoundsMethod::EPSOBM_None ? TEXT("False") : TEXT("True"), FAssetRegistryTag::TT_Alphabetical));
+	Context.AddTag(FAssetRegistryTag("CPUCollision", UsesCPUCollision() ? TEXT("True") : TEXT("False"), FAssetRegistryTag::TT_Alphabetical));
+	Context.AddTag(FAssetRegistryTag("Looping", bAnyEmitterLoopsForever ? TEXT("True") : TEXT("False"), FAssetRegistryTag::TT_Alphabetical));
+	Context.AddTag(FAssetRegistryTag("Immortal", IsImmortal() ? TEXT("True") : TEXT("False"), FAssetRegistryTag::TT_Alphabetical));
+	Context.AddTag(FAssetRegistryTag("Becomes Zombie", WillBecomeZombie() ? TEXT("True") : TEXT("False"), FAssetRegistryTag::TT_Alphabetical));
+	Context.AddTag(FAssetRegistryTag("CanBeOccluded", OcclusionBoundsMethod == EParticleSystemOcclusionBoundsMethod::EPSOBM_None ? TEXT("False") : TEXT("True"), FAssetRegistryTag::TT_Alphabetical));
 
 	uint32 NumEmittersAtEachSig[(int32)EParticleSignificanceLevel::Num] = { 0, 0, 0, 0 };
 	for (UParticleEmitter* Emitter : Emitters)
@@ -2886,12 +2894,12 @@ void UParticleSystem::GetAssetRegistryTags(TArray<FAssetRegistryTag>& OutTags) c
 			++NumEmittersAtEachSig[(int32)Emitter->SignificanceLevel];			
 		}
 	}
-	OutTags.Add(FAssetRegistryTag("Critical Emitters", LexToString(NumEmittersAtEachSig[(int32)EParticleSignificanceLevel::Critical]), FAssetRegistryTag::TT_Numerical));
-	OutTags.Add(FAssetRegistryTag("High Emitters", LexToString(NumEmittersAtEachSig[(int32)EParticleSignificanceLevel::High]), FAssetRegistryTag::TT_Numerical));
-	OutTags.Add(FAssetRegistryTag("Medium Emitters", LexToString(NumEmittersAtEachSig[(int32)EParticleSignificanceLevel::Medium]), FAssetRegistryTag::TT_Numerical));
-	OutTags.Add(FAssetRegistryTag("Low Emitters", LexToString(NumEmittersAtEachSig[(int32)EParticleSignificanceLevel::Low]), FAssetRegistryTag::TT_Numerical));
+	Context.AddTag(FAssetRegistryTag("Critical Emitters", LexToString(NumEmittersAtEachSig[(int32)EParticleSignificanceLevel::Critical]), FAssetRegistryTag::TT_Numerical));
+	Context.AddTag(FAssetRegistryTag("High Emitters", LexToString(NumEmittersAtEachSig[(int32)EParticleSignificanceLevel::High]), FAssetRegistryTag::TT_Numerical));
+	Context.AddTag(FAssetRegistryTag("Medium Emitters", LexToString(NumEmittersAtEachSig[(int32)EParticleSignificanceLevel::Medium]), FAssetRegistryTag::TT_Numerical));
+	Context.AddTag(FAssetRegistryTag("Low Emitters", LexToString(NumEmittersAtEachSig[(int32)EParticleSignificanceLevel::Low]), FAssetRegistryTag::TT_Numerical));
 
-	Super::GetAssetRegistryTags(OutTags);
+	Super::GetAssetRegistryTags(Context);
 }
 
 bool UParticleSystem::UsesCPUCollision() const
diff --git a/Engine/Source/Runtime/Engine/Private/PhysicsEngine/PhysicsAsset.cpp b/Engine/Source/Runtime/Engine/Private/PhysicsEngine/PhysicsAsset.cpp
--- a/Engine/Source/Runtime/Engine/Private/PhysicsEngine/PhysicsAsset.cpp
+++ b/Engine/Source/Runtime/Engine/Private/PhysicsEngine/PhysicsAsset.cpp
@@ -12,6 +12,7 @@
 #include "Serialization/ObjectReader.h"
 #include "Components/SkeletalMeshComponent.h"
 #include "PhysicsEngine/PhysicsConstraintTemplate.h"
+#include "UObject/AssetRegistryTagsContext.h"
 #include "UObject/ReleaseObjectVersion.h"
 #include "UObject/UObjectIterator.h"
 #include "UObject/FortniteSeasonBranchObjectVersion.h"
@@ -842,10 +843,17 @@ FString UPhysicsAsset::GetDesc()
 
 void UPhysicsAsset::GetAssetRegistryTags(TArray<FAssetRegistryTag>& OutTags) const
 {
-	OutTags.Add( FAssetRegistryTag("Bodies", FString::FromInt(SkeletalBodySetups.Num()), FAssetRegistryTag::TT_Numerical) );
-	OutTags.Add( FAssetRegistryTag("Constraints", FString::FromInt(ConstraintSetup.Num()), FAssetRegistryTag::TT_Numerical) );
-
+	PRAGMA_DISABLE_DEPRECATION_WARNINGS;
 	Super::GetAssetRegistryTags(OutTags);
+	PRAGMA_ENABLE_DEPRECATION_WARNINGS;
+}
+
+void UPhysicsAsset::GetAssetRegistryTags(FAssetRegistryTagsContext Context) const
+{
+	Context.AddTag( FAssetRegistryTag("Bodies", FString::FromInt(SkeletalBodySetups.Num()), FAssetRegistryTag::TT_Numerical) );
+	Context.AddTag( FAssetRegistryTag("Constraints", FString::FromInt(ConstraintSetup.Num()), FAssetRegistryTag::TT_Numerical) );
+
+	Super::GetAssetRegistryTags(Context);
 }
 
 
diff --git a/Engine/Source/Runtime/Engine/Private/SkeletalMesh.cpp b/Engine/Source/Runtime/Engine/Private/SkeletalMesh.cpp
--- a/Engine/Source/Runtime/Engine/Private/SkeletalMesh.cpp
+++ b/Engine/Source/Runtime/Engine/Private/SkeletalMesh.cpp
@@ -23,6 +23,7 @@
 #include "Components/SkeletalMeshComponent.h"
 #include "ComponentReregisterContext.h"
 #include "Math/ScaleRotationTranslationMatrix.h"
+#include "UObject/AssetRegistryTagsContext.h"
 #include "UObject/FrameworkObjectVersion.h"
 #include "PhysicsEngine/SphereElem.h"
 #include "UObject/NiagaraObjectVersion.h"
@@ -3491,6 +3492,13 @@ void USkeletalMesh::RebuildRefSkeletonNameToIndexMap()
 
 
 void USkeletalMesh::GetAssetRegistryTags(TArray<FAssetRegistryTag>& OutTags) const
+{
+	PRAGMA_DISABLE_DEPRECATION_WARNINGS;
+	Super::GetAssetRegistryTags(OutTags);
+	PRAGMA_ENABLE_DEPRECATION_WARNINGS;
+}
+
+void USkeletalMesh::GetAssetRegistryTags(FAssetRegistryTagsContext Context) const
 {
 #if WITH_EDITOR
 	// Avoid accessing properties being compiled, this function will get called again after compilation is finished.
@@ -3500,9 +3508,9 @@ void USkeletalMesh::GetAssetRegistryTags(TArray<FAssetRegistryTag>& OutTags) con
 		if (AssetRegistry)
 		{
 			FAssetData AssetData = AssetRegistry->GetAssetByObjectPath(FSoftObjectPath(this), true /* bIncludeOnlyOnDiskAssets */);
-			AssetData.EnumerateTags([&OutTags](const TPair<FName, FAssetTagValueRef>& Pair)
+			AssetData.EnumerateTags([&Context](const TPair<FName, FAssetTagValueRef>& Pair)
 				{
-					OutTags.Add(FAssetRegistryTag(Pair.Key, Pair.Value.GetStorageString(), FAssetRegistryTag::TT_Alphabetical));
+					Context.AddTag(FAssetRegistryTag(Pair.Key, Pair.Value.GetStorageString(), FAssetRegistryTag::TT_Alphabetical));
 				});
 		}
 		return;
@@ -3521,19 +3529,27 @@ void USkeletalMesh::GetAssetRegistryTags(TArray<FAssetRegistryTag>& OutTags) con
 	
 	int32 NumLODs = GetLODInfoArray().Num();
 
-	OutTags.Add(FAssetRegistryTag("Vertices", FString::FromInt(NumVertices), FAssetRegistryTag::TT_Numerical));
-	OutTags.Add(FAssetRegistryTag("Triangles", FString::FromInt(NumTriangles), FAssetRegistryTag::TT_Numerical));
-	OutTags.Add(FAssetRegistryTag("LODs", FString::FromInt(NumLODs), FAssetRegistryTag::TT_Numerical));
-	OutTags.Add(FAssetRegistryTag("Bones", FString::FromInt(GetRefSkeleton().GetRawBoneNum()), FAssetRegistryTag::TT_Numerical));
-	OutTags.Add(FAssetRegistryTag("MorphTargets", FString::FromInt(GetMorphTargets().Num()), FAssetRegistryTag::TT_Numerical));
-	OutTags.Add(FAssetRegistryTag("SkinWeightProfiles", FString::FromInt(GetSkinWeightProfiles().Num()), FAssetRegistryTag::TT_Numerical));
+	Context.AddTag(FAssetRegistryTag("Vertices", FString::FromInt(NumVertices), FAssetRegistryTag::TT_Numerical));
+	Context.AddTag(FAssetRegistryTag("Triangles", FString::FromInt(NumTriangles), FAssetRegistryTag::TT_Numerical));
+	Context.AddTag(FAssetRegistryTag("LODs", FString::FromInt(NumLODs), FAssetRegistryTag::TT_Numerical));
+	Context.AddTag(FAssetRegistryTag("Bones", FString::FromInt(GetRefSkeleton().GetRawBoneNum()), FAssetRegistryTag::TT_Numerical));
+	Context.AddTag(FAssetRegistryTag("MorphTargets", FString::FromInt(GetMorphTargets().Num()), FAssetRegistryTag::TT_Numerical));
+	Context.AddTag(FAssetRegistryTag("SkinWeightProfiles", FString::FromInt(GetSkinWeightProfiles().Num()), FAssetRegistryTag::TT_Numerical));
 
 #if WITH_EDITORONLY_DATA
 	if (GetAssetImportData())
 	{
-		OutTags.Add( FAssetRegistryTag(SourceFileTagName(), GetAssetImportData()->GetSourceData().ToJson(), FAssetRegistryTag::TT_Hidden) );
+		Context.AddTag( FAssetRegistryTag(SourceFileTagName(), GetAssetImportData()->GetSourceData().ToJson(), FAssetRegistryTag::TT_Hidden) );
 #if WITH_EDITOR
-		GetAssetImportData()->AppendAssetRegistryTags(OutTags);
+		PRAGMA_DISABLE_DEPRECATION_WARNINGS;
+		TArray<UObject::FAssetRegistryTag> DeprecatedFunctionTags;
+		GetAssetImportData()->AppendAssetRegistryTags(DeprecatedFunctionTags);
+		for (UObject::FAssetRegistryTag& Tag : DeprecatedFunctionTags)
+		{
+			Context.AddTag(MoveTemp(Tag));
+		}
+		PRAGMA_ENABLE_DEPRECATION_WARNINGS;
+		GetAssetImportData()->AppendAssetRegistryTags(Context);
 #endif
 	}
 
@@ -3555,19 +3571,19 @@ void USkeletalMesh::GetAssetRegistryTags(TArray<FAssetRegistryTag>& OutTags) con
 
 	// The tag must be added unconditionally, because some code calls this function on the CDO to find out what
 	// tags are available.
-	OutTags.Add(FAssetRegistryTag("MaxBoneInfluences", MaxBoneInfluencesString, FAssetRegistryTag::TT_Numerical));
+	Context.AddTag(FAssetRegistryTag("MaxBoneInfluences", MaxBoneInfluencesString, FAssetRegistryTag::TT_Numerical));
 
 	// Allow asset user data to output tags
 	for(UAssetUserData* AssetUserDataItem : *GetAssetUserDataArray())
 	{
 		if(AssetUserDataItem)
 		{
-			AssetUserDataItem->GetAssetRegistryTags(OutTags);
+			AssetUserDataItem->GetAssetRegistryTags(Context);
 		}
 	}
 #endif // WITH_EDITORONLY_DATA
 	
-	Super::GetAssetRegistryTags(OutTags);
+	Super::GetAssetRegistryTags(Context);
 }
 
 #if WITH_EDITOR
diff --git a/Engine/Source/Runtime/Engine/Private/SoundWave.cpp b/Engine/Source/Runtime/Engine/Private/SoundWave.cpp
--- a/Engine/Source/Runtime/Engine/Private/SoundWave.cpp
+++ b/Engine/Source/Runtime/Engine/Private/SoundWave.cpp
@@ -12,6 +12,7 @@
 #include "Sound/SoundSourceBus.h"
 #include "Sound/SoundSubmix.h"
 #include "Stats/StatsTrace.h"
+#include "UObject/AssetRegistryTagsContext.h"
 #include "UObject/FrameworkObjectVersion.h"
 #include "UObject/Package.h"
 #include "UObject/UObjectIterator.h"
@@ -1062,12 +1063,19 @@ FString USoundWave::GetDesc()
 
 void USoundWave::GetAssetRegistryTags(TArray<FAssetRegistryTag>& OutTags) const
 {
+	PRAGMA_DISABLE_DEPRECATION_WARNINGS;
 	Super::GetAssetRegistryTags(OutTags);
+	PRAGMA_ENABLE_DEPRECATION_WARNINGS;
+}
+
+void USoundWave::GetAssetRegistryTags(FAssetRegistryTagsContext Context) const
+{
+	Super::GetAssetRegistryTags(Context);
 
 #if WITH_EDITORONLY_DATA
 	if (AssetImportData)
 	{
-		OutTags.Add( FAssetRegistryTag(SourceFileTagName(), AssetImportData->GetSourceData().ToJson(), FAssetRegistryTag::TT_Hidden) );
+		Context.AddTag( FAssetRegistryTag(SourceFileTagName(), AssetImportData->GetSourceData().ToJson(), FAssetRegistryTag::TT_Hidden) );
 	}
 #endif
 }
diff --git a/Engine/Source/Runtime/Engine/Private/SoundWaveProcedural.cpp b/Engine/Source/Runtime/Engine/Private/SoundWaveProcedural.cpp
--- a/Engine/Source/Runtime/Engine/Private/SoundWaveProcedural.cpp
+++ b/Engine/Source/Runtime/Engine/Private/SoundWaveProcedural.cpp
@@ -4,6 +4,7 @@
 
 #include "AudioDevice.h"
 #include "Engine/Engine.h"
+#include "UObject/AssetRegistryTagsContext.h"
 
 #include UE_INLINE_GENERATED_CPP_BY_NAME(SoundWaveProcedural)
 
@@ -150,7 +151,14 @@ int32 USoundWaveProcedural::GetResourceSizeForFormat(FName Format)
 
 void USoundWaveProcedural::GetAssetRegistryTags(TArray<FAssetRegistryTag>& OutTags) const
 {
+	PRAGMA_DISABLE_DEPRECATION_WARNINGS;
 	Super::GetAssetRegistryTags(OutTags);
+	PRAGMA_ENABLE_DEPRECATION_WARNINGS;
+}
+
+void USoundWaveProcedural::GetAssetRegistryTags(FAssetRegistryTagsContext Context) const
+{
+	Super::GetAssetRegistryTags(Context);
 }
 
 bool USoundWaveProcedural::HasCompressedData(FName Format, ITargetPlatform* TargetPlatform) const
diff --git a/Engine/Source/Runtime/Engine/Private/SparseVolumeTexture/SparseVolumeTexture.cpp b/Engine/Source/Runtime/Engine/Private/SparseVolumeTexture/SparseVolumeTexture.cpp
--- a/Engine/Source/Runtime/Engine/Private/SparseVolumeTexture/SparseVolumeTexture.cpp
+++ b/Engine/Source/Runtime/Engine/Private/SparseVolumeTexture/SparseVolumeTexture.cpp
@@ -14,6 +14,7 @@
 #include "Shader/ShaderTypes.h"
 #include "RenderingThread.h"
 #include "GlobalRenderResources.h"
+#include "UObject/AssetRegistryTagsContext.h"
 #include "UObject/Package.h"
 #include "SparseVolumeTexture/SparseVolumeTextureData.h"
 #include "SparseVolumeTexture/SparseVolumeTextureUtility.h"
@@ -1314,12 +1315,19 @@ void UStreamableSparseVolumeTexture::GetResourceSizeEx(FResourceSizeEx& Cumulati
 
 void UStreamableSparseVolumeTexture::GetAssetRegistryTags(TArray<FAssetRegistryTag>& OutTags) const
 {
+	PRAGMA_DISABLE_DEPRECATION_WARNINGS;
 	Super::GetAssetRegistryTags(OutTags);
+	PRAGMA_ENABLE_DEPRECATION_WARNINGS;
+}
+
+void UStreamableSparseVolumeTexture::GetAssetRegistryTags(FAssetRegistryTagsContext Context) const
+{
+	Super::GetAssetRegistryTags(Context);
 
 #if WITH_EDITORONLY_DATA
 	if (AssetImportData)
 	{
-		OutTags.Add(FAssetRegistryTag(SourceFileTagName(), AssetImportData->GetSourceData().ToJson(), FAssetRegistryTag::TT_Hidden));
+		Context.AddTag(FAssetRegistryTag(SourceFileTagName(), AssetImportData->GetSourceData().ToJson(), FAssetRegistryTag::TT_Hidden));
 	}
 #endif
 }
diff --git a/Engine/Source/Runtime/Engine/Private/StaticMesh.cpp b/Engine/Source/Runtime/Engine/Private/StaticMesh.cpp
--- a/Engine/Source/Runtime/Engine/Private/StaticMesh.cpp
+++ b/Engine/Source/Runtime/Engine/Private/StaticMesh.cpp
@@ -13,6 +13,7 @@
 #include "Misc/ScopedSlowTask.h"
 #include "Engine/CollisionProfile.h"
 #include "Math/ScaleRotationTranslationMatrix.h"
+#include "UObject/AssetRegistryTagsContext.h"
 #include "UObject/EditorObjectVersion.h"
 #include "PhysicsEngine/BoxElem.h"
 #include "UObject/FrameworkObjectVersion.h"
@@ -4728,14 +4729,21 @@ int32 UStaticMesh::GetNumSectionsWithCollision() const
 }
 
 void UStaticMesh::GetAssetRegistryTags(TArray<FAssetRegistryTag>& OutTags) const
+{
+	PRAGMA_DISABLE_DEPRECATION_WARNINGS;
+	Super::GetAssetRegistryTags(OutTags);
+	PRAGMA_ENABLE_DEPRECATION_WARNINGS;
+}
+
+void UStaticMesh::GetAssetRegistryTags(FAssetRegistryTagsContext Context) const
 {
 #if WITH_EDITORONLY_DATA
-	OutTags.Add(FAssetRegistryTag("NaniteEnabled", IsNaniteEnabled() ? TEXT("True") : TEXT("False"), FAssetRegistryTag::TT_Alphabetical));
-	OutTags.Add(FAssetRegistryTag("NaniteFallbackPercent", FString::Printf(TEXT("%.1f"), NaniteSettings.FallbackPercentTriangles * 100.0f), FAssetRegistryTag::TT_Numerical));
+	Context.AddTag(FAssetRegistryTag("NaniteEnabled", IsNaniteEnabled() ? TEXT("True") : TEXT("False"), FAssetRegistryTag::TT_Alphabetical));
+	Context.AddTag(FAssetRegistryTag("NaniteFallbackPercent", FString::Printf(TEXT("%.1f"), NaniteSettings.FallbackPercentTriangles * 100.0f), FAssetRegistryTag::TT_Numerical));
 
 	if (AssetImportData)
 	{
-		OutTags.Add(FAssetRegistryTag(SourceFileTagName(), AssetImportData->GetSourceData().ToJson(), FAssetRegistryTag::TT_Hidden));
+		Context.AddTag(FAssetRegistryTag(SourceFileTagName(), AssetImportData->GetSourceData().ToJson(), FAssetRegistryTag::TT_Hidden));
 	}
 #endif
 
@@ -4810,29 +4818,29 @@ void UStaticMesh::GetAssetRegistryTags(TArray<FAssetRegistryTag>& OutTags) const
 	}
 #endif
 
-	OutTags.Add(FAssetRegistryTag("NaniteTriangles", FString::FromInt(NumNaniteTriangles), FAssetRegistryTag::TT_Numerical));
-	OutTags.Add(FAssetRegistryTag("NaniteVertices", FString::FromInt(NumNaniteVertices), FAssetRegistryTag::TT_Numerical));
-	OutTags.Add(FAssetRegistryTag("Triangles", FString::FromInt(NumTriangles), FAssetRegistryTag::TT_Numerical) );
-	OutTags.Add(FAssetRegistryTag("Vertices", FString::FromInt(NumVertices), FAssetRegistryTag::TT_Numerical) );
-	OutTags.Add(FAssetRegistryTag("UVChannels", FString::FromInt(NumUVChannels), FAssetRegistryTag::TT_Numerical) );
-	OutTags.Add(FAssetRegistryTag("Materials", FString::FromInt(GetStaticMaterials().Num()), FAssetRegistryTag::TT_Numerical) );
-	OutTags.Add(FAssetRegistryTag("ApproxSize", ApproxSizeStr, FAssetRegistryTag::TT_Dimensional) );
-	OutTags.Add(FAssetRegistryTag("CollisionPrims", FString::FromInt(NumCollisionPrims), FAssetRegistryTag::TT_Numerical));
-	OutTags.Add(FAssetRegistryTag("LODs", FString::FromInt(NumLODs), FAssetRegistryTag::TT_Numerical));
-	OutTags.Add(FAssetRegistryTag("MinLOD", GetMinLOD().ToString(), FAssetRegistryTag::TT_Alphabetical));
-	OutTags.Add(FAssetRegistryTag("QualityLevelMinLOD", GetQualityLevelMinLOD().ToString(), FAssetRegistryTag::TT_Alphabetical));
-	OutTags.Add(FAssetRegistryTag("SectionsWithCollision", FString::FromInt(NumSectionsWithCollision), FAssetRegistryTag::TT_Numerical));
-	OutTags.Add(FAssetRegistryTag("DefaultCollision", DefaultCollisionName.ToString(), FAssetRegistryTag::TT_Alphabetical));
-	OutTags.Add(FAssetRegistryTag("CollisionComplexity", ComplexityString, FAssetRegistryTag::TT_Alphabetical));
-	OutTags.Add(FAssetRegistryTag("DistanceFieldSize", FString::FromInt(DistanceFieldSize), FAssetRegistryTag::TT_Numerical, FAssetRegistryTag::TD_Memory));
-	OutTags.Add(FAssetRegistryTag("EstTotalCompressedSize", FString::Printf(TEXT("%llu"), EstimatedCompressedSize), FAssetRegistryTag::TT_Numerical, FAssetRegistryTag::TD_Memory));
-	OutTags.Add(FAssetRegistryTag("EstNaniteCompressedSize", FString::Printf(TEXT("%llu"), EstimatedNaniteCompressedSize), FAssetRegistryTag::TT_Numerical, FAssetRegistryTag::TD_Memory));
+	Context.AddTag(FAssetRegistryTag("NaniteTriangles", FString::FromInt(NumNaniteTriangles), FAssetRegistryTag::TT_Numerical));
+	Context.AddTag(FAssetRegistryTag("NaniteVertices", FString::FromInt(NumNaniteVertices), FAssetRegistryTag::TT_Numerical));
+	Context.AddTag(FAssetRegistryTag("Triangles", FString::FromInt(NumTriangles), FAssetRegistryTag::TT_Numerical) );
+	Context.AddTag(FAssetRegistryTag("Vertices", FString::FromInt(NumVertices), FAssetRegistryTag::TT_Numerical) );
+	Context.AddTag(FAssetRegistryTag("UVChannels", FString::FromInt(NumUVChannels), FAssetRegistryTag::TT_Numerical) );
+	Context.AddTag(FAssetRegistryTag("Materials", FString::FromInt(GetStaticMaterials().Num()), FAssetRegistryTag::TT_Numerical) );
+	Context.AddTag(FAssetRegistryTag("ApproxSize", ApproxSizeStr, FAssetRegistryTag::TT_Dimensional) );
+	Context.AddTag(FAssetRegistryTag("CollisionPrims", FString::FromInt(NumCollisionPrims), FAssetRegistryTag::TT_Numerical));
+	Context.AddTag(FAssetRegistryTag("LODs", FString::FromInt(NumLODs), FAssetRegistryTag::TT_Numerical));
+	Context.AddTag(FAssetRegistryTag("MinLOD", GetMinLOD().ToString(), FAssetRegistryTag::TT_Alphabetical));
+	Context.AddTag(FAssetRegistryTag("QualityLevelMinLOD", GetQualityLevelMinLOD().ToString(), FAssetRegistryTag::TT_Alphabetical));
+	Context.AddTag(FAssetRegistryTag("SectionsWithCollision", FString::FromInt(NumSectionsWithCollision), FAssetRegistryTag::TT_Numerical));
+	Context.AddTag(FAssetRegistryTag("DefaultCollision", DefaultCollisionName.ToString(), FAssetRegistryTag::TT_Alphabetical));
+	Context.AddTag(FAssetRegistryTag("CollisionComplexity", ComplexityString, FAssetRegistryTag::TT_Alphabetical));
+	Context.AddTag(FAssetRegistryTag("DistanceFieldSize", FString::FromInt(DistanceFieldSize), FAssetRegistryTag::TT_Numerical, FAssetRegistryTag::TD_Memory));
+	Context.AddTag(FAssetRegistryTag("EstTotalCompressedSize", FString::Printf(TEXT("%llu"), EstimatedCompressedSize), FAssetRegistryTag::TT_Numerical, FAssetRegistryTag::TD_Memory));
+	Context.AddTag(FAssetRegistryTag("EstNaniteCompressedSize", FString::Printf(TEXT("%llu"), EstimatedNaniteCompressedSize), FAssetRegistryTag::TT_Numerical, FAssetRegistryTag::TD_Memory));
 	
 #if WITH_EDITORONLY_DATA
-	OutTags.Add(FAssetRegistryTag("HasHiResMesh", IsHiResMeshDescriptionValid() ? TEXT("True") : TEXT("False"), FAssetRegistryTag::TT_Alphabetical));
+	Context.AddTag(FAssetRegistryTag("HasHiResMesh", IsHiResMeshDescriptionValid() ? TEXT("True") : TEXT("False"), FAssetRegistryTag::TT_Alphabetical));
 #endif
 
-	Super::GetAssetRegistryTags(OutTags);
+	Super::GetAssetRegistryTags(Context);
 }
 
 #if WITH_EDITOR
diff --git a/Engine/Source/Runtime/Engine/Private/Texture.cpp b/Engine/Source/Runtime/Engine/Private/Texture.cpp
--- a/Engine/Source/Runtime/Engine/Private/Texture.cpp
+++ b/Engine/Source/Runtime/Engine/Private/Texture.cpp
@@ -9,6 +9,7 @@
 #include "Math/ColorList.h"
 #include "Misc/ConfigCacheIni.h"
 #include "Misc/FeedbackContext.h"
+#include "UObject/AssetRegistryTagsContext.h"
 #include "UObject/UObjectIterator.h"
 #include "UObject/ObjectSaveContext.h"
 #include "UObject/Package.h"
@@ -1284,21 +1285,28 @@ void UTexture::PreSave(FObjectPreSaveContext ObjectSaveContext)
 
 #if WITH_EDITORONLY_DATA
 void UTexture::GetAssetRegistryTags(TArray<FAssetRegistryTag>& OutTags) const
+{
+	PRAGMA_DISABLE_DEPRECATION_WARNINGS;
+	Super::GetAssetRegistryTags(OutTags);
+	PRAGMA_ENABLE_DEPRECATION_WARNINGS;
+}
+
+void UTexture::GetAssetRegistryTags(FAssetRegistryTagsContext Context) const
 {
 	if (AssetImportData)
 	{
-		OutTags.Add( FAssetRegistryTag(SourceFileTagName(), AssetImportData->GetSourceData().ToJson(), FAssetRegistryTag::TT_Hidden) );
+		Context.AddTag( FAssetRegistryTag(SourceFileTagName(), AssetImportData->GetSourceData().ToJson(), FAssetRegistryTag::TT_Hidden) );
 	}
 
-	OutTags.Add(FAssetRegistryTag("SourceCompression", Source.GetSourceCompressionAsString(), FAssetRegistryTag::TT_Alphabetical));
+	Context.AddTag(FAssetRegistryTag("SourceCompression", Source.GetSourceCompressionAsString(), FAssetRegistryTag::TT_Alphabetical));
 
-	OutTags.Add(FAssetRegistryTag("SourceFormat", 
+	Context.AddTag(FAssetRegistryTag("SourceFormat", 
 		StaticEnum<ETextureSourceFormat>()->GetDisplayNameTextByValue(Source.GetFormat()).ToString(),
 		FAssetRegistryTag::TT_Alphabetical));
 	
-	OutTags.Add(FAssetRegistryTag("IsSourceValid", Source.IsValid() ? TEXT("True") : TEXT("False"), FAssetRegistryTag::TT_Alphabetical));
+	Context.AddTag(FAssetRegistryTag("IsSourceValid", Source.IsValid() ? TEXT("True") : TEXT("False"), FAssetRegistryTag::TT_Alphabetical));
 
-	Super::GetAssetRegistryTags(OutTags);
+	Super::GetAssetRegistryTags(Context);
 }
 #endif
 
diff --git a/Engine/Source/Runtime/Engine/Private/Texture2D.cpp b/Engine/Source/Runtime/Engine/Private/Texture2D.cpp
--- a/Engine/Source/Runtime/Engine/Private/Texture2D.cpp
+++ b/Engine/Source/Runtime/Engine/Private/Texture2D.cpp
@@ -10,6 +10,7 @@
 #include "AsyncCompilationHelpers.h"
 #include "Containers/ResourceArray.h"
 #include "EngineLogs.h"
+#include "UObject/AssetRegistryTagsContext.h"
 #include "UObject/Package.h"
 #include "RenderUtils.h"
 #include "RenderGraphBuilder.h"
@@ -691,18 +692,25 @@ void UTexture2D::PreSave(FObjectPreSaveContext ObjectSaveContext)
 }
 
 void UTexture2D::GetAssetRegistryTags(TArray<FAssetRegistryTag>& OutTags) const
+{
+	PRAGMA_DISABLE_DEPRECATION_WARNINGS;
+	Super::GetAssetRegistryTags(OutTags);
+	PRAGMA_ENABLE_DEPRECATION_WARNINGS;
+}
+
+void UTexture2D::GetAssetRegistryTags(FAssetRegistryTagsContext Context) const
 {
 	FIntPoint SourceSize = GetImportedSize();
 
 	const FString DimensionsStr = FString::Printf(TEXT("%dx%d"), SourceSize.X, SourceSize.Y);
-	OutTags.Add( FAssetRegistryTag("Dimensions", DimensionsStr, FAssetRegistryTag::TT_Dimensional) );
+	Context.AddTag( FAssetRegistryTag("Dimensions", DimensionsStr, FAssetRegistryTag::TT_Dimensional) );
 	
 	// This "Has Alpha Channel" is whether the GPU format can represent alpha in the format (eg. is it DXT1 vs DXT5)
 	//	it does not tell you if the texture actually has non-opaque alpha
-	OutTags.Add( FAssetRegistryTag("HasAlphaChannel", HasAlphaChannel() ? TEXT("True") : TEXT("False"), FAssetRegistryTag::TT_Alphabetical) );
-	OutTags.Add( FAssetRegistryTag("Format", GPixelFormats[GetPixelFormat()].Name, FAssetRegistryTag::TT_Alphabetical) );
+	Context.AddTag( FAssetRegistryTag("HasAlphaChannel", HasAlphaChannel() ? TEXT("True") : TEXT("False"), FAssetRegistryTag::TT_Alphabetical) );
+	Context.AddTag( FAssetRegistryTag("Format", GPixelFormats[GetPixelFormat()].Name, FAssetRegistryTag::TT_Alphabetical) );
 
-	Super::GetAssetRegistryTags(OutTags);
+	Super::GetAssetRegistryTags(Context);
 }
 
 void UTexture2D::UpdateResource()
diff --git a/Engine/Source/Runtime/Engine/Private/TextureCube.cpp b/Engine/Source/Runtime/Engine/Private/TextureCube.cpp
--- a/Engine/Source/Runtime/Engine/Private/TextureCube.cpp
+++ b/Engine/Source/Runtime/Engine/Private/TextureCube.cpp
@@ -16,6 +16,7 @@
 #include "Interfaces/ITargetPlatform.h"
 #include "TextureCompiler.h"
 #include "Misc/ScopedSlowTask.h"
+#include "UObject/AssetRegistryTagsContext.h"
 #include "UObject/Package.h"
 #include "UObject/StrongObjectPtr.h"
 #include "ImageUtils.h"
@@ -176,6 +177,13 @@ void UTextureCube::PostLoad()
 }
 
 void UTextureCube::GetAssetRegistryTags(TArray<FAssetRegistryTag>& OutTags) const
+{
+	PRAGMA_DISABLE_DEPRECATION_WARNINGS;
+	Super::GetAssetRegistryTags(OutTags);
+	PRAGMA_ENABLE_DEPRECATION_WARNINGS;
+}
+
+void UTextureCube::GetAssetRegistryTags(FAssetRegistryTagsContext Context) const
 {
 #if WITH_EDITOR
 	int32 SizeX = Source.GetSizeX();
@@ -186,10 +194,10 @@ void UTextureCube::GetAssetRegistryTags(TArray<FAssetRegistryTag>& OutTags) cons
 #endif
 
 	const FString Dimensions = FString::Printf(TEXT("%dx%d"), SizeX, SizeY);
-	OutTags.Add( FAssetRegistryTag("Dimensions", Dimensions, FAssetRegistryTag::TT_Dimensional) );
-	OutTags.Add( FAssetRegistryTag("Format", GPixelFormats[GetPixelFormat()].Name, FAssetRegistryTag::TT_Alphabetical) );
+	Context.AddTag( FAssetRegistryTag("Dimensions", Dimensions, FAssetRegistryTag::TT_Dimensional) );
+	Context.AddTag( FAssetRegistryTag("Format", GPixelFormats[GetPixelFormat()].Name, FAssetRegistryTag::TT_Alphabetical) );
 
-	Super::GetAssetRegistryTags(OutTags);
+	Super::GetAssetRegistryTags(Context);
 }
 
 void UTextureCube::UpdateResource()
diff --git a/Engine/Source/Runtime/Engine/Private/TextureCubeArray.cpp b/Engine/Source/Runtime/Engine/Private/TextureCubeArray.cpp
--- a/Engine/Source/Runtime/Engine/Private/TextureCubeArray.cpp
+++ b/Engine/Source/Runtime/Engine/Private/TextureCubeArray.cpp
@@ -17,6 +17,7 @@
 #include "Stats/StatsTrace.h"
 #include "TextureCompiler.h"
 #include "TextureResource.h"
+#include "UObject/AssetRegistryTagsContext.h"
 #include "UObject/Package.h"
 #include "UObject/StrongObjectPtr.h"
 
@@ -803,6 +804,13 @@ void UTextureCubeArray::PostLoad()
 };
 
 void UTextureCubeArray::GetAssetRegistryTags(TArray<FAssetRegistryTag>& OutTags) const
+{
+	PRAGMA_DISABLE_DEPRECATION_WARNINGS;
+	Super::GetAssetRegistryTags(OutTags);
+	PRAGMA_ENABLE_DEPRECATION_WARNINGS;
+}
+
+void UTextureCubeArray::GetAssetRegistryTags(FAssetRegistryTagsContext Context) const
 {
 #if WITH_EDITOR
 	int32 SizeX = Source.GetSizeX();
@@ -814,10 +822,10 @@ void UTextureCubeArray::GetAssetRegistryTags(TArray<FAssetRegistryTag>& OutTags)
 	int32 ArraySize = 0;
 #endif
 	const FString Dimensions = FString::Printf(TEXT("%dx%d*%d"), SizeX, SizeY, ArraySize);
-	OutTags.Add(FAssetRegistryTag("Dimensions", Dimensions, FAssetRegistryTag::TT_Dimensional));
-	OutTags.Add(FAssetRegistryTag("Format", GPixelFormats[GetPixelFormat()].Name, FAssetRegistryTag::TT_Alphabetical));
+	Context.AddTag(FAssetRegistryTag("Dimensions", Dimensions, FAssetRegistryTag::TT_Dimensional));
+	Context.AddTag(FAssetRegistryTag("Format", GPixelFormats[GetPixelFormat()].Name, FAssetRegistryTag::TT_Alphabetical));
 
-	Super::GetAssetRegistryTags(OutTags);
+	Super::GetAssetRegistryTags(Context);
 }
 
 FString UTextureCubeArray::GetDesc()
diff --git a/Engine/Source/Runtime/Engine/Private/UserDefinedEnum.cpp b/Engine/Source/Runtime/Engine/Private/UserDefinedEnum.cpp
--- a/Engine/Source/Runtime/Engine/Private/UserDefinedEnum.cpp
+++ b/Engine/Source/Runtime/Engine/Private/UserDefinedEnum.cpp
@@ -2,6 +2,7 @@
 
 #include "Engine/UserDefinedEnum.h"
 #include "Templates/SubclassOf.h"
+#include "UObject/AssetRegistryTagsContext.h"
 #include "UObject/EditorObjectVersion.h"
 #include "UObject/ObjectSaveContext.h"
 #include "UObject/Package.h"
@@ -129,11 +130,18 @@ void UUserDefinedEnum::PostEditChangeProperty(FPropertyChangedEvent& PropertyCha
 
 void UUserDefinedEnum::GetAssetRegistryTags(TArray<FAssetRegistryTag>& OutTags) const
 {
+	PRAGMA_DISABLE_DEPRECATION_WARNINGS;
 	Super::GetAssetRegistryTags(OutTags);
+	PRAGMA_ENABLE_DEPRECATION_WARNINGS;
+}
+
+void UUserDefinedEnum::GetAssetRegistryTags(FAssetRegistryTagsContext Context) const
+{
+	Super::GetAssetRegistryTags(Context);
 
 	FString DescriptionString;
 	FTextStringHelper::WriteToBuffer(/*out*/ DescriptionString, EnumDescription);
-	OutTags.Emplace(GET_MEMBER_NAME_CHECKED(UUserDefinedEnum, EnumDescription), DescriptionString, FAssetRegistryTag::TT_Hidden);
+	Context.AddTag(FAssetRegistryTag(GET_MEMBER_NAME_CHECKED(UUserDefinedEnum, EnumDescription), DescriptionString, FAssetRegistryTag::TT_Hidden));
 }
 
 void UUserDefinedEnum::PreSaveRoot(FObjectPreSaveRootContext ObjectSaveContext)
diff --git a/Engine/Source/Runtime/Engine/Private/UserDefinedStruct.cpp b/Engine/Source/Runtime/Engine/Private/UserDefinedStruct.cpp
--- a/Engine/Source/Runtime/Engine/Private/UserDefinedStruct.cpp
+++ b/Engine/Source/Runtime/Engine/Private/UserDefinedStruct.cpp
@@ -2,6 +2,7 @@
 
 #include "Engine/UserDefinedStruct.h"
 #include "Templates/SubclassOf.h"
+#include "UObject/AssetRegistryTagsContext.h"
 #include "UObject/Package.h"
 #include "UObject/Package.h"
 #include "UObject/UnrealType.h"
@@ -193,9 +194,16 @@ void UUserDefinedStruct::PostSaveRoot(FObjectPostSaveRootContext ObjectSaveConte
 
 void UUserDefinedStruct::GetAssetRegistryTags(TArray<FAssetRegistryTag>& OutTags) const
 {
+	PRAGMA_DISABLE_DEPRECATION_WARNINGS;
 	Super::GetAssetRegistryTags(OutTags);
+	PRAGMA_ENABLE_DEPRECATION_WARNINGS;
+}
+
+void UUserDefinedStruct::GetAssetRegistryTags(FAssetRegistryTagsContext Context) const
+{
+	Super::GetAssetRegistryTags(Context);
 
-	OutTags.Add(FAssetRegistryTag(TEXT("Tooltip"), FStructureEditorUtils::GetTooltip(this), FAssetRegistryTag::TT_Hidden));
+	Context.AddTag(FAssetRegistryTag(TEXT("Tooltip"), FStructureEditorUtils::GetTooltip(this), FAssetRegistryTag::TT_Hidden));
 }
 
 void UUserDefinedStruct::ValidateGuid()
diff --git a/Engine/Source/Runtime/Engine/Private/VT/RuntimeVirtualTexture.cpp b/Engine/Source/Runtime/Engine/Private/VT/RuntimeVirtualTexture.cpp
--- a/Engine/Source/Runtime/Engine/Private/VT/RuntimeVirtualTexture.cpp
+++ b/Engine/Source/Runtime/Engine/Private/VT/RuntimeVirtualTexture.cpp
@@ -8,6 +8,7 @@
 #include "RendererInterface.h"
 #include "RenderingThread.h"
 #include "Shader/ShaderTypes.h"
+#include "UObject/AssetRegistryTagsContext.h"
 #include "UObject/UnrealType.h"
 #include "VT/RuntimeVirtualTextureNotify.h"
 #include "VT/UploadingVirtualTexture.h"
@@ -595,12 +596,19 @@ void URuntimeVirtualTexture::InitNullResource()
 
 void URuntimeVirtualTexture::GetAssetRegistryTags(TArray<FAssetRegistryTag>& OutTags) const
 {
+	PRAGMA_DISABLE_DEPRECATION_WARNINGS;
 	Super::GetAssetRegistryTags(OutTags);
+	PRAGMA_ENABLE_DEPRECATION_WARNINGS;
+}
+
+void URuntimeVirtualTexture::GetAssetRegistryTags(FAssetRegistryTagsContext Context) const
+{
+	Super::GetAssetRegistryTags(Context);
 
-	OutTags.Add(FAssetRegistryTag("Size", FString::FromInt(GetSize()), FAssetRegistryTag::TT_Numerical));
-	OutTags.Add(FAssetRegistryTag("TileCount", FString::FromInt(GetTileCount()), FAssetRegistryTag::TT_Numerical));
-	OutTags.Add(FAssetRegistryTag("TileSize", FString::FromInt(GetTileSize()), FAssetRegistryTag::TT_Numerical));
-	OutTags.Add(FAssetRegistryTag("TileBorderSize", FString::FromInt(GetTileBorderSize()), FAssetRegistryTag::TT_Numerical));
+	Context.AddTag(FAssetRegistryTag("Size", FString::FromInt(GetSize()), FAssetRegistryTag::TT_Numerical));
+	Context.AddTag(FAssetRegistryTag("TileCount", FString::FromInt(GetTileCount()), FAssetRegistryTag::TT_Numerical));
+	Context.AddTag(FAssetRegistryTag("TileSize", FString::FromInt(GetTileSize()), FAssetRegistryTag::TT_Numerical));
+	Context.AddTag(FAssetRegistryTag("TileBorderSize", FString::FromInt(GetTileBorderSize()), FAssetRegistryTag::TT_Numerical));
 }
 
 void URuntimeVirtualTexture::PostLoad()
diff --git a/Engine/Source/Runtime/Engine/Private/VectorField.cpp b/Engine/Source/Runtime/Engine/Private/VectorField.cpp
--- a/Engine/Source/Runtime/Engine/Private/VectorField.cpp
+++ b/Engine/Source/Runtime/Engine/Private/VectorField.cpp
@@ -15,6 +15,7 @@
 #include "Engine/CollisionProfile.h"
 #include "Engine/Texture2D.h"
 #include "ComponentReregisterContext.h"
+#include "UObject/AssetRegistryTagsContext.h"
 #include "UObject/UnrealType.h"
 #include "VectorFieldVisualization.h"
 #include "Serialization/MemoryWriter.h"
@@ -658,13 +659,20 @@ void UVectorFieldStatic::PostEditChangeProperty(FPropertyChangedEvent& PropertyC
 
 #if WITH_EDITORONLY_DATA
 void UVectorFieldStatic::GetAssetRegistryTags(TArray<FAssetRegistryTag>& OutTags) const
+{
+	PRAGMA_DISABLE_DEPRECATION_WARNINGS;
+	Super::GetAssetRegistryTags(OutTags);
+	PRAGMA_ENABLE_DEPRECATION_WARNINGS;
+}
+
+void UVectorFieldStatic::GetAssetRegistryTags(FAssetRegistryTagsContext Context) const
 {
 	if (AssetImportData)
 	{
-		OutTags.Add( FAssetRegistryTag(SourceFileTagName(), AssetImportData->GetSourceData().ToJson(), FAssetRegistryTag::TT_Hidden) );
+		Context.AddTag( FAssetRegistryTag(SourceFileTagName(), AssetImportData->GetSourceData().ToJson(), FAssetRegistryTag::TT_Hidden) );
 	}
 
-	Super::GetAssetRegistryTags(OutTags);
+	Super::GetAssetRegistryTags(Context);
 }
 
 void UVectorFieldStatic::PostInitProperties()
diff --git a/Engine/Source/Runtime/Engine/Private/VolumeTexture.cpp b/Engine/Source/Runtime/Engine/Private/VolumeTexture.cpp
--- a/Engine/Source/Runtime/Engine/Private/VolumeTexture.cpp
+++ b/Engine/Source/Runtime/Engine/Private/VolumeTexture.cpp
@@ -18,6 +18,7 @@
 #include "Streaming/TextureStreamOut.h"
 #include "Streaming/VolumeTextureStreaming.h"
 #include "TextureCompiler.h"
+#include "UObject/AssetRegistryTagsContext.h"
 #include "UObject/Package.h"
 #include "UObject/StrongObjectPtr.h"
 #include "ImageCoreUtils.h"
@@ -336,6 +337,13 @@ void UVolumeTexture::PostLoad()
 }
 
 void UVolumeTexture::GetAssetRegistryTags(TArray<FAssetRegistryTag>& OutTags) const
+{
+	PRAGMA_DISABLE_DEPRECATION_WARNINGS;
+	Super::GetAssetRegistryTags(OutTags);
+	PRAGMA_ENABLE_DEPRECATION_WARNINGS;
+}
+
+void UVolumeTexture::GetAssetRegistryTags(FAssetRegistryTagsContext Context) const
 {
 #if WITH_EDITOR
 	int32 SizeX = Source.GetSizeX();
@@ -348,10 +356,10 @@ void UVolumeTexture::GetAssetRegistryTags(TArray<FAssetRegistryTag>& OutTags) co
 #endif
 
 	const FString Dimensions = FString::Printf(TEXT("%dx%dx%d"), SizeX, SizeY, SizeZ);
-	OutTags.Add( FAssetRegistryTag("Dimensions", Dimensions, FAssetRegistryTag::TT_Dimensional) );
-	OutTags.Add( FAssetRegistryTag("Format", GPixelFormats[GetPixelFormat()].Name, FAssetRegistryTag::TT_Alphabetical) );
+	Context.AddTag( FAssetRegistryTag("Dimensions", Dimensions, FAssetRegistryTag::TT_Dimensional) );
+	Context.AddTag( FAssetRegistryTag("Format", GPixelFormats[GetPixelFormat()].Name, FAssetRegistryTag::TT_Alphabetical) );
 
-	Super::GetAssetRegistryTags(OutTags);
+	Super::GetAssetRegistryTags(Context);
 }
 
 void UVolumeTexture::UpdateResource()
