GetAssetRegistryTags deprecation: Make the old GetAssetRegistryTags and related functions deprecated.

diff --git a/Engine/Source/Runtime/Engine/Private/World.cpp b/Engine/Source/Runtime/Engine/Private/World.cpp
--- a/Engine/Source/Runtime/Engine/Private/World.cpp
+++ b/Engine/Source/Runtime/Engine/Private/World.cpp
@@ -20,6 +20,7 @@
 #include "Stats/StatsMisc.h"
 #include "Misc/ScopedSlowTask.h"
 #include "SceneInterface.h"
+#include "UObject/AssetRegistryTagsContext.h"
 #include "UObject/ObjectRedirector.h"
 #include "SceneView.h"
 #include "UObject/ObjectSaveContext.h"
@@ -562,7 +563,10 @@ FWorldDelegates::FOnLevelChanged FWorldDelegates::PreLevelRemovedFromWorld;
 FWorldDelegates::FOnLevelChanged FWorldDelegates::LevelRemovedFromWorld;
 FWorldDelegates::FLevelOffsetEvent FWorldDelegates::PostApplyLevelOffset;
 FWorldDelegates::FLevelTransformEvent FWorldDelegates::PostApplyLevelTransform;
+PRAGMA_DISABLE_DEPRECATION_WARNINGS;
 FWorldDelegates::FWorldGetAssetTags FWorldDelegates::GetAssetTags;
+PRAGMA_ENABLE_DEPRECATION_WARNINGS;
+FWorldDelegates::FWorldGetAssetTagsWithContext FWorldDelegates::GetAssetTagsWithContext;
 FWorldDelegates::FOnWorldTickStart FWorldDelegates::OnWorldTickStart;
 FWorldDelegates::FOnWorldTickEnd FWorldDelegates::OnWorldTickEnd;
 FWorldDelegates::FOnWorldPreActorTick FWorldDelegates::OnWorldPreActorTick;
@@ -9093,19 +9097,26 @@ void UWorld::RestoreScene()
 
 void UWorld::GetAssetRegistryTags(TArray<FAssetRegistryTag>& OutTags) const
 {
+	PRAGMA_DISABLE_DEPRECATION_WARNINGS;
 	Super::GetAssetRegistryTags(OutTags);
+	PRAGMA_ENABLE_DEPRECATION_WARNINGS;
+}
+
+void UWorld::GetAssetRegistryTags(FAssetRegistryTagsContext Context) const
+{
+	Super::GetAssetRegistryTags(Context);
 
-	if(PersistentLevel && PersistentLevel->OwningWorld)
+	if (PersistentLevel && PersistentLevel->OwningWorld)
 	{
 		if (ULevelScriptBlueprint* Blueprint = PersistentLevel->GetLevelScriptBlueprint(true))
 		{
-			Blueprint->GetAssetRegistryTags(OutTags);
+			Blueprint->GetAssetRegistryTags(Context);
 		}
 		// If there are no blueprints FiBData will be empty, the search manager will treat this as indexed
 								
-		if(UWorldPartition* WorldPartition = GetWorldPartition())
+		if (UWorldPartition* WorldPartition = GetWorldPartition())
 		{
-			WorldPartition->AppendAssetRegistryTags(OutTags);
+			WorldPartition->AppendAssetRegistryTags(Context);
 		}
 		else
 		{
@@ -9124,22 +9135,22 @@ void UWorld::GetAssetRegistryTags(TArray<FAssetRegistryTag>& OutTags) const
 			LevelBounds.GetCenterAndExtents(LevelBoundsLocation, LevelBoundsExtent);
 
 			static const FName NAME_LevelBoundsLocation(TEXT("LevelBoundsLocation"));
-			OutTags.Add(FAssetRegistryTag(NAME_LevelBoundsLocation, LevelBoundsLocation.ToCompactString(), FAssetRegistryTag::TT_Hidden));
+			Context.AddTag(FAssetRegistryTag(NAME_LevelBoundsLocation, LevelBoundsLocation.ToCompactString(), FAssetRegistryTag::TT_Hidden));
 
 			static const FName NAME_LevelBoundsExtent(TEXT("LevelBoundsExtent"));
-			OutTags.Add(FAssetRegistryTag(NAME_LevelBoundsExtent, LevelBoundsExtent.ToCompactString(), FAssetRegistryTag::TT_Hidden));
+			Context.AddTag(FAssetRegistryTag(NAME_LevelBoundsExtent, LevelBoundsExtent.ToCompactString(), FAssetRegistryTag::TT_Hidden));
 		}
 	
 		if (PersistentLevel->IsUsingExternalActors())
 		{
 			static const FName NAME_LevelIsUsingExternalActors(TEXT("LevelIsUsingExternalActors"));
-			OutTags.Add(FAssetRegistryTag(NAME_LevelIsUsingExternalActors, TEXT("1"), FAssetRegistryTag::TT_Hidden));
+			Context.AddTag(FAssetRegistryTag(NAME_LevelIsUsingExternalActors, TEXT("1"), FAssetRegistryTag::TT_Hidden));
 		}
 
 		if (PersistentLevel->IsUsingActorFolders())
 		{
 			static const FName NAME_LevelIsUsingActorFolders(TEXT("LevelIsUsingActorFolders"));
-			OutTags.Add(FAssetRegistryTag(NAME_LevelIsUsingActorFolders, TEXT("1"), FAssetRegistryTag::TT_Hidden));
+			Context.AddTag(FAssetRegistryTag(NAME_LevelIsUsingActorFolders, TEXT("1"), FAssetRegistryTag::TT_Hidden));
 		}
 
 		if (AWorldSettings* WorldSettings = GetWorldSettings(/*bCheckStreamingPersistent*/false, /*bChecked*/false))
@@ -9148,7 +9159,7 @@ void UWorld::GetAssetRegistryTags(TArray<FAssetRegistryTag>& OutTags) const
 			if (!LevelInstancePivotOffset.IsNearlyZero())
 			{
 				static const FName NAME_LevelInstancePivotOffset(TEXT("LevelInstancePivotOffset"));
-				OutTags.Add(FAssetRegistryTag(NAME_LevelInstancePivotOffset, LevelInstancePivotOffset.ToCompactString(), FAssetRegistryTag::TT_Hidden));
+				Context.AddTag(FAssetRegistryTag(NAME_LevelInstancePivotOffset, LevelInstancePivotOffset.ToCompactString(), FAssetRegistryTag::TT_Hidden));
 			}
 		}
 	}
@@ -9159,10 +9170,18 @@ void UWorld::GetAssetRegistryTags(TArray<FAssetRegistryTag>& OutTags) const
 	{
 		// Save/Display the modify date. File size is handled generically for all packages
 		FDateTime AssetDateModified = IFileManager::Get().GetTimeStamp(*FullFilePath);
-		OutTags.Add(FAssetRegistryTag("DateModified", AssetDateModified.ToString(), FAssetRegistryTag::TT_Chronological, FAssetRegistryTag::TD_Date));
+		Context.AddTag(FAssetRegistryTag("DateModified", AssetDateModified.ToString(), FAssetRegistryTag::TT_Chronological, FAssetRegistryTag::TD_Date));
 	}
 
-	FWorldDelegates::GetAssetTags.Broadcast(this, OutTags);
+	PRAGMA_DISABLE_DEPRECATION_WARNINGS;
+	TArray<UObject::FAssetRegistryTag> DeprecatedTags;
+	FWorldDelegates::GetAssetTags.Broadcast(this, DeprecatedTags);
+	for (UObject::FAssetRegistryTag& Tag : DeprecatedTags)
+	{
+		Context.AddTag(MoveTemp(Tag));
+	}
+	PRAGMA_ENABLE_DEPRECATION_WARNINGS;
+	FWorldDelegates::GetAssetTagsWithContext.Broadcast(this, Context);
 }
 
 void UWorld::PostLoadAssetRegistryTags(const FAssetData& InAssetData, TArray<FAssetRegistryTag>& OutTagsAndValuesToUpdate) const
diff --git a/Engine/Source/Runtime/Engine/Private/WorldPartition/DataLayer/DataLayerInstance.cpp b/Engine/Source/Runtime/Engine/Private/WorldPartition/DataLayer/DataLayerInstance.cpp
--- a/Engine/Source/Runtime/Engine/Private/WorldPartition/DataLayer/DataLayerInstance.cpp
+++ b/Engine/Source/Runtime/Engine/Private/WorldPartition/DataLayer/DataLayerInstance.cpp
@@ -3,6 +3,7 @@
 #include "WorldPartition/DataLayer/DataLayerInstance.h"
 
 #include "Internationalization/Text.h"
+#include "UObject/AssetRegistryTagsContext.h"
 #include "UObject/UnrealType.h"
 #include "WorldPartition/WorldPartitionLog.h"
 #include "WorldPartition/WorldPartitionActorLoaderInterface.h"
@@ -76,27 +77,34 @@ namespace DataLayerInstance
 
 void UDataLayerInstance::GetAssetRegistryTags(TArray<FAssetRegistryTag>& OutTags) const
 {
+	PRAGMA_DISABLE_DEPRECATION_WARNINGS;
 	Super::GetAssetRegistryTags(OutTags);
+	PRAGMA_ENABLE_DEPRECATION_WARNINGS;
+}
+
+void UDataLayerInstance::GetAssetRegistryTags(FAssetRegistryTagsContext Context) const
+{
+	Super::GetAssetRegistryTags(Context);
 
 	if (IsPackageExternal())
 	{
 		// Set generic FPrimaryAssetId::PrimaryAssetDisplayNameTag
-		OutTags.Add(FAssetRegistryTag(FPrimaryAssetId::PrimaryAssetDisplayNameTag, *GetDataLayerShortName(), FAssetRegistryTag::TT_Hidden));
+		Context.AddTag(FAssetRegistryTag(FPrimaryAssetId::PrimaryAssetDisplayNameTag, *GetDataLayerShortName(), FAssetRegistryTag::TT_Hidden));
 
 		// Set DataLayerInstance specific tags
-		OutTags.Add(FAssetRegistryTag(DataLayerInstance::NAME_DataLayerInstanceName, *GetDataLayerFName().ToString(), FAssetRegistryTag::TT_Hidden));
+		Context.AddTag(FAssetRegistryTag(DataLayerInstance::NAME_DataLayerInstanceName, *GetDataLayerFName().ToString(), FAssetRegistryTag::TT_Hidden));
 		if (GetParent())
 		{
-			OutTags.Add(FAssetRegistryTag(DataLayerInstance::NAME_DataLayerInstanceParentName, *GetParent()->GetDataLayerFName().ToString(), FAssetRegistryTag::TT_Hidden));
+			Context.AddTag(FAssetRegistryTag(DataLayerInstance::NAME_DataLayerInstanceParentName, *GetParent()->GetDataLayerFName().ToString(), FAssetRegistryTag::TT_Hidden));
 		}
 		if (const UDataLayerAsset* DataLayerAsset = GetAsset())
 		{
-			OutTags.Add(FAssetRegistryTag(DataLayerInstance::NAME_DataLayerInstanceAssetPath, *DataLayerAsset->GetPathName(), FAssetRegistryTag::TT_Hidden));
-			OutTags.Add(FAssetRegistryTag(DataLayerInstance::NAME_DataLayerInstanceIsPrivate, DataLayerAsset->IsPrivate() ? TEXT("1") : TEXT("0"), FAssetRegistryTag::TT_Hidden));
+			Context.AddTag(FAssetRegistryTag(DataLayerInstance::NAME_DataLayerInstanceAssetPath, *DataLayerAsset->GetPathName(), FAssetRegistryTag::TT_Hidden));
+			Context.AddTag(FAssetRegistryTag(DataLayerInstance::NAME_DataLayerInstanceIsPrivate, DataLayerAsset->IsPrivate() ? TEXT("1") : TEXT("0"), FAssetRegistryTag::TT_Hidden));
 		}
-		OutTags.Add(FAssetRegistryTag(DataLayerInstance::NAME_DataLayerInstanceIsIncludedInActorFilterDefault, IsIncludedInActorFilterDefault() ? TEXT("1") : TEXT("0"), FAssetRegistryTag::TT_Hidden));
-		OutTags.Add(FAssetRegistryTag(DataLayerInstance::NAME_DataLayerInstancePrivateDataLayerSupportsActorFilter, SupportsActorFilters() ? TEXT("1") : TEXT("0"), FAssetRegistryTag::TT_Hidden));
-		OutTags.Add(FAssetRegistryTag(DataLayerInstance::NAME_DataLayerInstancePrivateShortName, *GetDataLayerShortName(), FAssetRegistryTag::TT_Hidden));
+		Context.AddTag(FAssetRegistryTag(DataLayerInstance::NAME_DataLayerInstanceIsIncludedInActorFilterDefault, IsIncludedInActorFilterDefault() ? TEXT("1") : TEXT("0"), FAssetRegistryTag::TT_Hidden));
+		Context.AddTag(FAssetRegistryTag(DataLayerInstance::NAME_DataLayerInstancePrivateDataLayerSupportsActorFilter, SupportsActorFilters() ? TEXT("1") : TEXT("0"), FAssetRegistryTag::TT_Hidden));
+		Context.AddTag(FAssetRegistryTag(DataLayerInstance::NAME_DataLayerInstancePrivateShortName, *GetDataLayerShortName(), FAssetRegistryTag::TT_Hidden));
 	}
 }
 
diff --git a/Engine/Source/Runtime/Engine/Private/WorldPartition/WorldPartition.cpp b/Engine/Source/Runtime/Engine/Private/WorldPartition/WorldPartition.cpp
--- a/Engine/Source/Runtime/Engine/Private/WorldPartition/WorldPartition.cpp
+++ b/Engine/Source/Runtime/Engine/Private/WorldPartition/WorldPartition.cpp
@@ -5,6 +5,7 @@
 =============================================================================*/
 #include "WorldPartition/WorldPartition.h"
 #include "Misc/PackageName.h"
+#include "UObject/AssetRegistryTagsContext.h"
 #include "UObject/UObjectIterator.h"
 #include "Misc/Paths.h"
 #include "WorldPartition/WorldPartitionLog.h"
@@ -1947,14 +1948,24 @@ void UWorldPartition::DumpActorDescs(const FString& Path)
 }
 
 void UWorldPartition::AppendAssetRegistryTags(TArray<FAssetRegistryTag>& OutTags) const
+{
+	FAssetRegistryTagsContextData Context(this, EAssetRegistryTagsCaller::Uncategorized);
+	AppendAssetRegistryTags(Context);
+	for (TPair<FName, FAssetRegistryTag>& Pair : Context.Tags)
+	{
+		OutTags.Add(MoveTemp(Pair.Value));
+	}
+}
+
+void UWorldPartition::AppendAssetRegistryTags(FAssetRegistryTagsContext Context) const
 {
 	static const FName NAME_LevelIsPartitioned(TEXT("LevelIsPartitioned"));
-	OutTags.Add(FAssetRegistryTag(NAME_LevelIsPartitioned, TEXT("1"), FAssetRegistryTag::TT_Hidden));
+	Context.AddTag(FAssetRegistryTag(NAME_LevelIsPartitioned, TEXT("1"), FAssetRegistryTag::TT_Hidden));
 
 	if (!IsStreamingEnabled())
 	{
 		static const FName NAME_LevelHasStreamingDisabled(TEXT("LevelHasStreamingDisabled"));
-		OutTags.Add(FAssetRegistryTag(NAME_LevelHasStreamingDisabled, TEXT("1"), FAssetRegistryTag::TT_Hidden));
+		Context.AddTag(FAssetRegistryTag(NAME_LevelHasStreamingDisabled, TEXT("1"), FAssetRegistryTag::TT_Hidden));
 	}
 
 	// Append level script references so we can perform changelists validations without loading the world
@@ -1975,7 +1986,7 @@ void UWorldPartition::AppendAssetRegistryTags(TArray<FAssetRegistryTag>& OutTags
 			StringBuilder.RemoveSuffix(1);
 
 			static const FName NAME_LevelScriptExternalActorsReferences(TEXT("LevelScriptExternalActorsReferences"));
-			OutTags.Add(FAssetRegistryTag(NAME_LevelScriptExternalActorsReferences, StringBuilder.ToString(), FAssetRegistryTag::TT_Hidden));
+			Context.AddTag(FAssetRegistryTag(NAME_LevelScriptExternalActorsReferences, StringBuilder.ToString(), FAssetRegistryTag::TT_Hidden));
 		}
 	}
 }
diff --git a/Engine/Source/Runtime/Engine/Private/WorldPartition/WorldPartitionActorDescUtils.cpp b/Engine/Source/Runtime/Engine/Private/WorldPartition/WorldPartitionActorDescUtils.cpp
--- a/Engine/Source/Runtime/Engine/Private/WorldPartition/WorldPartitionActorDescUtils.cpp
+++ b/Engine/Source/Runtime/Engine/Private/WorldPartition/WorldPartitionActorDescUtils.cpp
@@ -10,6 +10,7 @@
 #include "WorldPartition/WorldPartitionLog.h"
 #include "Engine/Level.h"
 #include "GameFramework/Actor.h"
+#include "UObject/AssetRegistryTagsContext.h"
 #include "UObject/CoreRedirects.h"
 #include "Misc/Base64.h"
 #include "WorldPartition/WorldPartitionActorDesc.h"
@@ -100,6 +101,16 @@ TUniquePtr<FWorldPartitionActorDesc> FWorldPartitionActorDescUtils::GetActorDesc
 }
 
 void FWorldPartitionActorDescUtils::AppendAssetDataTagsFromActor(const AActor* InActor, TArray<UObject::FAssetRegistryTag>& OutTags)
+{
+	FAssetRegistryTagsContextData Context(InActor, EAssetRegistryTagsCaller::Uncategorized);
+	AppendAssetDataTagsFromActor(InActor, Context);
+	for (TPair<FName, UObject::FAssetRegistryTag>& Pair : Context.Tags)
+	{
+		OutTags.Add(MoveTemp(Pair.Value));
+	}
+}
+
+void FWorldPartitionActorDescUtils::AppendAssetDataTagsFromActor(const AActor* InActor, FAssetRegistryTagsContext Context)
 {
 	// Avoid an assert when calling StaticFindObject during save to retrieve the actor's class.
 	// Since we are only looking for a native class, the call to StaticFindObject is legit.
@@ -132,10 +143,10 @@ void FWorldPartitionActorDescUtils::AppendAssetDataTagsFromActor(const AActor* I
 	}
 
 	const FString ActorMetaDataClass = GetParentNativeClass(InActor->GetClass())->GetPathName();
-	OutTags.Add(UObject::FAssetRegistryTag(NAME_ActorMetaDataClass, ActorMetaDataClass, UObject::FAssetRegistryTag::TT_Hidden));
+	Context.AddTag(UObject::FAssetRegistryTag(NAME_ActorMetaDataClass, ActorMetaDataClass, UObject::FAssetRegistryTag::TT_Hidden));
 
 	const FString ActorMetaData = GetAssetDataFromActorDescriptor(ActorDesc);
-	OutTags.Add(UObject::FAssetRegistryTag(NAME_ActorMetaData, ActorMetaData, UObject::FAssetRegistryTag::TT_Hidden));
+	Context.AddTag(UObject::FAssetRegistryTag(NAME_ActorMetaData, ActorMetaData, UObject::FAssetRegistryTag::TT_Hidden));
 }
 
 FString FWorldPartitionActorDescUtils::GetAssetDataFromActorDescriptor(TUniquePtr<FWorldPartitionActorDesc>& InActorDesc)
diff --git a/Engine/Source/Runtime/Engine/Public/WorldPartition/DataLayer/DataLayerInstance.h b/Engine/Source/Runtime/Engine/Public/WorldPartition/DataLayer/DataLayerInstance.h
--- a/Engine/Source/Runtime/Engine/Public/WorldPartition/DataLayer/DataLayerInstance.h
+++ b/Engine/Source/Runtime/Engine/Public/WorldPartition/DataLayer/DataLayerInstance.h
@@ -70,7 +70,9 @@ public:
 #if WITH_EDITOR
 	//~ Begin UObject
 	ENGINE_API virtual bool IsAsset() const override;
-	ENGINE_API virtual void GetAssetRegistryTags(TArray<FAssetRegistryTag>& OutTags) const;
+	ENGINE_API virtual void GetAssetRegistryTags(FAssetRegistryTagsContext Context) const override;
+	UE_DEPRECATED(5.4, "Implement the version that takes FAssetRegistryTagsContext instead.")
+	ENGINE_API virtual void GetAssetRegistryTags(TArray<FAssetRegistryTag>& OutTags) const override;
 	ENGINE_API virtual void PreEditUndo() override;
 	ENGINE_API virtual void PostEditUndo() override;
 	//~ End UObject
diff --git a/Engine/Source/Runtime/Engine/Public/WorldPartition/WorldPartition.h b/Engine/Source/Runtime/Engine/Public/WorldPartition/WorldPartition.h
--- a/Engine/Source/Runtime/Engine/Public/WorldPartition/WorldPartition.h
+++ b/Engine/Source/Runtime/Engine/Public/WorldPartition/WorldPartition.h
@@ -315,6 +315,8 @@ public:
 	using FStreamingGenerationErrorHandlerOverride = TFunction<IStreamingGenerationErrorHandler*(IStreamingGenerationErrorHandler* InErrorHandler)>;
 	ENGINE_API inline static TOptional<FStreamingGenerationErrorHandlerOverride> StreamingGenerationErrorHandlerOverride;
 
+	ENGINE_API void AppendAssetRegistryTags(FAssetRegistryTagsContext Context) const;
+	UE_DEPRECATED(5.4, "Implement the version that takes FAssetRegistryTagsContext instead.")
 	ENGINE_API void AppendAssetRegistryTags(TArray<FAssetRegistryTag>& OutTags) const;
 
 	struct FContainerRegistrationParams
diff --git a/Engine/Source/Runtime/Engine/Public/WorldPartition/WorldPartitionActorDescUtils.h b/Engine/Source/Runtime/Engine/Public/WorldPartition/WorldPartitionActorDescUtils.h
--- a/Engine/Source/Runtime/Engine/Public/WorldPartition/WorldPartitionActorDescUtils.h
+++ b/Engine/Source/Runtime/Engine/Public/WorldPartition/WorldPartitionActorDescUtils.h
@@ -37,11 +37,13 @@ struct FWorldPartitionActorDescUtils
 	 */
 	static ENGINE_API TUniquePtr<FWorldPartitionActorDesc> GetActorDescriptorFromAssetData(const FAssetData& InAssetData);
 
-	/** 
+	/**
 	 * Appends the actor's actor descriptor data into the provided asset registry tags.
 	 * @param InActor		The actor that will append its actor descriptor.
 	 * @param OutTags		Output tags to output into.
 	 */
+	static ENGINE_API void AppendAssetDataTagsFromActor(const AActor* InActor, FAssetRegistryTagsContext Context);
+	UE_DEPRECATED(5.4, "Call the version that takes FAssetRegistryTagsContext instead.")
 	static ENGINE_API void AppendAssetDataTagsFromActor(const AActor* InActor, TArray<UObject::FAssetRegistryTag>& OutTags);
 
 	/** 
diff --git a/Engine/Source/Runtime/Experimental/GeometryCollectionEngine/Private/GeometryCollection/GeometryCollectionCache.cpp b/Engine/Source/Runtime/Experimental/GeometryCollectionEngine/Private/GeometryCollection/GeometryCollectionCache.cpp
--- a/Engine/Source/Runtime/Experimental/GeometryCollectionEngine/Private/GeometryCollection/GeometryCollectionCache.cpp
+++ b/Engine/Source/Runtime/Experimental/GeometryCollectionEngine/Private/GeometryCollection/GeometryCollectionCache.cpp
@@ -5,6 +5,7 @@
 #include "HAL/IConsoleManager.h"
 #include "Serialization/ArchiveCountMem.h"
 #include "Features/IModularFeatures.h"
+#include "UObject/AssetRegistryTagsContext.h"
 
 #include UE_INLINE_GENERATED_CPP_BY_NAME(GeometryCollectionCache)
 
@@ -38,9 +39,18 @@ void UGeometryCollectionCache::SetSupportedCollection(const UGeometryCollection*
 
 void UGeometryCollectionCache::GetAssetRegistryTags(TArray<FAssetRegistryTag>& OutTags) const
 {
-	OutTags.Add(FAssetRegistryTag(TagName_Name, SupportedCollection ? SupportedCollection->GetName() : FString(TEXT("None")), FAssetRegistryTag::TT_Alphabetical));
-	OutTags.Add(FAssetRegistryTag(TagName_IdGuid, SupportedCollection ? SupportedCollection->GetIdGuid().ToString() : FString(TEXT("INVALID")), FAssetRegistryTag::TT_Hidden));
-	OutTags.Add(FAssetRegistryTag(TagName_StateGuid, SupportedCollection ? CompatibleCollectionState.ToString() : FString(TEXT("INVALID")), FAssetRegistryTag::TT_Hidden));
+	PRAGMA_DISABLE_DEPRECATION_WARNINGS;
+	Super::GetAssetRegistryTags(OutTags);
+	PRAGMA_ENABLE_DEPRECATION_WARNINGS;
+}
+
+void UGeometryCollectionCache::GetAssetRegistryTags(FAssetRegistryTagsContext Context) const
+{
+	Super::GetAssetRegistryTags(Context);
+
+	Context.AddTag(FAssetRegistryTag(TagName_Name, SupportedCollection ? SupportedCollection->GetName() : FString(TEXT("None")), FAssetRegistryTag::TT_Alphabetical));
+	Context.AddTag(FAssetRegistryTag(TagName_IdGuid, SupportedCollection ? SupportedCollection->GetIdGuid().ToString() : FString(TEXT("INVALID")), FAssetRegistryTag::TT_Hidden));
+	Context.AddTag(FAssetRegistryTag(TagName_StateGuid, SupportedCollection ? CompatibleCollectionState.ToString() : FString(TEXT("INVALID")), FAssetRegistryTag::TT_Hidden));
 }
 
 UGeometryCollectionCache* UGeometryCollectionCache::CreateCacheForCollection(const UGeometryCollection* InCollection)
diff --git a/Engine/Source/Runtime/Experimental/GeometryCollectionEngine/Public/GeometryCollection/GeometryCollectionCache.h b/Engine/Source/Runtime/Experimental/GeometryCollectionEngine/Public/GeometryCollection/GeometryCollectionCache.h
--- a/Engine/Source/Runtime/Experimental/GeometryCollectionEngine/Public/GeometryCollection/GeometryCollectionCache.h
+++ b/Engine/Source/Runtime/Experimental/GeometryCollectionEngine/Public/GeometryCollection/GeometryCollectionCache.h
@@ -36,6 +36,8 @@ public:
 	GEOMETRYCOLLECTIONENGINE_API void SetSupportedCollection(const UGeometryCollection* InCollection);
 
 	/** UObject Interface */
+	GEOMETRYCOLLECTIONENGINE_API virtual void GetAssetRegistryTags(FAssetRegistryTagsContext Context) const override;
+	UE_DEPRECATED(5.4, "Implement the version that takes FAssetRegistryTagsContext instead.")
 	GEOMETRYCOLLECTIONENGINE_API virtual void GetAssetRegistryTags(TArray<FAssetRegistryTag>& OutTags) const override;
 	/** End UObject Interface */
 
diff --git a/Engine/Source/Runtime/Interchange/Core/Public/Nodes/InterchangeBaseNode.h b/Engine/Source/Runtime/Interchange/Core/Public/Nodes/InterchangeBaseNode.h
--- a/Engine/Source/Runtime/Interchange/Core/Public/Nodes/InterchangeBaseNode.h
+++ b/Engine/Source/Runtime/Interchange/Core/Public/Nodes/InterchangeBaseNode.h
@@ -17,6 +17,7 @@
 #include "Nodes/InterchangeBaseNodeUtilities.h"
 #include "Templates/SharedPointer.h"
 #include "Types/AttributeStorage.h"
+#include "UObject/AssetRegistryTagsContext.h"
 #include "UObject/Class.h"
 #include "UObject/Field.h"
 #include "UObject/NameTypes.h"
@@ -29,6 +30,7 @@
 #include "InterchangeBaseNode.generated.h"
 
 class FArchive;
+class FAssetRegistryTagsContext;
 class UInterchangeBaseNode;
 struct FFrame;
 
@@ -526,9 +528,13 @@ public:
 
 	static INTERCHANGECORE_API void CopyStorage(const UInterchangeBaseNode* SourceNode, UInterchangeBaseNode* DestinationNode);
 	
+	virtual void AppendAssetRegistryTags(FAssetRegistryTagsContext Context) const
+	{
+	}
+
+	UE_DEPRECATED(5.4, "Implement the version that takes FAssetRegistryTagsContext instead.")
 	virtual void AppendAssetRegistryTags(TArray<FAssetRegistryTag>& OutTags) const
 	{
-		return;
 	}
 
 protected:
diff --git a/Engine/Source/Runtime/Interchange/Engine/Private/InterchangeSceneImportAsset.cpp b/Engine/Source/Runtime/Interchange/Engine/Private/InterchangeSceneImportAsset.cpp
--- a/Engine/Source/Runtime/Interchange/Engine/Private/InterchangeSceneImportAsset.cpp
+++ b/Engine/Source/Runtime/Interchange/Engine/Private/InterchangeSceneImportAsset.cpp
@@ -9,6 +9,7 @@
 #include "Engine/Blueprint.h"
 #include "Engine/World.h"
 #include "EngineUtils.h"
+#include "UObject/AssetRegistryTagsContext.h"
 #include "UObject/Package.h"
 
 UInterchangeSceneImportAsset::~UInterchangeSceneImportAsset()
@@ -23,15 +24,22 @@ UInterchangeSceneImportAsset::~UInterchangeSceneImportAsset()
 }
 
 void UInterchangeSceneImportAsset::GetAssetRegistryTags(TArray<FAssetRegistryTag>& OutTags) const
+{
+	PRAGMA_DISABLE_DEPRECATION_WARNINGS;
+	Super::GetAssetRegistryTags(OutTags);
+	PRAGMA_ENABLE_DEPRECATION_WARNINGS;
+}
+
+void UInterchangeSceneImportAsset::GetAssetRegistryTags(FAssetRegistryTagsContext Context) const
 {
 #if WITH_EDITORONLY_DATA
 	if (AssetImportData)
 	{
-		OutTags.Add(FAssetRegistryTag(SourceFileTagName(), AssetImportData->GetSourceData().ToJson(), FAssetRegistryTag::TT_Hidden));
+		Context.AddTag(FAssetRegistryTag(SourceFileTagName(), AssetImportData->GetSourceData().ToJson(), FAssetRegistryTag::TT_Hidden));
 	}
 #endif
 
-	Super::GetAssetRegistryTags(OutTags);
+	Super::GetAssetRegistryTags(Context);
 }
 
 void UInterchangeSceneImportAsset::PostLoad()
diff --git a/Engine/Source/Runtime/Interchange/Engine/Public/InterchangeAssetImportData.h b/Engine/Source/Runtime/Interchange/Engine/Public/InterchangeAssetImportData.h
--- a/Engine/Source/Runtime/Interchange/Engine/Public/InterchangeAssetImportData.h
+++ b/Engine/Source/Runtime/Interchange/Engine/Public/InterchangeAssetImportData.h
@@ -6,6 +6,7 @@
 #include "EditorFramework/AssetImportData.h"
 #include "Nodes/InterchangeBaseNode.h"
 #include "Nodes/InterchangeBaseNodeContainer.h"
+#include "UObject/AssetRegistryTagsContext.h"
 #include "UObject/Class.h"
 #include "UObject/Object.h"
 #include "UObject/ObjectMacros.h"
@@ -69,10 +70,18 @@ public:
 
 #if WITH_EDITORONLY_DATA
 #if WITH_EDITOR
+	UE_DEPRECATED(5.4, "Implement the version that takes FAssetRegistryTagsContext instead.")
+	virtual void AppendAssetRegistryTags(TArray<FAssetRegistryTag>& OutTags) override
+	{
+		PRAGMA_DISABLE_DEPRECATION_WARNINGS;
+		Super::AppendAssetRegistryTags(OutTags);
+		PRAGMA_ENABLE_DEPRECATION_WARNINGS;
+	}
+
 	/**
 	 * This function add tags to the asset registry.
 	 */
-	virtual void AppendAssetRegistryTags(TArray<FAssetRegistryTag>& OutTags) override
+	virtual void AppendAssetRegistryTags(FAssetRegistryTagsContext Context) override
 	{
 		if(const UInterchangeBaseNodeContainer* NodeContainerTmp = GetNodeContainer())
 		{
@@ -80,11 +89,19 @@ public:
 			{
 				if (const UInterchangeBaseNode* Node = GetStoredNode(NodeUniqueID))
 				{
-					Node->AppendAssetRegistryTags(OutTags);
+					PRAGMA_DISABLE_DEPRECATION_WARNINGS;
+					TArray<UObject::FAssetRegistryTag> DeprecatedFunctionTags;
+					Node->AppendAssetRegistryTags(DeprecatedFunctionTags);
+					for (UObject::FAssetRegistryTag& Tag : DeprecatedFunctionTags)
+					{
+						Context.AddTag(MoveTemp(Tag));
+					}
+					PRAGMA_ENABLE_DEPRECATION_WARNINGS;
+					Node->AppendAssetRegistryTags(Context);
 				}
 			}
 		}
-		Super::AppendAssetRegistryTags(OutTags);
+		Super::AppendAssetRegistryTags(Context);
 	}
 #endif
 #endif
diff --git a/Engine/Source/Runtime/Interchange/Engine/Public/InterchangeSceneImportAsset.h b/Engine/Source/Runtime/Interchange/Engine/Public/InterchangeSceneImportAsset.h
--- a/Engine/Source/Runtime/Interchange/Engine/Public/InterchangeSceneImportAsset.h
+++ b/Engine/Source/Runtime/Interchange/Engine/Public/InterchangeSceneImportAsset.h
@@ -33,6 +33,8 @@ public:
 #endif // #if WITH_EDITORONLY_DATA
 
 	//~ Begin UObject Interface
+	INTERCHANGEENGINE_API virtual void GetAssetRegistryTags(FAssetRegistryTagsContext Context) const override;
+	UE_DEPRECATED(5.4, "Implement the version that takes FAssetRegistryTagsContext instead.")
 	INTERCHANGEENGINE_API virtual void GetAssetRegistryTags(TArray<FAssetRegistryTag>& OutTags) const override;
 	INTERCHANGEENGINE_API virtual void PostLoad() override;
 	//~ End UObject Interface
diff --git a/Engine/Source/Runtime/LevelSequence/Private/LevelSequence.cpp b/Engine/Source/Runtime/LevelSequence/Private/LevelSequence.cpp
--- a/Engine/Source/Runtime/LevelSequence/Private/LevelSequence.cpp
+++ b/Engine/Source/Runtime/LevelSequence/Private/LevelSequence.cpp
@@ -19,6 +19,7 @@
 #include "Engine/Engine.h"
 #include "MovieScene.h"
 #include "MovieSceneCommonHelpers.h"
+#include "UObject/AssetRegistryTagsContext.h"
 #include "UObject/Package.h"
 #include "UObject/UObjectHash.h"
 #include "Animation/AnimInstance.h"
@@ -160,11 +161,18 @@ ETrackSupport ULevelSequence::IsTrackSupported(TSubclassOf<class UMovieSceneTrac
 }
 
 void ULevelSequence::GetAssetRegistryTags(TArray<FAssetRegistryTag>& OutTags) const
+{
+	PRAGMA_DISABLE_DEPRECATION_WARNINGS;
+	Super::GetAssetRegistryTags(OutTags);
+	PRAGMA_ENABLE_DEPRECATION_WARNINGS;
+}
+
+void ULevelSequence::GetAssetRegistryTags(FAssetRegistryTagsContext Context) const
 {
 #if WITH_EDITORONLY_DATA
 	if (DirectorBlueprint)
 	{
-		DirectorBlueprint->GetAssetRegistryTags(OutTags);
+		DirectorBlueprint->GetAssetRegistryTags(Context);
 	}
 #endif
 
@@ -173,11 +181,19 @@ void ULevelSequence::GetAssetRegistryTags(TArray<FAssetRegistryTag>& OutTags) co
 		IMovieSceneMetaDataInterface* MetaDataInterface = Cast<IMovieSceneMetaDataInterface>(MetaData);
 		if (MetaDataInterface)
 		{
-			MetaDataInterface->ExtendAssetRegistryTags(OutTags);
+			PRAGMA_DISABLE_DEPRECATION_WARNINGS;
+			TArray<UObject::FAssetRegistryTag> DeprecatedFunctionTags;
+			MetaDataInterface->ExtendAssetRegistryTags(DeprecatedFunctionTags);
+			for (UObject::FAssetRegistryTag& Tag : DeprecatedFunctionTags)
+			{
+				Context.AddTag(MoveTemp(Tag));
+			}
+			PRAGMA_ENABLE_DEPRECATION_WARNINGS;
+			MetaDataInterface->ExtendAssetRegistryTags(Context);
 		}
 	}
 
-	Super::GetAssetRegistryTags(OutTags);
+	Super::GetAssetRegistryTags(Context);
 }
 
 void ULevelSequence::GetAssetRegistryTagMetadata(TMap<FName, FAssetRegistryTagMetadata>& OutMetadata) const
diff --git a/Engine/Source/Runtime/LevelSequence/Public/LevelSequence.h b/Engine/Source/Runtime/LevelSequence/Public/LevelSequence.h
--- a/Engine/Source/Runtime/LevelSequence/Public/LevelSequence.h
+++ b/Engine/Source/Runtime/LevelSequence/Public/LevelSequence.h
@@ -75,6 +75,8 @@ public:
 #if WITH_EDITOR
 	LEVELSEQUENCE_API virtual ETrackSupport IsTrackSupported(TSubclassOf<class UMovieSceneTrack> InTrackClass) const override;
 	LEVELSEQUENCE_API virtual void GetAssetRegistryTagMetadata(TMap<FName, FAssetRegistryTagMetadata>& OutMetadata) const override;
+	LEVELSEQUENCE_API virtual void GetAssetRegistryTags(FAssetRegistryTagsContext Context) const override;
+	UE_DEPRECATED(5.4, "Implement the version that takes FAssetRegistryTagsContext instead.")
 	LEVELSEQUENCE_API virtual void GetAssetRegistryTags(TArray<FAssetRegistryTag>& OutTags) const override;
 	LEVELSEQUENCE_API virtual void PostLoadAssetRegistryTags(const FAssetData& InAssetData, TArray<FAssetRegistryTag>& OutTagsAndValuesToUpdate) const override;
 	
diff --git a/Engine/Source/Runtime/MediaAssets/Private/Assets/BaseMediaSource.cpp b/Engine/Source/Runtime/MediaAssets/Private/Assets/BaseMediaSource.cpp
--- a/Engine/Source/Runtime/MediaAssets/Private/Assets/BaseMediaSource.cpp
+++ b/Engine/Source/Runtime/MediaAssets/Private/Assets/BaseMediaSource.cpp
@@ -2,6 +2,7 @@
 
 #include "BaseMediaSource.h"
 
+#include "UObject/AssetRegistryTagsContext.h"
 #include "UObject/ObjectSaveContext.h"
 #include "UObject/SequencerObjectVersion.h"
 #include "UObject/MediaFrameWorkObjectVersion.h"
@@ -20,17 +21,24 @@
  *****************************************************************************/
 
 void UBaseMediaSource::GetAssetRegistryTags(TArray<FAssetRegistryTag>& OutTags) const
+{
+	PRAGMA_DISABLE_DEPRECATION_WARNINGS;
+	Super::GetAssetRegistryTags(OutTags);
+	PRAGMA_ENABLE_DEPRECATION_WARNINGS;
+}
+
+void UBaseMediaSource::GetAssetRegistryTags(FAssetRegistryTagsContext Context) const
 {
 	FString Url = GetUrl();
 
 	if (!Url.IsEmpty())
 	{
-		OutTags.Add(FAssetRegistryTag("Url", Url, FAssetRegistryTag::TT_Alphabetical));
+		Context.AddTag(FAssetRegistryTag("Url", Url, FAssetRegistryTag::TT_Alphabetical));
 	}
 	
-	OutTags.Add(FAssetRegistryTag("Validate", Validate() ? TEXT("True") : TEXT("False"),
+	Context.AddTag(FAssetRegistryTag("Validate", Validate() ? TEXT("True") : TEXT("False"),
 		FAssetRegistryTag::TT_Alphabetical));
-	Super::GetAssetRegistryTags(OutTags);
+	Super::GetAssetRegistryTags(Context);
 }
 
 void UBaseMediaSource::PreSave(const class ITargetPlatform* TargetPlatform)
diff --git a/Engine/Source/Runtime/MediaAssets/Public/BaseMediaSource.h b/Engine/Source/Runtime/MediaAssets/Public/BaseMediaSource.h
--- a/Engine/Source/Runtime/MediaAssets/Public/BaseMediaSource.h
+++ b/Engine/Source/Runtime/MediaAssets/Public/BaseMediaSource.h
@@ -52,6 +52,8 @@ public:
 	MEDIAASSETS_API virtual void PreSave(const class ITargetPlatform* TargetPlatform);
 	MEDIAASSETS_API PRAGMA_ENABLE_DEPRECATION_WARNINGS
 	virtual void PreSave(FObjectPreSaveContext ObjectSaveContext);
+	MEDIAASSETS_API virtual void GetAssetRegistryTags(FAssetRegistryTagsContext Context) const override;
+	UE_DEPRECATED(5.4, "Implement the version that takes FAssetRegistryTagsContext instead.")
 	MEDIAASSETS_API virtual void GetAssetRegistryTags(TArray<FAssetRegistryTag>& OutTags) const override;
 	MEDIAASSETS_API virtual void Serialize(FArchive& Ar) override;
 
diff --git a/Engine/Source/Runtime/MovieScene/Private/MovieSceneMetaData.cpp b/Engine/Source/Runtime/MovieScene/Private/MovieSceneMetaData.cpp
--- a/Engine/Source/Runtime/MovieScene/Private/MovieSceneMetaData.cpp
+++ b/Engine/Source/Runtime/MovieScene/Private/MovieSceneMetaData.cpp
@@ -1,6 +1,7 @@
 // Copyright Epic Games, Inc. All Rights Reserved.
 
 #include "MovieSceneMetaData.h"
+#include "UObject/AssetRegistryTagsContext.h"
 #include "UObject/Package.h"
 #include "Engine/Level.h"
 
@@ -67,11 +68,13 @@ void UMovieSceneMetaData::SetNotes(FString InNotes)
 	Notes = InNotes;
 }
 
-void UMovieSceneMetaData::ExtendAssetRegistryTags(TArray<FAssetRegistryTag>& OutTags) const
+void UMovieSceneMetaData::ExtendAssetRegistryTags(FAssetRegistryTagsContext Context) const
 {
-	OutTags.Emplace(AssetRegistryTag_Author, Author, FAssetRegistryTag::ETagType::TT_Alphabetical, FAssetRegistryTag::TD_None);
-	OutTags.Emplace(AssetRegistryTag_Created, Created.ToString(), FAssetRegistryTag::ETagType::TT_Chronological, FAssetRegistryTag::TD_Date | FAssetRegistryTag::TD_Time);
-	OutTags.Emplace(AssetRegistryTag_Notes, Notes, FAssetRegistryTag::ETagType::TT_Alphabetical, FAssetRegistryTag::TD_None);
+	IMovieSceneMetaDataInterface::ExtendAssetRegistryTags(Context);
+
+	Context.AddTag(FAssetRegistryTag(AssetRegistryTag_Author, Author, FAssetRegistryTag::ETagType::TT_Alphabetical, FAssetRegistryTag::TD_None));
+	Context.AddTag(FAssetRegistryTag(AssetRegistryTag_Created, Created.ToString(), FAssetRegistryTag::ETagType::TT_Chronological, FAssetRegistryTag::TD_Date | FAssetRegistryTag::TD_Time));
+	Context.AddTag(FAssetRegistryTag(AssetRegistryTag_Notes, Notes, FAssetRegistryTag::ETagType::TT_Alphabetical, FAssetRegistryTag::TD_None));
 }
 
 #if WITH_EDITOR
diff --git a/Engine/Source/Runtime/MovieScene/Public/IMovieSceneMetaData.h b/Engine/Source/Runtime/MovieScene/Public/IMovieSceneMetaData.h
--- a/Engine/Source/Runtime/MovieScene/Public/IMovieSceneMetaData.h
+++ b/Engine/Source/Runtime/MovieScene/Public/IMovieSceneMetaData.h
@@ -2,6 +2,7 @@
 
 #pragma once
 
+#include "UObject/AssetRegistryTagsContext.h"
 #include "UObject/Interface.h"
 #include "IMovieSceneMetaData.generated.h"
 
@@ -25,6 +26,8 @@ public:
 	 * Called from ULevelSequence::GetAssetRegistryTags in order to
 	 * extend its default set of tags to include any from this meta-data object.
 	 */
+	virtual void ExtendAssetRegistryTags(FAssetRegistryTagsContext Context) const {}
+	UE_DEPRECATED(5.4, "Implement the version that takes FAssetRegistryTagsContext instead.")
 	virtual void ExtendAssetRegistryTags(TArray<UObject::FAssetRegistryTag>& OutTags) const {}
 
 #if WITH_EDITOR
diff --git a/Engine/Source/Runtime/MovieScene/Public/MovieSceneMetaData.h b/Engine/Source/Runtime/MovieScene/Public/MovieSceneMetaData.h
--- a/Engine/Source/Runtime/MovieScene/Public/MovieSceneMetaData.h
+++ b/Engine/Source/Runtime/MovieScene/Public/MovieSceneMetaData.h
@@ -49,7 +49,9 @@ public:
 	/**
 	 * Extend the default asset registry tags
 	 */
-	MOVIESCENE_API virtual void ExtendAssetRegistryTags(TArray<UObject::FAssetRegistryTag>& OutTags) const override;
+	MOVIESCENE_API virtual void ExtendAssetRegistryTags(FAssetRegistryTagsContext Context) const override;
+	UE_DEPRECATED(5.4, "Implement the version that takes FAssetRegistryTagsContext instead.")
+	MOVIESCENE_API virtual void ExtendAssetRegistryTags(TArray<UObject::FAssetRegistryTag>& OutTags) const override {}
 
 #if WITH_EDITOR
 
diff --git a/Engine/Source/Runtime/NNE/Private/NNEModelData.cpp b/Engine/Source/Runtime/NNE/Private/NNEModelData.cpp
--- a/Engine/Source/Runtime/NNE/Private/NNEModelData.cpp
+++ b/Engine/Source/Runtime/NNE/Private/NNEModelData.cpp
@@ -8,6 +8,7 @@
 #include "NNEModelOptimizerInterface.h"
 #include "NNERuntimeFormat.h"
 #include "Serialization/CustomVersion.h"
+#include "UObject/AssetRegistryTagsContext.h"
 #include "UObject/WeakInterfacePtr.h"
 
 #if WITH_EDITOR
@@ -197,8 +198,15 @@ namespace UE::NNE
 
 void UNNEModelData::GetAssetRegistryTags(TArray<FAssetRegistryTag>& OutTags) const
 {
-	OutTags.Add(FAssetRegistryTag("TargetRuntimes", UE::NNE::ModelData::GetRuntimesAsString(GetTargetRuntimes()), FAssetRegistryTag::TT_Alphabetical));
+	PRAGMA_DISABLE_DEPRECATION_WARNINGS;
 	Super::GetAssetRegistryTags(OutTags);
+	PRAGMA_ENABLE_DEPRECATION_WARNINGS;
+}
+
+void UNNEModelData::GetAssetRegistryTags(FAssetRegistryTagsContext Context) const
+{
+	Context.AddTag(FAssetRegistryTag("TargetRuntimes", UE::NNE::ModelData::GetRuntimesAsString(GetTargetRuntimes()), FAssetRegistryTag::TT_Alphabetical));
+	Super::GetAssetRegistryTags(Context);
 }
 
 void UNNEModelData::Serialize(FArchive& Ar)
diff --git a/Engine/Source/Runtime/Overlay/Private/Assets/BasicOverlays.cpp b/Engine/Source/Runtime/Overlay/Private/Assets/BasicOverlays.cpp
--- a/Engine/Source/Runtime/Overlay/Private/Assets/BasicOverlays.cpp
+++ b/Engine/Source/Runtime/Overlay/Private/Assets/BasicOverlays.cpp
@@ -4,6 +4,7 @@
 
 #include UE_INLINE_GENERATED_CPP_BY_NAME(BasicOverlays)
 
+#include "UObject/AssetRegistryTagsContext.h"
 #if WITH_EDITORONLY_DATA
 #include "EditorFramework/AssetImportData.h"
 #endif
@@ -21,15 +22,22 @@ void UBasicOverlays::PostInitProperties()
 }
 
 void UBasicOverlays::GetAssetRegistryTags(TArray<FAssetRegistryTag>& OutTags) const
+{
+	PRAGMA_DISABLE_DEPRECATION_WARNINGS;
+	Super::GetAssetRegistryTags(OutTags);
+	PRAGMA_ENABLE_DEPRECATION_WARNINGS;
+}
+
+void UBasicOverlays::GetAssetRegistryTags(FAssetRegistryTagsContext Context) const
 {
 #if WITH_EDITORONLY_DATA
 	if (AssetImportData)
 	{
-		OutTags.Add(FAssetRegistryTag(SourceFileTagName(), AssetImportData->GetSourceData().ToJson(), FAssetRegistryTag::TT_Hidden));
+		Context.AddTag(FAssetRegistryTag(SourceFileTagName(), AssetImportData->GetSourceData().ToJson(), FAssetRegistryTag::TT_Hidden));
 	}
 #endif	// WITH_EDITORONLY_DATA
 
-	Super::GetAssetRegistryTags(OutTags);
+	Super::GetAssetRegistryTags(Context);
 }
 
 TArray<FOverlayItem> UBasicOverlays::GetAllOverlays() const
diff --git a/Engine/Source/Runtime/Overlay/Private/Assets/LocalizedOverlays.cpp b/Engine/Source/Runtime/Overlay/Private/Assets/LocalizedOverlays.cpp
--- a/Engine/Source/Runtime/Overlay/Private/Assets/LocalizedOverlays.cpp
+++ b/Engine/Source/Runtime/Overlay/Private/Assets/LocalizedOverlays.cpp
@@ -5,6 +5,7 @@
 #include "BasicOverlays.h"
 #include "Internationalization/Internationalization.h"
 #include "Internationalization/Culture.h"
+#include "UObject/AssetRegistryTagsContext.h"
 
 #include UE_INLINE_GENERATED_CPP_BY_NAME(LocalizedOverlays)
 
@@ -25,15 +26,22 @@ void ULocalizedOverlays::PostInitProperties()
 }
 
 void ULocalizedOverlays::GetAssetRegistryTags(TArray<FAssetRegistryTag>& OutTags) const
+{
+	PRAGMA_DISABLE_DEPRECATION_WARNINGS;
+	Super::GetAssetRegistryTags(OutTags);
+	PRAGMA_ENABLE_DEPRECATION_WARNINGS;
+}
+
+void ULocalizedOverlays::GetAssetRegistryTags(FAssetRegistryTagsContext Context) const
 {
 #if WITH_EDITORONLY_DATA
 	if (AssetImportData)
 	{
-		OutTags.Add(FAssetRegistryTag(SourceFileTagName(), AssetImportData->GetSourceData().ToJson(), FAssetRegistryTag::TT_Hidden));
+		Context.AddTag(FAssetRegistryTag(SourceFileTagName(), AssetImportData->GetSourceData().ToJson(), FAssetRegistryTag::TT_Hidden));
 	}
 #endif	// WITH_EDITORONLY_DATA
 
-	Super::GetAssetRegistryTags(OutTags);
+	Super::GetAssetRegistryTags(Context);
 }
 
 UBasicOverlays* ULocalizedOverlays::GetCurrentLocaleOverlays() const
diff --git a/Engine/Source/Runtime/Overlay/Public/BasicOverlays.h b/Engine/Source/Runtime/Overlay/Public/BasicOverlays.h
--- a/Engine/Source/Runtime/Overlay/Public/BasicOverlays.h
+++ b/Engine/Source/Runtime/Overlay/Public/BasicOverlays.h
@@ -53,6 +53,8 @@ public:
 	//~ UObject interface
 	
 	OVERLAY_API virtual void PostInitProperties() override;
+	OVERLAY_API virtual void GetAssetRegistryTags(FAssetRegistryTagsContext Context) const override;
+	UE_DEPRECATED(5.4, "Implement the version that takes FAssetRegistryTagsContext instead.")
 	OVERLAY_API virtual void GetAssetRegistryTags(TArray<FAssetRegistryTag>& OutTags) const override;
 	// End UObject interface
 
diff --git a/Engine/Source/Runtime/Overlay/Public/LocalizedOverlays.h b/Engine/Source/Runtime/Overlay/Public/LocalizedOverlays.h
--- a/Engine/Source/Runtime/Overlay/Public/LocalizedOverlays.h
+++ b/Engine/Source/Runtime/Overlay/Public/LocalizedOverlays.h
@@ -68,6 +68,8 @@ public:
 	//~ UObject interface
 
 	OVERLAY_API virtual void PostInitProperties() override;
+	OVERLAY_API virtual void GetAssetRegistryTags(FAssetRegistryTagsContext Context) const override;
+	UE_DEPRECATED(5.4, "Implement the version that takes FAssetRegistryTagsContext instead.")
 	OVERLAY_API virtual void GetAssetRegistryTags(TArray<FAssetRegistryTag>& OutTags) const override;
 
 private:
diff --git a/Engine/Source/Runtime/UMG/Private/WidgetBlueprintGeneratedClass.cpp b/Engine/Source/Runtime/UMG/Private/WidgetBlueprintGeneratedClass.cpp
--- a/Engine/Source/Runtime/UMG/Private/WidgetBlueprintGeneratedClass.cpp
+++ b/Engine/Source/Runtime/UMG/Private/WidgetBlueprintGeneratedClass.cpp
@@ -10,6 +10,7 @@
 #include "Extensions/WidgetBlueprintGeneratedClassExtension.h"
 #include "UMGPrivate.h"
 #include "Interfaces/ITargetPlatform.h"
+#include "UObject/AssetRegistryTagsContext.h"
 #include "UObject/EditorObjectVersion.h"
 #include "UObject/ObjectResource.h"
 #include "UObject/LinkerLoad.h"
@@ -103,7 +104,10 @@ FAutoConsoleCommand GDumpTemplateSizesCommand(
 int32 TemplatePreviewInEditor = 0;
 static FAutoConsoleVariableRef CVarTemplatePreviewInEditor(TEXT("Widget.TemplatePreviewInEditor"), TemplatePreviewInEditor, TEXT("Should a dynamic template be generated at runtime for the editor for widgets?  Useful for debugging templates."), ECVF_Default);
 
+PRAGMA_DISABLE_DEPRECATION_WARNINGS;
 FWidgetBlueprintGeneratedClassDelegates::FGetAssetTags FWidgetBlueprintGeneratedClassDelegates::GetAssetTags;
+PRAGMA_ENABLE_DEPRECATION_WARNINGS;
+FWidgetBlueprintGeneratedClassDelegates::FGetAssetTagsWithContext FWidgetBlueprintGeneratedClassDelegates::GetAssetTagsWithContext;
 
 #endif
 
@@ -531,9 +535,24 @@ void UWidgetBlueprintGeneratedClass::GetExtensions(TArray<UWidgetBlueprintGenera
 #if WITH_EDITOR
 void UWidgetBlueprintGeneratedClass::GetAssetRegistryTags(TArray<FAssetRegistryTag>& OutTags) const
 {
+	PRAGMA_DISABLE_DEPRECATION_WARNINGS;
 	Super::GetAssetRegistryTags(OutTags);
+	PRAGMA_ENABLE_DEPRECATION_WARNINGS;
+}
+
+void UWidgetBlueprintGeneratedClass::GetAssetRegistryTags(FAssetRegistryTagsContext Context) const
+{
+	Super::GetAssetRegistryTags(Context);
 
-	FWidgetBlueprintGeneratedClassDelegates::GetAssetTags.Broadcast(this, OutTags);
+	PRAGMA_DISABLE_DEPRECATION_WARNINGS;
+	TArray<UObject::FAssetRegistryTag> DeprecatedFunctionTags;
+	FWidgetBlueprintGeneratedClassDelegates::GetAssetTags.Broadcast(this, DeprecatedFunctionTags);
+	for (UObject::FAssetRegistryTag& Tag : DeprecatedFunctionTags)
+	{
+		Context.AddTag(MoveTemp(Tag));
+	}
+	PRAGMA_ENABLE_DEPRECATION_WARNINGS;
+	FWidgetBlueprintGeneratedClassDelegates::GetAssetTagsWithContext.Broadcast(this, Context);
 }
 #endif
 
diff --git a/Engine/Source/Runtime/UMG/Public/Blueprint/WidgetBlueprintGeneratedClass.h b/Engine/Source/Runtime/UMG/Public/Blueprint/WidgetBlueprintGeneratedClass.h
--- a/Engine/Source/Runtime/UMG/Public/Blueprint/WidgetBlueprintGeneratedClass.h
+++ b/Engine/Source/Runtime/UMG/Public/Blueprint/WidgetBlueprintGeneratedClass.h
@@ -10,6 +10,7 @@
 
 #include "WidgetBlueprintGeneratedClass.generated.h"
 
+class FAssetRegistryTagsContext;
 class UWidget;
 class UUserWidget;
 class UWidgetAnimation;
@@ -56,9 +57,12 @@ class FWidgetBlueprintGeneratedClassDelegates
 {
 public:
 	// delegate for generating widget asset registry tags.
+	DECLARE_MULTICAST_DELEGATE_TwoParams(FGetAssetTagsWithContext, const UWidgetBlueprintGeneratedClass*, FAssetRegistryTagsContext);
 	DECLARE_MULTICAST_DELEGATE_TwoParams(FGetAssetTags, const UWidgetBlueprintGeneratedClass*, TArray<UObject::FAssetRegistryTag>&);
 
 	// called by UWidgetBlueprintGeneratedClass::GetAssetRegistryTags()
+	static UMG_API FGetAssetTagsWithContext GetAssetTagsWithContext;
+	UE_DEPRECATED(5.4, "Subscribe to GetAssetTagsWithContext instead.")
 	static UMG_API FGetAssetTags GetAssetTags;
 };
 #endif
@@ -152,6 +156,8 @@ public:
 	UMG_API virtual void PostLoad() override;
 	UMG_API virtual bool NeedsLoadForServer() const override;
 #if WITH_EDITOR
+	UMG_API virtual void GetAssetRegistryTags(FAssetRegistryTagsContext Context) const override;
+	UE_DEPRECATED(5.4, "Implement the version that takes FAssetRegistryTagsContext instead.")
 	UMG_API virtual void GetAssetRegistryTags(TArray<FAssetRegistryTag>& OutTags) const override;
 #endif
 	//~ End UObject interface
