[Material Editor] Enable DDC translation cache again.

diff --git a/Engine/Source/Runtime/Engine/Private/Materials/HLSLMaterialTranslator.cpp b/Engine/Source/Runtime/Engine/Private/Materials/HLSLMaterialTranslator.cpp
--- a/Engine/Source/Runtime/Engine/Private/Materials/HLSLMaterialTranslator.cpp
+++ b/Engine/Source/Runtime/Engine/Private/Materials/HLSLMaterialTranslator.cpp
@@ -206,7 +206,7 @@ static FAutoConsoleVariableRef CVarPreshaderGapInterval(
 static const bool GUseMaterialTranslationResultsGrouping = true;
 
 /* Controls whether DDC caching of material translation results should be forcefully disabled. */
-#define FORCE_DISABLE_MATERIAL_TRANSLATION_DDC true
+#define FORCE_DISABLE_MATERIAL_TRANSLATION_DDC false
 
 /* Helper macro to check whether the DDC material translation data has arrived and
  * therefore quit material translation. */
