Introduce FNetworkMetricsDatabase and integrate it with UNetDriver and UNetDriver::UpdateNetworkStats().

diff --git a/Engine/Config/BaseEngine.ini b/Engine/Config/BaseEngine.ini
--- a/Engine/Config/BaseEngine.ini
+++ b/Engine/Config/BaseEngine.ini
@@ -3462,3 +3462,69 @@ SceneImportSettings=(PipelineStacks=(("Scene", (Pipelines=("/Interchange/Pipelin
 FilePickerClass=/Script/InterchangeEditorUtilities.InterchangeFilePickerGeneric
 bStaticMeshUseSmoothEdgesIfSmoothingInformationIsMissing=True
 GenericPipelineClass=/Script/InterchangePipelines.InterchangeGenericAssetsPipeline
+
+[/Script/Engine.NetworkMetricsCSV_Replication]
+IntervalSeconds=0
+
+[/Script/Engine.NetworkMetricsPerfCounters]
+IntervalSeconds=1
+
+[/Script/Engine.NetworkMetricsConfig]
++Listeners=(MetricName=Connections, Class=/Script/Engine.NetworkMetricsCSV_Replication)
++Listeners=(MetricName=GatherPrioritizeTimeMS, Class=/Script/Engine.NetworkMetricsCSV_Replication)
++Listeners=(MetricName=InPacketsClientAvg, Class=/Script/Engine.NetworkMetricsCSV_Replication)
++Listeners=(MetricName=InPacketsClientMax, Class=/Script/Engine.NetworkMetricsCSV_Replication)
++Listeners=(MetricName=NumberOfActiveActors, Class=/Script/Engine.NetworkMetricsCSV_Replication)
++Listeners=(MetricName=NumberOfFullyDormantActors, Class=/Script/Engine.NetworkMetricsCSV_Replication)
++Listeners=(MetricName=NumClientUpdateLevelVisibility, Class=/Script/Engine.NetworkMetricsCSV_Replication)
++Listeners=(MetricName=NumOpenChannels, Class=/Script/Engine.NetworkMetricsCSV_Replication)
++Listeners=(MetricName=NumReplicateActorCallsPerConAvg, Class=/Script/Engine.NetworkMetricsCSV_Replication)
++Listeners=(MetricName=NumSkippedObjectEmptyUpdates, Class=/Script/Engine.NetworkMetricsCSV_Replication)
++Listeners=(MetricName=NumTickingChannels, Class=/Script/Engine.NetworkMetricsCSV_Replication)
++Listeners=(MetricName=OutKBytes, Class=/Script/Engine.NetworkMetricsCSV_Replication)
++Listeners=(MetricName=OutNetGUIDKBytesSec, Class=/Script/Engine.NetworkMetricsCSV_Replication)
++Listeners=(MetricName=OutPacketsClientAvg, Class=/Script/Engine.NetworkMetricsCSV_Replication)
++Listeners=(MetricName=OutPacketsClientMax, Class=/Script/Engine.NetworkMetricsCSV_Replication)
++Listeners=(MetricName=ReplicateActorTimeMS, Class=/Script/Engine.NetworkMetricsCSV_Replication)
++Listeners=(MetricName=SatConnections, Class=/Script/Engine.NetworkMetricsCSV_Replication)
++Listeners=(MetricName=ServerReplicateActorTimeMS, Class=/Script/Engine.NetworkMetricsCSV_Replication)
++Listeners=(MetricName=SharedSerializationPropertyHit, Class=/Script/Engine.NetworkMetricsCSV_Replication)
++Listeners=(MetricName=SharedSerializationPropertyMiss, Class=/Script/Engine.NetworkMetricsCSV_Replication)
++Listeners=(MetricName=SharedSerializationRPCHit, Class=/Script/Engine.NetworkMetricsCSV_Replication)
++Listeners=(MetricName=SharedSerializationRPCMiss, Class=/Script/Engine.NetworkMetricsCSV_Replication)
++Listeners=(MetricName=AvgPing, Class=/Script/Engine.NetworkMetricsPerfCounters)
++Listeners=(MetricName=InBunches, Class=/Script/Engine.NetworkMetricsPerfCounters)
++Listeners=(MetricName=InPackets, Class=/Script/Engine.NetworkMetricsPerfCounters)
++Listeners=(MetricName=InPacketsClientPerSecondAvg, Class=/Script/Engine.NetworkMetricsPerfCounters)
++Listeners=(MetricName=InPacketsClientPerSecondMax, Class=/Script/Engine.NetworkMetricsPerfCounters)
++Listeners=(MetricName=InPacketsClientPerSecondMin, Class=/Script/Engine.NetworkMetricsPerfCounters)
++Listeners=(MetricName=InPacketsLost, Class=/Script/Engine.NetworkMetricsPerfCounters)
++Listeners=(MetricName=InRate, Class=/Script/Engine.NetworkMetricsPerfCounters)
++Listeners=(MetricName=InRateClientAvg, Class=/Script/Engine.NetworkMetricsPerfCounters)
++Listeners=(MetricName=InRateClientMax, Class=/Script/Engine.NetworkMetricsPerfCounters)
++Listeners=(MetricName=InRateClientMin, Class=/Script/Engine.NetworkMetricsPerfCounters)
++Listeners=(MetricName=MaxPacketOverhead, Class=/Script/Engine.NetworkMetricsPerfCounters)
++Listeners=(MetricName=MaxPing, Class=/Script/Engine.NetworkMetricsPerfCounters)
++Listeners=(MetricName=MinPing, Class=/Script/Engine.NetworkMetricsPerfCounters)
++Listeners=(MetricName=NumClients, Class=/Script/Engine.NetworkMetricsPerfCounters)
++Listeners=(MetricName=Connections, Class=/Script/Engine.NetworkMetricsPerfCounters)
++Listeners=(MetricName=OutBunches, Class=/Script/Engine.NetworkMetricsPerfCounters)
++Listeners=(MetricName=OutPackets, Class=/Script/Engine.NetworkMetricsPerfCounters)
++Listeners=(MetricName=OutPacketsClientPerSecondAvg, Class=/Script/Engine.NetworkMetricsPerfCounters)
++Listeners=(MetricName=OutPacketsClientPerSecondMax, Class=/Script/Engine.NetworkMetricsPerfCounters)
++Listeners=(MetricName=OutPacketsClientPerSecondMin, Class=/Script/Engine.NetworkMetricsPerfCounters)
++Listeners=(MetricName=OutPacketsLost, Class=/Script/Engine.NetworkMetricsPerfCounters)
++Listeners=(MetricName=OutRate, Class=/Script/Engine.NetworkMetricsPerfCounters)
++Listeners=(MetricName=OutRateClientAvg, Class=/Script/Engine.NetworkMetricsPerfCounters)
++Listeners=(MetricName=OutRateClientMax, Class=/Script/Engine.NetworkMetricsPerfCounters)
++Listeners=(MetricName=OutRateClientMin, Class=/Script/Engine.NetworkMetricsPerfCounters)
++Listeners=(MetricName=PingBucketInt0, Class=/Script/Engine.NetworkMetricsPerfCounters)
++Listeners=(MetricName=PingBucketInt1, Class=/Script/Engine.NetworkMetricsPerfCounters)
++Listeners=(MetricName=PingBucketInt2, Class=/Script/Engine.NetworkMetricsPerfCounters)
++Listeners=(MetricName=PingBucketInt3, Class=/Script/Engine.NetworkMetricsPerfCounters)
++Listeners=(MetricName=PingBucketInt4, Class=/Script/Engine.NetworkMetricsPerfCounters)
++Listeners=(MetricName=PingBucketInt5, Class=/Script/Engine.NetworkMetricsPerfCounters)
++Listeners=(MetricName=PingBucketInt6, Class=/Script/Engine.NetworkMetricsPerfCounters)
++Listeners=(MetricName=PingBucketInt7, Class=/Script/Engine.NetworkMetricsPerfCounters)
++Listeners=(MetricName=ClosedConnectionsDueToReliableBufferOverflow, Class=/Script/Engine.NetworkMetricsPerfCounters)
++Listeners=(MetricName=AddedConnections, Class=/Script/Engine.NetworkMetricsPerfCounters)
diff --git a/Engine/Source/Runtime/Engine/Classes/Engine/NetDriver.h b/Engine/Source/Runtime/Engine/Classes/Engine/NetDriver.h
--- a/Engine/Source/Runtime/Engine/Classes/Engine/NetDriver.h
+++ b/Engine/Source/Runtime/Engine/Classes/Engine/NetDriver.h
@@ -7,6 +7,7 @@
 #include "CoreMinimal.h"
 #include "Engine/EngineBaseTypes.h"
 #include "Engine/EngineTypes.h"
+#include "Net/NetworkMetricsDatabase.h"
 #include "HAL/IConsoleManager.h"
 #include "Math/RandomStream.h"
 #include "UObject/ObjectMacros.h"
@@ -740,6 +741,21 @@ enum class EProcessRemoteFunctionFlags : uint32
 };
 ENUM_CLASS_FLAGS(EProcessRemoteFunctionFlags);
 
+/** A metrics listener that writes a metric to the 'Replication' CSV category. */
+UCLASS()
+class ENGINE_API UNetworkMetricsCSV_Replication : public UNetworkMetricsCSV
+{
+	GENERATED_BODY()
+
+public:
+	UNetworkMetricsCSV_Replication()
+	{
+		SetCategory("Replication");
+	}
+
+	virtual ~UNetworkMetricsCSV_Replication() = default;
+};
+
 UCLASS(Abstract, customConstructor, transient, MinimalAPI, config=Engine)
 class UNetDriver : public UObject, public FExec
 {
@@ -972,6 +988,23 @@ private:
 	/** Cached copy of MaxChannelsOverride from the net driver definition to avoid extra lookups */
 	int32 MaxChannelsOverride;
 
+	/** A metrics database that holds statistics calcluated by the networking system. */
+	UPROPERTY()
+	TObjectPtr<UNetworkMetricsDatabase> NetworkMetricsDatabase;
+
+	/** A cache of UNetworkMetricsBaseListener sub-class instances provided by the *.ini file (one instance per sub-class). */
+	UPROPERTY()
+	TMap<FName, TObjectPtr<UNetworkMetricsBaseListener>> NetworkMetricsListeners;
+
+	/** Register each metric used by the networking system. */
+	void SetupNetworkMetrics();
+
+	/** Register metric listeners provided by the *.ini file. */
+	void SetupNetworkMetricsListeners();
+
+	/** Create an instance of UNetworkMetricsStats that is associated with a given Stat and cached with other listeners in NetworkMetricsListeners. */
+	void RegisterStatsListener(const FName MetricName, const FName StatName);
+
 public:
 	/** Get the value of MaxChannelsOverride cached from the net driver definition */
 	int32 GetMaxChannelsOverride() const { return MaxChannelsOverride; }
@@ -1459,6 +1492,8 @@ public:
 	 */
 	ENGINE_API virtual void ProcessRemoteFunction(class AActor* Actor, class UFunction* Function, void* Parameters, struct FOutParmRec* OutParms, struct FFrame* Stack, class UObject* SubObject = nullptr );
 
+	/** Return a reference to the database that holds metrics calcluated by the networking system. */
+	ENGINE_API TObjectPtr<UNetworkMetricsDatabase> GetMetrics() { return NetworkMetricsDatabase; };
 
 	enum ERemoteFunctionSendPolicy
 	{		
diff --git a/Engine/Source/Runtime/Engine/Private/DataChannel.cpp b/Engine/Source/Runtime/Engine/Private/DataChannel.cpp
--- a/Engine/Source/Runtime/Engine/Private/DataChannel.cpp
+++ b/Engine/Source/Runtime/Engine/Private/DataChannel.cpp
@@ -24,6 +24,7 @@
 #include "DrawDebugHelpers.h"
 #include "Net/NetworkProfiler.h"
 #include "Net/DataReplication.h"
+#include "Net/NetworkMetricsDefs.h"
 #include "Engine/ActorChannel.h"
 #include "Engine/ControlChannel.h"
 #include "Engine/DemoNetDriver.h"
@@ -3752,7 +3753,7 @@ int64 UActorChannel::ReplicateActor()
 
 	bForceCompareProperties = false;		// Only do this once per frame when set
 	
-	INC_DWORD_STAT_BY(STAT_NumReplicatedActorBytes, (NumBitsWrote + 7) >> 3);
+	Connection->GetDriver()->GetMetrics()->IncrementInt(UE::Net::Metric::NumReplicatedActorBytes, (NumBitsWrote + 7) >> 3);
 	return NumBitsWrote;
 }
 
diff --git a/Engine/Source/Runtime/Engine/Private/Net/NetworkMetricsConfig.cpp b/Engine/Source/Runtime/Engine/Private/Net/NetworkMetricsConfig.cpp
new file mode 100644
--- /dev/null
+++ b/Engine/Source/Runtime/Engine/Private/Net/NetworkMetricsConfig.cpp
@@ -0,0 +1,4 @@
+// Copyright Epic Games, Inc. All Rights Reserved.
+#include "Net/NetworkMetricsConfig.h"
+
+#include UE_INLINE_GENERATED_CPP_BY_NAME(NetworkMetricsConfig)
diff --git a/Engine/Source/Runtime/Engine/Private/Net/NetworkMetricsDatabase.cpp b/Engine/Source/Runtime/Engine/Private/Net/NetworkMetricsDatabase.cpp
new file mode 100644
--- /dev/null
+++ b/Engine/Source/Runtime/Engine/Private/Net/NetworkMetricsDatabase.cpp
@@ -0,0 +1,263 @@
+// Copyright Epic Games, Inc. All Rights Reserved.
+
+#include "Net/NetworkMetricsDatabase.h"
+#include "EngineStats.h"
+#include "ProfilingDebugging/CsvProfiler.h"
+
+#if USE_SERVER_PERF_COUNTERS
+#include "PerfCountersModule.h"
+#endif
+
+#include UE_INLINE_GENERATED_CPP_BY_NAME(NetworkMetricsDatabase)
+
+void UNetworkMetricsDatabase::CreateFloat(const FName MetricName, float DefaultValue)
+{
+	if (ensureMsgf(!MetricTypes.Contains(MetricName), TEXT("Metric %s already exists in the database."), *MetricName.ToString()))
+	{
+		UE::Net::FNetworkMetric<float> Metric;
+		Metric.Name = MetricName;
+		Metric.Value = DefaultValue;
+		MetricFloats.Add(MetricName, Metric);
+		MetricTypes.Add(MetricName, EMetricType::Float);
+	}
+}
+
+void UNetworkMetricsDatabase::CreateInt(const FName MetricName, int64 DefaultValue)
+{
+	if (ensureMsgf(!MetricTypes.Contains(MetricName), TEXT("Metric %s already exists in the database."), *MetricName.ToString()))
+	{
+		UE::Net::FNetworkMetric<int64> Metric;
+		Metric.Name = MetricName;
+		Metric.Value = DefaultValue;
+		MetricInts.Add(MetricName, Metric);
+		MetricTypes.Add(MetricName, EMetricType::Integer);
+	}
+}
+
+void UNetworkMetricsDatabase::SetFloat(const FName MetricName, float Value)
+{
+	UE::Net::FNetworkMetric<float>* Metric = MetricFloats.Find(MetricName);
+	if (ensureMsgf(Metric, TEXT("Cannot find float metric %s."), *MetricName.ToString()))
+	{
+		Metric->Value = Value;
+	}
+}
+
+void UNetworkMetricsDatabase::SetInt(const FName MetricName, int64 Value)
+{
+	UE::Net::FNetworkMetric<int64>* Metric = MetricInts.Find(MetricName);
+	if (ensureMsgf(Metric, TEXT("Cannot find integer metric %s."), *MetricName.ToString()))
+	{
+		Metric->Value = Value;
+	}
+}
+
+void UNetworkMetricsDatabase::IncrementInt(const FName MetricName, int64 Value)
+{
+	UE::Net::FNetworkMetric<int64>* Metric = MetricInts.Find(MetricName);
+	if (ensureMsgf(Metric, TEXT("Cannot find integer metric %s."), *MetricName.ToString()))
+	{
+		Metric->Value += Value;
+	}
+}
+
+bool UNetworkMetricsDatabase::Contains(const FName MetricName) const
+{
+	return MetricTypes.Contains(MetricName);
+}
+
+void UNetworkMetricsDatabase::ProcessListeners()
+{
+	double CurrentTimeSeconds = FPlatformTime::Seconds();
+
+	UE::Net::FNetworkMetricSnapshot Snapshot;
+	Snapshot.MetricFloats.Reserve(MetricFloats.Num());
+	Snapshot.MetricInts.Reserve(MetricInts.Num());
+
+	for (const TPair<TWeakObjectPtr<UNetworkMetricsBaseListener>, TSet<FNameAndType>>& Elem: ListenersToMetrics)
+	{	
+		UNetworkMetricsBaseListener* Listener = Elem.Key.Get();
+		const TSet<FNameAndType>& ListenerMetrics = Elem.Value;
+
+		if (ensure(Listener))
+		{
+			if (!HasReportIntervalPassed(CurrentTimeSeconds, Listener))
+			{
+				continue;
+			}
+
+			Snapshot.Reset();
+
+			for (const FNameAndType& NameAndType : ListenerMetrics)
+			{
+				const FName& MetricName = NameAndType.Key;
+				const EMetricType& MetricType = NameAndType.Value;
+
+				if (MetricType == EMetricType::Integer)
+				{
+					const UE::Net::FNetworkMetric<int64>* Metric = MetricInts.Find(MetricName);
+					if (ensureMsgf(Metric, TEXT("Unable to find metric %s when reporting to listeners."), *MetricName.ToString()))
+					{
+						Snapshot.MetricInts.Add(*Metric);
+					}
+				}
+				else if (MetricType == EMetricType::Float)
+				{
+					const UE::Net::FNetworkMetric<float>* Metric = MetricFloats.Find(MetricName);
+					if (ensureMsgf(Metric, TEXT("Unable to find metric %s when reporting to listeners."), *MetricName.ToString()))
+					{
+						Snapshot.MetricFloats.Add(*Metric);
+					}
+				}
+			}
+
+			Listener->Report(Snapshot);
+		}
+	}
+}
+
+void UNetworkMetricsDatabase::Reset()
+{
+	MetricInts.Reset();
+	MetricFloats.Reset();
+	MetricTypes.Reset();
+	LastReportListener.Reset();
+	ListenersToMetrics.Reset();
+}
+
+void UNetworkMetricsDatabase::Register(const FName MetricName, TWeakObjectPtr<UNetworkMetricsBaseListener> Reporter)
+{
+	const EMetricType* MetricType = MetricTypes.Find(MetricName);
+
+	if (ensureMsgf(MetricType, TEXT("Cannot find metric %s to register listener."), *MetricName.ToString()))
+	{
+		TSet<FNameAndType>& ListenerMetrics = ListenersToMetrics.FindOrAdd(Reporter);
+		ListenerMetrics.Add(FNameAndType(MetricName, *MetricType));
+		LastReportListener.Add(Reporter, 0.0);
+	}
+}
+
+bool UNetworkMetricsDatabase::HasReportIntervalPassed(double CurrentTimeSeconds, UNetworkMetricsBaseListener* Listener)
+{
+	if (ensureMsgf(Listener->GetInterval() >= 0, TEXT("Listener has a negative reporting time interval.")))
+	{
+		double* LastReportSeconds = LastReportListener.Find(Listener);
+
+		if (ensure(LastReportSeconds))
+		{
+			double DurationSeconds = CurrentTimeSeconds - *LastReportSeconds;
+
+			if (DurationSeconds >= Listener->GetInterval())
+			{
+				*LastReportSeconds = CurrentTimeSeconds;
+				return true;
+			}
+		}
+	}
+
+	return false;
+}
+
+UNetworkMetricsBaseListener::UNetworkMetricsBaseListener() :
+	UObject(),
+	IntervalSeconds(0)
+{
+
+}
+
+UNetworkMetricsCSV::UNetworkMetricsCSV() :
+	CategoryIndex(-1)
+{
+	// The default CSV category for networking metrics.
+	SetCategory("Networking");
+}
+
+void UNetworkMetricsCSV::SetCategory(const FString& CategoryName)
+{
+#if CSV_PROFILER
+	int32 Index = FCsvProfiler::GetCategoryIndex(CategoryName);
+	if (ensureMsgf(Index != 1, TEXT("Unable to find CSV category %s"), *CategoryName))
+	{
+		CategoryIndex = Index;
+	}
+#endif
+}
+
+void UNetworkMetricsCSV::Report(const UE::Net::FNetworkMetricSnapshot& Snapshot)
+{
+#if CSV_PROFILER
+	if (ensureMsgf(CategoryIndex != -1, TEXT("SetCategory() must be called before being registered as a listener.")))
+	{
+		for (const UE::Net::FNetworkMetric<int64>& Metric : Snapshot.MetricInts)
+		{
+			if (ensureMsgf(IntFitsIn<int32>(Metric.Value), TEXT("Integer metric %s truncated when reporting to CSV."), *Metric.Name.ToString()))
+			{
+				FCsvProfiler::RecordCustomStat(Metric.Name, CategoryIndex, static_cast<int32>(Metric.Value), ECsvCustomStatOp::Set);
+			}
+		}
+
+		for (const UE::Net::FNetworkMetric<float>& Metric : Snapshot.MetricFloats)
+		{
+			FCsvProfiler::RecordCustomStat(Metric.Name, CategoryIndex, Metric.Value, ECsvCustomStatOp::Set);
+		}
+	}
+#endif
+}
+
+void UNetworkMetricsPerfCounters::Report(const UE::Net::FNetworkMetricSnapshot& Snapshot)
+{
+#if USE_SERVER_PERF_COUNTERS
+	IPerfCounters* PerfCounters = IPerfCountersModule::Get().GetPerformanceCounters();
+
+	if (PerfCounters)
+	{
+		for (const UE::Net::FNetworkMetric<int64>& Metric : Snapshot.MetricInts)
+		{
+			if (ensureMsgf(IntFitsIn<int32>(Metric.Value), TEXT("Integer metric %s truncated when reporting to PerfCounters."), *Metric.Name.ToString()))
+			{
+				PerfCounters->Set(Metric.Name.ToString(), static_cast<uint32>(Metric.Value));
+			}
+		}
+
+		for (const UE::Net::FNetworkMetric<float>& Metric : Snapshot.MetricFloats)
+		{
+			PerfCounters->Set(Metric.Name.ToString(), Metric.Value);
+		}
+	}
+#endif
+}
+
+UNetworkMetricsStats::UNetworkMetricsStats()
+{
+
+}
+
+void UNetworkMetricsStats::Report(const UE::Net::FNetworkMetricSnapshot& Snapshot)
+{
+#if STATS
+	if (!FThreadStats::IsCollectingData())
+	{
+		return;
+	}
+
+	const int32 TotalMetrics = Snapshot.MetricFloats.Num() + Snapshot.MetricInts.Num();
+
+	// An instance of UNetworkMetricStats is bound to a specific Stat value defined with the DEFINE_STAT macro 
+	// so there should only be one metric provided to this function.
+	if (ensureMsgf(TotalMetrics <= 1, TEXT("UNetworkMetricsStats should only be registered to listen to one metric.")))
+	{
+		if (LIKELY(TotalMetrics == 1))
+		{
+			for (const UE::Net::FNetworkMetric<int64>& Metric : Snapshot.MetricInts)
+			{
+				SET_DWORD_STAT_FName(StatName, Metric.Value);
+			}
+
+			for (const UE::Net::FNetworkMetric<float>& Metric : Snapshot.MetricFloats)
+			{
+				SET_FLOAT_STAT_FName(StatName, Metric.Value);
+			}
+		}
+	}
+#endif
+}
diff --git a/Engine/Source/Runtime/Engine/Private/Net/NetworkMetricsDefs.cpp b/Engine/Source/Runtime/Engine/Private/Net/NetworkMetricsDefs.cpp
new file mode 100644
--- /dev/null
+++ b/Engine/Source/Runtime/Engine/Private/Net/NetworkMetricsDefs.cpp
@@ -0,0 +1,97 @@
+// Copyright Epic Games, Inc. All Rights Reserved.
+
+#include "Net/NetworkMetricsDefs.h"
+
+namespace UE::Net::Metric
+{
+
+const FName AddedConnections("AddedConnections");
+const FName ClosedConnectionsDueToReliableBufferOverflow("ClosedConnectionsDueToReliableBufferOverflow");
+const FName InPackets("InPackets");
+const FName OutPackets("OutPackets");
+const FName InPacketsClientPerSecondAvg("InPacketsClientPerSecondAvg");
+const FName InPacketsClientPerSecondMax("InPacketsClientPerSecondMax");
+const FName InPacketsClientPerSecondMin("InPacketsClientPerSecondMin");
+const FName OutPacketsClientPerSecondAvg("OutPacketsClientPerSecondAvg");
+const FName OutPacketsClientPerSecondMax("OutPacketsClientPerSecondMax");
+const FName OutPacketsClientPerSecondMin("OutPacketsClientPerSecondMin");
+const FName InPacketsClientAvg("InPacketsClientAvg");
+const FName InPacketsClientMax("InPacketsClientMax");
+const FName OutPacketsClientAvg("OutPacketsClientAvg");
+const FName OutPacketsClientMax("OutPacketsClientMax");
+const FName InRate("InRate");
+const FName OutRate("OutRate");
+const FName InRateClientAvg("InRateClientAvg");
+const FName InRateClientMax("InRateClientMax");
+const FName InRateClientMin("InRateClientMin");
+const FName OutRateClientAvg("OutRateClientAvg");
+const FName OutRateClientMax("OutRateClientMax");
+const FName OutRateClientMin("OutRateClientMin");
+const FName InPacketsLost("InPacketsLost");
+const FName OutPacketsLost("OutPacketsLost");
+const FName InBunches("InBunches");
+const FName OutBunches("OutBunches");
+const FName Ping("Ping");
+const FName AvgPing("AvgPing");
+const FName MinPing("MinPing");
+const FName MaxPing("MaxPing");
+const FName OutKBytes("OutKBytes");
+const FName OutNetGUIDKBytesSec("OutNetGUIDKBytesSec");
+const FName ServerReplicateActorTimeMS("ServerReplicateActorTimeMS");
+const FName GatherPrioritizeTimeMS("GatherPrioritizeTimeMS");
+const FName ReplicateActorTimeMS("ReplicateActorTimeMS");
+const FName NumReplicatedActors("NumReplicatedActors");
+const FName Connections("Connections");
+const FName NumReplicateActorCallsPerConAvg("NumReplicateActorCallsPerConAvg");
+const FName NumberOfActiveActors("NumberOfActiveActors");
+const FName NumberOfFullyDormantActors("NumberOfFullyDormantActors");
+const FName NumSkippedObjectEmptyUpdates("NumSkippedObjectEmptyUpdates");
+const FName NumOpenChannels("NumOpenChannels");
+const FName NumTickingChannels("NumTickingChannels");
+const FName SatConnections("SatConnections");
+const FName SharedSerializationPropertyHit("SharedSerializationPropertyHit");
+const FName SharedSerializationPropertyMiss("SharedSerializationPropertyMiss");
+const FName SharedSerializationRPCHit("SharedSerializationRPCHit");
+const FName SharedSerializationRPCMiss("SharedSerializationRPCMiss");
+const FName NumClientUpdateLevelVisibility("NumClientUpdateLevelVisibility");
+const FName Channels("Channels");
+const FName NumActorChannels("NumActorChannels");
+const FName NumDormantActors("NumDormantActors");
+const FName NumActors("NumActors");
+const FName NumNetActors("NumNetActors");
+const FName NumNetGUIDsAckd("NumNetGUIDsAckd");
+const FName NumNetGUIDsPending("NumNetGUIDsPending");
+const FName NumNetGUIDsUnAckd("NumNetGUIDsUnAckd");
+const FName NetSaturated("NetSaturated");
+const FName MaxPacketOverhead("MaxPacketOverhead");
+const FName NetGUIDInRate("NetGUIDInRate");
+const FName NetGUIDOutRate("NetGUIDOutRate");
+const FName NetNumClients("NetNumClients");
+const FName NumClients("NumClients");
+const FName NumConnections("NumConnections");
+const FName NumConsideredActors("NumConsideredActors");
+const FName NumInitiallyDormantActors("NumInitiallyDormantActors");
+const FName PrioritizedActors("PrioritizedActors");
+const FName NumRelevantDeletedActors("NumRelevantDeletedActors");
+const FName NumReplicatedActorBytes("NumReplicatedActorBytes");
+const FName ImportedNetGuids("ImportedNetGuids");
+const FName PendingOuterNetGuids("PendingOuterNetGuids");
+const FName NetInBunchTimeOvershootPercent("NetInBunchTimeOvershootPercent");
+const FName UnmappedReplicators("UnmappedReplicators");
+const FName PercentInVoice("PercentInVoice");
+const FName PercentOutVoice("PercentOutVoice");
+const FName VoiceBytesRecv("VoiceBytesRecv");
+const FName VoiceBytesSent("VoiceBytesSent");
+const FName VoicePacketsRecv("VoicePacketsRecv");
+const FName VoicePacketsSent("VoicePacketsSent");
+const FName PingBucketInt0("PingBucketInt0");
+const FName PingBucketInt1("PingBucketInt1");
+const FName PingBucketInt2("PingBucketInt2");
+const FName PingBucketInt3("PingBucketInt3");
+const FName PingBucketInt4("PingBucketInt4");
+const FName PingBucketInt5("PingBucketInt5");
+const FName PingBucketInt6("PingBucketInt6");
+const FName PingBucketInt7("PingBucketInt7");
+
+}
+
diff --git a/Engine/Source/Runtime/Engine/Private/NetDriver.cpp b/Engine/Source/Runtime/Engine/Private/NetDriver.cpp
--- a/Engine/Source/Runtime/Engine/Private/NetDriver.cpp
+++ b/Engine/Source/Runtime/Engine/Private/NetDriver.cpp
@@ -44,6 +44,8 @@
 #include "Engine/LevelScriptActor.h"
 #include "Engine/NetworkSettings.h"
 #include "Net/NetEmulationHelper.h"
+#include "Net/NetworkMetricsDefs.h"
+#include "Net/NetworkMetricsConfig.h"
 #include "Net/NetSubObjectRegistryGetter.h"
 #include "Net/NetworkGranularMemoryLogging.h"
 #include "UObject/Stack.h"
@@ -149,6 +151,13 @@ DECLARE_CYCLE_STAT(TEXT("NetDriver TickFlush GatherStatsPerfCounters"), STAT_Net
 DECLARE_CYCLE_STAT(TEXT("ReceiveRPC_ProcessRemoteFunction"), STAT_NetReceiveRPC_ProcessRemoteFunction, STATGROUP_Game);
 DECLARE_DWORD_ACCUMULATOR_STAT(TEXT("Net bunch time % overshoot frames"), STAT_NetInBunchTimeOvershootPercent, STATGROUP_Net);
 
+DECLARE_DWORD_COUNTER_STAT(TEXT("Num Saturated Connections"), STAT_NumSaturatedConnections, STATGROUP_Net);
+DECLARE_DWORD_COUNTER_STAT(TEXT("SharedSerialization RPC Hit"), STAT_SharedSerializationRPCHit, STATGROUP_Net);
+DECLARE_DWORD_COUNTER_STAT(TEXT("SharedSerialization RPC Miss"), STAT_SharedSerializationRPCMiss, STATGROUP_Net);
+DECLARE_DWORD_COUNTER_STAT(TEXT("Empty Object Replicate Properties Skipped"), STAT_NumSkippedObjectEmptyUpdates, STATGROUP_Net);
+DECLARE_DWORD_COUNTER_STAT(TEXT("SharedSerialization Property Hit"), STAT_SharedSerializationPropertyHit, STATGROUP_Net);
+DECLARE_DWORD_COUNTER_STAT(TEXT("SharedSerialization Property Miss"), STAT_SharedSerializationPropertyMiss, STATGROUP_Net);
+
 DEFINE_LOG_CATEGORY_STATIC(LogNetSyncLoads, Log, All);
 
 CSV_DEFINE_CATEGORY(Networking, true);
@@ -160,38 +169,6 @@ int32 GNumSkippedObjectEmptyUpdates;
 
 extern int32 GNetRPCDebug;
 
-namespace UE::Net::Private
-{
-	static const FString PerfCounter_NumClients(TEXT("NumClients"));
-	static const FString PerfCounter_MaxPacketOverhead(TEXT("MaxPacketOverhead"));
-	static const FString PerfCounter_InRateClientMax(TEXT("InRateClientMax"));
-	static const FString PerfCounter_InRateClientMin(TEXT("InRateClientMin"));
-	static const FString PerfCounter_InRateClientAvg(TEXT("InRateClientAvg"));
-	static const FString PerfCounter_InPacketsClientMax(TEXT("InPacketsClientMax"));
-	static const FString PerfCounter_InPacketsClientMin(TEXT("InPacketsClientMin"));
-	static const FString PerfCounter_InPacketsClientAvg(TEXT("InPacketsClientAvg"));
-	static const FString PerfCounter_OutRateClientMax(TEXT("OutRateClientMax"));
-	static const FString PerfCounter_OutRateClientMin(TEXT("OutRateClientMin"));
-	static const FString PerfCounter_OutRateClientAvg(TEXT("OutRateClientAvg"));
-	static const FString PerfCounter_OutPacketsClientMax(TEXT("OutPacketsClientMax"));
-	static const FString PerfCounter_OutPacketsClientMin(TEXT("OutPacketsClientMin"));
-	static const FString PerfCounter_OutPacketsClientAvg(TEXT("OutPacketsClientAvg"));
-
-	static const FString PerfCounter_InRate(TEXT("InRate"));
-	static const FString PerfCounter_OutRate(TEXT("OutRate"));
-	static const FString PerfCounter_InPacketsLost(TEXT("InPacketsLost"));
-	static const FString PerfCounter_OutPacketsLost(TEXT("OutPacketsLost"));
-	static const FString PerfCounter_InPackets(TEXT("InPackets"));
-	static const FString PerfCounter_OutPackets(TEXT("OutPackets"));
-	static const FString PerfCounter_InBunches(TEXT("InBunches"));
-	static const FString PerfCounter_OutBunches(TEXT("OutBunches"));
-
-	static const FString PerfCounter_AvgPing(TEXT("AvgPing"));
-	static const FString PerfCounter_MaxPing(TEXT("MaxPing"));
-	static const FString PerfCounter_MinPing(TEXT("MinPing"));
-	static const FString PerfCounter_NumConnections(TEXT("NumConnections"));
-}
-
 namespace UE::Net::Private
 {
 	struct FAutoDestructProperty
@@ -1027,9 +1004,9 @@ void UNetDriver::TickFlush(float DeltaSeconds)
 			TSet<FNetworkGUID>& ImportedNetGuidsRef = GuidCache->ImportedNetGuids;
 			TMap<FNetworkGUID, TSet<FNetworkGUID>>& PendingOuterNetGuidsRef = GuidCache->PendingOuterNetGuids;
 
-			SET_DWORD_STAT(STAT_ImportedNetGuids, ImportedNetGuidsRef.Num());
-			SET_DWORD_STAT(STAT_PendingOuterNetGuids, PendingOuterNetGuidsRef.Num());
-			SET_DWORD_STAT(STAT_UnmappedReplicators, UnmappedReplicators.Num());
+			GetMetrics()->SetInt(UE::Net::Metric::ImportedNetGuids, ImportedNetGuidsRef.Num());
+			GetMetrics()->SetInt(UE::Net::Metric::PendingOuterNetGuids, PendingOuterNetGuidsRef.Num());
+			GetMetrics()->SetInt(UE::Net::Metric::UnmappedReplicators, UnmappedReplicators.Num());
 
 			TSet<FObjectReplicator*> ForceUpdateReplicators;
 
@@ -1225,6 +1202,9 @@ void UNetDriver::TickFlush(float DeltaSeconds)
 
 	UpdateNetworkStats();
 
+	// Send the current values of metrics to all of the metrics listeners.
+	GetMetrics()->ProcessListeners();
+
 	// Update the lag state
 	UpdateNetworkLagState();
 }
@@ -1542,9 +1522,320 @@ bool UNetDriver::InitBase(bool bInitAsClient, FNetworkNotify* InNotify, const FU
 
 	CachedGlobalNetTravelCount = GEngine->GetGlobalNetTravelCount();
 
+	// Add all of the metrics used by the networking system and register metrics listeners.
+	SetupNetworkMetrics();
+
+	if (NetDriverDefinition == NAME_GameNetDriver)
+	{
+		SetupNetworkMetricsListeners();
+	}
+
 	return bSuccess;
 }
 
+void UNetDriver::SetupNetworkMetrics()
+{
+	NetworkMetricsDatabase = NewObject<UNetworkMetricsDatabase>();
+
+	// The total number of connections added since the driver was started.
+	GetMetrics()->CreateInt(UE::Net::Metric::AddedConnections, 0);
+
+	// The total number of connections that were closed because of a reliable buffer overflow since the driver was started.
+	GetMetrics()->CreateInt(UE::Net::Metric::ClosedConnectionsDueToReliableBufferOverflow, 0);
+
+	// The number of incoming packets per second across all connections.
+	GetMetrics()->CreateInt(UE::Net::Metric::InPackets, 0);
+
+	// The number of outgoing packets per second across all connections.
+	GetMetrics()->CreateInt(UE::Net::Metric::OutPackets, 0);
+
+	// The average/min/max incoming packets per second across all connections.
+	GetMetrics()->CreateInt(UE::Net::Metric::InPacketsClientPerSecondAvg, 0);
+	GetMetrics()->CreateInt(UE::Net::Metric::InPacketsClientPerSecondMax, 0);
+	GetMetrics()->CreateInt(UE::Net::Metric::InPacketsClientPerSecondMin, 0);
+	
+	// The average/min/max outgoing packets per second across all connections.
+	GetMetrics()->CreateInt(UE::Net::Metric::OutPacketsClientPerSecondAvg, 0);
+	GetMetrics()->CreateInt(UE::Net::Metric::OutPacketsClientPerSecondMax, 0);
+	GetMetrics()->CreateInt(UE::Net::Metric::OutPacketsClientPerSecondMin, 0);
+	
+	// (Server only) The average/max incoming packets per frame across all connections.
+	GetMetrics()->CreateInt(UE::Net::Metric::InPacketsClientAvg, 0);
+	GetMetrics()->CreateInt(UE::Net::Metric::InPacketsClientMax, 0);
+
+	// (Server only) The average/max outgoing packets per frame across all connections.
+	GetMetrics()->CreateInt(UE::Net::Metric::OutPacketsClientAvg, 0);
+	GetMetrics()->CreateInt(UE::Net::Metric::OutPacketsClientMax, 0);
+
+	// The total number of incoming bytes per second across all connections.
+	GetMetrics()->CreateInt(UE::Net::Metric::InRate, 0);
+
+	// The total number of outgoing bytes per second across all connections.
+	GetMetrics()->CreateInt(UE::Net::Metric::OutRate, 0);
+
+	// (Server only) The average/min/max incoming bytes across all connections.
+	GetMetrics()->CreateInt(UE::Net::Metric::InRateClientAvg, 0);
+	GetMetrics()->CreateInt(UE::Net::Metric::InRateClientMax, 0);
+	GetMetrics()->CreateInt(UE::Net::Metric::InRateClientMin, 0);
+
+	// (Server only) The average/min/max outgoing bytes per second across all connections.
+	GetMetrics()->CreateInt(UE::Net::Metric::OutRateClientAvg, 0);
+	GetMetrics()->CreateInt(UE::Net::Metric::OutRateClientMax, 0);
+	GetMetrics()->CreateInt(UE::Net::Metric::OutRateClientMin, 0);
+
+	// The percentage of incoming/outgoing packets per second that have been lost.
+	GetMetrics()->CreateInt(UE::Net::Metric::InPacketsLost, 0);
+	GetMetrics()->CreateInt(UE::Net::Metric::OutPacketsLost, 0);
+
+	// The number of incoming bunches per second across all connections.
+	GetMetrics()->CreateInt(UE::Net::Metric::InBunches, 0);
+
+	// The number of outgoing bunches per second across all connections.
+	GetMetrics()->CreateInt(UE::Net::Metric::OutBunches, 0);
+
+	// (Client only) The current ping in milliseconds to the server.
+	GetMetrics()->CreateInt(UE::Net::Metric::Ping, 0);
+
+	// The average/min/max ping across all connections.
+	GetMetrics()->CreateFloat(UE::Net::Metric::AvgPing, 0.0f);
+	GetMetrics()->CreateFloat(UE::Net::Metric::MinPing, 0.0f);
+	GetMetrics()->CreateFloat(UE::Net::Metric::MaxPing, 0.0f);
+
+	// (Non-Iris only) The number of kilobytes sent in the last frame (only if GReplicateActorTimingEnabled is true).
+	GetMetrics()->CreateFloat(UE::Net::Metric::OutKBytes, 0);
+
+	// (Non-Iris only) The number of outgoing kilobytes that made up NetGUID bunches in the last frame (only if GReplicateActorTimingEnabled is true).
+	GetMetrics()->CreateFloat(UE::Net::Metric::OutNetGUIDKBytesSec, 0);
+	
+	// (Non-Iris only) The time, in milliseconds, spent by the server performing replication in the last frame. This value represents the sum of ReplicateActorTimeMS 
+	// and GatherPrioritizeTimeMS (only if GReplicateActorTimingEnabled is true).
+	GetMetrics()->CreateFloat(UE::Net::Metric::ServerReplicateActorTimeMS, 0.0f);
+	// (Non-Iris only) The time, in milliseconds, spent preparing to replicate actors during server replication in the last frame (only if GReplicateActorTimingEnabled is true).
+	GetMetrics()->CreateFloat(UE::Net::Metric::GatherPrioritizeTimeMS, 0.0f);
+	// (Non-Iris only) The time, in milliseconds, spent replicating actors during server replication in the last frame (only if GReplicateActorTimingEnabled is true).
+	GetMetrics()->CreateFloat(UE::Net::Metric::ReplicateActorTimeMS, 0.0f);
+
+	// (Non-Iris only) (Server only) The number of actors replicated in the last frame (only if GReplicateActorTimingEnabled is true).
+	GetMetrics()->CreateInt(UE::Net::Metric::NumReplicatedActors, 0);
+
+	// The number of active connections in the last frame.
+	GetMetrics()->CreateInt(UE::Net::Metric::Connections, 0);
+
+	// (Non-Iris only) The average number of replicated actors per connection in the last frame (only if GReplicateActorTimingEnabled is true).
+	GetMetrics()->CreateFloat(UE::Net::Metric::NumReplicateActorCallsPerConAvg, 0.0f);
+
+	// (Non-Iris only) The number of networked actors in the world during the last frame (only if GReplicateActorTimingEnabled is true).
+	GetMetrics()->CreateInt(UE::Net::Metric::NumberOfActiveActors, 0);
+
+	// (Non-Iris only) (Server only) The number of dormant actors in the last frame (only if GReplicateActorTimingEnabled is true).
+	GetMetrics()->CreateInt(UE::Net::Metric::NumberOfFullyDormantActors, 0);
+
+	// (Non-Iris only) (Server only) The number of network objects that skipped replication across all connections in the last frame (only if GReplicateActorTimingEnabled is true).
+	GetMetrics()->CreateInt(UE::Net::Metric::NumSkippedObjectEmptyUpdates, 0);
+
+	// (Non-Iris only) (Server only) The number of open channels across all connections in the last frame (only if GReplicateActorTimingEnabled is true).
+	GetMetrics()->CreateInt(UE::Net::Metric::NumOpenChannels, 0);
+
+	// (Non-Iris only) (Server only) The number of open channels that have dormant actors or queued bunches across all connections in the last frame (only if GReplicateActorTimingEnabled is true).
+	GetMetrics()->CreateInt(UE::Net::Metric::NumTickingChannels, 0);
+
+	// (Non-Iris only) (Server only) The number of connections that are skipped due to bandwidth saturation in the last frame (only if GReplicateActorTimingEnabled is true).
+	GetMetrics()->CreateInt(UE::Net::Metric::SatConnections, 0);
+
+	// (Non-Iris only) (Server only) The number of property updates that used shared serialization in the last frame (only if GReplicateActorTimingEnabled is true).
+	GetMetrics()->CreateInt(UE::Net::Metric::SharedSerializationPropertyHit, 0);
+
+	// (Non-Iris only) (Server only) The number of property updates that couldn't used shared serialization in the last frame (only if GReplicateActorTimingEnabled is true).
+	GetMetrics()->CreateInt(UE::Net::Metric::SharedSerializationPropertyMiss, 0);
+
+	// (Non-Iris only) (Server only) The number of RPCs that used shared serialization in the last frame (only if GReplicateActorTimingEnabled is true).
+	GetMetrics()->CreateInt(UE::Net::Metric::SharedSerializationRPCHit, 0);
+
+	// (Non-Iris only) (Server only) The number of RPCs that couldn't use shared serialization in the last frame (only if GReplicateActorTimingEnabled is true).
+	GetMetrics()->CreateInt(UE::Net::Metric::SharedSerializationRPCMiss, 0);
+
+	// (Non-Iris only) The number of received level visibility requests (for level changes) in the last frame (only if GReplicateActorTimingEnabled is true).
+	GetMetrics()->CreateFloat(UE::Net::Metric::NumClientUpdateLevelVisibility, 0);
+
+	// The number of open channels.
+	GetMetrics()->CreateInt(UE::Net::Metric::Channels, 0);
+
+	// The number of actor channels.
+	GetMetrics()->CreateInt(UE::Net::Metric::NumActorChannels, 0);
+
+	// The number of dormant actors.
+	GetMetrics()->CreateInt(UE::Net::Metric::NumDormantActors, 0);
+
+	// The number of actors in the world.
+	GetMetrics()->CreateInt(UE::Net::Metric::NumActors, 0);
+
+	// The number of networked actors in the world.
+	GetMetrics()->CreateInt(UE::Net::Metric::NumNetActors, 0);
+
+	// The number of NetGUID that have been acked, are unacked or are pending ack.
+	GetMetrics()->CreateInt(UE::Net::Metric::NumNetGUIDsAckd, 0);
+	GetMetrics()->CreateInt(UE::Net::Metric::NumNetGUIDsPending, 0);
+	GetMetrics()->CreateInt(UE::Net::Metric::NumNetGUIDsUnAckd, 0);
+
+	// 1 if one or more connections are unable to send packets due to bandwidth saturation.
+	// 0 if all connections are ready to send packets.
+	GetMetrics()->CreateInt(UE::Net::Metric::NetSaturated, 0);
+
+	// The additional overhead, in bytes, associated with each packet (e.g. a UDP header).
+	GetMetrics()->CreateInt(UE::Net::Metric::MaxPacketOverhead, 0);
+
+	// The number of incoming/outgoing bytes per second used for NetGUID bunches.
+	GetMetrics()->CreateInt(UE::Net::Metric::NetGUIDInRate, 0);
+	GetMetrics()->CreateInt(UE::Net::Metric::NetGUIDOutRate, 0);
+
+	// (Server only) The number of active client connections (can be smaller than NumConnections).
+	GetMetrics()->CreateInt(UE::Net::Metric::NetNumClients, 0);
+
+	// (Server only) The number of active client connections (can be smaller than NumConnections).
+	GetMetrics()->CreateInt(UE::Net::Metric::NumClients, 0);
+
+	// The number of connections.
+	GetMetrics()->CreateInt(UE::Net::Metric::NumConnections, 0);
+
+	// (Server only) The number of actors considered for replication across all connections.
+	GetMetrics()->CreateInt(UE::Net::Metric::NumConsideredActors, 0);
+
+	// (Server only) The number of actors that are configured to be initially dormant (DORM_Initial).
+	GetMetrics()->CreateInt(UE::Net::Metric::NumInitiallyDormantActors, 0);
+
+	// (Server only) The number of prioritized considered actors across all connections.
+	GetMetrics()->CreateInt(UE::Net::Metric::PrioritizedActors, 0);
+
+	// (Server only) The number of prioritized actors that are deleted actors across all connections.
+	GetMetrics()->CreateInt(UE::Net::Metric::NumRelevantDeletedActors, 0);
+
+	// (Non-Iris only) (Server only) The number of bytes sent when replicating actors across all connections.
+	GetMetrics()->CreateInt(UE::Net::Metric::NumReplicatedActorBytes, 0);
+
+	// The number of game objects that aren't currently mapped to a replicator across all connections.
+	GetMetrics()->CreateInt(UE::Net::Metric::UnmappedReplicators, 0);
+
+	// The percentage of incoming/outgoing bytes per second that are used for voice communications across all connections.
+	GetMetrics()->CreateInt(UE::Net::Metric::PercentInVoice, 0);
+	GetMetrics()->CreateInt(UE::Net::Metric::PercentOutVoice, 0);
+
+	// The number of incoming/outgoing bytes/packets per second used by voice communications across all connections.
+	GetMetrics()->CreateInt(UE::Net::Metric::VoiceBytesRecv, 0);
+	GetMetrics()->CreateInt(UE::Net::Metric::VoiceBytesSent, 0);
+	GetMetrics()->CreateInt(UE::Net::Metric::VoicePacketsRecv, 0);
+	GetMetrics()->CreateInt(UE::Net::Metric::VoicePacketsSent, 0);
+
+	// (Server only) A histogram of ping across all connections.
+	GetMetrics()->CreateInt(UE::Net::Metric::PingBucketInt0, 0);
+	GetMetrics()->CreateInt(UE::Net::Metric::PingBucketInt1, 0);
+	GetMetrics()->CreateInt(UE::Net::Metric::PingBucketInt2, 0);
+	GetMetrics()->CreateInt(UE::Net::Metric::PingBucketInt3, 0);
+	GetMetrics()->CreateInt(UE::Net::Metric::PingBucketInt4, 0);
+	GetMetrics()->CreateInt(UE::Net::Metric::PingBucketInt5, 0);
+	GetMetrics()->CreateInt(UE::Net::Metric::PingBucketInt6, 0);
+	GetMetrics()->CreateInt(UE::Net::Metric::PingBucketInt7, 0);
+
+	GetMetrics()->CreateInt(UE::Net::Metric::ImportedNetGuids, 0);
+	GetMetrics()->CreateInt(UE::Net::Metric::PendingOuterNetGuids, 0);
+	GetMetrics()->CreateInt(UE::Net::Metric::NetInBunchTimeOvershootPercent, 0);
+}
+
+void UNetDriver::SetupNetworkMetricsListeners()
+{
+	// Register metrics in the ini configuration file with a listener.
+	const UNetworkMetricsConfig* Config = GetDefault<UNetworkMetricsConfig>();
+
+	for (const FNetworkMetricConfig& MetricConfig : Config->Listeners)
+	{
+		FName ListenerClassName = *MetricConfig.Class.ToString();
+		TObjectPtr<UNetworkMetricsBaseListener>* Listener = NetworkMetricsListeners.Find(ListenerClassName);
+
+		// Create the listener if it's not in the cache.
+		if (Listener == nullptr)
+		{
+			const UClass* ListenerClass = MetricConfig.Class.Get();
+			if (ListenerClass == nullptr)
+			{
+				UE_LOG(LogNet, Warning, TEXT("NetDriver::SetupNetworkMetricsListeners: Unknown metric listener %s provided in config for metric %s."), *MetricConfig.Class.ToString(), *MetricConfig.MetricName.ToString());
+				continue;
+			}
+
+			if (!GetMetrics()->Contains(MetricConfig.MetricName))
+			{
+				UE_LOG(LogNet, Warning, TEXT("NetDriver::SetupNetworkMetricListeners: Cannot register metric listener %s to an unknown metric '%s'."), *MetricConfig.Class.ToString(), *MetricConfig.MetricName.ToString());
+				continue;
+			}
+
+			Listener = &NetworkMetricsListeners.Add(ListenerClassName, NewObject<UNetworkMetricsBaseListener>(this, ListenerClass));
+		}
+
+		GetMetrics()->Register(MetricConfig.MetricName, Listener->Get());
+		UE_LOG(LogNet, Log, TEXT("Registering network metrics listener %s for metric %s"), *ListenerClassName.ToString(), *MetricConfig.MetricName.ToString());
+	}
+
+	// The listeners for the Stats system are hardcoded here because Stats is always expected to be running and some/all of the metrics should
+	// not be disabled through a configuration file.
+#if STATS
+	RegisterStatsListener(UE::Net::Metric::ImportedNetGuids, GET_STATFNAME(STAT_ImportedNetGuids));
+	RegisterStatsListener(UE::Net::Metric::PendingOuterNetGuids, GET_STATFNAME(STAT_PendingOuterNetGuids));
+	RegisterStatsListener(UE::Net::Metric::UnmappedReplicators, GET_STATFNAME(STAT_UnmappedReplicators));
+	RegisterStatsListener(UE::Net::Metric::NumInitiallyDormantActors, GET_STATFNAME(STAT_NumInitiallyDormantActors));
+	RegisterStatsListener(UE::Net::Metric::NumConsideredActors, GET_STATFNAME(STAT_NumConsideredActors));
+	RegisterStatsListener(UE::Net::Metric::PrioritizedActors, GET_STATFNAME(STAT_PrioritizedActors));
+	RegisterStatsListener(UE::Net::Metric::NumRelevantDeletedActors, GET_STATFNAME(STAT_NumRelevantDeletedActors));
+	RegisterStatsListener(UE::Net::Metric::SharedSerializationRPCHit, GET_STATFNAME(STAT_SharedSerializationRPCHit));
+	RegisterStatsListener(UE::Net::Metric::SharedSerializationRPCMiss, GET_STATFNAME(STAT_SharedSerializationRPCMiss));
+	RegisterStatsListener(UE::Net::Metric::NumReplicatedActors, GET_STATFNAME(STAT_NumReplicatedActors));
+	RegisterStatsListener(UE::Net::Metric::SatConnections, GET_STATFNAME(STAT_NumSaturatedConnections));
+	RegisterStatsListener(UE::Net::Metric::NumSkippedObjectEmptyUpdates, GET_STATFNAME(STAT_NumSkippedObjectEmptyUpdates));
+	RegisterStatsListener(UE::Net::Metric::SharedSerializationPropertyHit, GET_STATFNAME(STAT_SharedSerializationPropertyHit));
+	RegisterStatsListener(UE::Net::Metric::SharedSerializationPropertyMiss, GET_STATFNAME(STAT_SharedSerializationPropertyMiss));
+	RegisterStatsListener(UE::Net::Metric::NumReplicatedActorBytes, GET_STATFNAME(STAT_NumReplicatedActorBytes));
+	RegisterStatsListener(UE::Net::Metric::Ping, GET_STATFNAME(STAT_Ping));
+	RegisterStatsListener(UE::Net::Metric::Channels, GET_STATFNAME(STAT_Channels));
+	RegisterStatsListener(UE::Net::Metric::MaxPacketOverhead, GET_STATFNAME(STAT_MaxPacketOverhead));
+	RegisterStatsListener(UE::Net::Metric::OutPacketsLost, GET_STATFNAME(STAT_OutLoss));
+	RegisterStatsListener(UE::Net::Metric::InPacketsLost, GET_STATFNAME(STAT_InLoss));
+	RegisterStatsListener(UE::Net::Metric::InRate, GET_STATFNAME(STAT_InRate));
+	RegisterStatsListener(UE::Net::Metric::OutRate, GET_STATFNAME(STAT_OutRate));
+	RegisterStatsListener(UE::Net::Metric::InRateClientMax, GET_STATFNAME(STAT_InRateClientMax));
+	RegisterStatsListener(UE::Net::Metric::InRateClientMin, GET_STATFNAME(STAT_InRateClientMin));
+	RegisterStatsListener(UE::Net::Metric::InRateClientAvg, GET_STATFNAME(STAT_InRateClientAvg));
+	RegisterStatsListener(UE::Net::Metric::InPacketsClientPerSecondMax, GET_STATFNAME(STAT_InPacketsClientMax));
+	RegisterStatsListener(UE::Net::Metric::InPacketsClientPerSecondMin, GET_STATFNAME(STAT_InPacketsClientMin));
+	RegisterStatsListener(UE::Net::Metric::InPacketsClientPerSecondAvg, GET_STATFNAME(STAT_InPacketsClientAvg));
+	RegisterStatsListener(UE::Net::Metric::OutRateClientMax, GET_STATFNAME(STAT_OutRateClientMax));
+	RegisterStatsListener(UE::Net::Metric::OutRateClientMin, GET_STATFNAME(STAT_OutRateClientMin));
+	RegisterStatsListener(UE::Net::Metric::OutRateClientAvg, GET_STATFNAME(STAT_OutRateClientAvg));
+	RegisterStatsListener(UE::Net::Metric::OutPacketsClientPerSecondMax, GET_STATFNAME(STAT_OutPacketsClientMax));
+	RegisterStatsListener(UE::Net::Metric::OutPacketsClientPerSecondMin, GET_STATFNAME(STAT_OutPacketsClientMin));
+	RegisterStatsListener(UE::Net::Metric::OutPacketsClientPerSecondAvg, GET_STATFNAME(STAT_OutPacketsClientAvg));
+	RegisterStatsListener(UE::Net::Metric::NetNumClients, GET_STATFNAME(STAT_NetNumClients));
+	RegisterStatsListener(UE::Net::Metric::InPackets, GET_STATFNAME(STAT_InPackets));
+	RegisterStatsListener(UE::Net::Metric::OutPackets, GET_STATFNAME(STAT_OutPackets));
+	RegisterStatsListener(UE::Net::Metric::InBunches, GET_STATFNAME(STAT_InBunches));
+	RegisterStatsListener(UE::Net::Metric::OutBunches, GET_STATFNAME(STAT_OutBunches));
+	RegisterStatsListener(UE::Net::Metric::NetGUIDInRate, GET_STATFNAME(STAT_NetGUIDInRate));
+	RegisterStatsListener(UE::Net::Metric::NetGUIDOutRate, GET_STATFNAME(STAT_NetGUIDOutRate));
+	RegisterStatsListener(UE::Net::Metric::VoicePacketsSent, GET_STATFNAME(STAT_VoicePacketsSent));
+	RegisterStatsListener(UE::Net::Metric::VoicePacketsRecv, GET_STATFNAME(STAT_VoicePacketsRecv));
+	RegisterStatsListener(UE::Net::Metric::VoiceBytesSent, GET_STATFNAME(STAT_VoiceBytesSent));
+	RegisterStatsListener(UE::Net::Metric::VoiceBytesRecv, GET_STATFNAME(STAT_VoiceBytesRecv));
+	RegisterStatsListener(UE::Net::Metric::PercentInVoice, GET_STATFNAME(STAT_PercentInVoice));
+	RegisterStatsListener(UE::Net::Metric::PercentOutVoice, GET_STATFNAME(STAT_PercentOutVoice));
+	RegisterStatsListener(UE::Net::Metric::NumActorChannels, GET_STATFNAME(STAT_NumActorChannels));
+	RegisterStatsListener(UE::Net::Metric::NumDormantActors, GET_STATFNAME(STAT_NumDormantActors));
+	RegisterStatsListener(UE::Net::Metric::NumActors, GET_STATFNAME(STAT_NumActors));
+	RegisterStatsListener(UE::Net::Metric::NumNetActors, GET_STATFNAME(STAT_NumNetActors));
+	RegisterStatsListener(UE::Net::Metric::NumNetGUIDsAckd, GET_STATFNAME(STAT_NumNetGUIDsAckd));
+	RegisterStatsListener(UE::Net::Metric::NumNetGUIDsPending, GET_STATFNAME(STAT_NumNetGUIDsPending));
+	RegisterStatsListener(UE::Net::Metric::NumNetGUIDsUnAckd, GET_STATFNAME(STAT_NumNetGUIDsUnAckd));
+	RegisterStatsListener(UE::Net::Metric::NetSaturated, GET_STATFNAME(STAT_NetSaturated));
+	RegisterStatsListener(UE::Net::Metric::NetInBunchTimeOvershootPercent, GET_STATFNAME(STAT_NetInBunchTimeOvershootPercent));
+	RegisterStatsListener(UE::Net::Metric::GatherPrioritizeTimeMS, GET_STATFNAME(STAT_NetServerGatherPrioritizeRepActorsTime));
+#endif
+}
+
 FString UNetDriver::GetReplicationModelName() const
 {
 	if (UReplicationDriver* RepDriver = GetReplicationDriver())
@@ -1671,6 +1962,9 @@ static FAutoConsoleVariableRef CVarLogPendingGuidsOnShutdown(
 /** Shutdown all connections managed by this net driver */
 void UNetDriver::Shutdown()
 {
+	NetworkMetricsDatabase->Reset();
+	NetworkMetricsListeners.Reset();
+
 	// Client closing connection to server
 	if (ServerConnection)
 	{
@@ -2292,7 +2586,7 @@ void UNetDriver::ProcessRemoteFunctionForChannelPrivate(
 			Connection->FlushNet(true);
 			Connection->Close(ENetCloseResult::RPCReliableBufferOverflow);
 #if USE_SERVER_PERF_COUNTERS
-			PerfCountersIncrement(TEXT("ClosedConnectionsDueToReliableBufferOverflow"));
+			GetMetrics()->IncrementInt(UE::Net::Metric::ClosedConnectionsDueToReliableBufferOverflow, 1);
 #endif
 		}
 		return;
@@ -4546,8 +4840,8 @@ void UNetDriver::ServerReplicateActors_BuildConsiderList( TArray<FNetworkObjectI
 	}
 
 	// Update stats
-	SET_DWORD_STAT( STAT_NumInitiallyDormantActors, NumInitiallyDormant );
-	SET_DWORD_STAT( STAT_NumConsideredActors, OutConsiderList.Num() );
+	GetMetrics()->SetInt(UE::Net::Metric::NumInitiallyDormantActors, NumInitiallyDormant);
+	GetMetrics()->SetInt(UE::Net::Metric::NumConsideredActors, OutConsiderList.Num());
 }
 
 // Returns true if this actor should replicate to *any* of the passed in connections
@@ -4758,8 +5052,8 @@ int32 UNetDriver::ServerReplicateActors_PrioritizeActors( UNetConnection* Connec
 	UE_LOG( LogNetTraffic, Log, TEXT( "ServerReplicateActors_PrioritizeActors: Potential %04i ConsiderList %03i FinalSortedCount %03i" ), MaxSortedActors, ConsiderList.Num(), FinalSortedCount );
 
 	// Setup stats
-	SET_DWORD_STAT( STAT_PrioritizedActors, FinalSortedCount );
-	SET_DWORD_STAT( STAT_NumRelevantDeletedActors, DeletedCount );
+	GetMetrics()->SetInt(UE::Net::Metric::PrioritizedActors, FinalSortedCount);
+	GetMetrics()->SetInt(UE::Net::Metric::NumRelevantDeletedActors, DeletedCount);
 
 	return FinalSortedCount;
 }
@@ -5131,15 +5425,6 @@ double GServerReplicateActorTimeSeconds;
 int32 GNumClientConnections;
 int32 GNumClientUpdateLevelVisibility;
 
-DECLARE_DWORD_COUNTER_STAT(TEXT("Num Saturated Connections"), STAT_NumSaturatedConnections, STATGROUP_Net);
-
-DECLARE_DWORD_COUNTER_STAT(TEXT("SharedSerialization RPC Hit"), STAT_SharedSerializationRPCHit, STATGROUP_Net);
-DECLARE_DWORD_COUNTER_STAT(TEXT("SharedSerialization RPC Miss"), STAT_SharedSerializationRPCMiss, STATGROUP_Net);
-DECLARE_DWORD_COUNTER_STAT(TEXT("Empty Object Replicate Properties Skipped"), STAT_NumSkippedObjectEmptyUpdates, STATGROUP_Net);
-
-DECLARE_DWORD_COUNTER_STAT(TEXT("SharedSerialization Property Hit"), STAT_SharedSerializationPropertyHit, STATGROUP_Net);
-DECLARE_DWORD_COUNTER_STAT(TEXT("SharedSerialization Property Miss"), STAT_SharedSerializationPropertyMiss, STATGROUP_Net);
-
 struct FReplicationAutoCapture
 {
 	int32 CaptureFrames=-1;
@@ -5190,7 +5475,7 @@ FReplicationAutoCapture GReplicationAutoCapture;
 #if CSV_PROFILER
 struct FScopedNetDriverStats
 {
-	FScopedNetDriverStats(uint32& OutBytesRef, UNetDriver* InNetDriver) : OutBytes(OutBytesRef), NetDriver(InNetDriver)
+	FScopedNetDriverStats(UNetDriver* InNetDriver) : NetDriver(InNetDriver)
 	{
 		GReplicationAutoCapture.DoFrame();
 
@@ -5207,15 +5492,12 @@ struct FScopedNetDriverStats
 			GServerReplicateActorTimeSeconds = 0.f;
 
 			// Whatever these values currently are were (mostly) set by RPCs (technically something else could have force ReplicateActor to be called but this is rare).
-			SET_DWORD_STAT(STAT_SharedSerializationRPCHit, GNumSharedSerializationHit);
-			SET_DWORD_STAT(STAT_SharedSerializationRPCMiss, GNumSharedSerializationMiss);
-
-			CSV_CUSTOM_STAT(Replication, SharedSerializationRPCHit, GNumSharedSerializationHit, ECsvCustomStatOp::Set);
-			CSV_CUSTOM_STAT(Replication, SharedSerializationRPCMiss, GNumSharedSerializationMiss, ECsvCustomStatOp::Set);
+			InNetDriver->GetMetrics()->SetInt(UE::Net::Metric::SharedSerializationRPCHit, GNumSharedSerializationHit);
+			InNetDriver->GetMetrics()->SetInt(UE::Net::Metric::SharedSerializationRPCMiss, GNumSharedSerializationMiss);
 
 			const FNetworkObjectList& NetworkObjectList = NetDriver->GetNetworkObjectList();
-			CSV_CUSTOM_STAT(Replication, NumberOfActiveActors, NetworkObjectList.GetActiveObjects().Num(), ECsvCustomStatOp::Set);
-			CSV_CUSTOM_STAT(Replication, NumberOfFullyDormantActors, NetworkObjectList.GetDormantObjectsOnAllConnections().Num(), ECsvCustomStatOp::Set);
+			InNetDriver->GetMetrics()->SetInt(UE::Net::Metric::NumberOfActiveActors, NetworkObjectList.GetActiveObjects().Num());
+			InNetDriver->GetMetrics()->SetInt(UE::Net::Metric::NumberOfFullyDormantActors, NetworkObjectList.GetDormantObjectsOnAllConnections().Num());
 
 			StartTime = FPlatformTime::Seconds();
 			StartOutBytes = GNetOutBytes;
@@ -5233,27 +5515,18 @@ struct FScopedNetDriverStats
 
 			uint32 FrameOutBytes = GNetOutBytes - StartOutBytes;
 
-			SET_FLOAT_STAT(STAT_NetServerGatherPrioritizeRepActorsTime, GReplicationGatherPrioritizeTimeSeconds * 1000.0);
-			SET_DWORD_STAT(STAT_NumReplicatedActors, GNumReplicateActorCalls);
-			SET_DWORD_STAT(STAT_NumSaturatedConnections, GNumSaturatedConnections);
-			SET_DWORD_STAT(STAT_NumSkippedObjectEmptyUpdates, GNumSkippedObjectEmptyUpdates);
-
-			CSV_CUSTOM_STAT(Replication, ServerReplicateActorTimeMS, (float)(GServerReplicateActorTimeSeconds * 1000.0), ECsvCustomStatOp::Set );
-			CSV_CUSTOM_STAT(Replication, GatherPrioritizeTimeMS, (float)(GReplicationGatherPrioritizeTimeSeconds * 1000.0), ECsvCustomStatOp::Set );
-			CSV_CUSTOM_STAT(Replication, ReplicateActorTimeMS, (float)(GReplicateActorTimeSeconds * 1000.0), ECsvCustomStatOp::Set );
-			CSV_CUSTOM_STAT(Replication, NumReplicateActorCallsPerConAvg, ((float)GNumReplicateActorCalls)/(float)GNumClientConnections, ECsvCustomStatOp::Set );
-			CSV_CUSTOM_STAT(Replication, Connections, (float)GNumClientConnections, ECsvCustomStatOp::Set );
-			CSV_CUSTOM_STAT(Replication, SatConnections, (float)GNumSaturatedConnections, ECsvCustomStatOp::Set );
-			CSV_CUSTOM_STAT(Replication, OutKBytes, ((float)FrameOutBytes) / 1024.f, ECsvCustomStatOp::Set );
-			CSV_CUSTOM_STAT(Replication, OutNetGUIDKBytesSec, ((float)NetDriver->NetGUIDOutBytes / 1024.f), ECsvCustomStatOp::Set );
-			CSV_CUSTOM_STAT(Replication, NumClientUpdateLevelVisibility, ((float)GNumClientUpdateLevelVisibility), ECsvCustomStatOp::Set );
-			CSV_CUSTOM_STAT(Replication, NumSkippedObjectEmptyUpdates, (float)GNumSkippedObjectEmptyUpdates, ECsvCustomStatOp::Set);
-
-			SET_DWORD_STAT(STAT_SharedSerializationPropertyHit, GNumSharedSerializationHit);
-			SET_DWORD_STAT(STAT_SharedSerializationPropertyMiss, GNumSharedSerializationMiss);
-
-			CSV_CUSTOM_STAT(Replication, SharedSerializationPropertyHit, GNumSharedSerializationHit, ECsvCustomStatOp::Set);
-			CSV_CUSTOM_STAT(Replication, SharedSerializationPropertyMiss, GNumSharedSerializationMiss, ECsvCustomStatOp::Set);
+			NetDriver->GetMetrics()->SetFloat(UE::Net::Metric::ServerReplicateActorTimeMS, static_cast<float>(GServerReplicateActorTimeSeconds * 1000.0));
+			NetDriver->GetMetrics()->SetInt(UE::Net::Metric::NumReplicatedActors, GNumReplicateActorCalls);
+			NetDriver->GetMetrics()->SetInt(UE::Net::Metric::SatConnections, GNumSaturatedConnections);
+			NetDriver->GetMetrics()->SetInt(UE::Net::Metric::NumSkippedObjectEmptyUpdates, GNumSkippedObjectEmptyUpdates);
+			NetDriver->GetMetrics()->SetFloat(UE::Net::Metric::GatherPrioritizeTimeMS, static_cast<float>(GReplicationGatherPrioritizeTimeSeconds * 1000.0));
+			NetDriver->GetMetrics()->SetFloat(UE::Net::Metric::ReplicateActorTimeMS, static_cast<float>(GReplicateActorTimeSeconds * 1000.0));
+			NetDriver->GetMetrics()->SetFloat(UE::Net::Metric::NumReplicateActorCallsPerConAvg, (static_cast<float>(GNumReplicateActorCalls) / static_cast<float>(GNumClientConnections)));
+			NetDriver->GetMetrics()->SetFloat(UE::Net::Metric::OutKBytes, static_cast<float>(FrameOutBytes) / 1024.f);
+			NetDriver->GetMetrics()->SetFloat(UE::Net::Metric::OutNetGUIDKBytesSec, static_cast<float>(NetDriver->NetGUIDOutBytes) / 1024.f);
+			NetDriver->GetMetrics()->SetFloat(UE::Net::Metric::NumClientUpdateLevelVisibility, static_cast<float>(GNumClientUpdateLevelVisibility));
+			NetDriver->GetMetrics()->SetInt(UE::Net::Metric::SharedSerializationPropertyHit, GNumSharedSerializationHit);
+			NetDriver->GetMetrics()->SetInt(UE::Net::Metric::SharedSerializationPropertyMiss, GNumSharedSerializationMiss);
 
 			int32 NumOpenChannels = 0;
 			int32 NumTickingChannels = 0;
@@ -5267,8 +5540,8 @@ struct FScopedNetDriverStats
 				MaxOutPackets = FMath::Max(MaxOutPackets, ClientConnection->OutPackets);
 			}
 
-			CSV_CUSTOM_STAT(Replication, NumOpenChannels, (float)NumOpenChannels, ECsvCustomStatOp::Set);
-			CSV_CUSTOM_STAT(Replication, NumTickingChannels, (float)NumTickingChannels, ECsvCustomStatOp::Set);
+			NetDriver->GetMetrics()->SetInt(UE::Net::Metric::NumOpenChannels, NumOpenChannels);
+			NetDriver->GetMetrics()->SetInt(UE::Net::Metric::NumTickingChannels, NumTickingChannels);
 
 			// Note: we want to reset this at the end of the frame since the RPC stats are incremented at the top (recv)
 			GNumSharedSerializationHit = 0;
@@ -5277,7 +5550,6 @@ struct FScopedNetDriverStats
 		}
 	}
 	
-	uint32& OutBytes;
 	double StartTime;
 	uint32 StartOutBytes;
 	UNetDriver* NetDriver;
@@ -5299,11 +5571,11 @@ int32 UNetDriver::ServerReplicateActors(float DeltaSeconds)
 		return 0;
 	}
 
-	SET_DWORD_STAT(STAT_NumReplicatedActors, 0);
-	SET_DWORD_STAT(STAT_NumReplicatedActorBytes, 0);
+	GetMetrics()->SetInt(UE::Net::Metric::NumReplicatedActors,0 );
+	GetMetrics()->SetInt(UE::Net::Metric::NumReplicatedActorBytes, 0);
 
 #if CSV_PROFILER
-	FScopedNetDriverStats NetDriverStats(OutBytes, this);
+	FScopedNetDriverStats NetDriverStats(this);
 	GNumClientConnections = ClientConnections.Num();
 #endif
 	
@@ -5558,6 +5830,19 @@ UChannel* UNetDriver::InternalCreateChannelByName(const FName& ChName)
 	return NewObject<UChannel>(this, ChannelDefinitionMap[ChName].ChannelClass);
 }
 
+void UNetDriver::RegisterStatsListener(const FName MetricName, const FName StatName)
+{
+	FName ListenerKeyName = FName(FString::Printf(TEXT("NetworkMetricsStats_%s"), *StatName.ToString()));
+	if (ensureMsgf(!NetworkMetricsListeners.Contains(ListenerKeyName), TEXT("Stat metrics listener %s already added to the database."), *StatName.ToString()))
+	{
+		UNetworkMetricsStats* Listener = NewObject<UNetworkMetricsStats>();
+		Listener->SetStatName(StatName);
+		Listener->SetInterval(StatPeriod);
+		NetworkMetricsListeners.Add(ListenerKeyName, Listener);
+		GetMetrics()->Register(MetricName, Listener);
+	}
+}
+
 void UNetDriver::ReleaseToChannelPool(UChannel* Channel)
 {
 	LLM_SCOPE_BYTAG(NetChannel);
@@ -5887,7 +6172,7 @@ void UNetDriver::AddClientConnection(UNetConnection* NewConnection)
 	}
 
 #if USE_SERVER_PERF_COUNTERS
-	PerfCountersIncrement(TEXT("AddedConnections"));
+	GetMetrics()->IncrementInt(UE::Net::Metric::AddedConnections, 1);
 #endif
 
 	CreateInitialServerChannels(NewConnection);
@@ -7463,6 +7748,13 @@ void UNetDriver::UpdateNetworkStats()
 				OutPacketsLost = FMath::TruncToInt(100.f * OutPacketsLost / FMath::Max((float)OutPackets, 1.f));
 				InPacketsLost = FMath::TruncToInt(100.f * InPacketsLost / FMath::Max((float)InPackets + InPacketsLost, 1.f));
 
+				GetMetrics()->SetInt(UE::Net::Metric::InPackets, InPackets);
+				GetMetrics()->SetInt(UE::Net::Metric::OutPackets, OutPackets);
+				GetMetrics()->SetInt(UE::Net::Metric::InBunches, InBunches);
+				GetMetrics()->SetInt(UE::Net::Metric::OutBunches, OutBunches);
+				GetMetrics()->SetInt(UE::Net::Metric::OutPacketsLost, OutPacketsLost);
+				GetMetrics()->SetInt(UE::Net::Metric::InPacketsLost, InPacketsLost);
+
 				if (ServerConnection != nullptr && ServerConnection->PlayerController != nullptr && ServerConnection->PlayerController->PlayerState != nullptr)
 				{
 					Ping = FMath::TruncToInt(ServerConnection->PlayerController->PlayerState->ExactPing);
@@ -7533,54 +7825,35 @@ void UNetDriver::UpdateNetworkStats()
 			if (!bSkipLocalStats)
 			{
 				// Copy the net status values over
-				SET_DWORD_STAT(STAT_Ping, Ping);
-				SET_DWORD_STAT(STAT_Channels, NumOpenChannels);
-				SET_DWORD_STAT(STAT_MaxPacketOverhead, MaxPacketOverhead);
-
-				SET_DWORD_STAT(STAT_OutLoss, OutPacketsLost);
-				SET_DWORD_STAT(STAT_InLoss, InPacketsLost);
-				SET_DWORD_STAT(STAT_InRate, InBytes);
-				SET_DWORD_STAT(STAT_OutRate, OutBytes);
-				SET_DWORD_STAT(STAT_InRateClientMax, ClientInBytesMax);
-				SET_DWORD_STAT(STAT_InRateClientMin, ClientInBytesMin);
-				SET_DWORD_STAT(STAT_InRateClientAvg, ClientInBytesAvg);
-				SET_DWORD_STAT(STAT_InPacketsClientMax, ClientInPacketsMax);
-				SET_DWORD_STAT(STAT_InPacketsClientMin, ClientInPacketsMin);
-				SET_DWORD_STAT(STAT_InPacketsClientAvg, ClientInPacketsAvg);
-				SET_DWORD_STAT(STAT_OutRateClientMax, ClientOutBytesMax);
-				SET_DWORD_STAT(STAT_OutRateClientMin, ClientOutBytesMin);
-				SET_DWORD_STAT(STAT_OutRateClientAvg, ClientOutBytesAvg);
-				SET_DWORD_STAT(STAT_OutPacketsClientMax, ClientOutPacketsMax);
-				SET_DWORD_STAT(STAT_OutPacketsClientMin, ClientOutPacketsMin);
-				SET_DWORD_STAT(STAT_OutPacketsClientAvg, ClientOutPacketsAvg);
-
-				SET_DWORD_STAT(STAT_NetNumClients, NumClients);
-				SET_DWORD_STAT(STAT_InPackets, InPackets);
-				SET_DWORD_STAT(STAT_OutPackets, OutPackets);
-				SET_DWORD_STAT(STAT_InBunches, InBunches);
-				SET_DWORD_STAT(STAT_OutBunches, OutBunches);
-
-				SET_DWORD_STAT(STAT_NetGUIDInRate, NetGUIDInBytes);
-				SET_DWORD_STAT(STAT_NetGUIDOutRate, NetGUIDOutBytes);
-
-				SET_DWORD_STAT(STAT_VoicePacketsSent, VoicePacketsSent);
-				SET_DWORD_STAT(STAT_VoicePacketsRecv, VoicePacketsRecv);
-				SET_DWORD_STAT(STAT_VoiceBytesSent, VoiceBytesSent);
-				SET_DWORD_STAT(STAT_VoiceBytesRecv, VoiceBytesRecv);
-
-				SET_DWORD_STAT(STAT_PercentInVoice, VoiceInPercent);
-				SET_DWORD_STAT(STAT_PercentOutVoice, VoiceOutPercent);
-
-				SET_DWORD_STAT(STAT_NumActorChannels, NumActorChannels);
-				SET_DWORD_STAT(STAT_NumDormantActors, NumDormantActors);
-				SET_DWORD_STAT(STAT_NumActors, NumActors);
-				SET_DWORD_STAT(STAT_NumNetActors, NumActiveNetActors);
-				SET_DWORD_STAT(STAT_NumNetGUIDsAckd, AckCount);
-				SET_DWORD_STAT(STAT_NumNetGUIDsPending, UnAckCount);
-				SET_DWORD_STAT(STAT_NumNetGUIDsUnAckd, PendingCount);
-				SET_DWORD_STAT(STAT_NetSaturated, NetSaturated);
-
-				SET_DWORD_STAT(STAT_NetInBunchTimeOvershootPercent, InBunchTimeOvershootPercent);
+				GetMetrics()->SetInt(UE::Net::Metric::Ping, Ping);
+				GetMetrics()->SetInt(UE::Net::Metric::Channels, NumOpenChannels);
+				GetMetrics()->SetInt(UE::Net::Metric::MaxPacketOverhead, MaxPacketOverhead);
+				GetMetrics()->SetInt(UE::Net::Metric::InRate, InBytesPerSecond);
+				GetMetrics()->SetInt(UE::Net::Metric::OutRate, OutBytesPerSecond);
+
+				GetMetrics()->SetInt(UE::Net::Metric::NetNumClients, NumClients);
+
+				GetMetrics()->SetInt(UE::Net::Metric::NetGUIDInRate, NetGUIDInBytes);
+				GetMetrics()->SetInt(UE::Net::Metric::NetGUIDOutRate, NetGUIDOutBytes);
+
+				GetMetrics()->SetInt(UE::Net::Metric::VoicePacketsSent, VoicePacketsSent);
+				GetMetrics()->SetInt(UE::Net::Metric::VoicePacketsRecv, VoicePacketsRecv);
+				GetMetrics()->SetInt(UE::Net::Metric::VoiceBytesSent, VoiceBytesSent);
+				GetMetrics()->SetInt(UE::Net::Metric::VoiceBytesRecv, VoiceBytesRecv);
+
+				GetMetrics()->SetInt(UE::Net::Metric::PercentInVoice, VoiceInPercent);
+				GetMetrics()->SetInt(UE::Net::Metric::PercentOutVoice, VoiceOutPercent);
+
+				GetMetrics()->SetInt(UE::Net::Metric::NumActorChannels, NumActorChannels);
+				GetMetrics()->SetInt(UE::Net::Metric::NumDormantActors, NumDormantActors);
+				GetMetrics()->SetInt(UE::Net::Metric::NumActors, NumActors);
+				GetMetrics()->SetInt(UE::Net::Metric::NumNetActors, NumActiveNetActors);
+				GetMetrics()->SetInt(UE::Net::Metric::NumNetGUIDsAckd, AckCount);
+				GetMetrics()->SetInt(UE::Net::Metric::NumNetGUIDsPending, UnAckCount);
+				GetMetrics()->SetInt(UE::Net::Metric::NumNetGUIDsUnAckd, PendingCount);
+				GetMetrics()->SetInt(UE::Net::Metric::NetSaturated, NetSaturated);
+
+				GetMetrics()->SetInt(UE::Net::Metric::NetInBunchTimeOvershootPercent, InBunchTimeOvershootPercent);
 			}
 
 			// If we are to replicate server stats out to an observer, then set those values
@@ -7592,8 +7865,8 @@ void UNetDriver::UpdateNetworkStats()
 					GMBase->ServerStatReplicator->MaxPacketOverhead = MaxPacketOverhead;
 					GMBase->ServerStatReplicator->OutLoss = OutPacketsLost;
 					GMBase->ServerStatReplicator->InLoss = InPacketsLost;
-					GMBase->ServerStatReplicator->InRate = InBytes;
-					GMBase->ServerStatReplicator->OutRate = OutBytes;
+					GMBase->ServerStatReplicator->InRate = InBytesPerSecond;
+					GMBase->ServerStatReplicator->OutRate = OutBytesPerSecond;
 					GMBase->ServerStatReplicator->InRateClientMax = ClientInBytesMax;
 					GMBase->ServerStatReplicator->InRateClientMin = ClientInBytesMin;
 					GMBase->ServerStatReplicator->InRateClientAvg = ClientInBytesAvg;
@@ -7632,13 +7905,12 @@ void UNetDriver::UpdateNetworkStats()
 #endif // STATS
 
 #if USE_SERVER_PERF_COUNTERS
-			IPerfCounters* PerfCounters = IPerfCountersModule::Get().GetPerformanceCounters();
-			if (PerfCounters)
 			{
 				SCOPE_CYCLE_COUNTER(STAT_NetTickFlushGatherStatsPerfCounters);
 
 				// Update total connections
-				PerfCounters->Set(PerfCounter_NumConnections, ClientConnections.Num());
+				GetMetrics()->SetInt(UE::Net::Metric::Connections, ClientConnections.Num());
+				GetMetrics()->SetInt(UE::Net::Metric::NumConnections, ClientConnections.Num());
 
 				const int kNumBuckets = 8;	// evenly spaced with increment of 30 ms; last bucket collects all off-scale pings as well
 				if (ClientConnections.Num() > 0)
@@ -7686,62 +7958,63 @@ void UNetDriver::UpdateNetworkStats()
 						AvgPing /= static_cast<float>(PingCount);
 					}
 
-					PerfCounters->Set(PerfCounter_AvgPing, AvgPing, IPerfCounters::Flags::Transient);
-					PerfCounters->Set(PerfCounter_MaxPing, MaxPing, IPerfCounters::Flags::Transient);
-					PerfCounters->Set(PerfCounter_MinPing, MinPing, IPerfCounters::Flags::Transient);
+					GetMetrics()->SetFloat(UE::Net::Metric::AvgPing, AvgPing);
+					GetMetrics()->SetFloat(UE::Net::Metric::MaxPing, MaxPing);
+					GetMetrics()->SetFloat(UE::Net::Metric::MinPing, MinPing);
 
-					// update buckets
-					for (int BucketIdx = 0; BucketIdx < UE_ARRAY_COUNT(Buckets); ++BucketIdx)
+					// Unrolling the loop through ping buckets to prevent dynamically creating 
+					// the metric name string from the ping bucket id using FString::Printf().
+					if (ensure(kNumBuckets == 8))
 					{
-						PerfCountersIncrement(FString::Printf(TEXT("PingBucketInt%d"), BucketIdx), Buckets[BucketIdx], 0, IPerfCounters::Flags::Transient);
+						GetMetrics()->IncrementInt(UE::Net::Metric::PingBucketInt0, Buckets[0]);
+						GetMetrics()->IncrementInt(UE::Net::Metric::PingBucketInt1, Buckets[1]);
+						GetMetrics()->IncrementInt(UE::Net::Metric::PingBucketInt2, Buckets[2]);
+						GetMetrics()->IncrementInt(UE::Net::Metric::PingBucketInt3, Buckets[3]);
+						GetMetrics()->IncrementInt(UE::Net::Metric::PingBucketInt4, Buckets[4]);
+						GetMetrics()->IncrementInt(UE::Net::Metric::PingBucketInt5, Buckets[5]);
+						GetMetrics()->IncrementInt(UE::Net::Metric::PingBucketInt6, Buckets[6]);
+						GetMetrics()->IncrementInt(UE::Net::Metric::PingBucketInt7, Buckets[7]);
 					}
 				}
 				else
 				{
-					PerfCounters->Set(PerfCounter_AvgPing, 0.0f, IPerfCounters::Flags::Transient);
-					PerfCounters->Set(PerfCounter_MaxPing, -FLT_MAX, IPerfCounters::Flags::Transient);
-					PerfCounters->Set(PerfCounter_MinPing, FLT_MAX, IPerfCounters::Flags::Transient);
-
-					for (int BucketIdx = 0; BucketIdx < kNumBuckets; ++BucketIdx)
-					{
-						PerfCounters->Set(FString::Printf(TEXT("PingBucketInt%d"), BucketIdx), 0, IPerfCounters::Flags::Transient);
-					}
+					GetMetrics()->SetFloat(UE::Net::Metric::AvgPing, 0.0f);
+					GetMetrics()->SetFloat(UE::Net::Metric::MaxPing, -FLT_MAX);
+					GetMetrics()->SetFloat(UE::Net::Metric::MinPing, FLT_MAX);
+
+					GetMetrics()->SetInt(UE::Net::Metric::PingBucketInt0, 0);
+					GetMetrics()->SetInt(UE::Net::Metric::PingBucketInt1, 0);
+					GetMetrics()->SetInt(UE::Net::Metric::PingBucketInt2, 0);
+					GetMetrics()->SetInt(UE::Net::Metric::PingBucketInt3, 0);
+					GetMetrics()->SetInt(UE::Net::Metric::PingBucketInt4, 0);
+					GetMetrics()->SetInt(UE::Net::Metric::PingBucketInt5, 0);
+					GetMetrics()->SetInt(UE::Net::Metric::PingBucketInt6, 0);
+					GetMetrics()->SetInt(UE::Net::Metric::PingBucketInt7, 0);
 				}
 
 				// set the per connection stats (these are calculated earlier).
 				// Note that NumClients may be != NumConnections. Also, if NumClients is 0, the rest of counters should be 0 as well
-				PerfCounters->Set(PerfCounter_NumClients, NumClients);
-				PerfCounters->Set(PerfCounter_MaxPacketOverhead, MaxPacketOverhead);
-				PerfCounters->Set(PerfCounter_InRateClientMax, ClientInBytesMax);
-				PerfCounters->Set(PerfCounter_InRateClientMin, ClientInBytesMin);
-				PerfCounters->Set(PerfCounter_InRateClientAvg, ClientInBytesAvg);
-				PerfCounters->Set(PerfCounter_InPacketsClientMax, ClientInPacketsMax);
-				PerfCounters->Set(PerfCounter_InPacketsClientMin, ClientInPacketsMin);
-				PerfCounters->Set(PerfCounter_InPacketsClientAvg, ClientInPacketsAvg);
-				PerfCounters->Set(PerfCounter_OutRateClientMax, ClientOutBytesMax);
-				PerfCounters->Set(PerfCounter_OutRateClientMin, ClientOutBytesMin);
-				PerfCounters->Set(PerfCounter_OutRateClientAvg, ClientOutBytesAvg);
-				PerfCounters->Set(PerfCounter_OutPacketsClientMax, ClientOutPacketsMax);
-				PerfCounters->Set(PerfCounter_OutPacketsClientMin, ClientOutPacketsMin);
-				PerfCounters->Set(PerfCounter_OutPacketsClientAvg, ClientOutPacketsAvg);
-
-				PerfCounters->Set(PerfCounter_InRate, InBytes);
-				PerfCounters->Set(PerfCounter_OutRate, OutBytes);
-				PerfCounters->Set(PerfCounter_InPacketsLost, InPacketsLost);
-				PerfCounters->Set(PerfCounter_OutPacketsLost, OutPacketsLost);
-				PerfCounters->Set(PerfCounter_InPackets, InPackets);
-				PerfCounters->Set(PerfCounter_OutPackets, OutPackets);
-				PerfCounters->Set(PerfCounter_InBunches, InBunches);
-				PerfCounters->Set(PerfCounter_OutBunches, OutBunches);
+				GetMetrics()->SetInt(UE::Net::Metric::NumClients, NumClients);
+				GetMetrics()->SetInt(UE::Net::Metric::MaxPacketOverhead, MaxPacketOverhead);
+				GetMetrics()->SetInt(UE::Net::Metric::InRateClientMax, ClientInBytesMax);
+				GetMetrics()->SetInt(UE::Net::Metric::InRateClientMin, ClientInBytesMin);
+				GetMetrics()->SetInt(UE::Net::Metric::InRateClientAvg, ClientInBytesAvg);
+				GetMetrics()->SetInt(UE::Net::Metric::InPacketsClientPerSecondMax, ClientInPacketsMax);
+				GetMetrics()->SetInt(UE::Net::Metric::InPacketsClientPerSecondMin, ClientInPacketsMin);
+				GetMetrics()->SetInt(UE::Net::Metric::InPacketsClientPerSecondAvg, ClientInPacketsAvg);
+				GetMetrics()->SetInt(UE::Net::Metric::OutRateClientMax, ClientOutBytesMax);
+				GetMetrics()->SetInt(UE::Net::Metric::OutRateClientMin, ClientOutBytesMin);
+				GetMetrics()->SetInt(UE::Net::Metric::OutRateClientAvg, ClientOutBytesAvg);
+				GetMetrics()->SetInt(UE::Net::Metric::OutPacketsClientPerSecondMax, ClientOutPacketsMax);
+				GetMetrics()->SetInt(UE::Net::Metric::OutPacketsClientPerSecondMin, ClientOutPacketsMin);
+				GetMetrics()->SetInt(UE::Net::Metric::OutPacketsClientPerSecondAvg, ClientOutPacketsAvg);
 			}
 #endif // USE_SERVER_PERF_COUNTERS
 
-#if CSV_PROFILER
-			CSV_CUSTOM_STAT(Replication, InPacketsClientAvg, ClientsInPacketsThisFrameAvg, ECsvCustomStatOp::Set);
-			CSV_CUSTOM_STAT(Replication, InPacketsClientMax, ClientsInPacketsThisFrameMax, ECsvCustomStatOp::Set);
-			CSV_CUSTOM_STAT(Replication, OutPacketsClientAvg, ClientsOutPacketsThisFrameAvg, ECsvCustomStatOp::Set);
-			CSV_CUSTOM_STAT(Replication, OutPacketsClientMax, ClientsOutPacketsThisFrameMax, ECsvCustomStatOp::Set);
-#endif
+			GetMetrics()->SetInt(UE::Net::Metric::InPacketsClientAvg, ClientsInPacketsThisFrameAvg);
+			GetMetrics()->SetInt(UE::Net::Metric::InPacketsClientMax, ClientsInPacketsThisFrameMax);
+			GetMetrics()->SetInt(UE::Net::Metric::OutPacketsClientAvg, ClientsOutPacketsThisFrameAvg);
+			GetMetrics()->SetInt(UE::Net::Metric::OutPacketsClientMax, ClientsOutPacketsThisFrameMax);
 
 			// Reset everything
 			InBytes = 0;
@@ -7798,10 +8071,10 @@ void UNetDriver::UpdateNetworkStats()
 				ClientsOutPacketsThisFrameAvg /= NumClients;
 			}
 
-			CSV_CUSTOM_STAT(Replication, InPacketsClientAvg, ClientsInPacketsThisFrameAvg, ECsvCustomStatOp::Set);
-			CSV_CUSTOM_STAT(Replication, InPacketsClientMax, ClientsInPacketsThisFrameMax, ECsvCustomStatOp::Set);
-			CSV_CUSTOM_STAT(Replication, OutPacketsClientAvg, ClientsOutPacketsThisFrameAvg, ECsvCustomStatOp::Set);
-			CSV_CUSTOM_STAT(Replication, OutPacketsClientMax, ClientsOutPacketsThisFrameMax, ECsvCustomStatOp::Set);
+			GetMetrics()->SetInt(UE::Net::Metric::InPacketsClientAvg, ClientsInPacketsThisFrameAvg);
+			GetMetrics()->SetInt(UE::Net::Metric::InPacketsClientMax, ClientsInPacketsThisFrameMax);
+			GetMetrics()->SetInt(UE::Net::Metric::OutPacketsClientAvg, ClientsOutPacketsThisFrameAvg);
+			GetMetrics()->SetInt(UE::Net::Metric::OutPacketsClientMax, ClientsOutPacketsThisFrameMax);
 #else
 			// Reset the per-frame stats here
 			for (UNetConnection* Client : ClientConnections)
diff --git a/Engine/Source/Runtime/Engine/Public/Net/NetworkMetricsConfig.h b/Engine/Source/Runtime/Engine/Public/Net/NetworkMetricsConfig.h
new file mode 100644
--- /dev/null
+++ b/Engine/Source/Runtime/Engine/Public/Net/NetworkMetricsConfig.h
@@ -0,0 +1,31 @@
+// Copyright Epic Games, Inc. All Rights Reserved.
+#pragma once
+
+#include "CoreMinimal.h"
+#include "NetworkMetricsConfig.generated.h"
+
+USTRUCT()
+struct FNetworkMetricConfig
+{
+	GENERATED_BODY()
+
+public:
+	/** The name of the metric to register the listener. */
+	UPROPERTY()
+	FName MetricName;
+
+	/** A sub-class of UNetworkMetricBaseListener. */
+	UPROPERTY()
+	TSoftClassPtr<class UNetworkMetricsBaseListener> Class;
+};
+
+UCLASS(Config=Engine)
+class UNetworkMetricsConfig : public UObject
+{
+	GENERATED_BODY()
+
+public:
+
+	UPROPERTY(Config)
+	TArray<FNetworkMetricConfig> Listeners;
+};
diff --git a/Engine/Source/Runtime/Engine/Public/Net/NetworkMetricsDatabase.h b/Engine/Source/Runtime/Engine/Public/Net/NetworkMetricsDatabase.h
new file mode 100644
--- /dev/null
+++ b/Engine/Source/Runtime/Engine/Public/Net/NetworkMetricsDatabase.h
@@ -0,0 +1,256 @@
+// Copyright Epic Games, Inc. All Rights Reserved.
+
+#pragma once
+
+#include "CoreMinimal.h"
+#include "NetworkMetricsDatabase.generated.h"
+
+namespace UE::Net
+{
+
+template<class MetricType>
+struct FNetworkMetric
+{
+	FName Name;
+	MetricType Value;
+};
+
+struct FNetworkMetricSnapshot
+{
+	TArray<FNetworkMetric<int64>> MetricInts;
+	TArray<FNetworkMetric<float>> MetricFloats;
+
+	void Reset()
+	{
+		MetricInts.Reset();
+		MetricFloats.Reset();
+	}
+};
+
+} // namespace UE::Net
+
+UCLASS()
+class ENGINE_API UNetworkMetricsDatabase : public UObject
+{
+	GENERATED_BODY()
+
+public:
+	/* Add a floating point metric. */
+	void CreateFloat(const FName MetricName, float DefaultValue);
+	/* Add an integer metric. */
+	void CreateInt(const FName MetricName, int64 DefaultValue);
+	/* Set the value of an existing floating point metric. */
+	void SetFloat(const FName MetricName, float Value);
+	/* Set the value of an existing integer metric. */
+	void SetInt(const FName MetricName, int64 Value);
+	/* Increment the value of an existing integer metric. */
+	void IncrementInt(const FName MetricName, int64 Value);
+	/* Returns true if a metric has been created in the database. */
+	bool Contains(const FName MetricName) const;
+
+	/* Call all registered listeners.*/
+	void ProcessListeners();
+	/* Remove all registered metrics and listeners. */
+	void Reset();
+
+	/* Register a listener to be called for a given metric. */
+	void Register(const FName MetricName, TWeakObjectPtr<UNetworkMetricsBaseListener> Reporter);
+
+private:
+	/* Return true if the report interval time for a listener has elapsed. */
+	bool HasReportIntervalPassed(double CurrentTimeSeconds, UNetworkMetricsBaseListener* Listener);
+
+	enum class EMetricType { Integer, Float };
+
+	TMap<FName, EMetricType> MetricTypes;
+
+	TMap<FName, UE::Net::FNetworkMetric<int64>> MetricInts;
+	TMap<FName, UE::Net::FNetworkMetric<float>> MetricFloats;
+
+	using FNameAndType = TPair<FName, EMetricType>;
+
+	/* The time, in seconds, metrics were reported to a listener. */
+	TMap<TWeakObjectPtr<UNetworkMetricsBaseListener>, double> LastReportListener;
+
+	TMap<TWeakObjectPtr<UNetworkMetricsBaseListener>, TSet<FNameAndType>> ListenersToMetrics;
+};
+
+/** 
+ * An abstract class for metrics listeners that are registered with FNetworkMetricsDatabase.
+ * 
+ * Listeners are the recommended method for reading the current value of metrics from FNetworkMetricsDatabase.
+ * 
+ * Begin by creating a sub-class of UNetworkMetricsBaseListener that overrides the Report() function. This function will be called 
+ * by FNetworkMetricsDatabase::ProcessListeners() once a frame and will be provided an array of metrics that are registered 
+ * with FNetworkMetricsDatabase::Register():
+ *
+ * UCLASS()
+ * class UNetworkMetricsMyListener : public UNetworkMetricsBaseListener
+ * {
+ *		GENERATED_BODY()
+ * public:
+ *		virtual ~UNetworkMetricsMyListener() = default;
+ * 
+ *		void Report(const TArray<UE::Net::FNetworkMetricSnapshot>& Stats)
+ *		{
+ *			for (const UE::Net::FNetworkMetric<int64>& Metric : Snapshot.MetricInts)
+ *			{
+ *				// Do something with integer metrics...
+ *			}
+ *
+ *			for (const UE::Net::FNetworkMetric<float>& Metric : Snapshot.MetricFloats)
+ *			{
+ *				// Do something with floating point metrics...
+ *			}
+ *		}
+ * };
+ *
+ * Listeners can either be registered explicitly using FNetworkMetricsDatabase::Register() or through the engine configuration files. A configuration 
+ * file is the prefered way to register a listener because it allows metrics reporting to be configured without rebuilding the application.
+ *
+ * This is an example configuration from an ini file (e.g. DefaultEngine.ini) that registers metrics with the example listener above:
+ * 
+ * [/Script/Engine.NetworkMetricsConfig]
+ * +Listeners=(MetricName=ExampleMetric1, ClassName=/Script/Engine.NetworkMetricsMyListener)
+ * 
+ * All sub-classes of UNetworkMetricsBaseListener can set a time interval between calls to Report(). This is a useful method for limiting the rate
+ * at which metrics need to be recorded (e.g. you may only want to report metrics to an external analytics services every 60 seconds). This time interval
+ * can be set by calling UNetworkMetricsBaseListener::SetInterval() or in a configuration file by setting the IntervalSeconds property on the
+ * listener sub-class.
+ * 
+ * This is an example configuration from an ini file (e.g. DefaultEngine.ini) that sets the interval between calling UNetworkMetricsMyListener::Report()
+ * to 1 second:
+ * 
+ * [/Script/Engine.NetworkMetricsMyListener]
+ * IntervalSeconds=1
+ */
+UCLASS(abstract, Config=Engine)
+class ENGINE_API UNetworkMetricsBaseListener : public UObject
+{
+	GENERATED_BODY()
+
+public:
+	UNetworkMetricsBaseListener();
+	virtual ~UNetworkMetricsBaseListener() = default;
+
+	/* Set the interval, in seconds, between calling Report(). */
+	void SetInterval(double Seconds)
+	{
+		if (ensureMsgf(Seconds >= 0, TEXT("SetInterval() called with a negative time interval.")))
+		{
+			IntervalSeconds = Seconds;
+		}
+	}
+
+	/* Get the interval, in seconds, between calling Report(). */
+	double GetInterval() const
+	{
+		return IntervalSeconds;
+	}
+
+	virtual void Report(const UE::Net::FNetworkMetricSnapshot& Snapshot) {};
+
+private:
+	UPROPERTY(Config)
+	double IntervalSeconds;
+};
+
+/**
+ * A metrics listener that reports an array of metrics to CSV.
+ * 
+ * The function SetCategory() is expected to be called before the listener is registered with
+ * FNetworkMetricsDatabase::Register(). This function will associate the instance of UNetworkMetricsCSV
+ * with a category of values in CSV.
+ * 
+ * To use UNetworkMetricsCSV in a configuration file, a sub-class of UNetworkMetricsCSV must be 
+ * created that calls SetCategory() from the constructor to provide the CSV category to use.
+ * 
+ * UCLASS()
+ * class UNetworkMetricsCSV_ExampleCategory : public UNetworkMetricsCSV
+ * {
+ *		GENERATED_BODY()
+ * public:
+ *		virtual ~UNetworkMetricsCSV_ExampleCategory() = default;
+ * 
+ *		UNetworkMetricsCSV_ExampleCategory() : UNetworkMetricsCSV()
+ *		{
+ *			SetCategory("ExampleCategory");
+ *		}
+ * };
+ * 
+ * This sub-class can then be used in the configuration file when registering a listener:
+ * 
+ *	[/Script/Engine.NetworkMetricsConfig]
+ *	+Listeners=(MetricName=ExampleMetric, ClassName=/Script/Engine.NetworkMetricsCSV_ExampleCategory)
+ * 
+ * If the base UNetworkMetricsCSV class is used in the configuration file, CSV stats will be recorded to the default 'Networking' category.
+ */
+UCLASS()
+class ENGINE_API UNetworkMetricsCSV : public UNetworkMetricsBaseListener
+{
+	GENERATED_BODY()
+
+public:
+	UNetworkMetricsCSV();
+	virtual ~UNetworkMetricsCSV() = default;
+
+	/* Set the CSV category. */
+	void SetCategory(const FString& CategoryName);
+
+	virtual void Report(const UE::Net::FNetworkMetricSnapshot& Snapshot);
+
+private:
+	int32 CategoryIndex;
+};
+
+/**
+ * A metrics listener that reports an array of metrics to PerfCounters.
+ * 
+ * To use UNetworkMetricsPerfCounters in a configuration file, the class must be used when registering a listener:
+ * 
+ *	[/Script/Engine.NetworkMetricsConfig]
+ *	+Listeners=(MetricName=ExampleMetric, ClassName=/Script/Engine.NetworkMetricsPerfCounters)
+ */
+UCLASS()
+class ENGINE_API UNetworkMetricsPerfCounters : public UNetworkMetricsBaseListener
+{
+	GENERATED_BODY()
+
+public:
+	virtual ~UNetworkMetricsPerfCounters() = default;
+
+	virtual void Report(const UE::Net::FNetworkMetricSnapshot& Snapshot);
+};
+
+/**
+ * A metrics listener that reports a metric to a single Stat. 
+ * 
+ * The function SetStatName() is expected to be called before the listener is registered with
+ * FNetworkMetricsDatabase::Register(). This function will associate the instance of UNetworkMetricStats
+ * with a specific Stat.
+ * 
+ * Since each instance of this class is associated with a single Stat it can only be registered
+ * as a listener to a single metric in FNetworkMetricsDatabase.
+ * 
+ * UNetworkMetricsStats is not intended to be used from configuration files!
+ */
+UCLASS()
+class ENGINE_API UNetworkMetricsStats : public UNetworkMetricsBaseListener
+{
+	GENERATED_BODY()
+
+public:
+	UNetworkMetricsStats();
+	virtual ~UNetworkMetricsStats() = default;
+
+	/** Set the name of the pre-defined Stat (normally defined with DEFINE_STAT()). */
+	void SetStatName(const FName Name)
+	{
+		StatName = Name;
+	}
+
+	virtual void Report(const UE::Net::FNetworkMetricSnapshot& Snapshot);
+
+private:
+	FName StatName;
+};
\ No newline at end of file
diff --git a/Engine/Source/Runtime/Engine/Public/Net/NetworkMetricsDefs.h b/Engine/Source/Runtime/Engine/Public/Net/NetworkMetricsDefs.h
new file mode 100644
--- /dev/null
+++ b/Engine/Source/Runtime/Engine/Public/Net/NetworkMetricsDefs.h
@@ -0,0 +1,98 @@
+// Copyright Epic Games, Inc. All Rights Reserved.
+
+#pragma once
+
+#include "UObject\NameTypes.h"
+
+namespace UE::Net::Metric
+{
+
+extern ENGINE_API const FName AddedConnections;
+extern ENGINE_API const FName ClosedConnectionsDueToReliableBufferOverflow;
+extern ENGINE_API const FName InPackets;
+extern ENGINE_API const FName OutPackets;
+extern ENGINE_API const FName InPacketsClientPerSecondAvg;
+extern ENGINE_API const FName InPacketsClientPerSecondMax;
+extern ENGINE_API const FName InPacketsClientPerSecondMin;
+extern ENGINE_API const FName OutPacketsClientPerSecondAvg;
+extern ENGINE_API const FName OutPacketsClientPerSecondMax;
+extern ENGINE_API const FName OutPacketsClientPerSecondMin;
+extern ENGINE_API const FName InPacketsClientAvg;
+extern ENGINE_API const FName InPacketsClientMax;
+extern ENGINE_API const FName OutPacketsClientAvg;
+extern ENGINE_API const FName OutPacketsClientMax;
+extern ENGINE_API const FName InRate;
+extern ENGINE_API const FName OutRate;
+extern ENGINE_API const FName InRateClientAvg;
+extern ENGINE_API const FName InRateClientMax;
+extern ENGINE_API const FName InRateClientMin;
+extern ENGINE_API const FName OutRateClientAvg;
+extern ENGINE_API const FName OutRateClientMax;
+extern ENGINE_API const FName OutRateClientMin;
+extern ENGINE_API const FName InPacketsLost;
+extern ENGINE_API const FName OutPacketsLost;
+extern ENGINE_API const FName InBunches;
+extern ENGINE_API const FName OutBunches;
+extern ENGINE_API const FName Ping;
+extern ENGINE_API const FName AvgPing;
+extern ENGINE_API const FName MinPing;
+extern ENGINE_API const FName MaxPing;
+extern ENGINE_API const FName OutKBytes;
+extern ENGINE_API const FName OutNetGUIDKBytesSec;
+extern ENGINE_API const FName ServerReplicateActorTimeMS;
+extern ENGINE_API const FName GatherPrioritizeTimeMS;
+extern ENGINE_API const FName ReplicateActorTimeMS;
+extern ENGINE_API const FName NumReplicatedActors;
+extern ENGINE_API const FName Connections;
+extern ENGINE_API const FName NumReplicateActorCallsPerConAvg;
+extern ENGINE_API const FName NumberOfActiveActors;
+extern ENGINE_API const FName NumberOfFullyDormantActors;
+extern ENGINE_API const FName NumSkippedObjectEmptyUpdates;
+extern ENGINE_API const FName NumOpenChannels;
+extern ENGINE_API const FName NumTickingChannels;
+extern ENGINE_API const FName SatConnections;
+extern ENGINE_API const FName SharedSerializationPropertyHit;
+extern ENGINE_API const FName SharedSerializationPropertyMiss;
+extern ENGINE_API const FName SharedSerializationRPCHit;
+extern ENGINE_API const FName SharedSerializationRPCMiss;
+extern ENGINE_API const FName NumClientUpdateLevelVisibility;
+extern ENGINE_API const FName Channels;
+extern ENGINE_API const FName NumActorChannels;
+extern ENGINE_API const FName NumDormantActors;
+extern ENGINE_API const FName NumActors;
+extern ENGINE_API const FName NumNetActors;
+extern ENGINE_API const FName NumNetGUIDsAckd;
+extern ENGINE_API const FName NumNetGUIDsPending;
+extern ENGINE_API const FName NumNetGUIDsUnAckd;
+extern ENGINE_API const FName NetSaturated;
+extern ENGINE_API const FName MaxPacketOverhead;
+extern ENGINE_API const FName NetGUIDInRate;
+extern ENGINE_API const FName NetGUIDOutRate;
+extern ENGINE_API const FName NetNumClients;
+extern ENGINE_API const FName NumClients;
+extern ENGINE_API const FName NumConnections;
+extern ENGINE_API const FName NumConsideredActors;
+extern ENGINE_API const FName NumInitiallyDormantActors;
+extern ENGINE_API const FName PrioritizedActors;
+extern ENGINE_API const FName NumRelevantDeletedActors;
+extern ENGINE_API const FName NumReplicatedActorBytes;
+extern ENGINE_API const FName ImportedNetGuids;
+extern ENGINE_API const FName PendingOuterNetGuids;
+extern ENGINE_API const FName NetInBunchTimeOvershootPercent;
+extern ENGINE_API const FName UnmappedReplicators;
+extern ENGINE_API const FName PercentInVoice;
+extern ENGINE_API const FName PercentOutVoice;
+extern ENGINE_API const FName VoiceBytesRecv;
+extern ENGINE_API const FName VoiceBytesSent;
+extern ENGINE_API const FName VoicePacketsRecv;
+extern ENGINE_API const FName VoicePacketsSent;
+extern ENGINE_API const FName PingBucketInt0;
+extern ENGINE_API const FName PingBucketInt1;
+extern ENGINE_API const FName PingBucketInt2;
+extern ENGINE_API const FName PingBucketInt3;
+extern ENGINE_API const FName PingBucketInt4;
+extern ENGINE_API const FName PingBucketInt5;
+extern ENGINE_API const FName PingBucketInt6;
+extern ENGINE_API const FName PingBucketInt7;
+
+}
