[NetPhysics] Fix issues with replication target alignment
The issue was that target alignment was initialized incorrectly in various different scenarios, making the starting alignment incorrect which would keep the replication incorrect through the objects replication lifetime.
- The server can send the first RepMovement with ServerFrame 0 which made SendInterval calculate incorrect on the next received RepMovement.
- The first RepMovement might get delayed a lot while the second might arrive just after the first making ReceiveInterval and SendInterval vastly different.
- The fake state created client-side when a physics-object wakes up was counted as a valid state to calculate ReceiveInterval from.
- If the client or server would stall for a short time the offset in physics-ticks between them could reach a level that would make the calculation for time dilation invert itself and go into a feedback loop and never respolve the offset. This would make the replication target alignment break since it relies on physics offset being stable / managed.

The fix ensures that SendInterval and ReceiveInterval is calculated and used correctly and that we start target alignment first after the 3rd valid state is received.

diff --git a/Engine/Source/Runtime/Engine/Private/PhysicsEngine/PhysicsReplication.cpp b/Engine/Source/Runtime/Engine/Private/PhysicsEngine/PhysicsReplication.cpp
--- a/Engine/Source/Runtime/Engine/Private/PhysicsEngine/PhysicsReplication.cpp
+++ b/Engine/Source/Runtime/Engine/Private/PhysicsEngine/PhysicsReplication.cpp
@@ -197,11 +197,8 @@ namespace PhysicsReplicationCVars
 		static float SleepSecondsClearTarget = 15.0f;
 		static FAutoConsoleVariableRef CVarSleepSecondsClearTarget(TEXT("np2.PredictiveInterpolation.SleepSecondsClearTarget"), SleepSecondsClearTarget, TEXT("Wait for the object to sleep for this many seconds before clearing the replication target, to ensure nothing wakes up the object just after it goes to sleep on the client."));
 		
-		static int32 TargetTickAlignmentClampMultiplier = 20;
+		static int32 TargetTickAlignmentClampMultiplier = 10;
 		static FAutoConsoleVariableRef CVarTargetTickAlignmentClampMultiplier(TEXT("np2.PredictiveInterpolation.TargetTickAlignmentClampMultiplier"), TargetTickAlignmentClampMultiplier, TEXT("Multiplier to adjust clamping of target alignment via TickCount. Multiplier is performed on AverageReceiveInterval."));
-		
-		static bool LegacyTargetUpdateCheck = false;
-		static FAutoConsoleVariableRef CVarTargetLegacyTargetUpdateCheck(TEXT("np2.PredictiveInterpolation.LegacyTargetUpdateCheck"), LegacyTargetUpdateCheck, TEXT("Use Input.ServerFrame >= Target->PrevServerFrame when checking if a received input state should be cached as a target. This is the legacy version of the check."));
 	}
 
 }
@@ -775,8 +772,7 @@ void FPhysicsReplicationAsync::OnPreSimulate_Internal()
 					// If final resim frame, mark interpolated targets as waiting for up to date data from the server.
 					if (Target.RepMode == EPhysicsReplicationMode::PredictiveInterpolation)
 					{
-						Target.bWaiting = true;
-						Target.ServerFrame = RigidsSolver->GetCurrentFrame() + Target.FrameOffset;
+						Target.SetWaiting(RigidsSolver->GetCurrentFrame() + Target.FrameOffset);
 					}
 				}
 			}
@@ -858,9 +854,7 @@ void FPhysicsReplicationAsync::UpdateAsyncTarget(const FPhysicsRepAsyncInputData
 	bool bFirstTarget = Target == nullptr;
 	if (bFirstTarget)
 	{
-		// First time we add a target, set it's previous and correction
-		// positions to the target position to avoid math with uninitialized
-		// memory.
+		// First time we add a target, set previous state to current input
 		Target = &ObjectToTarget.Add(Input.PhysicsObject, FReplicatedPhysicsTargetAsync());
 		Target->PrevPos = Input.TargetState.Position;
 		Target->PrevPosTarget = Input.TargetState.Position;
@@ -868,27 +862,76 @@ void FPhysicsReplicationAsync::UpdateAsyncTarget(const FPhysicsRepAsyncInputData
 		Target->PrevLinVel = Input.TargetState.LinVel;
 	}
 
-	// Note: If target is waiting, check PrevServerFrame since ServerFrame has been modified and we still want to cache the input target
-	if (PhysicsReplicationCVars::PredictiveInterpolationCVars::LegacyTargetUpdateCheck
-		? (Input.ServerFrame >= Target->PrevServerFrame) // Legacy check
-		: (bFirstTarget || Input.ServerFrame == 0 || (Target->bWaiting ? Input.ServerFrame > Target->PrevServerFrame : Input.ServerFrame > Target->ServerFrame)))
-	{
-		const int32 PrevTickCount = Target->TickCount;
-		const int32 PrevReceiveInterval = Target->ReceiveInterval;
+	/** Target Update Description
+	* @param Input = incoming state target for replication.
+	* 
+	* Input comes mainly from the server but can be a faked state produced by the client for example if the client object wakes up from sleeping.
+	* Fake inputs should have a ServerFrame of -1 (bool bIsFake = Input.ServerFrame == -1)
+	* Server inputs can have ServerFrame values of either 0 or an incrementing integer value.
+	*	If the ServerFrame is 0 it should always be 0. If it's incrementing it will always increment.
+	*
+	* @local Target = The current state target used for replication, to be updated with data from Input.
+	* Read about the different target properties in FReplicatedPhysicsTargetAsync
+	* 
+	* IMPORTANT:
+	* Target.ServerFrame can be -1 if the target is newly created or if it has data from a fake input.
+	* 
+	* SendInterval is calculated by taking Input.ServerFrame - Target.ServerFrame
+	*	Note, can only be calculated if the server is sending incrementing SendIntervals and if we have received a valid input previously so we have the previous ServerFrame cached in Target.
+	* 
+	* ReceiveInterval is calculated by taking RigidsSolver->GetCurrentFrame() - Target.ReceiveFrame
+	*	Note that ReceiveInterval is only used if SendInterval is 0
+	* 
+	* Target.TickCount starts at 0 and is incremented each tick that the target is used for, TickCount is reset back to 0 each time Target is updated with new Input.
+	* 
+	* NOTE: With perfect network conditions SendInterval, ReceiveInterval and Target.TickCount will be the same value.
+	*/
+
+	// Update target from input if input is newer than target or this is the first input received (target is empty)
+	if ((bFirstTarget || Input.ServerFrame == 0 || Input.ServerFrame > Target->ServerFrame))
+	{
+		// Get the current physics frame
+		const int32 CurrentFrame = RigidsSolver->GetCurrentFrame();
+
+		// Cache TickCount before updating it, force to 0 if ServerFrame is -1
+		const int32 PrevTickCount = (Target->ServerFrame < 0) ? 0 : Target->TickCount;
+		
+		// Cache SendInterval, only calculate if we have a valid Target->ServerFrame, else leave at 0.
 		const int32 SendInterval = (Target->ServerFrame <= 0) ? 0 : Input.ServerFrame - Target->ServerFrame;
-		const int32 AdjustedAverageReceiveInterval = FMath::CeilToInt(Target->AverageReceiveInterval) * PhysicsReplicationCVars::PredictiveInterpolationCVars::TargetTickAlignmentClampMultiplier;
+		
+		// Cache if this target was previously allowed to be altered, before this update
+		const bool bPrevAllowTargetAltering = Target->bAllowTargetAltering;
+		
+		// Set if the target is allowed to be altered after this update
+		Target->bAllowTargetAltering = !(Target->TargetState.Flags & ERigidBodyFlags::Sleeping) && !(Input.TargetState.Flags & ERigidBodyFlags::Sleeping);
+		
+		// Set Target->ReceiveInterval from either SendInterval or the number of physics ticks between receiving input states
+		if (SendInterval > 0)
+		{
+			Target->ReceiveInterval = SendInterval;
+		}
+		else
+		{
+			const int32 PrevReceiveFrame = Target->ReceiveFrame < 0 ? (CurrentFrame - 1) : Target->ReceiveFrame;
+			Target->ReceiveInterval = (CurrentFrame - PrevReceiveFrame);
+		}
 
-		Target->PrevServerFrame = Target->bWaiting ? Input.ServerFrame : Target->ServerFrame;
-		Target->ServerFrame = Target->bWaiting ? Target->ServerFrame : Input.ServerFrame;
-		Target->PrevReceiveFrame = (Target->ReceiveFrame == INDEX_NONE) ? (RigidsSolver->GetCurrentFrame() - 1) : Target->ReceiveFrame;
-		Target->ReceiveFrame = RigidsSolver->GetCurrentFrame();
-		Target->ReceiveInterval = FMath::Clamp((Target->ReceiveFrame - Target->PrevReceiveFrame), 1, 255);
+		// Update target from input and reset properties
+		PRAGMA_DISABLE_DEPRECATION_WARNINGS
+		Target->PrevServerFrame = Target->bWaiting ? Input.ServerFrame : Target->ServerFrame; // DEPRECATED UE5.4
+		Target->PrevReceiveFrame = (Target->ReceiveFrame == INDEX_NONE) ? (RigidsSolver->GetCurrentFrame() - 1) : Target->ReceiveFrame; // DEPRECATED UE5.4
+		PRAGMA_ENABLE_DEPRECATION_WARNINGS
+		Target->ServerFrame = Input.ServerFrame;
+		Target->ReceiveFrame = CurrentFrame;
 		Target->TargetState = Input.TargetState;
 		Target->RepMode = Input.RepMode;
 		Target->FrameOffset = Input.FrameOffset;
 		Target->TickCount = 0;
 		Target->AccumulatedSleepSeconds = 0.0f;
 
+		// Update waiting state
+		Target->UpdateWaiting(Input.ServerFrame);
+
 		if (Input.RepMode == EPhysicsReplicationMode::PredictiveInterpolation)
 		{
 #if !(UE_BUILD_SHIPPING || UE_BUILD_TEST)
@@ -899,18 +942,45 @@ void FPhysicsReplicationAsync::UpdateAsyncTarget(const FPhysicsRepAsyncInputData
 			}
 #endif
 
-			// Set the TickCount to the physics tick offset value from where we expected this target to arrive.
-			// If the client has ticked 2 times ahead from the last target and this target is 3 ticks in front of the previous target then the TickOffset should be -1
-			Target->TickCount = bFirstTarget ? 0 : FMath::Clamp((PrevTickCount - (SendInterval > 0 ? SendInterval : PrevReceiveInterval)), -AdjustedAverageReceiveInterval, AdjustedAverageReceiveInterval);
-
 			// Cache the position we received this target at, Predictive Interpolation will alter the target state but use this as the source position for reconciliation.
 			Target->PrevPosTarget = Input.TargetState.Position;
 			Target->PrevRotTarget = Input.TargetState.Quaternion;
-
-			// If we extrapolated the previous target past the receive interval, extrapolate this target by the overshoot
-			if (!Target->bWaiting)
+		
 			{
-				FPhysicsReplicationAsync::ExtrapolateTarget(*Target, Target->TickCount, GetDeltaTime_Internal());
+				/** Target Alignment Feature
+				* With variable network conditions state inputs from the server can arrive both later or earlier than expected.
+				* Target Alignment can adjust for this to make replication act on a target in the timeline that the client is currently replicating in.
+				* 
+				* If SendInterval is 4 we expect TickCount to be 4. TickCount - SendInterval = 0, meaning the client and server has ticked physics the same amount between the target states.
+				* 
+				* If SendInterval is 4 and TickCount is 2 we have only simulated physics for 2 ticks with the previous target while the server had simulated physics 4 ticks between previous target and new target
+				*	TickCount - SendInterval = -2
+				*	To align this we need to adjust the new target by predicting backwards by 2 ticks, else the replication will start replicating towards a state that is 2 ticks further ahead than expected, making replication speed up.
+				* 
+				* Same goes for vice-versa:
+				* If SendInterval is 4 and TickCount is 6 we have simulated physics for 6 ticks with the previous target while the server had simulated physics 4 ticks between previous target and new target
+				*	TickCount - SendInterval = 2
+				*	To align this we need to adjust the new target by predicting forwards by 2 ticks, else the replication will start replicating towards a state that is 2 ticks behind than expected, making replication slow down.
+				* 
+				* Note that state inputs from the server can arrive fluctuating between above examples, but over time the alignment is evened out to 0.
+				* If the clients latency is raised or lowered since replication started there might be a consistent offset in the TickCount which is handled by TimeDilation of client physics through APlayerController::UpdateServerAsyncPhysicsTickOffset()
+				*/
+
+				// Run target alignment if we have been allowed to alter the target during the last two target updates
+				if (!bFirstTarget && bPrevAllowTargetAltering && Target->bAllowTargetAltering)
+				{
+					const int32 AdjustedAverageReceiveInterval = FMath::CeilToInt(Target->AverageReceiveInterval) * PhysicsReplicationCVars::PredictiveInterpolationCVars::TargetTickAlignmentClampMultiplier;
+
+					// Set the TickCount to the physics tick offset value from where we expected this target to arrive.
+					// If the client has ticked 2 times ahead from the last target and this target is 3 ticks in front of the previous target then the TickOffset should be -1
+					Target->TickCount = FMath::Clamp(PrevTickCount - Target->ReceiveInterval, -AdjustedAverageReceiveInterval, AdjustedAverageReceiveInterval);
+
+					// Apply target alignment if we aren't waiting for a newer state from the server
+					if (!Target->IsWaiting())
+					{
+						FPhysicsReplicationAsync::ExtrapolateTarget(*Target, Target->TickCount, GetDeltaTime_Internal());
+					}
+				}
 			}
 		}
 	}
@@ -921,6 +991,11 @@ void FPhysicsReplicationAsync::UpdateAsyncTarget(const FPhysicsRepAsyncInputData
 
 void FPhysicsReplicationAsync::CacheResimInteractions()
 {
+	if(!PhysicsReplicationCVars::ResimulationCVars::bDisableReplicationOnInteraction)
+	{
+		return;
+	}
+
 	Chaos::FPBDRigidsSolver* RigidsSolver = static_cast<Chaos::FPBDRigidsSolver*>(GetSolver());
 	if (RigidsSolver == nullptr)
 	{
@@ -1310,11 +1385,10 @@ bool FPhysicsReplicationAsync::PredictiveInterpolation(Chaos::FPBDRigidParticleH
 		return true;
 	}
 
-	if (Target.bWaiting && (Target.PrevServerFrame + Target.TickCount) < Target.ServerFrame)
+	if (Target.IsWaiting())
 	{
 		return false;
 	}
-	Target.bWaiting = false;
 
 	Chaos::FPBDRigidsSolver* RigidsSolver = static_cast<Chaos::FPBDRigidsSolver*>(GetSolver());
 	if (RigidsSolver == nullptr)
@@ -1325,8 +1399,7 @@ bool FPhysicsReplicationAsync::PredictiveInterpolation(Chaos::FPBDRigidParticleH
 	if (PhysicsReplicationCVars::ResimulationCVars::bDisableReplicationOnInteraction && ParticlesInResimIslands.Contains(Handle->GetHandleIdx()))
 	{
 		// If particle is in an island with a resim object, don't run replication and wait for an up to date target (after leaving the island)
-		Target.bWaiting = true;
-		Target.ServerFrame = RigidsSolver->GetCurrentFrame() + Target.FrameOffset;
+		Target.SetWaiting(RigidsSolver->GetCurrentFrame() + Target.FrameOffset);
 		return false;
 	}
 
@@ -1366,10 +1439,10 @@ bool FPhysicsReplicationAsync::PredictiveInterpolation(Chaos::FPBDRigidParticleH
 			&& !PhysicsReplicationCVars::PredictiveInterpolationCVars::bDontClearTarget;
 
 		// --- Target Prediction ---
-		if (!bClearTarget)
+		if (!bClearTarget && Target.bAllowTargetAltering)
 		{
 			const int32 ExtrapolationTickLimit = FMath::Max(
-				FMath::CeilToInt(Target.ReceiveInterval * PhysicsReplicationCVars::PredictiveInterpolationCVars::ExtrapolationTimeMultiplier), // Extrapolate time based on receive interval * multiplier
+				FMath::CeilToInt(Target.AverageReceiveInterval * PhysicsReplicationCVars::PredictiveInterpolationCVars::ExtrapolationTimeMultiplier), // Extrapolate time based on receive interval * multiplier
 				FMath::CeilToInt(PhysicsReplicationCVars::PredictiveInterpolationCVars::ExtrapolationMinTime / DeltaSeconds)); // At least extrapolate for N seconds
 			if (Target.TickCount <= ExtrapolationTickLimit)
 			{
@@ -1408,8 +1481,8 @@ bool FPhysicsReplicationAsync::PredictiveInterpolation(Chaos::FPBDRigidParticleH
 		RigidsSolver->GetEvolution()->SetParticleObjectState(Handle, Chaos::EObjectStateType::Dynamic);
 	}
 
-	// Update the AverageReceiveInterval of targets from the server (receive rate = send-rate from server with network conditions taken into account)
-	Target.AverageReceiveInterval = FMath::Lerp(Target.AverageReceiveInterval, Target.ReceiveInterval, FMath::Clamp((1.0f / (Target.ReceiveInterval * PhysicsReplicationCVars::PredictiveInterpolationCVars::AverageReceiveIntervalSmoothing)), 0.0f, 1.0f));
+	// Update the AverageReceiveInterval if Target.ReceiveInterval has a valid value to update from
+	Target.AverageReceiveInterval = Target.ReceiveInterval == 0 ? Target.AverageReceiveInterval : FMath::Lerp(Target.AverageReceiveInterval, Target.ReceiveInterval, FMath::Clamp((1.0f / (Target.ReceiveInterval * PhysicsReplicationCVars::PredictiveInterpolationCVars::AverageReceiveIntervalSmoothing)), 0.0f, 1.0f));
 
 	// CurrentState
 	FRigidBodyState CurrentState;
diff --git a/Engine/Source/Runtime/Engine/Private/PlayerController.cpp b/Engine/Source/Runtime/Engine/Private/PlayerController.cpp
--- a/Engine/Source/Runtime/Engine/Private/PlayerController.cpp
+++ b/Engine/Source/Runtime/Engine/Private/PlayerController.cpp
@@ -140,6 +140,9 @@ namespace NetworkPhysicsCvars
 	float TimeDilationEscalationDecay = 0.05f;
 	FAutoConsoleVariableRef CVarTimeDilationEscalationDecay(TEXT("np2.TimeDilationEscalationDecay"), TimeDilationEscalationDecay, TEXT("Value is a multiplier, Default: 0.05. For each escalated TimeDilation amount, also decay by this much. Disable by setting to 0."));
 
+	float TimeDilationEscalationDecayMax = 0.5f;
+	FAutoConsoleVariableRef CVarTimeDilationEscalationDecayMax(TEXT("np2.TimeDilationEscalationDecayMax"), TimeDilationEscalationDecayMax, TEXT("Value is a multiplier, Default: 0.5. The max decay value for escalated time dilation. Lower value means higher decay."));
+
 	float TimeDilationMax = 1.1f;
 	FAutoConsoleVariableRef CVarTimeDilationMax(TEXT("np2.TimeDilationMax"), TimeDilationMax, TEXT("Max value of the time dilation multiplier."));
 
@@ -6255,7 +6258,7 @@ void APlayerController::ServerSendLatestAsyncPhysicsTimestamp_Implementation(FAs
 		}
 			
 		// Calculate desired dilation and send to client
-		const float TimeDilationDecay = 1.0f - (NetworkPhysicsCvars::TimeDilationEscalationDecay * FMath::Abs(CurrentFrameBufferOffset));
+		const float TimeDilationDecay = FMath::Clamp(1.0f - (NetworkPhysicsCvars::TimeDilationEscalationDecay * FMath::Abs(CurrentFrameBufferOffset)), NetworkPhysicsCvars::TimeDilationEscalationDecayMax, 1.0f);
 		float CalculatedTimeDilation = 1.0f + ((NetworkPhysicsCvars::TimeDilationAmount * -CurrentFrameBufferOffset) * TimeDilationDecay);
 		CalculatedTimeDilation = FMath::Clamp(CalculatedTimeDilation, NetworkPhysicsCvars::TimeDilationMin, NetworkPhysicsCvars::TimeDilationMax);
 
diff --git a/Engine/Source/Runtime/Engine/Public/PhysicsReplication.h b/Engine/Source/Runtime/Engine/Public/PhysicsReplication.h
--- a/Engine/Source/Runtime/Engine/Public/PhysicsReplication.h
+++ b/Engine/Source/Runtime/Engine/Public/PhysicsReplication.h
@@ -96,11 +96,22 @@ struct FReplicatedPhysicsTargetAsync
 		, ServerFrame(INDEX_NONE)
 		, ReceiveFrame(INDEX_NONE)
 		, ReceiveInterval(5)
-		, AverageReceiveInterval(5.f)
+		, AverageReceiveInterval(5.0f)
+		PRAGMA_DISABLE_DEPRECATION_WARNINGS
 		, PrevServerFrame(INDEX_NONE)
 		, bWaiting(false)
+		PRAGMA_ENABLE_DEPRECATION_WARNINGS
 		, AccumulatedSleepSeconds(0.0f)
+		, bAllowTargetAltering(false)
+		, WaitForServerFrame(INDEX_NONE)
 	{ }
+	PRAGMA_DISABLE_DEPRECATION_WARNINGS
+	~FReplicatedPhysicsTargetAsync() = default;
+	FReplicatedPhysicsTargetAsync(const FReplicatedPhysicsTargetAsync&) = default;
+	FReplicatedPhysicsTargetAsync(FReplicatedPhysicsTargetAsync&&) = default;
+	FReplicatedPhysicsTargetAsync& operator=(const FReplicatedPhysicsTargetAsync&) = default;
+	FReplicatedPhysicsTargetAsync& operator=(FReplicatedPhysicsTargetAsync&&) = default;
+	PRAGMA_ENABLE_DEPRECATION_WARNINGS
 
 	/** The target state replicated by server */
 	FRigidBodyState TargetState;
@@ -111,10 +122,10 @@ struct FReplicatedPhysicsTargetAsync
 	/** The amount of simulation ticks this target has been used for */
 	int32 TickCount;
 
-	/** ServerFrame this target was replicated on (must be converted to local frame prior to client-side use) */
+	/** ServerFrame this target was replicated on. (LocalFrame = ServerFrame - FrameOffset) */
 	int32 ServerFrame;
 
-	/** The frame offset between local client and server */
+	/** The frame offset between local client and server. (LocalFrame = ServerFrame - FrameOffset) */
 	int32 FrameOffset;
 
 	/** The local client frame when receiving this target from the server */
@@ -132,14 +143,47 @@ struct FReplicatedPhysicsTargetAsync
 	FQuat PrevRotTarget;
 	FVector PrevPos;
 	FVector PrevLinVel;
+
+	UE_DEPRECATED(5.4, "This property is deprecated. PrevServerFrame will no longer be cached.")
 	int32 PrevServerFrame;
+
+	UE_DEPRECATED(5.4, "This property is deprecated. PrevReceiveFrame will no longer be cached.")
 	int32 PrevReceiveFrame;
 
 	/** If this target is waiting for up-to-date data? */
+	UE_DEPRECATED(5.4, "This property is deprecated. Use IsWaiting() and SetWaiting() instead.")
 	bool bWaiting;
 
 	/** Accumulated seconds asleep */
 	float AccumulatedSleepSeconds;
+
+	/** If this target is allowed to be altered, via extrapolation or target alignment via TickCount */
+	bool bAllowTargetAltering;
+
+	/** ServerFrame for the target to wait on, no replication will be performed while waiting for up to date data. */
+	int32 WaitForServerFrame;
+
+public:
+	/** Is this target waiting for up to date data? */
+	const bool IsWaiting() { return WaitForServerFrame > INDEX_NONE; } const
+
+	/** Set target to wait for data newer than @param InWaitForServerFrame */
+	void SetWaiting(int32 InWaitForServerFrame)
+	{
+		WaitForServerFrame = InWaitForServerFrame; 
+		PRAGMA_DISABLE_DEPRECATION_WARNINGS
+		bWaiting = IsWaiting();
+		PRAGMA_ENABLE_DEPRECATION_WARNINGS
+	}
+
+	/** Update waiting status and clear waiting if @param InServerFrame is newer than the frame we are waiting for */
+	void UpdateWaiting(int32 InServerFrame)
+	{
+		if (InServerFrame > WaitForServerFrame)
+		{
+			SetWaiting(INDEX_NONE);
+		}
+	}
 };
 
 class FPhysicsReplicationAsync : public Chaos::TSimCallbackObject<
