World Partition - HLOD: Avoid creating an HLOD ISM Component pointing to a null/private/transient static mesh

diff --git a/Engine/Source/Runtime/Engine/Private/WorldPartition/HLOD/HLODBuilder.cpp b/Engine/Source/Runtime/Engine/Private/WorldPartition/HLOD/HLODBuilder.cpp
--- a/Engine/Source/Runtime/Engine/Private/WorldPartition/HLOD/HLODBuilder.cpp
+++ b/Engine/Source/Runtime/Engine/Private/WorldPartition/HLOD/HLODBuilder.cpp
@@ -178,6 +178,15 @@ TArray<UActorComponent*> UHLODBuilder::BatchInstances(const TArray<UActorCompone
 	TMap<uint32, FInstanceBatch> InstancesData;
 	for (UStaticMeshComponent* SMC : SourceStaticMeshComponents)
 	{
+		const UStaticMesh* StaticMesh = SMC->GetStaticMesh();
+		const bool bPrivateStaticMesh = StaticMesh && !StaticMesh->HasAnyFlags(RF_Public);
+		const bool bTransientStaticMesh = StaticMesh && StaticMesh->HasAnyFlags(RF_Transient);
+		if (!StaticMesh || bPrivateStaticMesh || bTransientStaticMesh)
+		{
+			UE_LOG(LogHLODBuilder, Warning, TEXT("Instanced HLOD source component %s points to a %s static mesh, ignoring."), *SMC->GetPathName(), !StaticMesh ? TEXT("null") : bPrivateStaticMesh ? TEXT("private") : TEXT("transient"));
+			continue;
+		}
+
 		TUniquePtr<FISMComponentDescriptor> ISMComponentDescriptor = ComponentClass->GetDefaultObject<UHLODInstancedStaticMeshComponent>()->AllocateISMComponentDescriptor();
 		ISMComponentDescriptor->InitFrom(SMC, false);
 
