Dont set actors as external packaged when they're renamed to a new outer (Level) when in PIE, prevents issues with AssetRegistry asset iteration
Also make sure that external packages created for actors in PIE levels are using the PKG_PlayInEditor flag (if it ever happens again)

diff --git a/Engine/Source/Runtime/Engine/Private/Actor.cpp b/Engine/Source/Runtime/Engine/Private/Actor.cpp
--- a/Engine/Source/Runtime/Engine/Private/Actor.cpp
+++ b/Engine/Source/Runtime/Engine/Private/Actor.cpp
@@ -1345,6 +1345,7 @@ bool AActor::Rename( const TCHAR* InName, UObject* NewOuter, ERenameFlags Flags
 #if WITH_EDITOR
 	// if we have an external actor and the actor is changing name/outer, we will want to update its package
 	const bool bExternalActor = IsPackageExternal();
+	const bool bIsPIEPackage = bExternalActor ? !!(GetExternalPackage()->GetPackageFlags() & PKG_PlayInEditor) : false;
 	bool bShouldSetPackageExternal = false;
 	if (!bRenameTest)
 	{
@@ -1403,7 +1404,8 @@ bool AActor::Rename( const TCHAR* InName, UObject* NewOuter, ERenameFlags Flags
 	const bool bSuccess = Super::Rename( InName, NewOuter, Flags );
 
 #if WITH_EDITOR
-	if (bShouldSetPackageExternal)
+	// Restore external package state, except when in PIE since it's transient and serves no purpose
+	if (bShouldSetPackageExternal && !bIsPIEPackage)
 	{
 		if (ULevel* MyLevel = GetLevel())
 		{
diff --git a/Engine/Source/Runtime/Engine/Private/Level.cpp b/Engine/Source/Runtime/Engine/Private/Level.cpp
--- a/Engine/Source/Runtime/Engine/Private/Level.cpp
+++ b/Engine/Source/Runtime/Engine/Private/Level.cpp
@@ -4088,6 +4088,12 @@ UPackage* ULevel::CreateActorPackage(UPackage* InLevelPackage, EActorPackagingSc
 
 	UPackage* ActorPackage = CreatePackage(*PackageName);
 	ActorPackage->SetPackageFlags(PKG_EditorOnly | PKG_ContainsMapData | PKG_NewlyCreated);
+
+	// Should be prevented upstream but we propagate the flag to prevent issues in asset enumeration
+	if (!ensureMsgf(!(InLevelPackage->GetPackageFlags() & PKG_PlayInEditor), TEXT("Actor packages should not be created on PlayInEditor levels")))
+	{
+		ActorPackage->SetPackageFlags(PKG_PlayInEditor);
+	}
 	return ActorPackage;
 }
 
