Added a rename flag to disable component registration during the rename process.

#rnx

diff --git a/Engine/Source/Runtime/CoreUObject/Public/UObject/ObjectMacros.h b/Engine/Source/Runtime/CoreUObject/Public/UObject/ObjectMacros.h
--- a/Engine/Source/Runtime/CoreUObject/Public/UObject/ObjectMacros.h
+++ b/Engine/Source/Runtime/CoreUObject/Public/UObject/ObjectMacros.h
@@ -2126,7 +2126,7 @@ private:
 typedef uint32 ERenameFlags;
 
 /** Default rename behavior */
-#define REN_None				(0x0000)
+#define REN_None					(0x0000)
 /** Rename won't call ResetLoaders or flush async loading. You should pass this if you are renaming a deep subobject and do not need to reset loading for the outer package */
 #define REN_ForceNoResetLoaders		(0x0001) 
 /** Just test to make sure that the rename is guaranteed to succeed if an non test rename immediately follows */
@@ -2141,6 +2141,8 @@ typedef uint32 ERenameFlags;
 #define REN_ForceGlobalUnique		(0x0040) 
 /** Prevent renaming of any child generated classes and CDO's in blueprints */
 #define REN_SkipGeneratedClasses	(0x0080) 
+/** Prevents renaming from unregistering/registering all components */
+#define REN_SkipComponentRegWork	(0x0100) 
 
 /*-----------------------------------------------------------------------------
 	Misc.
diff --git a/Engine/Source/Runtime/Engine/Private/Actor.cpp b/Engine/Source/Runtime/Engine/Private/Actor.cpp
--- a/Engine/Source/Runtime/Engine/Private/Actor.cpp
+++ b/Engine/Source/Runtime/Engine/Private/Actor.cpp
@@ -1340,6 +1340,7 @@ AActor* AActor::FindActorInPackage(UPackage* InPackage, bool bEvenIfPendingKill)
 bool AActor::Rename( const TCHAR* InName, UObject* NewOuter, ERenameFlags Flags )
 {
 	const bool bRenameTest = ((Flags & REN_Test) != 0);
+	const bool bPerformComponentRegWork = ((Flags & REN_SkipComponentRegWork) == 0);
 	const bool bChangingOuters = (NewOuter && (NewOuter != GetOuter()));
 
 #if WITH_EDITOR
@@ -1383,7 +1384,11 @@ bool AActor::Rename( const TCHAR* InName, UObject* NewOuter, ERenameFlags Flags
 	if (!bRenameTest && bChangingOuters)
 	{
 		RegisterAllActorTickFunctions(false, true); // unregister all tick functions
-		UnregisterAllComponents();
+
+		if (bPerformComponentRegWork)
+		{
+			UnregisterAllComponents();
+		}
 
 		if (ULevel* MyLevel = GetLevel())
 		{
@@ -1424,7 +1429,7 @@ bool AActor::Rename( const TCHAR* InName, UObject* NewOuter, ERenameFlags Flags
 			MyLevel->ActorsForGC.Add(this);
 
 			UWorld* World = MyLevel->GetWorld();
-			if (World && World->bIsWorldInitialized)
+			if (World && World->bIsWorldInitialized && bPerformComponentRegWork)
 			{
 				RegisterAllComponents();
 			}
