Fix for directional light function resulting in wrong clipping on some hardware due to super large sphere bound.

diff --git a/Engine/Source/Runtime/Engine/Private/LightSceneProxy.cpp b/Engine/Source/Runtime/Engine/Private/LightSceneProxy.cpp
--- a/Engine/Source/Runtime/Engine/Private/LightSceneProxy.cpp
+++ b/Engine/Source/Runtime/Engine/Private/LightSceneProxy.cpp
@@ -182,8 +182,10 @@ void FLightSceneProxy::ApplyWorldOffset(FVector InOffset)
 
 FSphere FLightSceneProxy::GetBoundingSphere() const
 {
-	// Directional lights will have a radius of WORLD_MAX
-	return FSphere(FVector::ZeroVector, WORLD_MAX);
+	// Directional lights will have a radius of WORLD_MAX,
+	// but we use UE_OLD_WORLD_MAX which is smaller, because WORLD_MAX is SUPER larger when set to UE_LARGE_WORLD_MAX,
+	// and in this case some GPUs clipper can then fail for camera with a narrow field of view.
+	return FSphere(FVector::ZeroVector, UE_OLD_WORLD_MAX);
 }
 
 FTexture* FLightSceneProxy::GetIESTextureResource() const
