WorldPartition: Add UActorDescContainerInstance/FWorldPartitionActorDescInstance which take transient/context data from UActorDescContainer/FWorldPartitionActorDesc leaving them as pure on disk repres

diff --git a/Engine/Source/Editor/WorldPartitionEditor/Private/WorldPartition/SWorldPartitionEditorGrid2D.cpp b/Engine/Source/Editor/WorldPartitionEditor/Private/WorldPartition/SWorldPartitionEditorGrid2D.cpp
--- a/Engine/Source/Editor/WorldPartitionEditor/Private/WorldPartition/SWorldPartitionEditorGrid2D.cpp
+++ b/Engine/Source/Editor/WorldPartitionEditor/Private/WorldPartition/SWorldPartitionEditorGrid2D.cpp
@@ -33,7 +33,7 @@
 #include "WorldPartition/WorldPartition.h"
 #include "WorldPartition/WorldPartitionLog.h"
 #include "WorldPartition/WorldPartitionActorDesc.h"
-#include "WorldPartition/WorldPartitionActorDescView.h"
+#include "WorldPartition/WorldPartitionActorDescInstance.h"
 #include "WorldPartition/WorldPartitionActorDescUtils.h"
 #include "WorldPartition/WorldPartitionEditorHash.h"
 #include "WorldPartition/WorldPartitionEditorLoaderAdapter.h"
@@ -58,6 +58,73 @@ static FAutoConsoleCommand CVarToggleShowEditorProfilingStats(
 	FConsoleCommandDelegate::CreateLambda([] { GShowEditorProfilingStats = !GShowEditorProfilingStats; })
 );
 
+class FEditorGrid2DActorDescInstance : public FWorldPartitionActorDescInstance
+{
+public:
+	FEditorGrid2DActorDescInstance(){}
+	virtual ~FEditorGrid2DActorDescInstance() {}
+
+	virtual FBox GetEditorBounds() const override
+	{
+		if (AActor* Actor = GetActor(false))
+		{
+			return Actor->GetStreamingBounds();
+		}
+
+		return FWorldPartitionActorDescInstance::GetEditorBounds();
+	}
+
+	virtual FName GetActorLabel() const override
+	{
+		FName ActorLabel = FWorldPartitionActorDescInstance::GetActorLabel();
+
+		if (ActorLabel.IsNone())
+		{
+			if (AActor* Actor = GetActor(false))
+			{
+				ActorLabel = *Actor->GetActorLabel(false);
+			}
+		}
+
+		return ActorLabel;
+	}
+};
+
+class FEditorGrid2DUnsavedActorDescInstance : public FEditorGrid2DActorDescInstance
+{
+public:
+	FEditorGrid2DUnsavedActorDescInstance(AActor* InActor)
+	{
+		ActorDescPtr = InActor->CreateActorDesc();
+		ActorDesc = ActorDescPtr.Get();
+	}
+
+	virtual ~FEditorGrid2DUnsavedActorDescInstance() {}
+	
+	// No need for container registration here
+	virtual void RegisterChildContainerInstance() override {}
+	virtual void UnregisterChildContainerInstance() override {}
+
+	TUniquePtr<FWorldPartitionActorDesc> ActorDescPtr;
+};
+
+class FEditorGrid2DDirtyActorDescInstance : public FEditorGrid2DActorDescInstance
+{
+public:
+	FEditorGrid2DDirtyActorDescInstance(const IWorldPartitionActorDescInstance* InActorDescInstance)
+		: ActorDescInstance(InActorDescInstance)
+	{
+	}
+
+	virtual ~FEditorGrid2DDirtyActorDescInstance() {}
+		
+	virtual const FWorldPartitionActorDesc* GetActorDesc() const override { return ActorDescInstance->GetActorDesc(); }
+
+protected:
+	
+	const IWorldPartitionActorDescInstance* ActorDescInstance;
+};
+
 class FWeightedMovingAverageScope
 {
 public:
@@ -192,15 +259,15 @@ void ForEachIntersectingLoaderAdapters(UWorldPartition* WorldPartition, const FB
 		}
 	}
 
-	FWorldPartitionHelpers::ForEachIntersectingActorDesc(WorldPartition, SelectBox, [&](const FWorldPartitionActorDesc* ActorDesc)
+	FWorldPartitionHelpers::ForEachIntersectingActorDescInstance(WorldPartition, SelectBox, [&](const FWorldPartitionActorDescInstance* ActorDescInstance)
 	{
-		if (AActor* Actor = ActorDesc->GetActor())
+		if (AActor* Actor = ActorDescInstance->GetActor())
 		{
-			if (ActorDesc->GetActorNativeClass()->ImplementsInterface(UWorldPartitionActorLoaderInterface::StaticClass()))
+			if (ActorDescInstance->GetActorNativeClass()->ImplementsInterface(UWorldPartitionActorLoaderInterface::StaticClass()))
 			{
 				if (IWorldPartitionActorLoaderInterface::ILoaderAdapter* LoaderAdapter = Cast<IWorldPartitionActorLoaderInterface>(Actor)->GetLoaderAdapter())
 				{
-					if (IsBoundsSelected(SelectBox, ActorDesc->GetEditorBounds()))
+					if (IsBoundsSelected(SelectBox, ActorDescInstance->GetEditorBounds()))
 					{
 						Func(Actor);
 					}
@@ -212,50 +279,6 @@ void ForEachIntersectingLoaderAdapters(UWorldPartition* WorldPartition, const FB
 	});
 }
 
-class FWorldPartitionActorDescViewBoundsProxy : public FWorldPartitionActorDescView
-{
-public:
-	FWorldPartitionActorDescViewBoundsProxy(const FWorldPartitionActorDesc* InActorDesc, bool bInUseActor)
-		: FWorldPartitionActorDescView(InActorDesc)
-		, bUseActor(bInUseActor)
-	{}
-
-	virtual FBox GetEditorBounds() const override
-	{
-		if (bUseActor)
-		{
-			if (AActor* Actor = GetActor())
-			{
-				return Actor->GetStreamingBounds();
-			}
-		}
-
-		return ActorDesc->GetEditorBounds();
-	}
-
-	virtual AActor* GetActor() const override
-	{
-		return bUseActor ? ActorDesc->GetActor(false) : nullptr;
-	}
-
-	virtual FName GetActorLabel() const override
-	{
-		FName ActorLabel = ActorDesc->GetActorLabel();
-		
-		if (ActorLabel.IsNone() && bUseActor)
-		{
-			if (AActor* Actor = GetActor())
-			{
-				ActorLabel = *Actor->GetActorLabel(false);
-			}
-		}
-
-		return ActorLabel;
-	}
-
-	bool bUseActor;
-};
-
 SWorldPartitionEditorGrid2D::FEditorCommands::FEditorCommands()
 	: TCommands<FEditorCommands>
 (
@@ -628,7 +651,7 @@ void SWorldPartitionEditorGrid2D::BindCommands()
 	{
 		if (UWorldPartitionSubsystem* WorldPartitionSubsystem = UWorld::GetSubsystem<UWorldPartitionSubsystem>(GetWorld()))
 		{
-			return IsInteractive() && (GEditor->GetSelectedActors()->Num() || WorldPartitionSubsystem->SelectedActorDescs.Num());
+			return IsInteractive() && (GEditor->GetSelectedActors()->Num() || WorldPartitionSubsystem->SelectedActorHandles.Num());
 		}
 
 		return false;
@@ -1272,19 +1295,19 @@ void SWorldPartitionEditorGrid2D::Tick(const FGeometry& AllottedGeometry, const
 	UWorldPartitionEditorHash::FForEachIntersectingActorParams ForEachIntersectingActorParams = UWorldPartitionEditorHash::FForEachIntersectingActorParams()
 		.SetMinimumBox(FBox(FVector::ZeroVector, FVector(GetSelectionSnap())));
 
-	GetWorldPartition()->EditorHash->ForEachIntersectingActor(ViewRectWorld, [&](FWorldPartitionActorDesc* ActorDesc)
+	GetWorldPartition()->EditorHash->ForEachIntersectingActor(ViewRectWorld, [&](FWorldPartitionActorDescInstance* ActorDescInstance)
 	{
 		if (bShowActors)
 		{
-			if (ActorDesc->IsListedInSceneOutliner() && (!GetWorldPartition()->IsStreamingEnabled() || ActorDesc->GetIsSpatiallyLoaded()))
+			if (ActorDescInstance->IsListedInSceneOutliner() && (!GetWorldPartition()->IsStreamingEnabled() || ActorDescInstance->GetIsSpatiallyLoaded()))
 			{
-				ShownActorGuids.Add(ActorDesc->GetGuid());
+				ShownActorGuids.Add(ActorDescInstance->GetGuid());
 			}
 		}
 
-		if (ActorDesc->GetActorNativeClass()->ImplementsInterface(UWorldPartitionActorLoaderInterface::StaticClass()))
+		if (ActorDescInstance->GetActorNativeClass()->ImplementsInterface(UWorldPartitionActorLoaderInterface::StaticClass()))
 		{
-			if (AActor* Actor = ActorDesc->GetActor())
+			if (AActor* Actor = ActorDescInstance->GetActor())
 			{
 				if (IWorldPartitionActorLoaderInterface::ILoaderAdapter* LoaderAdapter = Cast<IWorldPartitionActorLoaderInterface>(Actor)->GetLoaderAdapter())
 				{
@@ -1294,7 +1317,7 @@ void SWorldPartitionEditorGrid2D::Tick(const FGeometry& AllottedGeometry, const
 		}
 	}, ForEachIntersectingActorParams);
 
-	// Add dirty actors as an acceleration for FWorldPartitionActorDescViewBoundsProxy constructor
+	// Add dirty actors as an acceleration for FDirtyActorDescView constructor
 	for (auto& [Reference, Actor] : GetWorldPartition()->GetDirtyActors())
 	{
 		DirtyActorGuids.Add(Actor->GetActorGuid());
@@ -1302,7 +1325,7 @@ void SWorldPartitionEditorGrid2D::Tick(const FGeometry& AllottedGeometry, const
 
 	// Also include transient actor loader adapters that might have been spawned by blutilities, etc. Since these actors can't be saved because they are transient,
 	// they will never get an actor descriptor so they will never appear in the world partition editor. Also include unsaved, newly created actors for convenience.
-	for (FActorDescList::TIterator<> ActorDescIterator(&NewlyAddedUnsavedActorDescs); ActorDescIterator; ++ActorDescIterator)
+	for (FActorDescInstanceList::TIterator<> ActorDescIterator(&NewlyAddedUnsavedActorDescs); ActorDescIterator; ++ActorDescIterator)
 	{
 		if (AActor* NewlyAddedUnsavedActor = ActorDescIterator->GetActor())
 		{
@@ -1401,9 +1424,9 @@ void SWorldPartitionEditorGrid2D::Tick(const FGeometry& AllottedGeometry, const
 
 	// Include selected actor descriptors
 	UWorldPartitionSubsystem* WorldPartitionSubsystem = UWorld::GetSubsystem<UWorldPartitionSubsystem>(GetWorld());
-	for (FWorldPartitionActorDesc* SelectedActorDesc : WorldPartitionSubsystem->SelectedActorDescs)
+	for (const FWorldPartitionHandle& SelectedActorHandle : WorldPartitionSubsystem->SelectedActorHandles)
 	{
-		ShownActorGuids.Add(SelectedActorDesc->GetGuid());
+		ShownActorGuids.Add(SelectedActorHandle.GetInstance()->GetGuid());
 	}
 }
 
@@ -1414,8 +1437,12 @@ uint32 SWorldPartitionEditorGrid2D::PaintActors(const FGeometry& AllottedGeometr
 	const FBox ViewRectWorld(FVector(WorldViewRect.Min.X, WorldViewRect.Min.Y, -HALF_WORLD_MAX), FVector(WorldViewRect.Max.X, WorldViewRect.Max.Y, HALF_WORLD_MAX));
 	const UWorldPartitionSubsystem* WorldPartitionSubsystem = UWorld::GetSubsystem<UWorldPartitionSubsystem>(GetWorld());
 
-	TArray<FWorldPartitionActorDescViewBoundsProxy> ActorDescList;
-	ActorDescList.Reserve(ShownActorGuids.Num());
+	// Reserve here is done to avoid reallocations so that pointers in the BoundProxies set stay valid
+	TArray<FEditorGrid2DDirtyActorDescInstance> DirtyActorDescInstances;
+	DirtyActorDescInstances.Reserve(DirtyActorGuids.Num());
+	
+	TSet<const IWorldPartitionActorDescInstance*> BoundProxies;
+	BoundProxies.Reserve(ShownActorGuids.Num());
 
 	const IWorldPartitionActorLoaderInterface::ILoaderAdapter* LocalHoveredLoaderAdapter = nullptr;
 	if (IWorldPartitionActorLoaderInterface* HoveredLoaderAdapterInterface = Cast<IWorldPartitionActorLoaderInterface>(HoveredLoaderInterface.Get()))
@@ -1426,13 +1453,20 @@ uint32 SWorldPartitionEditorGrid2D::PaintActors(const FGeometry& AllottedGeometr
 	const UWorldPartition* ThisWorldPartition = GetWorldPartition();
 	for (const FGuid& ActorGuid : ShownActorGuids)
 	{
-		if (const FWorldPartitionActorDesc* ActorDesc = ThisWorldPartition->GetActorDesc(ActorGuid))
+		if (const FWorldPartitionActorDescInstance* ActorDescInstance = ThisWorldPartition->GetActorDescInstance(ActorGuid))
 		{
-			ActorDescList.Emplace(ActorDesc, DirtyActorGuids.Contains(ActorGuid));
+			if (DirtyActorGuids.Contains(ActorGuid))
+			{
+				BoundProxies.Add(&DirtyActorDescInstances.Emplace_GetRef(ActorDescInstance));
+			}
+			else
+			{
+				BoundProxies.Add(ActorDescInstance);
+			}
 		}
-		else if (const FWorldPartitionActorDesc* NewlyAddedUnsavedActorDesc = NewlyAddedUnsavedActorDescs.GetActorDesc(ActorGuid))
+		else if (const FWorldPartitionActorDescInstance* NewlyAddedUnsavedActorDesc = NewlyAddedUnsavedActorDescs.GetActorDescInstance(ActorGuid))
 		{
-			ActorDescList.Emplace(NewlyAddedUnsavedActorDesc, true);
+			BoundProxies.Add(NewlyAddedUnsavedActorDesc);
 		}
 	}
 
@@ -1614,13 +1648,13 @@ uint32 SWorldPartitionEditorGrid2D::PaintActors(const FGeometry& AllottedGeometr
 		}
 	}
 
-	if (ActorDescList.Num())
+	if (BoundProxies.Num())
 	{
 		TArray<FVector2D> LinePoints;
 		LinePoints.SetNum(5);
 
 		bool bIsBoundsEdgeHoveredFound = false;
-		for (const FWorldPartitionActorDescViewBoundsProxy& ActorDescView: ActorDescList)
+		for (const IWorldPartitionActorDescInstance* BoundProxy : BoundProxies)
 		{
 			auto ShowActorBox = [this, &AllottedGeometry, &OutDrawElements, &LayerId, &bIsBoundsEdgeHoveredFound, &DrawActorLabel](const FBox& ActorBounds, bool bIsSelected, bool bIsSpatiallyLoaded, const FName ActorLabel)
 			{
@@ -1688,15 +1722,16 @@ uint32 SWorldPartitionEditorGrid2D::PaintActors(const FGeometry& AllottedGeometr
 					);
 				}
 			};
-
-			const FBox ActorBounds = ActorDescView.GetEditorBounds();
-			const bool bIsSelected = SelectedActorGuids.Contains(ActorDescView.GetGuid()) || WorldPartitionSubsystem->SelectedActorDescs.Contains(ActorDescView.GetActorDesc());
-			const bool bIsSpatiallyLoaded = ActorDescView.GetIsSpatiallyLoaded();
-			const FName ActorLabel = ActorDescView.GetActorLabel();
+			
+			const FWorldPartitionHandle ActorHandle(BoundProxy->GetContainerInstance(), BoundProxy->GetGuid());
+			const FBox ActorBounds = BoundProxy->GetEditorBounds();
+			const bool bIsSelected = SelectedActorGuids.Contains(BoundProxy->GetGuid()) || (ActorHandle.IsValid() && WorldPartitionSubsystem->SelectedActorHandles.Contains(ActorHandle));
+			const bool bIsSpatiallyLoaded = BoundProxy->GetIsSpatiallyLoaded();
+			const FName ActorLabel = BoundProxy->GetActorLabel();
 
 			if (bIsSelected)
 			{
-				const FBox ActorDescBounds = ActorDescView.GetActorDesc()->GetEditorBounds();
+				const FBox ActorDescBounds = BoundProxy->GetActorDesc()->GetEditorBounds();
 				if (!ActorDescBounds.Equals(ActorBounds, 1.0f))
 				{
 					ShowActorBox(ActorDescBounds, false, bIsSpatiallyLoaded, ActorLabel);
@@ -2192,9 +2227,9 @@ void SWorldPartitionEditorGrid2D::FocusSelection()
 	}
 
 	UWorldPartitionSubsystem* WorldPartitionSubsystem = UWorld::GetSubsystem<UWorldPartitionSubsystem>(GetWorld());
-	for (FWorldPartitionActorDesc* SelectedActorDesc : WorldPartitionSubsystem->SelectedActorDescs)
+	for (const FWorldPartitionHandle& SelectedActorHandle : WorldPartitionSubsystem->SelectedActorHandles)
 	{
-		SelectionBox += SelectedActorDesc->GetEditorBounds();
+		SelectionBox += SelectedActorHandle.GetInstance()->GetEditorBounds();
 	}
 
 	if (SelectionBox.IsValid)
@@ -2338,14 +2373,26 @@ void SWorldPartitionEditorGrid2D::ClearSelection()
 
 void SWorldPartitionEditorGrid2D::FNewlyAddedUnsavedActorDescsDescRegistry::OnActorAdded(AActor* Actor)
 {
-	if (TUniquePtr<FWorldPartitionActorDesc>* ExistingActorDesc = GetActorDescriptor(Actor->GetActorGuid()))
+	if (TUniquePtr<FWorldPartitionActorDescInstance>* ExistingActorDescInstance = GetActorDescriptor(Actor->GetActorGuid()))
 	{
-		FWorldPartitionActorDescUtils::UpdateActorDescriptorFromActor(Actor, *ExistingActorDesc);
+		FEditorGrid2DUnsavedActorDescInstance* NewlyAddedActorDescInstance = StaticCast<FEditorGrid2DUnsavedActorDescInstance*>(ExistingActorDescInstance->Get());
+
+		FWorldPartitionActorDescUtils::UpdateActorDescriptorFromActor(Actor, NewlyAddedActorDescInstance->ActorDescPtr);
 	}
 	else
 	{
-		AddActor(Actor);
+		AddActorDescriptor(new FEditorGrid2DUnsavedActorDescInstance(Actor));
+	}
+}
+
+const FWorldPartitionActorDescInstance* SWorldPartitionEditorGrid2D::FNewlyAddedUnsavedActorDescsDescRegistry::GetActorDescInstance(const FGuid& InActorGuid) const
+{
+	if (const TUniquePtr<FWorldPartitionActorDescInstance>* ExistingActorDescInstance = GetActorDescriptor(InActorGuid))
+	{
+		return ExistingActorDescInstance->Get();
 	}
+
+	return nullptr;
 }
 
 #undef LOCTEXT_NAMESPACE
\ No newline at end of file
diff --git a/Engine/Source/Editor/WorldPartitionEditor/Public/WorldPartition/SWorldPartitionEditorGrid2D.h b/Engine/Source/Editor/WorldPartitionEditor/Public/WorldPartition/SWorldPartitionEditorGrid2D.h
--- a/Engine/Source/Editor/WorldPartitionEditor/Public/WorldPartition/SWorldPartitionEditorGrid2D.h
+++ b/Engine/Source/Editor/WorldPartitionEditor/Public/WorldPartition/SWorldPartitionEditorGrid2D.h
@@ -11,7 +11,7 @@
 #include "SWorldPartitionViewportWidget.h"
 #include "Widgets/Text/STextBlock.h"
 #include "SViewportToolBar.h"
-#include "WorldPartition/ActorDescList.h"
+#include "WorldPartition/ActorDescContainerInstance.h"
 #include "WorldPartition/WorldPartitionActorLoaderInterface.h"
 
 class SWorldPartitionEditorGrid2D : public SWorldPartitionEditorGrid
@@ -230,10 +230,12 @@ protected:
 	TSharedPtr<STextBlock> TextRulerWidget;
 
 	// List of actors without an actor descriptor
-	class FNewlyAddedUnsavedActorDescsDescRegistry : public FActorDescList
+	class FNewlyAddedUnsavedActorDescsDescRegistry : public FActorDescInstanceList
 	{
 	public:
 		void OnActorAdded(AActor* Actor);
+		
+		const FWorldPartitionActorDescInstance* GetActorDescInstance(const FGuid& InActorGuid) const;
 	};
 
 	FNewlyAddedUnsavedActorDescsDescRegistry NewlyAddedUnsavedActorDescs;
diff --git a/Engine/Source/Runtime/Engine/Private/ActorPartition/ActorPartitionSubsystem.cpp b/Engine/Source/Runtime/Engine/Private/ActorPartition/ActorPartitionSubsystem.cpp
--- a/Engine/Source/Runtime/Engine/Private/ActorPartition/ActorPartitionSubsystem.cpp
+++ b/Engine/Source/Runtime/Engine/Private/ActorPartition/ActorPartitionSubsystem.cpp
@@ -8,6 +8,7 @@
 #include "WorldPartition/DataLayer/DataLayerEditorContext.h"
 #include "WorldPartition/WorldPartitionSubsystem.h"
 #include "WorldPartition/WorldPartition.h"
+#include "WorldPartition/WorldPartitionActorDescInstance.h"
 #include "EngineUtils.h"
 
 #if WITH_EDITOR
@@ -210,19 +211,20 @@ public:
 		bool bUnloadedActorExists = false;
 		FGuid ContentBundleGuid = UContentBundleEngineSubsystem::Get()->GetEditingContentBundleGuid();
 
-		auto FindActor = [&FoundActor, &bUnloadedActorExists, &ContentBundleGuid, InCellCoord, InActorPartitionId, InGridSize, ThisWorld = World](const FWorldPartitionActorDesc* ActorDesc)
+		auto FindActor = [&FoundActor, &bUnloadedActorExists, &ContentBundleGuid, InCellCoord, InActorPartitionId, InGridSize, ThisWorld = World](const FWorldPartitionActorDescInstance* ActorDescInstance)
 		{
+			const FWorldPartitionActorDesc* ActorDesc = ActorDescInstance->GetActorDesc();
 			check(ActorDesc->GetActorNativeClass()->IsChildOf(InActorPartitionId.GetClass()));
-			FPartitionActorDesc* PartitionActorDesc = (FPartitionActorDesc*)ActorDesc;
+			const FPartitionActorDesc* PartitionActorDesc = (const FPartitionActorDesc*)ActorDesc;
 
 			if ((PartitionActorDesc->GridIndexX == InCellCoord.X) &&
 				(PartitionActorDesc->GridIndexY == InCellCoord.Y) &&
 				(PartitionActorDesc->GridIndexZ == InCellCoord.Z) &&
 				(PartitionActorDesc->GridSize == InGridSize) &&
 				(PartitionActorDesc->GridGuid == InActorPartitionId.GetGridGuid()) &&
-				(FActorPartitionContextHash::Get(PartitionActorDesc->GetContentBundleGuid(), ThisWorld, PartitionActorDesc->GetDataLayerInstanceNames()) == InActorPartitionId.GetContextHash()))
+				(FActorPartitionContextHash::Get(ActorDescInstance->GetContentBundleGuid(), ThisWorld, ActorDescInstance->GetDataLayerInstanceNames()) == InActorPartitionId.GetContextHash()))
 			{
-				AActor* DescActor = ActorDesc->GetActor();
+				AActor* DescActor = ActorDescInstance->GetActor();
 
 				if (!DescActor)
 				{
@@ -246,11 +248,11 @@ public:
 		FBox CellBounds = UActorPartitionSubsystem::FCellCoord::GetCellBounds(InCellCoord, InGridSize);
 		if (bInBoundsSearch)
 		{
-			FWorldPartitionHelpers::ForEachIntersectingActorDesc(WorldPartition, CellBounds, InActorPartitionId.GetClass(), FindActor);
+			FWorldPartitionHelpers::ForEachIntersectingActorDescInstance(WorldPartition, CellBounds, InActorPartitionId.GetClass(), FindActor);
 		}
 		else
 		{
-			FWorldPartitionHelpers::ForEachActorDesc(WorldPartition, InActorPartitionId.GetClass(), FindActor);
+			FWorldPartitionHelpers::ForEachActorDescInstance(WorldPartition, InActorPartitionId.GetClass(), FindActor);
 		}
 				
 		if (bUnloadedActorExists)
diff --git a/Engine/Source/Runtime/Engine/Private/Level.cpp b/Engine/Source/Runtime/Engine/Private/Level.cpp
--- a/Engine/Source/Runtime/Engine/Private/Level.cpp
+++ b/Engine/Source/Runtime/Engine/Private/Level.cpp
@@ -29,6 +29,7 @@ Level.cpp: Level-related functions
 #include "UObject/MetaData.h"
 #include "UObject/ObjectSaveContext.h"
 #include "UObject/UObjectIterator.h"
+#include "UObject/Package.h"
 #include "GameFramework/PlayerController.h"
 #include "GameFramework/WorldSettings.h"
 #include "Engine/BlueprintGeneratedClass.h"
@@ -69,6 +70,7 @@ Level.cpp: Level-related functions
 #include "WorldPartition/WorldPartition.h"
 #include "WorldPartition/WorldPartitionLog.h"
 #include "WorldPartition/WorldPartitionSubsystem.h"
+#include "WorldPartition/WorldPartitionActorDescInstance.h"
 #include "Engine/LevelStreaming.h"
 #include "LevelUtils.h"
 #include "Components/ModelComponent.h"
@@ -3874,10 +3876,10 @@ TSet<FGuid> ULevel::GetDeletedAndUnreferencedActorFolders() const
 
 	if (UWorldPartition* WorldPartition = GetWorldPartition())
 	{
-		FWorldPartitionHelpers::ForEachActorDesc(WorldPartition, [&FoldersToDelete](const FWorldPartitionActorDesc* ActorDesc)
+		FWorldPartitionHelpers::ForEachActorDescInstance(WorldPartition, [&FoldersToDelete](const FWorldPartitionActorDescInstance* ActorDescInstance)
 		{
-			AActor* Actor = ActorDesc->GetActor();
-			FoldersToDelete.Remove(Actor ? Actor->GetFolderGuid() : ActorDesc->GetFolderGuid());
+			AActor* Actor = ActorDescInstance->GetActor();
+			FoldersToDelete.Remove(Actor ? Actor->GetFolderGuid() : ActorDescInstance->GetFolderGuid());
 			return !FoldersToDelete.IsEmpty();
 		});
 	}
diff --git a/Engine/Source/Runtime/Engine/Private/LevelInstance/LevelInstanceEditorLevelStreaming.cpp b/Engine/Source/Runtime/Engine/Private/LevelInstance/LevelInstanceEditorLevelStreaming.cpp
--- a/Engine/Source/Runtime/Engine/Private/LevelInstance/LevelInstanceEditorLevelStreaming.cpp
+++ b/Engine/Source/Runtime/Engine/Private/LevelInstance/LevelInstanceEditorLevelStreaming.cpp
@@ -152,6 +152,20 @@ void ULevelStreamingLevelInstanceEditor::OnLevelActorAdded(AActor* InActor)
 	}
 }
 
+void ULevelStreamingLevelInstanceEditor::OnPreInitializeContainerInstance(UActorDescContainerInstance::FInitializeParams& InInitParams, UActorDescContainerInstance* InContainerInstance)
+{
+	// Apply override container
+	AActor* LevelInstanceActor = Cast<AActor>(GetLevelInstance());
+	if (UWorldPartition* OwningWorldPartition = LevelInstanceActor->GetLevel()->GetWorldPartition())
+	{
+		if (FWorldPartitionActorDescInstance* ActorDescInstance = OwningWorldPartition->GetActorDescInstance(LevelInstanceActor->GetActorGuid()); ActorDescInstance && ActorDescInstance->IsChildContainerInstance())
+		{
+			// Add parenting info to init param
+			InInitParams.SetParent(ActorDescInstance->GetContainerInstance(), ActorDescInstance->GetGuid());
+		}
+	}
+}
+
 void ULevelStreamingLevelInstanceEditor::OnLevelLoadedChanged(ULevel* InLevel)
 {
 	Super::OnLevelLoadedChanged(InLevel);
@@ -180,6 +194,7 @@ void ULevelStreamingLevelInstanceEditor::OnLevelLoadedChanged(ULevel* InLevel)
 			if (UWorldPartition* OuterWorldPartition = NewLoadedLevel->GetWorldPartition())
 			{
 				check(!OuterWorldPartition->IsInitialized());
+				OuterWorldPartition->OnActorDescContainerInstancePreInitialize.BindUObject(this, &ULevelStreamingLevelInstanceEditor::OnPreInitializeContainerInstance);
 				OuterWorldPartition->bOverrideEnableStreamingInEditor = false;
 			}
 
diff --git a/Engine/Source/Runtime/Engine/Private/LevelInstance/LevelInstanceLevelStreaming.cpp b/Engine/Source/Runtime/Engine/Private/LevelInstance/LevelInstanceLevelStreaming.cpp
--- a/Engine/Source/Runtime/Engine/Private/LevelInstance/LevelInstanceLevelStreaming.cpp
+++ b/Engine/Source/Runtime/Engine/Private/LevelInstance/LevelInstanceLevelStreaming.cpp
@@ -257,6 +257,20 @@ void ULevelStreamingLevelInstance::OnLoadedActorsRemovedFromLevelPostEvent(const
 	}
 }
 
+void ULevelStreamingLevelInstance::OnPreInitializeContainerInstance(UActorDescContainerInstance::FInitializeParams& InInitParams, UActorDescContainerInstance* InContainerInstance)
+{
+	// Apply override container
+	AActor* LevelInstanceActor = Cast<AActor>(GetLevelInstance());
+	if (UWorldPartition* OwningWorldPartition = LevelInstanceActor->GetLevel()->GetWorldPartition())
+	{
+		if (FWorldPartitionActorDescInstance* ActorDescInstance = OwningWorldPartition->GetActorDescInstance(LevelInstanceActor->GetActorGuid()); ActorDescInstance && ActorDescInstance->IsChildContainerInstance())
+		{
+			// Add parenting info to init param
+			InInitParams.SetParent(ActorDescInstance->GetContainerInstance(), ActorDescInstance->GetGuid());
+		}
+	}
+}
+
 #endif
 
 ULevelStreamingLevelInstance* ULevelStreamingLevelInstance::LoadInstance(ILevelInstanceInterface* LevelInstance)
@@ -433,6 +447,8 @@ void ULevelStreamingLevelInstance::OnLevelLoadedChanged(ULevel* InLevel)
 			if (UWorldPartition* OuterWorldPartition = NewLoadedLevel->GetWorldPartition())
 			{
 				check(!OuterWorldPartition->IsInitialized());
+				OuterWorldPartition->OnActorDescContainerInstancePreInitialize.BindUObject(this, &ULevelStreamingLevelInstance::OnPreInitializeContainerInstance);
+
 				if (UWorldPartition* OwningWorldPartition = GetWorld()->GetWorldPartition(); OwningWorldPartition && OwningWorldPartition->IsStreamingEnabled())
 				{
 					if (GDisableLevelInstanceEditorPartialLoading)
diff --git a/Engine/Source/Runtime/Engine/Private/LevelInstance/LevelInstanceSubsystem.cpp b/Engine/Source/Runtime/Engine/Private/LevelInstance/LevelInstanceSubsystem.cpp
--- a/Engine/Source/Runtime/Engine/Private/LevelInstance/LevelInstanceSubsystem.cpp
+++ b/Engine/Source/Runtime/Engine/Private/LevelInstance/LevelInstanceSubsystem.cpp
@@ -5,7 +5,7 @@
 #include "LevelInstance/LevelInstanceLevelStreaming.h"
 #include "Misc/StringFormatArg.h"
 #include "UObject/UObjectIterator.h"
-#include "WorldPartition/WorldPartitionSubsystem.h"
+#include "UObject/Package.h"
 #include "WorldPartition/DataLayer/WorldDataLayers.h"
 #include "WorldPartition/WorldPartitionLevelStreamingDynamic.h"
 #include "EngineUtils.h"
@@ -23,7 +23,9 @@
 #include "WorldPartition/LevelInstance/LevelInstanceActorDesc.h"
 #include "WorldPartition/WorldPartition.h"
 #include "WorldPartition/WorldPartitionSubsystem.h"
+#include "WorldPartition/ActorDescContainerSubsystem.h"
 #include "WorldPartition/WorldPartitionActorLoaderInterface.h"
+#include "WorldPartition/WorldPartitionActorDescInstance.h"
 #include "WorldPartition/DataLayer/DataLayerManager.h"
 #include "WorldPartition/DataLayer/DataLayerInstanceWithAsset.h"
 #include "WorldPartition/DataLayer/WorldDataLayersActorDesc.h"
@@ -833,9 +835,7 @@ bool ULevelInstanceSubsystem::GetLevelInstanceBounds(const ILevelInstanceInterfa
 
 		FString LevelPackage = LevelInstance->GetWorldAssetPackage();
 
-		UWorldPartitionSubsystem* WorldPartitionSubsystem = GetWorld()->GetSubsystem<UWorldPartitionSubsystem>();
-		check(WorldPartitionSubsystem);
-		if (FBox ContainerBounds = WorldPartitionSubsystem->GetContainerBounds(*LevelPackage); ContainerBounds.IsValid)
+		if (FBox ContainerBounds = UActorDescContainerSubsystem::GetChecked().GetContainerBounds(*LevelPackage); ContainerBounds.IsValid)
 		{
 			FTransform LevelInstancePivotOffsetTransform = FTransform(ULevel::GetLevelInstancePivotOffsetFromPackage(*LevelPackage));
 			FTransform LevelTransform = LevelInstancePivotOffsetTransform * CastChecked<AActor>(LevelInstance)->GetActorTransform();
@@ -2077,10 +2077,10 @@ bool ULevelInstanceSubsystem::EditLevelInstanceInternal(ILevelInstanceInterface*
 			TopLevelInstanceActor = CurrentTopLevelInstanceActor;
 		}
 
-		if (FWorldPartitionActorDesc* TopLevelInstanceActorActorDesc = WorldPartition->GetActorDesc(TopLevelInstanceActor->GetActorGuid()))
+		if (FWorldPartitionActorDescInstance* TopLevelInstanceActorDescInstance = WorldPartition->GetActorDescInstance(TopLevelInstanceActor->GetActorGuid()))
 		{
 			check(!CurrentEditLevelInstanceActor.IsValid());
-			CurrentEditLevelInstanceActor = FWorldPartitionReference(TopLevelInstanceActorActorDesc->GetContainer(), TopLevelInstanceActorActorDesc->GetGuid());
+			CurrentEditLevelInstanceActor = FWorldPartitionReference(TopLevelInstanceActorDescInstance->GetContainerInstance(), TopLevelInstanceActorDescInstance->GetGuid());
 		}
 	}
 
@@ -2215,9 +2215,7 @@ bool ULevelInstanceSubsystem::CommitLevelInstanceInternal(TUniquePtr<FLevelInsta
 	LevelInstance = GetLevelInstance(LevelInstanceID);
 
 	// Update Registered Container Bounds
-	UWorldPartitionSubsystem* WorldPartitionSubsystem = GetWorld()->GetSubsystem<UWorldPartitionSubsystem>();
-	check(WorldPartitionSubsystem);
-	WorldPartitionSubsystem->UpdateContainerBounds(*LevelInstance->GetWorldAssetPackage());
+	UActorDescContainerSubsystem::GetChecked().NotifyContainerUpdated(*LevelInstance->GetWorldAssetPackage());
 	
 	TArray<TPair<ULevelInstanceSubsystem*, FLevelInstanceID>> LevelInstancesToUpdate;
 	// Gather list to update
@@ -2433,9 +2431,9 @@ bool ULevelInstanceSubsystem::CheckForLoop(const ILevelInstanceInterface* LevelI
 	return CheckForLoop(LevelInstance, LevelInstance->GetWorldAsset(), LoopInfo, LoopStart);
 }
 
-bool ULevelInstanceSubsystem::PassLevelInstanceFilter(UWorld* World, const FWorldPartitionHandle& Actor) const
+bool ULevelInstanceSubsystem::PassLevelInstanceFilter(UWorld* World, const FWorldPartitionHandle& ActorHandle) const
 {
-	UWorld* ContainerOuterWorld = Actor->GetContainer()->GetWorldPartition()->GetTypedOuter<UWorld>();
+	UWorld* ContainerOuterWorld = ActorHandle.GetContainerInstance()->GetWorldPartition()->GetTypedOuter<UWorld>();
 	check(ContainerOuterWorld);
 	if (const ILevelInstanceInterface* TopAncestor = GetOwningLevelInstance(ContainerOuterWorld->PersistentLevel))
 	{
@@ -2449,7 +2447,7 @@ bool ULevelInstanceSubsystem::PassLevelInstanceFilter(UWorld* World, const FWorl
 		const TMap<FActorContainerID, TSet<FGuid>>& FilteredActors = TopAncestor->GetFilteredActorsPerContainer();
 		if (const TSet<FGuid>* FilteredActorsForContainer = FilteredActors.Find(ContainerID))
 		{
-			if (FilteredActorsForContainer->Contains(Actor->GetGuid()))
+			if (FilteredActorsForContainer->Contains(ActorHandle.GetInstance()->GetGuid()))
 			{
 				return false;
 			}
diff --git a/Engine/Source/Runtime/Engine/Private/Tests/WorldPartition/WorldPartitionHandleTests.cpp b/Engine/Source/Runtime/Engine/Private/Tests/WorldPartition/WorldPartitionHandleTests.cpp
--- a/Engine/Source/Runtime/Engine/Private/Tests/WorldPartition/WorldPartitionHandleTests.cpp
+++ b/Engine/Source/Runtime/Engine/Private/Tests/WorldPartition/WorldPartitionHandleTests.cpp
@@ -7,6 +7,7 @@
 #include "WorldPartition/WorldPartition.h"
 #include "WorldPartition/WorldPartitionHelpers.h"
 #include "WorldPartition/ActorDescContainer.h"
+#include "WorldPartition/WorldPartitionActorDescInstance.h"
 #include "EditorWorldUtils.h"
 #include "PackageTools.h"
 #endif
@@ -18,14 +19,14 @@
 #if WITH_EDITOR
 struct FWorldPartitionActorDescUnitTestAcccessor
 {
-	inline static uint32 GetSoftRefCount(const FWorldPartitionActorDesc* ActorDesc)
+	inline static uint32 GetSoftRefCount(const FWorldPartitionActorDescInstance* ActorDescInstance)
 	{
-		return ActorDesc->GetSoftRefCount();
+		return ActorDescInstance->GetSoftRefCount();
 	}
 
-	inline static uint32 GetHardRefCount(const FWorldPartitionActorDesc* ActorDesc)
+	inline static uint32 GetHardRefCount(const FWorldPartitionActorDescInstance* ActorDescInstance)
 	{
-		return ActorDesc->GetHardRefCount();
+		return ActorDescInstance->GetHardRefCount();
 	}
 };
 #endif
@@ -62,84 +63,84 @@ namespace WorldPartitionTests
 
 		TestTrue(TEXT("World type"), World->WorldType == EWorldType::None);
 
-		UActorDescContainer* ActorDescContainer = NewObject<UActorDescContainer>(GetTransientPackage());
-		ActorDescContainer->Initialize({ nullptr, TEXT("/Engine/WorldPartition/WorldPartitionUnitTest") });
+		UActorDescContainerInstance* ActorDescContainerInstance = NewObject<UActorDescContainerInstance>(GetTransientPackage());
+		ActorDescContainerInstance->Initialize({ TEXT("/Engine/WorldPartition/WorldPartitionUnitTest") });
 
-		FWorldPartitionHandle Handle = FWorldPartitionHandle(ActorDescContainer, FGuid(TEXT("5D9F93BA407A811AFDDDAAB4F1CECC6A")));
+		FWorldPartitionHandle Handle = FWorldPartitionHandle(ActorDescContainerInstance, FGuid(TEXT("5D9F93BA407A811AFDDDAAB4F1CECC6A")));
 
 		{
-			FWorldPartitionHandle ActorWithChildActorComponentHandle = FWorldPartitionHandle(ActorDescContainer, FGuid(TEXT("538856174ECD465948488F9441AE9251")));
+			FWorldPartitionHandle ActorWithChildActorComponentHandle = FWorldPartitionHandle(ActorDescContainerInstance, FGuid(TEXT("538856174ECD465948488F9441AE9251")));
 			FWorldPartitionReference ActorWithChildActorComponentReference = ActorWithChildActorComponentHandle.ToReference();
 		}		
 
 		FWorldPartitionReference Reference;
 		{
 			LoadingContextType LoadingContext;
-			Reference = FWorldPartitionReference(ActorDescContainer, FGuid(TEXT("0D2B04D240BE5DE58FE437A8D2DBF5C9")));
+			Reference = FWorldPartitionReference(ActorDescContainerInstance, FGuid(TEXT("0D2B04D240BE5DE58FE437A8D2DBF5C9")));
 		}
 
-		TestTrue(TEXT("Handle container"), Handle->GetContainer() == ActorDescContainer);
-		TestTrue(TEXT("Reference container"), Reference->GetContainer() == ActorDescContainer);
+		TestTrue(TEXT("Handle container"), Handle.GetContainerInstance() == ActorDescContainerInstance);
+		TestTrue(TEXT("Reference container"), Reference.GetContainerInstance() == ActorDescContainerInstance);
 
-		TestTrue(TEXT("Handle soft refcount"), FWorldPartitionActorDescUnitTestAcccessor::GetSoftRefCount(Handle.Get()) == 1);
-		TestTrue(TEXT("Handle hard refcount"), FWorldPartitionActorDescUnitTestAcccessor::GetHardRefCount(Handle.Get()) == 0);
-		TestTrue(TEXT("Reference soft refcount"), FWorldPartitionActorDescUnitTestAcccessor::GetSoftRefCount(Reference.Get()) == 0);
-		TestTrue(TEXT("Reference hard refcount"), FWorldPartitionActorDescUnitTestAcccessor::GetHardRefCount(Reference.Get()) == 1);
+		TestTrue(TEXT("Handle soft refcount"), FWorldPartitionActorDescUnitTestAcccessor::GetSoftRefCount(Handle.GetInstance()) == 1);
+		TestTrue(TEXT("Handle hard refcount"), FWorldPartitionActorDescUnitTestAcccessor::GetHardRefCount(Handle.GetInstance()) == 0);
+		TestTrue(TEXT("Reference soft refcount"), FWorldPartitionActorDescUnitTestAcccessor::GetSoftRefCount(Reference.GetInstance()) == 0);
+		TestTrue(TEXT("Reference hard refcount"), FWorldPartitionActorDescUnitTestAcccessor::GetHardRefCount(Reference.GetInstance()) == 1);
 
 		// Pin handle scope test
 		{
 			FWorldPartitionHandlePinRefScope PinRefScopeHandle(Handle);
-			TestTrue(TEXT("Pin to Handle soft refcount"), FWorldPartitionActorDescUnitTestAcccessor::GetSoftRefCount(Handle.Get()) == 1);
-			TestTrue(TEXT("Pin to Handle hard refcount"), FWorldPartitionActorDescUnitTestAcccessor::GetHardRefCount(Handle.Get()) == 0);
+			TestTrue(TEXT("Pin to Handle soft refcount"), FWorldPartitionActorDescUnitTestAcccessor::GetSoftRefCount(Handle.GetInstance()) == 1);
+			TestTrue(TEXT("Pin to Handle hard refcount"), FWorldPartitionActorDescUnitTestAcccessor::GetHardRefCount(Handle.GetInstance()) == 0);
 		}
 
 		// Pin reference scope test
 		{
 			FWorldPartitionHandlePinRefScope PinRefScopeReference(Reference);
-			TestTrue(TEXT("Pin to Reference soft refcount"), FWorldPartitionActorDescUnitTestAcccessor::GetSoftRefCount(Reference.Get()) == 0);
-			TestTrue(TEXT("Pin to Reference hard refcount"), FWorldPartitionActorDescUnitTestAcccessor::GetHardRefCount(Reference.Get()) == 2);
+			TestTrue(TEXT("Pin to Reference soft refcount"), FWorldPartitionActorDescUnitTestAcccessor::GetSoftRefCount(Reference.GetInstance()) == 0);
+			TestTrue(TEXT("Pin to Reference hard refcount"), FWorldPartitionActorDescUnitTestAcccessor::GetHardRefCount(Reference.GetInstance()) == 2);
 		}
 
 		TestTrue(TEXT("Handle/Reference equality"), Handle != Reference);
 		TestTrue(TEXT("Reference/Handle equality"), Reference != Handle);
 		
-		TestTrue(TEXT("Handle soft refcount"), FWorldPartitionActorDescUnitTestAcccessor::GetSoftRefCount(Handle.Get()) == 1);
-		TestTrue(TEXT("Handle hard refcount"), FWorldPartitionActorDescUnitTestAcccessor::GetHardRefCount(Handle.Get()) == 0);
-		TestTrue(TEXT("Reference soft refcount"), FWorldPartitionActorDescUnitTestAcccessor::GetSoftRefCount(Reference.Get()) == 0);
-		TestTrue(TEXT("Reference hard refcount"), FWorldPartitionActorDescUnitTestAcccessor::GetHardRefCount(Reference.Get()) == 1);
+		TestTrue(TEXT("Handle soft refcount"), FWorldPartitionActorDescUnitTestAcccessor::GetSoftRefCount(Handle.GetInstance()) == 1);
+		TestTrue(TEXT("Handle hard refcount"), FWorldPartitionActorDescUnitTestAcccessor::GetHardRefCount(Handle.GetInstance()) == 0);
+		TestTrue(TEXT("Reference soft refcount"), FWorldPartitionActorDescUnitTestAcccessor::GetSoftRefCount(Reference.GetInstance()) == 0);
+		TestTrue(TEXT("Reference hard refcount"), FWorldPartitionActorDescUnitTestAcccessor::GetHardRefCount(Reference.GetInstance()) == 1);
 
 		// Conversions
 		{
 			FWorldPartitionHandle HandleToReference = Reference.ToHandle();
 			TestTrue(TEXT("Handle/Reference equality"), HandleToReference == Reference);
 			TestTrue(TEXT("Reference/Handle equality"), Reference == HandleToReference);
-			TestTrue(TEXT("Reference soft refcount"), FWorldPartitionActorDescUnitTestAcccessor::GetSoftRefCount(Reference.Get()) == 1);
-			TestTrue(TEXT("Reference hard refcount"), FWorldPartitionActorDescUnitTestAcccessor::GetHardRefCount(Reference.Get()) == 1);
+			TestTrue(TEXT("Reference soft refcount"), FWorldPartitionActorDescUnitTestAcccessor::GetSoftRefCount(Reference.GetInstance()) == 1);
+			TestTrue(TEXT("Reference hard refcount"), FWorldPartitionActorDescUnitTestAcccessor::GetHardRefCount(Reference.GetInstance()) == 1);
 
 			FWorldPartitionReference ReferenceToHandle = Handle.ToReference();
 			TestTrue(TEXT("Handle/Reference equality"), ReferenceToHandle == Handle);
 			TestTrue(TEXT("Handle/Reference equality"), Handle == ReferenceToHandle);
-			TestTrue(TEXT("Reference soft refcount"), FWorldPartitionActorDescUnitTestAcccessor::GetSoftRefCount(Reference.Get()) == 1);
-			TestTrue(TEXT("Reference hard refcount"), FWorldPartitionActorDescUnitTestAcccessor::GetHardRefCount(Reference.Get()) == 1);
+			TestTrue(TEXT("Reference soft refcount"), FWorldPartitionActorDescUnitTestAcccessor::GetSoftRefCount(Reference.GetInstance()) == 1);
+			TestTrue(TEXT("Reference hard refcount"), FWorldPartitionActorDescUnitTestAcccessor::GetHardRefCount(Reference.GetInstance()) == 1);
 		}
 
-		TestTrue(TEXT("Handle soft refcount"), FWorldPartitionActorDescUnitTestAcccessor::GetSoftRefCount(Handle.Get()) == 1);
-		TestTrue(TEXT("Handle hard refcount"), FWorldPartitionActorDescUnitTestAcccessor::GetHardRefCount(Handle.Get()) == 0);
-		TestTrue(TEXT("Reference soft refcount"), FWorldPartitionActorDescUnitTestAcccessor::GetSoftRefCount(Reference.Get()) == 0);
-		TestTrue(TEXT("Reference hard refcount"), FWorldPartitionActorDescUnitTestAcccessor::GetHardRefCount(Reference.Get()) == 1);
+		TestTrue(TEXT("Handle soft refcount"), FWorldPartitionActorDescUnitTestAcccessor::GetSoftRefCount(Handle.GetInstance()) == 1);
+		TestTrue(TEXT("Handle hard refcount"), FWorldPartitionActorDescUnitTestAcccessor::GetHardRefCount(Handle.GetInstance()) == 0);
+		TestTrue(TEXT("Reference soft refcount"), FWorldPartitionActorDescUnitTestAcccessor::GetSoftRefCount(Reference.GetInstance()) == 0);
+		TestTrue(TEXT("Reference hard refcount"), FWorldPartitionActorDescUnitTestAcccessor::GetHardRefCount(Reference.GetInstance()) == 1);
 
 		// inplace new test
 		{
 			uint8 Buffer[sizeof(FWorldPartitionHandle)];
 			FWorldPartitionHandle* HandlePtr = new (Buffer) FWorldPartitionHandle(Reference.ToHandle());
 
-			TestTrue(TEXT("Handle array soft refcount"), FWorldPartitionActorDescUnitTestAcccessor::GetSoftRefCount(Reference.Get()) == 1);
-			TestTrue(TEXT("Handle array hard refcount"), FWorldPartitionActorDescUnitTestAcccessor::GetHardRefCount(Reference.Get()) == 1);
+			TestTrue(TEXT("Handle array soft refcount"), FWorldPartitionActorDescUnitTestAcccessor::GetSoftRefCount(Reference.GetInstance()) == 1);
+			TestTrue(TEXT("Handle array hard refcount"), FWorldPartitionActorDescUnitTestAcccessor::GetHardRefCount(Reference.GetInstance()) == 1);
 
 			HandlePtr->~FWorldPartitionHandle();
 
-			TestTrue(TEXT("Handle array soft refcount"), FWorldPartitionActorDescUnitTestAcccessor::GetSoftRefCount(Reference.Get()) == 0);
-			TestTrue(TEXT("Handle array hard refcount"), FWorldPartitionActorDescUnitTestAcccessor::GetHardRefCount(Reference.Get()) == 1);
+			TestTrue(TEXT("Handle array soft refcount"), FWorldPartitionActorDescUnitTestAcccessor::GetSoftRefCount(Reference.GetInstance()) == 0);
+			TestTrue(TEXT("Handle array hard refcount"), FWorldPartitionActorDescUnitTestAcccessor::GetHardRefCount(Reference.GetInstance()) == 1);
 		}
 
 		// TArray test
@@ -147,14 +148,14 @@ namespace WorldPartitionTests
 			TArray<FWorldPartitionHandle> HandleList;
 			HandleList.Add(Handle);
 
-			TestTrue(TEXT("Handle array soft refcount"), FWorldPartitionActorDescUnitTestAcccessor::GetSoftRefCount(Handle.Get()) == 2);
-			TestTrue(TEXT("Handle array hard refcount"), FWorldPartitionActorDescUnitTestAcccessor::GetHardRefCount(Handle.Get()) == 0);
+			TestTrue(TEXT("Handle array soft refcount"), FWorldPartitionActorDescUnitTestAcccessor::GetSoftRefCount(Handle.GetInstance()) == 2);
+			TestTrue(TEXT("Handle array hard refcount"), FWorldPartitionActorDescUnitTestAcccessor::GetHardRefCount(Handle.GetInstance()) == 0);
 
 			FWorldPartitionReference ReferenceToHandle = Handle.ToReference();
 			TestTrue(TEXT("Handle/Reference equality"), ReferenceToHandle == Handle);
 			TestTrue(TEXT("Handle/Reference equality"), Handle == ReferenceToHandle);
-			TestTrue(TEXT("Handle soft refcount"), FWorldPartitionActorDescUnitTestAcccessor::GetSoftRefCount(Handle.Get()) == 2);
-			TestTrue(TEXT("Handle hard refcount"), FWorldPartitionActorDescUnitTestAcccessor::GetHardRefCount(Handle.Get()) == 1);
+			TestTrue(TEXT("Handle soft refcount"), FWorldPartitionActorDescUnitTestAcccessor::GetSoftRefCount(Handle.GetInstance()) == 2);
+			TestTrue(TEXT("Handle hard refcount"), FWorldPartitionActorDescUnitTestAcccessor::GetHardRefCount(Handle.GetInstance()) == 1);
 
 			TestTrue(TEXT("Handle array contains handle"), HandleList.Contains(Handle));
 			TestTrue(TEXT("Handle array contains reference"), HandleList.Contains(ReferenceToHandle));
@@ -162,36 +163,36 @@ namespace WorldPartitionTests
 			HandleList.Add(Reference.ToHandle());
 			
 			TestTrue(TEXT("Handle array contains reference"), HandleList.Contains(Reference));
-			TestTrue(TEXT("Handle array soft refcount"), FWorldPartitionActorDescUnitTestAcccessor::GetSoftRefCount(Reference.Get()) == 1);
-			TestTrue(TEXT("Handle array hard refcount"), FWorldPartitionActorDescUnitTestAcccessor::GetHardRefCount(Reference.Get()) == 1);
+			TestTrue(TEXT("Handle array soft refcount"), FWorldPartitionActorDescUnitTestAcccessor::GetSoftRefCount(Reference.GetInstance()) == 1);
+			TestTrue(TEXT("Handle array hard refcount"), FWorldPartitionActorDescUnitTestAcccessor::GetHardRefCount(Reference.GetInstance()) == 1);
 			
 			HandleList.Remove(Handle);
-			TestTrue(TEXT("Handle array soft refcount"), FWorldPartitionActorDescUnitTestAcccessor::GetSoftRefCount(Handle.Get()) == 1);
-			TestTrue(TEXT("Handle array hard refcount"), FWorldPartitionActorDescUnitTestAcccessor::GetHardRefCount(Handle.Get()) == 1);
+			TestTrue(TEXT("Handle array soft refcount"), FWorldPartitionActorDescUnitTestAcccessor::GetSoftRefCount(Handle.GetInstance()) == 1);
+			TestTrue(TEXT("Handle array hard refcount"), FWorldPartitionActorDescUnitTestAcccessor::GetHardRefCount(Handle.GetInstance()) == 1);
 
 			HandleList.Remove(Reference.ToHandle());
-			TestTrue(TEXT("Handle array soft refcount"), FWorldPartitionActorDescUnitTestAcccessor::GetSoftRefCount(Reference.Get()) == 0);
-			TestTrue(TEXT("Handle array hard refcount"), FWorldPartitionActorDescUnitTestAcccessor::GetHardRefCount(Reference.Get()) == 1);
+			TestTrue(TEXT("Handle array soft refcount"), FWorldPartitionActorDescUnitTestAcccessor::GetSoftRefCount(Reference.GetInstance()) == 0);
+			TestTrue(TEXT("Handle array hard refcount"), FWorldPartitionActorDescUnitTestAcccessor::GetHardRefCount(Reference.GetInstance()) == 1);
 		}
 
-		TestTrue(TEXT("Handle soft refcount"), FWorldPartitionActorDescUnitTestAcccessor::GetSoftRefCount(Handle.Get()) == 1);
-		TestTrue(TEXT("Handle hard refcount"), FWorldPartitionActorDescUnitTestAcccessor::GetHardRefCount(Handle.Get()) == 0);
-		TestTrue(TEXT("Reference soft refcount"), FWorldPartitionActorDescUnitTestAcccessor::GetSoftRefCount(Reference.Get()) == 0);
-		TestTrue(TEXT("Reference hard refcount"), FWorldPartitionActorDescUnitTestAcccessor::GetHardRefCount(Reference.Get()) == 1);
+		TestTrue(TEXT("Handle soft refcount"), FWorldPartitionActorDescUnitTestAcccessor::GetSoftRefCount(Handle.GetInstance()) == 1);
+		TestTrue(TEXT("Handle hard refcount"), FWorldPartitionActorDescUnitTestAcccessor::GetHardRefCount(Handle.GetInstance()) == 0);
+		TestTrue(TEXT("Reference soft refcount"), FWorldPartitionActorDescUnitTestAcccessor::GetSoftRefCount(Reference.GetInstance()) == 0);
+		TestTrue(TEXT("Reference hard refcount"), FWorldPartitionActorDescUnitTestAcccessor::GetHardRefCount(Reference.GetInstance()) == 1);
 
 		// TSet test
 		{
 			TSet<FWorldPartitionHandle> HandleSet;
 			HandleSet.Add(Handle);
 
-			TestTrue(TEXT("Handle set soft refcount"), FWorldPartitionActorDescUnitTestAcccessor::GetSoftRefCount(Handle.Get()) == 2);
-			TestTrue(TEXT("Handle set hard refcount"), FWorldPartitionActorDescUnitTestAcccessor::GetHardRefCount(Handle.Get()) == 0);
+			TestTrue(TEXT("Handle set soft refcount"), FWorldPartitionActorDescUnitTestAcccessor::GetSoftRefCount(Handle.GetInstance()) == 2);
+			TestTrue(TEXT("Handle set hard refcount"), FWorldPartitionActorDescUnitTestAcccessor::GetHardRefCount(Handle.GetInstance()) == 0);
 
 			FWorldPartitionReference ReferenceToHandle = Handle.ToReference();
 			TestTrue(TEXT("Handle/Reference equality"), ReferenceToHandle == Handle);
 			TestTrue(TEXT("Handle/Reference equality"), Handle == ReferenceToHandle);
-			TestTrue(TEXT("Handle soft refcount"), FWorldPartitionActorDescUnitTestAcccessor::GetSoftRefCount(Handle.Get()) == 2);
-			TestTrue(TEXT("Handle hard refcount"), FWorldPartitionActorDescUnitTestAcccessor::GetHardRefCount(Handle.Get()) == 1);
+			TestTrue(TEXT("Handle soft refcount"), FWorldPartitionActorDescUnitTestAcccessor::GetSoftRefCount(Handle.GetInstance()) == 2);
+			TestTrue(TEXT("Handle hard refcount"), FWorldPartitionActorDescUnitTestAcccessor::GetHardRefCount(Handle.GetInstance()) == 1);
 
 			TestTrue(TEXT("Handle set contains handle"), HandleSet.Contains(Handle));
 			TestTrue(TEXT("Handle set contains reference"), HandleSet.Contains(ReferenceToHandle.ToHandle()));
@@ -199,49 +200,49 @@ namespace WorldPartitionTests
 			HandleSet.Add(Reference.ToHandle());
 			TestTrue(TEXT("Handle set contains reference"), HandleSet.Contains(Reference.ToHandle()));
 
-			TestTrue(TEXT("Reference soft refcount"), FWorldPartitionActorDescUnitTestAcccessor::GetSoftRefCount(Reference.Get()) == 1);
-			TestTrue(TEXT("Reference hard refcount"), FWorldPartitionActorDescUnitTestAcccessor::GetHardRefCount(Reference.Get()) == 1);
+			TestTrue(TEXT("Reference soft refcount"), FWorldPartitionActorDescUnitTestAcccessor::GetSoftRefCount(Reference.GetInstance()) == 1);
+			TestTrue(TEXT("Reference hard refcount"), FWorldPartitionActorDescUnitTestAcccessor::GetHardRefCount(Reference.GetInstance()) == 1);
 		}
 
-		TestTrue(TEXT("Handle soft refcount"), FWorldPartitionActorDescUnitTestAcccessor::GetSoftRefCount(Handle.Get()) == 1);
-		TestTrue(TEXT("Handle hard refcount"), FWorldPartitionActorDescUnitTestAcccessor::GetHardRefCount(Handle.Get()) == 0);
-		TestTrue(TEXT("Reference soft refcount"), FWorldPartitionActorDescUnitTestAcccessor::GetSoftRefCount(Reference.Get()) == 0);
-		TestTrue(TEXT("Reference hard refcount"), FWorldPartitionActorDescUnitTestAcccessor::GetHardRefCount(Reference.Get()) == 1);
+		TestTrue(TEXT("Handle soft refcount"), FWorldPartitionActorDescUnitTestAcccessor::GetSoftRefCount(Handle.GetInstance()) == 1);
+		TestTrue(TEXT("Handle hard refcount"), FWorldPartitionActorDescUnitTestAcccessor::GetHardRefCount(Handle.GetInstance()) == 0);
+		TestTrue(TEXT("Reference soft refcount"), FWorldPartitionActorDescUnitTestAcccessor::GetSoftRefCount(Reference.GetInstance()) == 0);
+		TestTrue(TEXT("Reference hard refcount"), FWorldPartitionActorDescUnitTestAcccessor::GetHardRefCount(Reference.GetInstance()) == 1);
 
 		// Move tests
 		{
 			// Handle move
-			TestTrue(TEXT("Handle soft refcount"), FWorldPartitionActorDescUnitTestAcccessor::GetSoftRefCount(Handle.Get()) == 1);
-			TestTrue(TEXT("Handle hard refcount"), FWorldPartitionActorDescUnitTestAcccessor::GetHardRefCount(Handle.Get()) == 0);
+			TestTrue(TEXT("Handle soft refcount"), FWorldPartitionActorDescUnitTestAcccessor::GetSoftRefCount(Handle.GetInstance()) == 1);
+			TestTrue(TEXT("Handle hard refcount"), FWorldPartitionActorDescUnitTestAcccessor::GetHardRefCount(Handle.GetInstance()) == 0);
 			{
 				FWorldPartitionHandle HandleCopy(MoveTemp(Handle));
 				TestTrue(TEXT("Handle move src not valid"), !Handle.IsValid());
 				TestTrue(TEXT("Handle move dst valid"), HandleCopy.IsValid());
-				TestTrue(TEXT("Handle soft refcount"), FWorldPartitionActorDescUnitTestAcccessor::GetSoftRefCount(HandleCopy.Get()) == 1);
-				TestTrue(TEXT("Handle hard refcount"), FWorldPartitionActorDescUnitTestAcccessor::GetHardRefCount(HandleCopy.Get()) == 0);
+				TestTrue(TEXT("Handle soft refcount"), FWorldPartitionActorDescUnitTestAcccessor::GetSoftRefCount(HandleCopy.GetInstance()) == 1);
+				TestTrue(TEXT("Handle hard refcount"), FWorldPartitionActorDescUnitTestAcccessor::GetHardRefCount(HandleCopy.GetInstance()) == 0);
 
 				Handle = MoveTemp(HandleCopy);
 				TestTrue(TEXT("Handle move src not valid"), !HandleCopy.IsValid());
 				TestTrue(TEXT("Handle move dst valid"), Handle.IsValid());
-				TestTrue(TEXT("Handle soft refcount"), FWorldPartitionActorDescUnitTestAcccessor::GetSoftRefCount(Handle.Get()) == 1);
-				TestTrue(TEXT("Handle hard refcount"), FWorldPartitionActorDescUnitTestAcccessor::GetHardRefCount(Handle.Get()) == 0);
+				TestTrue(TEXT("Handle soft refcount"), FWorldPartitionActorDescUnitTestAcccessor::GetSoftRefCount(Handle.GetInstance()) == 1);
+				TestTrue(TEXT("Handle hard refcount"), FWorldPartitionActorDescUnitTestAcccessor::GetHardRefCount(Handle.GetInstance()) == 0);
 			}
 
 			// Reference move
-			TestTrue(TEXT("Reference soft refcount"), FWorldPartitionActorDescUnitTestAcccessor::GetSoftRefCount(Reference.Get()) == 0);
-			TestTrue(TEXT("Reference hard refcount"), FWorldPartitionActorDescUnitTestAcccessor::GetHardRefCount(Reference.Get()) == 1);
+			TestTrue(TEXT("Reference soft refcount"), FWorldPartitionActorDescUnitTestAcccessor::GetSoftRefCount(Reference.GetInstance()) == 0);
+			TestTrue(TEXT("Reference hard refcount"), FWorldPartitionActorDescUnitTestAcccessor::GetHardRefCount(Reference.GetInstance()) == 1);
 			{
 				FWorldPartitionReference ReferenceCopy(MoveTemp(Reference));
 				TestTrue(TEXT("Reference move src not valid"), !Reference.IsValid());
 				TestTrue(TEXT("Reference move dst valid"), ReferenceCopy.IsValid());
-				TestTrue(TEXT("Reference soft refcount"), FWorldPartitionActorDescUnitTestAcccessor::GetSoftRefCount(ReferenceCopy.Get()) == 0);
-				TestTrue(TEXT("Reference hard refcount"), FWorldPartitionActorDescUnitTestAcccessor::GetHardRefCount(ReferenceCopy.Get()) == 1);
+				TestTrue(TEXT("Reference soft refcount"), FWorldPartitionActorDescUnitTestAcccessor::GetSoftRefCount(ReferenceCopy.GetInstance()) == 0);
+				TestTrue(TEXT("Reference hard refcount"), FWorldPartitionActorDescUnitTestAcccessor::GetHardRefCount(ReferenceCopy.GetInstance()) == 1);
 
 				Reference = MoveTemp(ReferenceCopy);
 				TestTrue(TEXT("Reference move src not valid"), !ReferenceCopy.IsValid());
 				TestTrue(TEXT("Reference move dst valid"), Reference.IsValid());
-				TestTrue(TEXT("Reference soft refcount"), FWorldPartitionActorDescUnitTestAcccessor::GetSoftRefCount(Reference.Get()) == 0);
-				TestTrue(TEXT("Reference hard refcount"), FWorldPartitionActorDescUnitTestAcccessor::GetHardRefCount(Reference.Get()) == 1);
+				TestTrue(TEXT("Reference soft refcount"), FWorldPartitionActorDescUnitTestAcccessor::GetSoftRefCount(Reference.GetInstance()) == 0);
+				TestTrue(TEXT("Reference hard refcount"), FWorldPartitionActorDescUnitTestAcccessor::GetHardRefCount(Reference.GetInstance()) == 1);
 			}
 
 			// Handle reference move
@@ -249,14 +250,14 @@ namespace WorldPartitionTests
 				FWorldPartitionHandle HandleFromReference(MoveTemp(Reference));
 				TestTrue(TEXT("Handle move src not valid"), !Reference.IsValid());
 				TestTrue(TEXT("Handle move dst valid"), HandleFromReference.IsValid());
-				TestTrue(TEXT("Handle soft refcount"), FWorldPartitionActorDescUnitTestAcccessor::GetSoftRefCount(HandleFromReference.Get()) == 1);
-				TestTrue(TEXT("Handle hard refcount"), FWorldPartitionActorDescUnitTestAcccessor::GetHardRefCount(HandleFromReference.Get()) == 0);
+				TestTrue(TEXT("Handle soft refcount"), FWorldPartitionActorDescUnitTestAcccessor::GetSoftRefCount(HandleFromReference.GetInstance()) == 1);
+				TestTrue(TEXT("Handle hard refcount"), FWorldPartitionActorDescUnitTestAcccessor::GetHardRefCount(HandleFromReference.GetInstance()) == 0);
 
 				Reference = MoveTemp(HandleFromReference);
 				TestTrue(TEXT("Handle move src not valid"), !HandleFromReference.IsValid());
 				TestTrue(TEXT("Reference move dst valid"), Reference.IsValid());
-				TestTrue(TEXT("Reference soft refcount"), FWorldPartitionActorDescUnitTestAcccessor::GetSoftRefCount(Reference.Get()) == 0);
-				TestTrue(TEXT("Reference hard refcount"), FWorldPartitionActorDescUnitTestAcccessor::GetHardRefCount(Reference.Get()) == 1);
+				TestTrue(TEXT("Reference soft refcount"), FWorldPartitionActorDescUnitTestAcccessor::GetSoftRefCount(Reference.GetInstance()) == 0);
+				TestTrue(TEXT("Reference hard refcount"), FWorldPartitionActorDescUnitTestAcccessor::GetHardRefCount(Reference.GetInstance()) == 1);
 			}
 
 			// Reference handle move
@@ -264,27 +265,27 @@ namespace WorldPartitionTests
 				FWorldPartitionReference ReferenceFromHandle(MoveTemp(Handle));
 				TestTrue(TEXT("Handle move src not valid"), !Handle.IsValid());
 				TestTrue(TEXT("Reference move dst valid"), ReferenceFromHandle.IsValid());
-				TestTrue(TEXT("Reference soft refcount"), FWorldPartitionActorDescUnitTestAcccessor::GetSoftRefCount(ReferenceFromHandle.Get()) == 0);
-				TestTrue(TEXT("Reference hard refcount"), FWorldPartitionActorDescUnitTestAcccessor::GetHardRefCount(ReferenceFromHandle.Get()) == 1);
+				TestTrue(TEXT("Reference soft refcount"), FWorldPartitionActorDescUnitTestAcccessor::GetSoftRefCount(ReferenceFromHandle.GetInstance()) == 0);
+				TestTrue(TEXT("Reference hard refcount"), FWorldPartitionActorDescUnitTestAcccessor::GetHardRefCount(ReferenceFromHandle.GetInstance()) == 1);
 
 				Handle = MoveTemp(ReferenceFromHandle);
 				TestTrue(TEXT("Reference move src not valid"), !ReferenceFromHandle.IsValid());
 				TestTrue(TEXT("Handle move dst valid"), Handle.IsValid());
-				TestTrue(TEXT("Handle soft refcount"), FWorldPartitionActorDescUnitTestAcccessor::GetSoftRefCount(Handle.Get()) == 1);
-				TestTrue(TEXT("Handle hard refcount"), FWorldPartitionActorDescUnitTestAcccessor::GetHardRefCount(Handle.Get()) == 0);
+				TestTrue(TEXT("Handle soft refcount"), FWorldPartitionActorDescUnitTestAcccessor::GetSoftRefCount(Handle.GetInstance()) == 1);
+				TestTrue(TEXT("Handle hard refcount"), FWorldPartitionActorDescUnitTestAcccessor::GetHardRefCount(Handle.GetInstance()) == 0);
 			}
 
 			// Reset
 			{
 				FWorldPartitionReference ReferenceFromHandle(Handle.ToReference());
 				TestTrue(TEXT("Reference valid"), ReferenceFromHandle.IsValid());
-				TestTrue(TEXT("Reference soft refcount"), FWorldPartitionActorDescUnitTestAcccessor::GetSoftRefCount(ReferenceFromHandle.Get()) == 1);
-				TestTrue(TEXT("Reference hard refcount"), FWorldPartitionActorDescUnitTestAcccessor::GetHardRefCount(ReferenceFromHandle.Get()) == 1);
+				TestTrue(TEXT("Reference soft refcount"), FWorldPartitionActorDescUnitTestAcccessor::GetSoftRefCount(ReferenceFromHandle.GetInstance()) == 1);
+				TestTrue(TEXT("Reference hard refcount"), FWorldPartitionActorDescUnitTestAcccessor::GetHardRefCount(ReferenceFromHandle.GetInstance()) == 1);
 
 				ReferenceFromHandle.Reset();
 				TestTrue(TEXT("Reference not valid"), !ReferenceFromHandle.IsValid());
-				TestTrue(TEXT("Handle soft refcount"), FWorldPartitionActorDescUnitTestAcccessor::GetSoftRefCount(Handle.Get()) == 1);
-				TestTrue(TEXT("Handle hard refcount"), FWorldPartitionActorDescUnitTestAcccessor::GetHardRefCount(Handle.Get()) == 0);
+				TestTrue(TEXT("Handle soft refcount"), FWorldPartitionActorDescUnitTestAcccessor::GetSoftRefCount(Handle.GetInstance()) == 1);
+				TestTrue(TEXT("Handle hard refcount"), FWorldPartitionActorDescUnitTestAcccessor::GetHardRefCount(Handle.GetInstance()) == 0);
 			}
 		}
 
@@ -311,12 +312,12 @@ namespace WorldPartitionTests
 		WorldPartitionTests::PerformTests<FWorldPartitionLoadingContext::FDeferred>(this);
 
 		// Serialization tests
-		UActorDescContainer* ActorDescContainer = NewObject<UActorDescContainer>(GetTransientPackage());
-		ActorDescContainer->Initialize({ nullptr, TEXT("/Engine/WorldPartition/WorldPartitionUnitTest") });
+		UActorDescContainerInstance* ActorDescContainerInstance = NewObject<UActorDescContainerInstance>(GetTransientPackage());
+		ActorDescContainerInstance->Initialize({ TEXT("/Engine/WorldPartition/WorldPartitionUnitTest") });
 
-		for (FActorDescList::TIterator<> ActorDescIterator(ActorDescContainer); ActorDescIterator; ++ActorDescIterator)
+		for (UActorDescContainerInstance::TIterator<> Iterator(ActorDescContainerInstance); Iterator; ++Iterator)
 		{
-			UClass* ActorNativeClass = UClass::TryFindTypeSlow<UClass>(ActorDescIterator->GetNativeClass().ToString(), EFindFirstObjectOptions::ExactClass);
+			UClass* ActorNativeClass = UClass::TryFindTypeSlow<UClass>(Iterator->GetNativeClass().ToString(), EFindFirstObjectOptions::ExactClass);
 			TestFalse(TEXT("Actor Descriptor Serialization"), !ActorNativeClass);
 
 			if (ActorNativeClass)
@@ -325,13 +326,13 @@ namespace WorldPartitionTests
 
 				FWorldPartitionActorDescInitData ActorDescInitData = FWorldPartitionActorDescInitData()
 					.SetNativeClass(ActorNativeClass)
-					.SetPackageName(ActorDescIterator->GetActorPackage())
-					.SetActorPath(ActorDescIterator->GetActorSoftPath());
+					.SetPackageName(Iterator->GetActorPackage())
+					.SetActorPath(Iterator->GetActorSoftPath());
 			
-				ActorDescIterator->SerializeTo(ActorDescInitData.SerializedData);
+				Iterator->GetActorDesc()->SerializeTo(ActorDescInitData.SerializedData);
 				NewActorDesc->Init(ActorDescInitData);
 
-				TestTrue(TEXT("Actor Descriptor Serialization"), NewActorDesc->Equals(*ActorDescIterator));
+				TestTrue(TEXT("Actor Descriptor Serialization"), NewActorDesc->Equals(Iterator->GetActorDesc()));
 			}
 		}
 #endif
diff --git a/Engine/Source/Runtime/Engine/Private/Tests/WorldPartition/WorldPartitionPathConversionTests.cpp b/Engine/Source/Runtime/Engine/Private/Tests/WorldPartition/WorldPartitionPathConversionTests.cpp
--- a/Engine/Source/Runtime/Engine/Private/Tests/WorldPartition/WorldPartitionPathConversionTests.cpp
+++ b/Engine/Source/Runtime/Engine/Private/Tests/WorldPartition/WorldPartitionPathConversionTests.cpp
@@ -5,6 +5,7 @@
 #if WITH_EDITOR
 #include "WorldPartition/WorldPartition.h"
 #include "WorldPartition/WorldPartitionHelpers.h"
+#include "WorldPartition/WorldPartitionActorDescInstance.h"
 #include "EditorWorldUtils.h"
 #endif
 
@@ -50,13 +51,13 @@ namespace WorldPartitionTests
 		}
 		
 		check(WorldPartition);
-		UActorDescContainer* ActorDescMainContainer = WorldPartition->GetActorDescContainer();
-		if (!TestNotNull(TEXT("Missing World Partition Container"), ActorDescMainContainer))
+		UActorDescContainerInstance* ActorDescMainContainerInstance = WorldPartition->GetActorDescContainerInstance();
+		if (!TestNotNull(TEXT("Missing World Partition Container"), ActorDescMainContainerInstance))
 		{
 			return false;
 		}
 
-		UWorldPartition::FGenerateStreamingParams Params = UWorldPartition::FGenerateStreamingParams().SetActorDescContainer(ActorDescMainContainer);
+		UWorldPartition::FGenerateStreamingParams Params = UWorldPartition::FGenerateStreamingParams().SetActorDescContainerInstance(ActorDescMainContainerInstance);
 		UWorldPartition::FGenerateStreamingContext Context;
 
 		if (!TestTrue(TEXT("World Partition Generate Streaming"), WorldPartition->GenerateContainerStreaming(Params, Context)))
@@ -64,36 +65,36 @@ namespace WorldPartitionTests
 			return false;
 		}
 
-		FWorldPartitionReference ActorRef(ActorDescMainContainer, FGuid(TEXT("5D9F93BA407A811AFDDDAAB4F1CECC6A")));
+		FWorldPartitionReference ActorRef(ActorDescMainContainerInstance, FGuid(TEXT("5D9F93BA407A811AFDDDAAB4F1CECC6A")));
 		if (!TestTrue(TEXT("Invalid Actor Reference"), ActorRef.IsValid()))
 		{
 			return false;
 		}
 
-		AActor* Actor = ActorRef->GetActor();
+		AActor* Actor = ActorRef.GetActor();
 		if (!TestNotNull(TEXT("Missing Actor"), Actor))
 		{
 			return false;
 		}
 
-		FWorldPartitionHandle ActorHandle(ActorDescMainContainer, FGuid(TEXT("0D2B04D240BE5DE58FE437A8D2DBF5C9")));
+		FWorldPartitionHandle ActorHandle(ActorDescMainContainerInstance, FGuid(TEXT("0D2B04D240BE5DE58FE437A8D2DBF5C9")));
 		if (!TestTrue(TEXT("Invalid Actor Handle"), ActorHandle.IsValid()))
 		{
 			return false;
 		}
 
-		if (!TestNull(TEXT("Actor Handle Not Loaded"), ActorHandle->GetActor()))
+		if (!TestNull(TEXT("Actor Handle Not Loaded"), ActorHandle.GetActor()))
 		{
 			return false;
 		}
 
-		UObject* ResolvedObject = ActorHandle->GetActorSoftPath().TryLoad();
-		if (!TestNotNull(TEXT("Actor Handle Loaded"), ActorHandle->GetActor()))
+		UObject* ResolvedObject = ActorHandle.GetInstance()->GetActorSoftPath().TryLoad();
+		if (!TestNotNull(TEXT("Actor Handle Loaded"), ActorHandle.GetActor()))
 		{
 			return false;
 		}
 
-		if (!TestTrue(TEXT("Resolving Runtime Actor From Editor Path Failed"), ResolvedObject == ActorHandle->GetActor()))
+		if (!TestTrue(TEXT("Resolving Runtime Actor From Editor Path Failed"), ResolvedObject == ActorHandle.GetActor()))
 		{
 			return false;
 		}
diff --git a/Engine/Source/Runtime/Engine/Private/WorldPartition/ActorDescContainer.cpp b/Engine/Source/Runtime/Engine/Private/WorldPartition/ActorDescContainer.cpp
--- a/Engine/Source/Runtime/Engine/Private/WorldPartition/ActorDescContainer.cpp
+++ b/Engine/Source/Runtime/Engine/Private/WorldPartition/ActorDescContainer.cpp
@@ -33,11 +33,6 @@ void UActorDescContainer::Initialize(const FInitializeParams& InitParams)
 {
 	TRACE_CPUPROFILER_EVENT_SCOPE(UActorDescContainer::Initialize);
 
-	// @todo_ow: We need to pass the world context to RegisterActorDescriptor for FLevelInstanceActorDesc::RegisterContainerInstance to resolve the ULevelInstanceSubsystem.
-	// A better solution would be for ActorDescContainers to always be outered to an OwningWorldPartition (with the downside of not sharing between 2 instanced WorldPartition).
-	// With this, we could always find the owning WorldPartition and of course the owning world (GetOwningWorldPartition()->GetWorld()).
-	UWorld* OwningWorld = InitParams.World;
-
 	check(!bContainerInitialized);
 	ContainerPackageName = InitParams.PackageName;
 	TArray<FAssetData> Assets;
@@ -150,7 +145,7 @@ void UActorDescContainer::Initialize(const FInitializeParams& InitParams)
 		TRACE_CPUPROFILER_EVENT_SCOPE(RegisterDescriptors);
 		for (auto& [ActorGuid, ActorDesc] : ValidActorDescs)
 		{
-			RegisterActorDescriptor(ActorDesc.Release(), OwningWorld);
+			RegisterActorDescriptor(ActorDesc.Release());
 		}
 	}
 
@@ -199,53 +194,25 @@ bool UActorDescContainer::IsActorDescHandled(const AActor* Actor) const
 	if (Actor->GetContentBundleGuid() == GetContentBundleGuid())
 	{
 		const FString ActorPackageName = Actor->GetPackage()->GetName();
-		// Resolve ExternalActorPath using the actual package name instead of ContainerPackageName to properly handle 
-		// unsaved levels created with an existing template partitioned map.
-		// However, always rely on ContainerPackageName for content bundles to resolve the external actor path (no content bundles for template maps).
-		const FString ExternalActorPath = (GetContentBundleGuid().IsValid() ? GetExternalActorPath() : ULevel::GetExternalActorsPath(GetPackage()->GetName())) / TEXT("");
+		const FString ExternalActorPath = GetExternalActorPath() / TEXT("");
 		return ActorPackageName.StartsWith(ExternalActorPath);
 	}
 	return false;
 }
 
-bool UActorDescContainer::IsMainPartitionContainer() const
-{
-	return GetWorldPartition() && GetWorldPartition()->GetActorDescContainer() == this;
-}
-
-bool UActorDescContainer::IsTemplateContainer() const
-{
-	bool bIsTemplate = GetOuter() == GetTransientPackage();
-	check(bIsTemplate || GetWorldPartition());
-	return bIsTemplate;
-}
-
-UWorldPartition* UActorDescContainer::GetWorldPartition() const
-{
-	UWorldPartition* OuterWorldPartition = GetTypedOuter<UWorldPartition>();
-	check(OuterWorldPartition || IsTemplateContainer());
-	return OuterWorldPartition;
-}
-
-void UActorDescContainer::RegisterActorDescriptor(FWorldPartitionActorDesc* ActorDesc, UWorld* InWorldContext)
+void UActorDescContainer::RegisterActorDescriptor(FWorldPartitionActorDesc* ActorDesc)
 {
 	FActorDescList::AddActorDescriptor(ActorDesc);
 
-	ActorDesc->SetContainer(this, InWorldContext);
+	ActorDesc->SetContainer(this);
 		
-	//@odo_ow get rid of this with a delegate
-	if (const UDataLayerManager* DataLayerManager = UDataLayerManager::GetDataLayerManager(InWorldContext))
-	{
-		DataLayerManager->ResolveActorDescDataLayers(ActorDesc);
-	}
-
 	ActorsByName.Add(ActorDesc->GetActorName(), ActorsByGuid.FindChecked(ActorDesc->GetGuid()));
 }
 
 void UActorDescContainer::UnregisterActorDescriptor(FWorldPartitionActorDesc* ActorDesc)
 {
 	FActorDescList::RemoveActorDescriptor(ActorDesc);
-	ActorDesc->SetContainer(nullptr, nullptr);
+	ActorDesc->SetContainer(nullptr);
 	verifyf(ActorsByName.Remove(ActorDesc->GetActorName()), TEXT("Missing actor '%s' from container '%s'"), *ActorDesc->GetActorName().ToString(), *ContainerPackageName.ToString());
 }
 
@@ -301,7 +268,7 @@ void UActorDescContainer::OnObjectPreSave(UObject* Object, FObjectPreSaveContext
 				{
 					// New actor
 					FWorldPartitionActorDesc* AddedActorDesc = Actor->CreateActorDesc().Release();
-					RegisterActorDescriptor(AddedActorDesc, Actor->GetWorld());
+					RegisterActorDescriptor(AddedActorDesc);
 					OnActorDescAdded(AddedActorDesc);
 				}
 			}
@@ -319,25 +286,6 @@ void UActorDescContainer::OnPackageDeleted(UPackage* Package)
 	}
 }
 
-void UActorDescContainer::OnObjectsReplaced(const TMap<UObject*, UObject*>& OldToNewObjectMap)
-{
-	// Patch up Actor pointers in ActorDescs
-	for (auto [OldObject, NewObject] : OldToNewObjectMap)
-	{
-		if (AActor* OldActor = Cast<AActor>(OldObject))
-		{
-			if (ShouldHandleActorEvent(OldActor))
-			{
-				AActor* NewActor = Cast<AActor>(NewObject);
-				if (FWorldPartitionActorDesc* ActorDesc = GetActorDesc(OldActor->GetActorGuid()))
-				{
-					FWorldPartitionActorDescUtils::ReplaceActorDescriptorPointerFromActor(OldActor, NewActor, ActorDesc);
-				}
-			}
-		}
-	}
-}
-
 void UActorDescContainer::OnClassDescriptorUpdated(const FWorldPartitionActorDesc* InClassDesc)
 {
 	FWorldPartitionClassDescRegistry& ClassDescRegistry = FWorldPartitionClassDescRegistry::Get();
@@ -400,26 +348,9 @@ bool UActorDescContainer::RemoveActor(const FGuid& ActorGuid)
 	return false;
 }
 
-void UActorDescContainer::LoadAllActors(TArray<FWorldPartitionReference>& OutReferences)
-{
-	FWorldPartitionLoadingContext::FDeferred LoadingContext;
-	OutReferences.Reserve(OutReferences.Num() + GetActorDescCount());
-	for (FActorDescList::TIterator<> ActorDescIterator(this); ActorDescIterator; ++ActorDescIterator)
-	{
-		FWorldPartitionActorDesc* ActorDesc = *ActorDescIterator;
-		OutReferences.Emplace(this, ActorDesc->GetGuid());
-	}
-}
-
 bool UActorDescContainer::ShouldRegisterDelegates()
 {
-	UWorld* OuterWorld = GetTypedOuter<UWorld>();
-	// No need to register delegates for level instances
-	bool bIsInstance = OuterWorld && OuterWorld->IsInstanced() && !OuterWorld->GetPackage()->HasAnyPackageFlags(PKG_NewlyCreated);
-	UWorld* OwningWorld = GetWorldPartition() ? GetWorldPartition()->GetWorld() : nullptr;
-	// Template container will always register
-	bool bShouldRegisterForWorld = (OwningWorld && !OwningWorld->IsGameWorld()) || (IsTemplateContainer());
-	return GEditor && !IsTemplate() && !IsRunningCookCommandlet() && !bIsInstance && bShouldRegisterForWorld;
+	return GEditor && !IsTemplate() && !IsRunningCookCommandlet();
 }
 
 void UActorDescContainer::RegisterEditorDelegates()
@@ -428,8 +359,7 @@ void UActorDescContainer::RegisterEditorDelegates()
 	{
 		FCoreUObjectDelegates::OnObjectPreSave.AddUObject(this, &UActorDescContainer::OnObjectPreSave);
 		FEditorDelegates::OnPackageDeleted.AddUObject(this, &UActorDescContainer::OnPackageDeleted);
-		FCoreUObjectDelegates::OnObjectsReplaced.AddUObject(this, &UActorDescContainer::OnObjectsReplaced);
-		
+
 		FWorldPartitionClassDescRegistry& ClassDescRegistry = FWorldPartitionClassDescRegistry::Get();
 		ClassDescRegistry.OnClassDescriptorUpdated().AddUObject(this, &UActorDescContainer::OnClassDescriptorUpdated);
 	}
@@ -441,7 +371,6 @@ void UActorDescContainer::UnregisterEditorDelegates()
 	{
 		FCoreUObjectDelegates::OnObjectPreSave.RemoveAll(this);
 		FEditorDelegates::OnPackageDeleted.RemoveAll(this);
-		FCoreUObjectDelegates::OnObjectsReplaced.RemoveAll(this);
 
 		FWorldPartitionClassDescRegistry& ClassDescRegistry = FWorldPartitionClassDescRegistry::Get();
 		ClassDescRegistry.OnClassDescriptorUpdated().RemoveAll(this);
@@ -451,63 +380,21 @@ void UActorDescContainer::UnregisterEditorDelegates()
 void UActorDescContainer::OnActorDescAdded(FWorldPartitionActorDesc* NewActorDesc)
 {
 	OnActorDescAddedEvent.Broadcast(NewActorDesc);
-
-	if (UWorldPartition* WorldPartition = GetWorldPartition())
-	{
-		WorldPartition->OnActorDescAdded(NewActorDesc);
-	}
 }
 
 void UActorDescContainer::OnActorDescRemoved(FWorldPartitionActorDesc* ActorDesc)
 {
 	OnActorDescRemovedEvent.Broadcast(ActorDesc);
-
-	if (UWorldPartition* WorldPartition = GetWorldPartition())
-	{
-		WorldPartition->OnActorDescRemoved(ActorDesc);
-	}
 }
 
 void UActorDescContainer::OnActorDescUpdating(FWorldPartitionActorDesc* ActorDesc)
 {
-	if (UWorldPartition* WorldPartition = GetWorldPartition())
-	{
-		WorldPartition->OnActorDescUpdating(ActorDesc);
-	}
+	OnActorDescUpdatingEvent.Broadcast(ActorDesc);
 }
 
 void UActorDescContainer::OnActorDescUpdated(FWorldPartitionActorDesc* ActorDesc)
 {
-	if (UWorldPartition* WorldPartition = GetWorldPartition())
-	{
-		if (const UDataLayerManager* DataLayerManager = UDataLayerManager::GetDataLayerManager(WorldPartition))
-		{
-			DataLayerManager->ResolveActorDescDataLayers(ActorDesc);
-		}
-		WorldPartition->OnActorDescUpdated(ActorDesc);
-	}
-}
-
-const FLinkerInstancingContext* UActorDescContainer::GetInstancingContext() const
-{
-	const FLinkerInstancingContext* InstancingContext = nullptr;
-
-	if (UWorldPartition* WorldPartition = GetWorldPartition())
-	{
-		WorldPartition->GetInstancingContext(InstancingContext);
-	}
-
-	return InstancingContext;
-}
-
-const FTransform& UActorDescContainer::GetInstanceTransform() const
-{
-	if (UWorldPartition* WorldPartition = GetWorldPartition())
-	{
-		return WorldPartition->GetInstanceTransform();
-	}
-
-	return FTransform::Identity;
+	OnActorDescUpdatedEvent.Broadcast(ActorDesc);
 }
 #endif
 
diff --git a/Engine/Source/Runtime/Engine/Private/WorldPartition/ActorDescList.cpp b/Engine/Source/Runtime/Engine/Private/WorldPartition/ActorDescList.cpp
--- a/Engine/Source/Runtime/Engine/Private/WorldPartition/ActorDescList.cpp
+++ b/Engine/Source/Runtime/Engine/Private/WorldPartition/ActorDescList.cpp
@@ -16,51 +16,25 @@ FWorldPartitionActorDesc* FActorDescList::AddActor(const AActor* InActor)
 
 const FWorldPartitionActorDesc* FActorDescList::GetActorDesc(const FGuid& Guid) const
 {
-	const TUniquePtr<FWorldPartitionActorDesc>* const * ActorDesc = ActorsByGuid.Find(Guid);
-	return ActorDesc ? (*ActorDesc)->Get() : nullptr;
+	const TUniquePtr<FWorldPartitionActorDesc>* ActorDesc = GetActorDescriptor(Guid);
+	return ActorDesc ? ActorDesc->Get() : nullptr;
 }
 
 FWorldPartitionActorDesc* FActorDescList::GetActorDesc(const FGuid& Guid)
 {
-	TUniquePtr<FWorldPartitionActorDesc>** ActorDesc = ActorsByGuid.Find(Guid);
-	return ActorDesc ? (*ActorDesc)->Get() : nullptr;
+	TUniquePtr<FWorldPartitionActorDesc>* ActorDesc = GetActorDescriptor(Guid);
+	return ActorDesc ? ActorDesc->Get() : nullptr;
 }
 
 const FWorldPartitionActorDesc& FActorDescList::GetActorDescChecked(const FGuid& Guid) const
 {
-	const TUniquePtr<FWorldPartitionActorDesc>* const ActorDesc = ActorsByGuid.FindChecked(Guid);
+	const TUniquePtr<FWorldPartitionActorDesc>* ActorDesc = GetActorDescriptorChecked(Guid);
 	return *ActorDesc->Get();
 }
 
 FWorldPartitionActorDesc& FActorDescList::GetActorDescChecked(const FGuid& Guid)
 {
-	TUniquePtr<FWorldPartitionActorDesc>* ActorDesc = ActorsByGuid.FindChecked(Guid);
+	TUniquePtr<FWorldPartitionActorDesc>* ActorDesc = GetActorDescriptorChecked(Guid);
 	return *ActorDesc->Get();
 }
-
-void FActorDescList::Empty()
-{
-	ActorsByGuid.Empty();
-	ActorDescList.Empty();
-}
-
-void FActorDescList::AddActorDescriptor(FWorldPartitionActorDesc* ActorDesc)
-{
-	check(ActorDesc);
-	checkf(!ActorsByGuid.Contains(ActorDesc->GetGuid()), TEXT("Duplicated actor descriptor guid '%s' detected: `%s`"), *ActorDesc->GetGuid().ToString(), *ActorDesc->GetActorName().ToString());
-
-	TUniquePtr<FWorldPartitionActorDesc>* NewActorDesc = new(ActorDescList) TUniquePtr<FWorldPartitionActorDesc>(ActorDesc);
-	ActorsByGuid.Add(ActorDesc->GetGuid(), NewActorDesc);
-}
-
-void FActorDescList::RemoveActorDescriptor(FWorldPartitionActorDesc* ActorDesc)
-{
-	check(ActorDesc);
-	verify(ActorsByGuid.Remove(ActorDesc->GetGuid()));
-}
-
-TUniquePtr<FWorldPartitionActorDesc>* FActorDescList::GetActorDescriptor(const FGuid& ActorGuid)
-{
-	return ActorsByGuid.FindRef(ActorGuid);
-}
 #endif
\ No newline at end of file
diff --git a/Engine/Source/Runtime/Engine/Private/WorldPartition/ContentBundle/ContentBundleEditor.cpp b/Engine/Source/Runtime/Engine/Private/WorldPartition/ContentBundle/ContentBundleEditor.cpp
--- a/Engine/Source/Runtime/Engine/Private/WorldPartition/ContentBundle/ContentBundleEditor.cpp
+++ b/Engine/Source/Runtime/Engine/Private/WorldPartition/ContentBundle/ContentBundleEditor.cpp
@@ -19,6 +19,7 @@
 #include "WorldPartition/Cook/WorldPartitionCookPackageContextInterface.h"
 #include "WorldPartition/Cook/WorldPartitionCookPackage.h"
 #include "WorldPartition/WorldPartitionLevelHelper.h"
+#include "WorldPartition/ActorDescContainerInstance.h"
 
 FContentBundleEditor::FContentBundleEditor(TSharedPtr<FContentBundleClient>& InClient, UWorld* InWorld)
 	: FContentBundleBase(InClient, InWorld)
@@ -61,10 +62,10 @@ void FContentBundleEditor::DoInjectContent()
 	if (bCreatedContainerPath)
 	{
 		UWorldPartition* WorldPartition = GetInjectedWorld()->GetWorldPartition();
-		UWorldPartition::FContainerRegistrationParams RegistrationsParams(*ActorDescContainerPackage);
+		UActorDescContainerInstance::FInitializeParams InitParams(*ActorDescContainerPackage);
 		
 		const FTopLevelAssetPath InjectedWorldAssetPath = FSoftObjectPath(GetInjectedWorld()).GetAssetPath();
-		RegistrationsParams.FilterActorDescFunc = [&](const FWorldPartitionActorDesc* ActorDesc)
+		InitParams.FilterActorDescFunc = [&](const FWorldPartitionActorDesc* ActorDesc)
 		{
 			if (ActorDesc->GetActorSoftPath().GetAssetPath() != InjectedWorldAssetPath)
 			{
@@ -79,17 +80,17 @@ void FContentBundleEditor::DoInjectContent()
 			return true;
 		};
 
-		ActorDescContainer = WorldPartition->RegisterActorDescContainer(RegistrationsParams);
-		if (ActorDescContainer.IsValid())
+		ActorDescContainerInstance = WorldPartition->RegisterActorDescContainerInstance(InitParams);
+		if (ActorDescContainerInstance.IsValid() && ActorDescContainerInstance->IsInitialized())
 		{
-			UE_LOG(LogContentBundle, Log, TEXT("%s ExternalActors in %s found. %u actors were injected"), *ContentBundle::Log::MakeDebugInfoString(*this), *ActorDescContainer->GetExternalActorPath(), ActorDescContainer->GetActorDescCount());
+			UE_LOG(LogContentBundle, Log, TEXT("%s ExternalActors in %s found. %u actors were injected"), *ContentBundle::Log::MakeDebugInfoString(*this), *ActorDescContainerInstance->GetExternalActorPath(), ActorDescContainerInstance->GetActorDescInstanceCount());
 
 			check(GetDescriptor()->GetGuid().IsValid());
-			ActorDescContainer->SetContentBundleGuid(GetDescriptor()->GetGuid());
+			ActorDescContainerInstance->GetContainer()->SetContentBundleGuid(GetDescriptor()->GetGuid());
 
-			if (!ActorDescContainer->IsEmpty())
+			if (!ActorDescContainerInstance->GetContainer()->IsEmpty())
 			{
-				WorldDataLayersActorReference = FWorldDataLayersReference(ActorDescContainer.Get(), BuildWorlDataLayersName());
+				WorldDataLayersActorReference = FWorldDataLayersReference(ActorDescContainerInstance.Get(), BuildWorlDataLayersName());
 				SetStatus(EContentBundleStatus::ContentInjected);
 			}
 			else
@@ -136,13 +137,14 @@ void FContentBundleEditor::DoRemoveContent()
 
 	UnsavedActorMonitor->Uninitialize();
 
-	if (ActorDescContainer.IsValid())
+	// Might have been uninitialized by the WorldPartition 
+	if (UActorDescContainerInstance* ContainerInstance = ActorDescContainerInstance.Get(); ContainerInstance && ContainerInstance->IsInitialized())
 	{
 		UnregisterDelegates();
 
-		GetInjectedWorld()->GetWorldPartition()->UnregisterActorDescContainer(ActorDescContainer.Get());
-		ActorDescContainer = nullptr;
+		GetInjectedWorld()->GetWorldPartition()->UnregisterActorDescContainerInstance(ActorDescContainerInstance.Get());
 	}
+	ActorDescContainerInstance = nullptr;
 
 	ExternalStreamingObject = nullptr;
 	CookPackageIdsToCell.Empty();
@@ -187,7 +189,7 @@ bool FContentBundleEditor::AddActor(AActor* InActor)
 	FSetActorContentBundleGuid SetActorContentBundleGuid(InActor, GetDescriptor()->GetGuid());
 
 	// Rename the actor so it is saved in the content bundle location
-	FString ActorPackageNameInContentBundle = ContentBundlePaths::MakeExternalActorPackagePath(ActorDescContainer->GetExternalActorPath(), InActor->GetName());
+	FString ActorPackageNameInContentBundle = ContentBundlePaths::MakeExternalActorPackagePath(ActorDescContainerInstance->GetExternalActorPath(), InActor->GetName());
 	verify(InActor->GetPackage()->Rename(*ActorPackageNameInContentBundle));
 
 	UnsavedActorMonitor->MonitorActor(InActor);
@@ -201,7 +203,7 @@ bool FContentBundleEditor::ContainsActor(const AActor* InActor) const
 {
 	if (InActor != nullptr)
 	{
-		return ActorDescContainer->GetActorDesc(InActor->GetActorGuid()) != nullptr || UnsavedActorMonitor->IsMonitoring(InActor);
+		return ActorDescContainerInstance->GetActorDescInstance(InActor->GetActorGuid()) != nullptr || UnsavedActorMonitor->IsMonitoring(InActor);
 	}
 
 	return false;
@@ -211,9 +213,9 @@ bool FContentBundleEditor::GetActors(TArray<AActor*>& Actors)
 {
 	Actors.Reserve(GetActorCount());
 
-	if (ActorDescContainer.IsValid())
+	if (ActorDescContainerInstance.IsValid())
 	{
-		for (FActorDescList::TIterator<> It(ActorDescContainer.Get()); It; ++It)
+		for (UActorDescContainerInstance::TIterator<> It(ActorDescContainerInstance.Get()); It; ++It)
 		{
 			if (AActor* Actor = It->GetActor())
 			{
@@ -242,7 +244,7 @@ bool FContentBundleEditor::GetActors(TArray<AActor*>& Actors)
 bool FContentBundleEditor::HasUserPlacedActors() const
 {
 	// If there is only one actor in the container its the WorldDataLayer automatically created when injecting base content.
-	bool bActorDescContHasUserPlacedActors = ActorDescContainer.IsValid() && ActorDescContainer->GetActorDescCount() > 1;
+	bool bActorDescContHasUserPlacedActors = ActorDescContainerInstance.IsValid() && ActorDescContainerInstance->GetActorDescInstanceCount() > 1;
 	return (bActorDescContHasUserPlacedActors || UnsavedActorMonitor->IsMonitoringActors());
 }
 
@@ -250,8 +252,8 @@ uint32 FContentBundleEditor::GetActorCount() const
 {
 	if (GetStatus() == EContentBundleStatus::ContentInjected)
 	{
-		uint32 UnsavedWorldDataLayerCount = WorldDataLayersActorReference.IsValid() && ActorDescContainer.IsValid() && ActorDescContainer->IsEmpty() ? 1 : 0;
-		return ActorDescContainer->GetActorDescCount() + UnsavedActorMonitor->GetActorCount() + UnsavedWorldDataLayerCount;
+		uint32 UnsavedWorldDataLayerCount = WorldDataLayersActorReference.IsValid() && ActorDescContainerInstance.IsValid() && ActorDescContainerInstance->IsEmpty() ? 1 : 0;
+		return ActorDescContainerInstance->GetActorDescInstanceCount() + UnsavedActorMonitor->GetActorCount() + UnsavedWorldDataLayerCount;
 	}
 
 	return 0;
@@ -269,9 +271,9 @@ uint32 FContentBundleEditor::GetUnsavedActorAcount() const
 
 void FContentBundleEditor::ReferenceAllActors()
 {
-	if (ActorDescContainer.IsValid())
+	if (ActorDescContainerInstance.IsValid())
 	{
-		ActorDescContainer->LoadAllActors(ForceLoadedActors);
+		ActorDescContainerInstance->LoadAllActors(ForceLoadedActors);
 	}
 }
 
@@ -337,7 +339,7 @@ void FContentBundleEditor::GenerateStreaming(TArray<FString>* OutPackageToGenera
 	}
 
 	UWorldPartition::FGenerateStreamingParams Params = UWorldPartition::FGenerateStreamingParams()
-		.SetActorDescContainer(ActorDescContainer.Get());
+		.SetActorDescContainerInstance(ActorDescContainerInstance.Get());
 
 	UWorldPartition::FGenerateStreamingContext Context = UWorldPartition::FGenerateStreamingContext()
 		.SetPackagesToGenerate(OutPackageToGenerate);
@@ -545,7 +547,7 @@ void FContentBundleEditor::BroadcastChanged()
 
 UPackage* FContentBundleEditor::CreateActorPackage(const FName& ActorName) const
 {
-	FString ActorPackagePath = ULevel::GetActorPackageName(ActorDescContainer->GetExternalActorPath(), EActorPackagingScheme::Reduced, ActorName.ToString());
+	FString ActorPackagePath = ULevel::GetActorPackageName(ActorDescContainerInstance->GetExternalActorPath(), EActorPackagingScheme::Reduced, ActorName.ToString());
 	UPackage* ActorPackage = CreatePackage(*ActorPackagePath);
 
 	ActorPackage->SetDirtyFlag(true);
@@ -560,29 +562,29 @@ FName FContentBundleEditor::BuildWorlDataLayersName() const
 
 void FContentBundleEditor::RegisterDelegates()
 {
-	ActorDescContainer->OnActorDescAddedEvent.AddRaw(this, &FContentBundleEditor::OnActorDescAdded);
-	ActorDescContainer->OnActorDescRemovedEvent.AddRaw(this, &FContentBundleEditor::OnActorDescRemoved);
+	ActorDescContainerInstance->OnActorDescInstanceAddedEvent.AddRaw(this, &FContentBundleEditor::OnActorDescInstanceAdded);
+	ActorDescContainerInstance->OnActorDescInstanceRemovedEvent.AddRaw(this, &FContentBundleEditor::OnActorDescInstanceRemoved);
 }
 
 void FContentBundleEditor::UnregisterDelegates()
 {
-	ActorDescContainer->OnActorDescAddedEvent.RemoveAll(this);
-	ActorDescContainer->OnActorDescRemovedEvent.RemoveAll(this);
+	ActorDescContainerInstance->OnActorDescInstanceAddedEvent.RemoveAll(this);
+	ActorDescContainerInstance->OnActorDescInstanceRemovedEvent.RemoveAll(this);
 }
 
-void FContentBundleEditor::OnActorDescAdded(FWorldPartitionActorDesc* ActorDesc)
+void FContentBundleEditor::OnActorDescInstanceAdded(FWorldPartitionActorDescInstance* ActorDescInstance)
 {
 	UE_LOG(LogContentBundle, Verbose, TEXT("%s Added actor %s to container, ActorCount: %u. Package %s."), 
-		*ContentBundle::Log::MakeDebugInfoString(*this), *ActorDesc->GetActorLabelOrName().ToString(), GetActorCount(), *ActorDesc->GetActorPackage().ToString());
+		*ContentBundle::Log::MakeDebugInfoString(*this), *ActorDescInstance->GetActorLabelOrName().ToString(), GetActorCount(), *ActorDescInstance->GetActorPackage().ToString());
 
-	AActor* Actor = ActorDesc->GetActor();
+	AActor* Actor = ActorDescInstance->GetActor();
 	UnsavedActorMonitor->StopMonitoringActor(Actor);
 }
 
-void FContentBundleEditor::OnActorDescRemoved(FWorldPartitionActorDesc* ActorDesc)
+void FContentBundleEditor::OnActorDescInstanceRemoved(FWorldPartitionActorDescInstance* ActorDescInstance)
 {
 	UE_LOG(LogContentBundle, Verbose, TEXT("%s Removed actor %s from container, ActorCount:  %u. Package %s."), 
-		*ContentBundle::Log::MakeDebugInfoString(*this), *ActorDesc->GetActorLabelOrName().ToString(), GetActorCount(), *ActorDesc->GetActorPackage().ToString());
+		*ContentBundle::Log::MakeDebugInfoString(*this), *ActorDescInstance->GetActorLabelOrName().ToString(), GetActorCount(), *ActorDescInstance->GetActorPackage().ToString());
 
 	if (!HasUserPlacedActors())
 	{
@@ -592,7 +594,7 @@ void FContentBundleEditor::OnActorDescRemoved(FWorldPartitionActorDesc* ActorDes
 		}
 	}
 
-	AActor* ActorInWorld = ActorDesc->GetActor(false, false);
+	AActor* ActorInWorld = ActorDescInstance->GetActor(false, false);
 	check(ActorInWorld == nullptr || !UnsavedActorMonitor->IsMonitoring(ActorInWorld)); // ActorDesc existed is being deleted. Make sure the actor is not present in the unsaved list as it should have been saved for the desc to exist.
 }
 
diff --git a/Engine/Source/Runtime/Engine/Private/WorldPartition/ContentBundle/WorldDataLayerReference.cpp b/Engine/Source/Runtime/Engine/Private/WorldPartition/ContentBundle/WorldDataLayerReference.cpp
--- a/Engine/Source/Runtime/Engine/Private/WorldPartition/ContentBundle/WorldDataLayerReference.cpp
+++ b/Engine/Source/Runtime/Engine/Private/WorldPartition/ContentBundle/WorldDataLayerReference.cpp
@@ -5,6 +5,7 @@
 
 #include "WorldPartition/DataLayer/WorldDataLayers.h"
 #include "WorldPartition/WorldPartition.h"
+#include "WorldPartition/WorldPartitionActorDescInstance.h"
 
 FWorldDataLayersReference::FWorldDataLayersReference()
 {
@@ -16,9 +17,9 @@ FWorldDataLayersReference::~FWorldDataLayersReference()
 	Reset();
 }
 
-FWorldDataLayersReference::FWorldDataLayersReference(UActorDescContainer* Container, FName WorldDataLayerName)
+FWorldDataLayersReference::FWorldDataLayersReference(UActorDescContainerInstance* ContainerInstance, FName WorldDataLayerName)
 {
-	TrySetReference(Container, WorldDataLayerName);
+	TrySetReference(ContainerInstance, WorldDataLayerName);
 }
 
 FWorldDataLayersReference::FWorldDataLayersReference(const FActorSpawnParameters& SpawnParameters)
@@ -59,7 +60,7 @@ const AWorldDataLayers* FWorldDataLayersReference::Get() const
 		const FWorldPartitionReference& WorldDataLayerReference = WorldDataLayersVariant.Get<FWorldPartitionReference>();
 		if (WorldDataLayerReference.IsValid())
 		{
-			return Cast<AWorldDataLayers>(WorldDataLayerReference->GetActor());
+			return Cast<AWorldDataLayers>(WorldDataLayerReference.GetActor());
 		}
 		return nullptr;
 	}
@@ -67,15 +68,15 @@ const AWorldDataLayers* FWorldDataLayersReference::Get() const
 	return WorldDataLayersVariant.Get<AWorldDataLayers*>();
 }
 
-bool FWorldDataLayersReference::TrySetReference(UActorDescContainer* Container, FName WorldDataLayerName)
+bool FWorldDataLayersReference::TrySetReference(UActorDescContainerInstance* ContainerInstance, FName WorldDataLayerName)
 {
-	for (UActorDescContainer::TIterator<> ActorDescIterator(Container); ActorDescIterator; ++ActorDescIterator)
+	for (UActorDescContainerInstance::TIterator<> Iterator(ContainerInstance); Iterator; ++Iterator)
 	{
-		if (ActorDescIterator->GetActorNativeClass()->IsChildOf<AWorldDataLayers>())
+		if (Iterator->GetActorNativeClass()->IsChildOf<AWorldDataLayers>())
 		{
-			if (ActorDescIterator->GetActorName() == WorldDataLayerName)
+			if (Iterator->GetActorName() == WorldDataLayerName)
 			{
-				WorldDataLayersVariant.Emplace<FWorldPartitionReference>(Container, ActorDescIterator->GetGuid());
+				WorldDataLayersVariant.Emplace<FWorldPartitionReference>(ContainerInstance, Iterator->GetGuid());
 				return true;
 			}
 		}
@@ -99,7 +100,7 @@ void FWorldDataLayersReference::Reset()
 				FWorldPartitionReference ResolvedReference(WorldPartition, WorldDataLayers->GetActorGuid());
 				if (ResolvedReference.IsValid())
 				{
-					check(WorldDataLayers == ResolvedReference->GetActor());
+					check(WorldDataLayers == ResolvedReference.GetActor());
 					WorldDataLayersVariant.Emplace<FWorldPartitionReference>(ResolvedReference);
 				}
 				else
@@ -114,7 +115,7 @@ void FWorldDataLayersReference::Reset()
 
 	if (WorldDataLayersVariant.IsType<FWorldPartitionReference>())
 	{
-		WorldDataLayersVariant.Get<FWorldPartitionReference>() = nullptr;
+		WorldDataLayersVariant.Get<FWorldPartitionReference>().Reset();
 	}
 }
 
diff --git a/Engine/Source/Runtime/Engine/Private/WorldPartition/DataLayer/DataLayerManager.cpp b/Engine/Source/Runtime/Engine/Private/WorldPartition/DataLayer/DataLayerManager.cpp
--- a/Engine/Source/Runtime/Engine/Private/WorldPartition/DataLayer/DataLayerManager.cpp
+++ b/Engine/Source/Runtime/Engine/Private/WorldPartition/DataLayer/DataLayerManager.cpp
@@ -12,15 +12,17 @@
 #include "WorldPartition/WorldPartitionHelpers.h"
 #include "Engine/Canvas.h"
 #include "Debug/DebugDrawService.h"
+#include "UObject/Package.h"
 
 #if WITH_EDITOR
 #include "WorldPartition/DataLayer/WorldDataLayersActorDesc.h"
 #include "WorldPartition/DataLayer/DataLayerUtils.h"
 #include "WorldPartition/DataLayer/IDataLayerEditorModule.h"
-#include "WorldPartition/ActorDescContainer.h"
+#include "WorldPartition/ActorDescContainerInstance.h"
 #include "WorldPartition/WorldPartitionEditorPerProjectUserSettings.h"
 #include "WorldPartition/WorldPartitionActorLoaderInterface.h"
 #include "WorldPartition/ActorDescList.h"
+#include "WorldPartition/WorldPartitionActorDescInstance.h"
 #include "LevelInstance/LevelInstanceSubsystem.h"
 #include "Modules/ModuleManager.h"
 #include "UObject/UObjectIterator.h"
@@ -156,26 +158,27 @@ void UDataLayerManager::Initialize()
 #if WITH_EDITOR
 	UWorld* OuterWorld = GetTypedOuter<UWorld>();
 	UWorldPartition* OuterWorldPartition = GetOuterUWorldPartition();
-	UActorDescContainer* ActorDescContainer = OuterWorldPartition->GetActorDescContainer();
-
+	UActorDescContainerInstance* ActorDescContainerInstance = OuterWorldPartition->GetActorDescContainerInstance();
+	
 	// Partitioned Level Instance's DataLayerManager will never resolve DataLayers (it's up to its owning WorldPartition DataLayerManager to do the job.
 	ULevelInstanceSubsystem* LevelInstanceSubsystem = !GetWorld()->IsGameWorld() ? UWorld::GetSubsystem<ULevelInstanceSubsystem>(GetWorld()) : nullptr;
 	ILevelInstanceInterface* LevelInstance = LevelInstanceSubsystem ? LevelInstanceSubsystem->GetOwningLevelInstance(OuterWorld->PersistentLevel) : nullptr;
 	bCanResolveDataLayers = (LevelInstance == nullptr);
 
 	// In PIE, main world partition doesn't have an ActorDescContainer and doesn't need it, but instanced world partitions do have one.
-	check(!ActorDescContainer || ActorDescContainer->IsInitialized());
+	check(!ActorDescContainerInstance || ActorDescContainerInstance->IsInitialized());
+
 	AWorldDataLayers* WorldDataLayers = GetWorldDataLayers();
 	if (!WorldDataLayers)
 	{
-		if (ActorDescContainer)
+		if (ActorDescContainerInstance)
 		{
 			// Try to find and load AWorldDataLayers actor
-			for (FActorDescList::TIterator<> ActorDescIterator(ActorDescContainer); ActorDescIterator; ++ActorDescIterator)
+			for (UActorDescContainerInstance::TIterator<> Iterator(ActorDescContainerInstance); Iterator; ++Iterator)
 			{
-				if (ActorDescIterator->GetActorNativeClass()->IsChildOf<AWorldDataLayers>())
+				if (Iterator->GetActorNativeClass()->IsChildOf<AWorldDataLayers>())
 				{
-					WorldDataLayersActor = FWorldPartitionReference(OuterWorldPartition, ActorDescIterator->GetGuid());
+					WorldDataLayersActor = FWorldPartitionReference(ActorDescContainerInstance, Iterator->GetGuid());
 					break;
 				}
 			}
@@ -215,12 +218,12 @@ void UDataLayerManager::Initialize()
 
 	if (CanResolveDataLayers())
 	{
-		UActorDescContainer::OnActorDescContainerInitialized.AddUObject(this, &UDataLayerManager::OnActorDescContainerInitialized);
+		UActorDescContainerInstance::OnActorDescContainerInstanceInitialized.AddUObject(this, &UDataLayerManager::OnActorDescContainerInstanceInitialized);
 
-		// Manually call OnActorDescContainerInitialized on already initialized outer world partition container
-		if (ActorDescContainer)
+		// Manually call OnActorDescContainerInstanceInitialized on already initialized outer world partition container instance
+		if (ActorDescContainerInstance)
 		{
-			OnActorDescContainerInitialized(ActorDescContainer);
+			OnActorDescContainerInstanceInitialized(ActorDescContainerInstance);
 		}
 	}
 
@@ -247,7 +250,7 @@ void UDataLayerManager::DeInitialize()
 	}
 
 #if WITH_EDITOR
-	UActorDescContainer::OnActorDescContainerInitialized.RemoveAll(this);
+	UActorDescContainerInstance::OnActorDescContainerInstanceInitialized.RemoveAll(this);
 
 	WorldDataLayersActor = FWorldPartitionReference();
 #endif
@@ -805,46 +808,45 @@ bool UDataLayerManager::CanResolveDataLayers() const
 	return (GetWorldDataLayers() != nullptr) && bCanResolveDataLayers;
 }
 
-void UDataLayerManager::OnActorDescContainerInitialized(UActorDescContainer* InActorDescContainer) const
+void UDataLayerManager::OnActorDescContainerInstanceInitialized(UActorDescContainerInstance* InActorDescContainerInstance)
 {
-	ResolveActorDescContainerDataLayers(InActorDescContainer);
+	ResolveActorDescContainerInstanceDataLayers(InActorDescContainerInstance);
 }
 
 void UDataLayerManager::ResolveActorDescContainersDataLayers() const
 {
 	check(CanResolveDataLayers());
-	for (TObjectIterator<UActorDescContainer> ContainerIt; ContainerIt; ++ContainerIt)
+	for (TObjectIterator<UActorDescContainerInstance> ContainerIt; ContainerIt; ++ContainerIt)
 	{
-		if (UActorDescContainer* ActorDescContainer = *ContainerIt; ActorDescContainer)
+		if (UActorDescContainerInstance* ContainerInstance = *ContainerIt; ContainerInstance)
 		{
-			ResolveActorDescContainerDataLayers(ActorDescContainer);
+			ResolveActorDescContainerInstanceDataLayers(ContainerInstance);
 		}
 	}
 }
 
-void UDataLayerManager::ResolveActorDescContainerDataLayers(UActorDescContainer* InActorDescContainer) const
+void UDataLayerManager::ResolveActorDescContainerInstanceDataLayers(UActorDescContainerInstance* InActorDescContainerInstance) const
 {
-	ResolveActorDescContainerDataLayersInternal(InActorDescContainer, nullptr);
+	ResolveActorDescContainerInstanceDataLayersInternal(InActorDescContainerInstance, nullptr);
 }
 
-void UDataLayerManager::ResolveActorDescDataLayers(FWorldPartitionActorDesc* InActorDesc) const
+void UDataLayerManager::ResolveActorDescInstanceDataLayers(FWorldPartitionActorDescInstance* InActorDescInstance) const
 {
-	ResolveActorDescContainerDataLayersInternal(InActorDesc->GetContainer(), InActorDesc);
+	ResolveActorDescContainerInstanceDataLayersInternal(InActorDescInstance->GetContainerInstance(), InActorDescInstance);
 }
 
-void UDataLayerManager::ResolveActorDescContainerDataLayersInternal(UActorDescContainer* InActorDescContainer, FWorldPartitionActorDesc* InActorDesc) const
+void UDataLayerManager::ResolveActorDescContainerInstanceDataLayersInternal(UActorDescContainerInstance* InActorDescContainerInstance, FWorldPartitionActorDescInstance* InActorDescInstance) const
 {
-	check(InActorDescContainer);
-	check(!InActorDesc || (InActorDesc->GetContainer() == InActorDescContainer));
-
+	check(InActorDescContainerInstance);
+	check(!InActorDescInstance || (InActorDescInstance->GetContainerInstance() == InActorDescContainerInstance));
+		
+	const UWorldPartition* ContainerOuterWorldPartition = InActorDescContainerInstance->GetWorldPartition();
 	// Skip resolving for template containers (will be done on ActorDescViews)
-	if (InActorDescContainer->IsTemplateContainer())
+	if (!ContainerOuterWorldPartition)
 	{
 		return;
 	}
 
-	// Find owning world partition for this ActorDescContainer
-	const UWorldPartition* ContainerOuterWorldPartition = InActorDescContainer->GetWorldPartition();
 	const ULevelStreaming* ContainerLevelStreaming = FLevelUtils::FindStreamingLevel(ContainerOuterWorldPartition->GetTypedOuter<UWorld>()->PersistentLevel);
 	const UWorld* ContainerLevelStreamingWorld = ContainerLevelStreaming ? ContainerLevelStreaming->GetWorld() : nullptr;
 	const UWorldPartition* ContainerOwningWorldPartition = ContainerLevelStreamingWorld && !ContainerLevelStreamingWorld->IsGameWorld() ? ContainerLevelStreamingWorld->GetWorldPartition() : ContainerOuterWorldPartition;
@@ -855,18 +857,18 @@ void UDataLayerManager::ResolveActorDescContainerDataLayersInternal(UActorDescCo
 	{
 		return;
 	}
-	
+
 	// Resolve ActorDescs DataLayerInstanceNames
 	check(CanResolveDataLayers());
-	if (InActorDesc)
+	if (InActorDescInstance)
 	{
-		InActorDesc->SetDataLayerInstanceNames(FDataLayerUtils::ResolvedDataLayerInstanceNames(this, InActorDesc));
+		InActorDescInstance->SetDataLayerInstanceNames(FDataLayerUtils::ResolvedDataLayerInstanceNames(this, InActorDescInstance->GetActorDesc()));
 	}
 	else
 	{
-		for (FActorDescList::TIterator<> Iterator(InActorDescContainer); Iterator; ++Iterator)
+		for (UActorDescContainerInstance::TIterator<> Iterator(InActorDescContainerInstance); Iterator; ++Iterator)
 		{
-			Iterator->SetDataLayerInstanceNames(FDataLayerUtils::ResolvedDataLayerInstanceNames(this, *Iterator));
+			Iterator->SetDataLayerInstanceNames(FDataLayerUtils::ResolvedDataLayerInstanceNames(this, Iterator->GetActorDesc()));
 		}
 	}
 }
