WorldPartition: Add UActorDescContainerInstance/FWorldPartitionActorDescInstance which take transient/context data from UActorDescContainer/FWorldPartitionActorDesc leaving them as pure on disk repres

diff --git a/Engine/Source/Runtime/Engine/Private/WorldPartition/DataLayer/DataLayerUtils.cpp b/Engine/Source/Runtime/Engine/Private/WorldPartition/DataLayer/DataLayerUtils.cpp
--- a/Engine/Source/Runtime/Engine/Private/WorldPartition/DataLayer/DataLayerUtils.cpp
+++ b/Engine/Source/Runtime/Engine/Private/WorldPartition/DataLayer/DataLayerUtils.cpp
@@ -80,7 +80,7 @@ TArray<FName> FDataLayerUtils::ResolvedDataLayerInstanceNames(const UDataLayerMa
 }
 
 // For performance reasons, this function assumes that InActorDesc's DataLayerInstanceNames was already resolved.
-bool FDataLayerUtils::ResolveRuntimeDataLayerInstanceNames(const UDataLayerManager* InDataLayerManager, const FWorldPartitionActorDescView& InActorDescView, const TArray<const FWorldPartitionActorDescView*>& InWorldDataLayerViews, TArray<FName>& OutRuntimeDataLayerInstanceNames)
+bool FDataLayerUtils::ResolveRuntimeDataLayerInstanceNames(const UDataLayerManager* InDataLayerManager, const IWorldPartitionActorDescInstanceView& InActorDescView, const FStreamingGenerationActorDescViewMap& InActorDescViewMap, TArray<FName>& OutRuntimeDataLayerInstanceNames)
 {
 	const TArray<FName>& ActorDescViewDataLayerInstanceNames = InActorDescView.GetDataLayerInstanceNames();
 
@@ -99,7 +99,7 @@ bool FDataLayerUtils::ResolveRuntimeDataLayerInstanceNames(const UDataLayerManag
 	}
 
 	// Fallback on FWorldDataLayersActorDesc
-	TArray<const FWorldDataLayersActorDesc*> WorldDataLayersActorDescs = FindWorldDataLayerActorDescs(InWorldDataLayerViews);
+	TArray<const FWorldDataLayersActorDesc*> WorldDataLayersActorDescs = FindWorldDataLayerActorDescs(InActorDescViewMap);
 	if (WorldDataLayersActorDescs.Num())
 	{
 		check(AreWorldDataLayersActorDescsSane(WorldDataLayersActorDescs));
@@ -146,12 +146,13 @@ const FDataLayerInstanceDesc* FDataLayerUtils::GetDataLayerInstanceDescFromAsset
 	return nullptr;
 }
 
-TArray<const FWorldDataLayersActorDesc*> FDataLayerUtils::FindWorldDataLayerActorDescs(const TArray<const FWorldPartitionActorDescView*>& InWorldDataLayerViews)
+TArray<const FWorldDataLayersActorDesc*> FDataLayerUtils::FindWorldDataLayerActorDescs(const FStreamingGenerationActorDescViewMap& InActorDescViewMap)
 {
 	TArray<const FWorldDataLayersActorDesc*> WorldDataLayersActorDescs;
-	Algo::TransformIf(InWorldDataLayerViews, WorldDataLayersActorDescs, 
-		[](const FWorldPartitionActorDescView* WorldDataLayersActorDescView) { return ((FWorldDataLayersActorDesc*)WorldDataLayersActorDescView->GetActorDesc())->IsValid(); },
-		[](const FWorldPartitionActorDescView* WorldDataLayersActorDescView) { return (FWorldDataLayersActorDesc*)WorldDataLayersActorDescView->GetActorDesc(); });
+	TArray<const FStreamingGenerationActorDescView*> WorldDataLayerViews = InActorDescViewMap.FindByExactNativeClass<AWorldDataLayers>();
+	Algo::TransformIf(WorldDataLayerViews, WorldDataLayersActorDescs,
+		[](const FStreamingGenerationActorDescView* WorldDataLayersActorDescView) { return ((FWorldDataLayersActorDesc*)WorldDataLayersActorDescView->GetActorDesc())->IsValid(); },
+		[](const FStreamingGenerationActorDescView* WorldDataLayersActorDescView) { return (FWorldDataLayersActorDesc*)WorldDataLayersActorDescView->GetActorDesc(); });
 	return WorldDataLayersActorDescs;
 }
 
@@ -204,4 +205,4 @@ FString FDataLayerUtils::GenerateUniqueDataLayerShortName(const UDataLayerManage
 	return UniqueNewDataLayerShortName;
 }
 
-#endif
+#endif
\ No newline at end of file
diff --git a/Engine/Source/Runtime/Engine/Private/WorldPartition/ErrorHandling/WorldPartitionStreamingGenerationErrorHandler.cpp b/Engine/Source/Runtime/Engine/Private/WorldPartition/ErrorHandling/WorldPartitionStreamingGenerationErrorHandler.cpp
--- a/Engine/Source/Runtime/Engine/Private/WorldPartition/ErrorHandling/WorldPartitionStreamingGenerationErrorHandler.cpp
+++ b/Engine/Source/Runtime/Engine/Private/WorldPartition/ErrorHandling/WorldPartitionStreamingGenerationErrorHandler.cpp
@@ -2,16 +2,16 @@
 
 #if WITH_EDITOR
 #include "WorldPartition/ErrorHandling/WorldPartitionStreamingGenerationErrorHandler.h"
-#include "WorldPartition/WorldPartitionActorDescView.h"
+#include "WorldPartition/WorldPartitionActorDescInstanceViewInterface.h"
 
-FString IStreamingGenerationErrorHandler::GetActorName(const FWorldPartitionActorDescView& ActorDescView)
+FString IStreamingGenerationErrorHandler::GetActorName(const IWorldPartitionActorDescInstanceView& ActorDescView)
 {
 	const FString ActorPath = ActorDescView.GetActorSoftPath().GetLongPackageName();
 	const FString ActorLabel = ActorDescView.GetActorLabelOrName().ToString();
 	return ActorPath + TEXT(".") + ActorLabel;
 }
 
-FString IStreamingGenerationErrorHandler::GetFullActorName(const FWorldPartitionActorDescView& ActorDescView)
+FString IStreamingGenerationErrorHandler::GetFullActorName(const IWorldPartitionActorDescInstanceView& ActorDescView)
 {
 	const FString ActorPackage = ActorDescView.GetActorPackage().ToString();
 	return GetActorName(ActorDescView) + TEXT(" (") + ActorPackage + TEXT(")");
diff --git a/Engine/Source/Runtime/Engine/Private/WorldPartition/ErrorHandling/WorldPartitionStreamingGenerationLogErrorHandler.cpp b/Engine/Source/Runtime/Engine/Private/WorldPartition/ErrorHandling/WorldPartitionStreamingGenerationLogErrorHandler.cpp
--- a/Engine/Source/Runtime/Engine/Private/WorldPartition/ErrorHandling/WorldPartitionStreamingGenerationLogErrorHandler.cpp
+++ b/Engine/Source/Runtime/Engine/Private/WorldPartition/ErrorHandling/WorldPartitionStreamingGenerationLogErrorHandler.cpp
@@ -2,24 +2,23 @@
 
 #if WITH_EDITOR
 #include "WorldPartition/ErrorHandling/WorldPartitionStreamingGenerationLogErrorHandler.h"
-#include "WorldPartition/WorldPartitionActorDescView.h"
 #include "WorldPartition/WorldPartitionLog.h"
 #include "WorldPartition/DataLayer/DataLayerInstanceWithAsset.h"
 
 #define UE_ASSET_LOG_ACTORDESCVIEW(CategoryName, Verbosity, ActorDescView, Format, ...) \
 	UE_ASSET_LOG(LogWorldPartition, Log, *ActorDescView.GetActorPackage().ToString(), Format, ##__VA_ARGS__)
 
-void FStreamingGenerationLogErrorHandler::OnInvalidRuntimeGrid(const FWorldPartitionActorDescView& ActorDescView, FName GridName)
+void FStreamingGenerationLogErrorHandler::OnInvalidRuntimeGrid(const IWorldPartitionActorDescInstanceView& ActorDescView, FName GridName)
 {
 	UE_ASSET_LOG_ACTORDESCVIEW(LogWorldPartition, Log, ActorDescView, TEXT("Actor %s has an invalid runtime grid %s"), *GetActorName(ActorDescView), *GridName.ToString());
 }
 
-void FStreamingGenerationLogErrorHandler::OnInvalidReference(const FWorldPartitionActorDescView& ActorDescView, const FGuid& ReferenceGuid, FWorldPartitionActorDescView* ReferenceActorDescView)
+void FStreamingGenerationLogErrorHandler::OnInvalidReference(const IWorldPartitionActorDescInstanceView& ActorDescView, const FGuid& ReferenceGuid, IWorldPartitionActorDescInstanceView* ReferenceActorDescView)
 {
 	UE_ASSET_LOG_ACTORDESCVIEW(LogWorldPartition, Log, ActorDescView, TEXT("Actor %s has an invalid reference to %s"), *GetActorName(ActorDescView), ReferenceActorDescView ? *GetActorName(*ReferenceActorDescView) : *ReferenceGuid.ToString());
 }
 
-void FStreamingGenerationLogErrorHandler::OnInvalidReferenceGridPlacement(const FWorldPartitionActorDescView& ActorDescView, const FWorldPartitionActorDescView& ReferenceActorDescView)
+void FStreamingGenerationLogErrorHandler::OnInvalidReferenceGridPlacement(const IWorldPartitionActorDescInstanceView& ActorDescView, const IWorldPartitionActorDescInstanceView& ReferenceActorDescView)
 {
 	static const FString SpatiallyLoadedActor(TEXT("Spatially loaded actor"));
 	static const FString NonSpatiallyLoadedActor(TEXT("Non-spatially loaded loaded actor"));
@@ -27,22 +26,22 @@ void FStreamingGenerationLogErrorHandler::OnInvalidReferenceGridPlacement(const
 	UE_ASSET_LOG_ACTORDESCVIEW(LogWorldPartition, Log, ActorDescView, TEXT("%s %s reference %s %s"), ActorDescView.GetIsSpatiallyLoaded() ? *SpatiallyLoadedActor : *NonSpatiallyLoadedActor, *GetActorName(ActorDescView), ReferenceActorDescView.GetIsSpatiallyLoaded() ? *SpatiallyLoadedActor : *NonSpatiallyLoadedActor, *GetActorName(ReferenceActorDescView));
 }
 
-void FStreamingGenerationLogErrorHandler::OnInvalidReferenceDataLayers(const FWorldPartitionActorDescView& ActorDescView, const FWorldPartitionActorDescView& ReferenceActorDescView)
+void FStreamingGenerationLogErrorHandler::OnInvalidReferenceDataLayers(const IWorldPartitionActorDescInstanceView& ActorDescView, const IWorldPartitionActorDescInstanceView& ReferenceActorDescView)
 {
 	UE_ASSET_LOG_ACTORDESCVIEW(LogWorldPartition, Log, ActorDescView, TEXT("Actor %s references an actor in a different set of runtime data layers %s"), *GetActorName(ActorDescView), *GetActorName(ReferenceActorDescView));
 }
 
-void FStreamingGenerationLogErrorHandler::OnInvalidReferenceRuntimeGrid(const FWorldPartitionActorDescView& ActorDescView, const FWorldPartitionActorDescView& ReferenceActorDescView)
+void FStreamingGenerationLogErrorHandler::OnInvalidReferenceRuntimeGrid(const IWorldPartitionActorDescInstanceView& ActorDescView, const IWorldPartitionActorDescInstanceView& ReferenceActorDescView)
 {
 	UE_ASSET_LOG_ACTORDESCVIEW(LogWorldPartition, Log, ActorDescView, TEXT("Actor %s references an actor in a different runtime grid %s"), *GetActorName(ActorDescView), *GetActorName(ReferenceActorDescView));
 }
 
-void FStreamingGenerationLogErrorHandler::OnInvalidReferenceLevelScriptStreamed(const FWorldPartitionActorDescView& ActorDescView)
+void FStreamingGenerationLogErrorHandler::OnInvalidReferenceLevelScriptStreamed(const IWorldPartitionActorDescInstanceView& ActorDescView)
 {
 	UE_ASSET_LOG_ACTORDESCVIEW(LogWorldPartition, Log, ActorDescView, TEXT("Level Script Blueprint references streamed actor %s"), *GetActorName(ActorDescView));
 }
 
-void FStreamingGenerationLogErrorHandler::OnInvalidReferenceLevelScriptDataLayers(const FWorldPartitionActorDescView& ActorDescView)
+void FStreamingGenerationLogErrorHandler::OnInvalidReferenceLevelScriptDataLayers(const IWorldPartitionActorDescInstanceView& ActorDescView)
 {
 	UE_ASSET_LOG_ACTORDESCVIEW(LogWorldPartition, Log, ActorDescView, TEXT("Level Script Blueprint references streamed actor %s with a non empty set of data layers"), *GetActorName(ActorDescView));
 }
@@ -62,12 +61,12 @@ void FStreamingGenerationLogErrorHandler::OnDataLayerAssetConflict(const UDataLa
 	UE_ASSET_LOG(LogWorldPartition, Log, DataLayerInstance, TEXT("Data Layer Instance %s and Data Layer Instance %s are both referencing Data Layer Asset %s"), *DataLayerInstance->GetDataLayerFName().ToString(), *ConflictingDataLayerInstance->GetDataLayerFName().ToString(), *DataLayerInstance->GetAsset()->GetFullName());
 }
 
-void FStreamingGenerationLogErrorHandler::OnActorNeedsResave(const FWorldPartitionActorDescView& ActorDescView)
+void FStreamingGenerationLogErrorHandler::OnActorNeedsResave(const IWorldPartitionActorDescInstanceView& ActorDescView)
 {
 	UE_ASSET_LOG_ACTORDESCVIEW(LogWorldPartition, Log, ActorDescView, TEXT("Actor %s needs to be resaved"), *GetActorName(ActorDescView));
 }
 
-void FStreamingGenerationLogErrorHandler::OnLevelInstanceInvalidWorldAsset(const FWorldPartitionActorDescView& ActorDescView, FName WorldAsset, ELevelInstanceInvalidReason Reason)
+void FStreamingGenerationLogErrorHandler::OnLevelInstanceInvalidWorldAsset(const IWorldPartitionActorDescInstanceView& ActorDescView, FName WorldAsset, ELevelInstanceInvalidReason Reason)
 {
 	const FString ActorName = GetActorName(ActorDescView);
 
@@ -91,12 +90,12 @@ void FStreamingGenerationLogErrorHandler::OnLevelInstanceInvalidWorldAsset(const
 	};
 }
 
-void FStreamingGenerationLogErrorHandler::OnInvalidActorFilterReference(const FWorldPartitionActorDescView& ActorDescView, const FWorldPartitionActorDescView& ReferenceActorDescView)
+void FStreamingGenerationLogErrorHandler::OnInvalidActorFilterReference(const IWorldPartitionActorDescInstanceView& ActorDescView, const IWorldPartitionActorDescInstanceView& ReferenceActorDescView)
 {
 	UE_ASSET_LOG_ACTORDESCVIEW(LogWorldPartition, Log, ReferenceActorDescView, TEXT("Actor %s will not be filtered out because it is referenced by Actor %s not part of the filter"), *GetActorName(ReferenceActorDescView), *GetActorName(ActorDescView));
 }
 
-void FStreamingGenerationLogErrorHandler::OnInvalidHLODLayer(const FWorldPartitionActorDescView& ActorDescView)
+void FStreamingGenerationLogErrorHandler::OnInvalidHLODLayer(const IWorldPartitionActorDescInstanceView& ActorDescView)
 {
 	UE_ASSET_LOG_ACTORDESCVIEW(LogWorldPartition, Log, ActorDescView, TEXT("Actor %s has an invalid HLOD layer %s"), *GetActorName(ActorDescView), *ActorDescView.GetHLODLayer().ToString());
 }
diff --git a/Engine/Source/Runtime/Engine/Private/WorldPartition/ErrorHandling/WorldPartitionStreamingGenerationProxyErrorHandler.cpp b/Engine/Source/Runtime/Engine/Private/WorldPartition/ErrorHandling/WorldPartitionStreamingGenerationProxyErrorHandler.cpp
--- a/Engine/Source/Runtime/Engine/Private/WorldPartition/ErrorHandling/WorldPartitionStreamingGenerationProxyErrorHandler.cpp
+++ b/Engine/Source/Runtime/Engine/Private/WorldPartition/ErrorHandling/WorldPartitionStreamingGenerationProxyErrorHandler.cpp
@@ -3,37 +3,37 @@
 #if WITH_EDITOR
 #include "WorldPartition/ErrorHandling/WorldPartitionStreamingGenerationProxyErrorHandler.h"
 
-void FStreamingGenerationProxyErrorHandler::OnInvalidRuntimeGrid(const FWorldPartitionActorDescView& ActorDescView, FName GridName)
+void FStreamingGenerationProxyErrorHandler::OnInvalidRuntimeGrid(const IWorldPartitionActorDescInstanceView& ActorDescView, FName GridName)
 {
 	InnerErrorHandler->OnInvalidRuntimeGrid(ActorDescView, GridName);
 }
 
-void FStreamingGenerationProxyErrorHandler::OnInvalidReference(const FWorldPartitionActorDescView& ActorDescView, const FGuid& ReferenceGuid, FWorldPartitionActorDescView* ReferenceActorDescView)
+void FStreamingGenerationProxyErrorHandler::OnInvalidReference(const IWorldPartitionActorDescInstanceView& ActorDescView, const FGuid& ReferenceGuid, IWorldPartitionActorDescInstanceView* ReferenceActorDescView)
 {
 	InnerErrorHandler->OnInvalidReference(ActorDescView, ReferenceGuid, ReferenceActorDescView);
 }
 
-void FStreamingGenerationProxyErrorHandler::OnInvalidReferenceGridPlacement(const FWorldPartitionActorDescView& ActorDescView, const FWorldPartitionActorDescView& ReferenceActorDescView)
+void FStreamingGenerationProxyErrorHandler::OnInvalidReferenceGridPlacement(const IWorldPartitionActorDescInstanceView& ActorDescView, const IWorldPartitionActorDescInstanceView& ReferenceActorDescView)
 {
 	InnerErrorHandler->OnInvalidReferenceGridPlacement(ActorDescView, ReferenceActorDescView);
 }
 
-void FStreamingGenerationProxyErrorHandler::OnInvalidReferenceDataLayers(const FWorldPartitionActorDescView& ActorDescView, const FWorldPartitionActorDescView& ReferenceActorDescView)
+void FStreamingGenerationProxyErrorHandler::OnInvalidReferenceDataLayers(const IWorldPartitionActorDescInstanceView& ActorDescView, const IWorldPartitionActorDescInstanceView& ReferenceActorDescView)
 {
 	InnerErrorHandler->OnInvalidReferenceDataLayers(ActorDescView, ReferenceActorDescView);
 }
 
-void FStreamingGenerationProxyErrorHandler::OnInvalidReferenceRuntimeGrid(const FWorldPartitionActorDescView& ActorDescView, const FWorldPartitionActorDescView& ReferenceActorDescView)
+void FStreamingGenerationProxyErrorHandler::OnInvalidReferenceRuntimeGrid(const IWorldPartitionActorDescInstanceView& ActorDescView, const IWorldPartitionActorDescInstanceView& ReferenceActorDescView)
 {
 	InnerErrorHandler->OnInvalidReferenceRuntimeGrid(ActorDescView, ReferenceActorDescView);
 }
 
-void FStreamingGenerationProxyErrorHandler::OnInvalidReferenceLevelScriptStreamed(const FWorldPartitionActorDescView& ActorDescView)
+void FStreamingGenerationProxyErrorHandler::OnInvalidReferenceLevelScriptStreamed(const IWorldPartitionActorDescInstanceView& ActorDescView)
 {
 	InnerErrorHandler->OnInvalidReferenceLevelScriptStreamed(ActorDescView);
 }
 
-void FStreamingGenerationProxyErrorHandler::OnInvalidReferenceLevelScriptDataLayers(const FWorldPartitionActorDescView& ActorDescView)
+void FStreamingGenerationProxyErrorHandler::OnInvalidReferenceLevelScriptDataLayers(const IWorldPartitionActorDescInstanceView& ActorDescView)
 {
 	InnerErrorHandler->OnInvalidReferenceLevelScriptDataLayers(ActorDescView);
 }
@@ -53,22 +53,22 @@ void FStreamingGenerationProxyErrorHandler::OnDataLayerAssetConflict(const UData
 	InnerErrorHandler->OnDataLayerAssetConflict(DataLayerInstance, ConflictingDataLayerInstance);
 }
 
-void FStreamingGenerationProxyErrorHandler::OnActorNeedsResave(const FWorldPartitionActorDescView& ActorDescView)
+void FStreamingGenerationProxyErrorHandler::OnActorNeedsResave(const IWorldPartitionActorDescInstanceView& ActorDescView)
 {
 	InnerErrorHandler->OnActorNeedsResave(ActorDescView);
 }
 
-void FStreamingGenerationProxyErrorHandler::OnLevelInstanceInvalidWorldAsset(const FWorldPartitionActorDescView& ActorDescView, FName WorldAsset, ELevelInstanceInvalidReason Reason)
+void FStreamingGenerationProxyErrorHandler::OnLevelInstanceInvalidWorldAsset(const IWorldPartitionActorDescInstanceView& ActorDescView, FName WorldAsset, ELevelInstanceInvalidReason Reason)
 {
 	InnerErrorHandler->OnLevelInstanceInvalidWorldAsset(ActorDescView, WorldAsset, Reason);
 }
 
-void FStreamingGenerationProxyErrorHandler::OnInvalidActorFilterReference(const FWorldPartitionActorDescView& ActorDescView, const FWorldPartitionActorDescView& ReferenceActorDescView)
+void FStreamingGenerationProxyErrorHandler::OnInvalidActorFilterReference(const IWorldPartitionActorDescInstanceView& ActorDescView, const IWorldPartitionActorDescInstanceView& ReferenceActorDescView)
 {
 	InnerErrorHandler->OnInvalidActorFilterReference(ActorDescView, ReferenceActorDescView);
 }
 
-void FStreamingGenerationProxyErrorHandler::OnInvalidHLODLayer(const FWorldPartitionActorDescView& ActorDescView)
+void FStreamingGenerationProxyErrorHandler::OnInvalidHLODLayer(const IWorldPartitionActorDescInstanceView& ActorDescView)
 {
 	InnerErrorHandler->OnInvalidHLODLayer(ActorDescView);
 }
diff --git a/Engine/Source/Runtime/Engine/Private/WorldPartition/ErrorHandling/WorldPartitionStreamingGenerationTokenizedMessageErrorHandler.cpp b/Engine/Source/Runtime/Engine/Private/WorldPartition/ErrorHandling/WorldPartitionStreamingGenerationTokenizedMessageErrorHandler.cpp
--- a/Engine/Source/Runtime/Engine/Private/WorldPartition/ErrorHandling/WorldPartitionStreamingGenerationTokenizedMessageErrorHandler.cpp
+++ b/Engine/Source/Runtime/Engine/Private/WorldPartition/ErrorHandling/WorldPartitionStreamingGenerationTokenizedMessageErrorHandler.cpp
@@ -7,11 +7,10 @@
 #include "Misc/UObjectToken.h"
 #include "Misc/MapErrors.h"
 #include "WorldPartition/DataLayer/DataLayerInstanceWithAsset.h"
-#include "WorldPartition/WorldPartitionActorDescView.h"
 
 #define LOCTEXT_NAMESPACE "WorldPartition"
 
-void ITokenizedMessageErrorHandler::OnInvalidRuntimeGrid(const FWorldPartitionActorDescView& ActorDescView, FName GridName)
+void ITokenizedMessageErrorHandler::OnInvalidRuntimeGrid(const IWorldPartitionActorDescInstanceView& ActorDescView, FName GridName)
 {
 	TSharedRef<FTokenizedMessage> Message = FTokenizedMessage::Create(EMessageSeverity::Error);
 	Message->AddToken(FTextToken::Create(LOCTEXT("TokenMessage_WorldPartition_Actor", "Actor")))
@@ -23,7 +22,7 @@ void ITokenizedMessageErrorHandler::OnInvalidRuntimeGrid(const FWorldPartitionAc
 	HandleTokenizedMessage(MoveTemp(Message));
 }
 
-void ITokenizedMessageErrorHandler::OnInvalidReference(const FWorldPartitionActorDescView& ActorDescView, const FGuid& ReferenceGuid, FWorldPartitionActorDescView* ReferenceActorDescView)
+void ITokenizedMessageErrorHandler::OnInvalidReference(const IWorldPartitionActorDescInstanceView& ActorDescView, const FGuid& ReferenceGuid, IWorldPartitionActorDescInstanceView* ReferenceActorDescView)
 {
 	// Don't report invalid references to non-existing actors as it can happen in valid scenarios, as the linker code will silently null out references when loading actors with missing
 	// references.
@@ -42,7 +41,7 @@ void ITokenizedMessageErrorHandler::OnInvalidReference(const FWorldPartitionActo
 	}
 }
 
-void ITokenizedMessageErrorHandler::OnInvalidReferenceGridPlacement(const FWorldPartitionActorDescView& ActorDescView, const FWorldPartitionActorDescView& ReferenceActorDescView)
+void ITokenizedMessageErrorHandler::OnInvalidReferenceGridPlacement(const IWorldPartitionActorDescInstanceView& ActorDescView, const IWorldPartitionActorDescInstanceView& ReferenceActorDescView)
 {
 	const FText SpatiallyLoadedActor(LOCTEXT("TokenMessage_WorldPartition_SpatiallyLoadedActor", "Spatially loaded actor"));
 	const FText NonSpatiallyLoadedActor(LOCTEXT("TokenMessage_WorldPartition_NonSpatiallyLoadedActor", "Non-spatially loaded actor"));
@@ -59,7 +58,7 @@ void ITokenizedMessageErrorHandler::OnInvalidReferenceGridPlacement(const FWorld
 	HandleTokenizedMessage(MoveTemp(Message));
 }
 
-void ITokenizedMessageErrorHandler::OnInvalidReferenceDataLayers(const FWorldPartitionActorDescView& ActorDescView, const FWorldPartitionActorDescView& ReferenceActorDescView)
+void ITokenizedMessageErrorHandler::OnInvalidReferenceDataLayers(const IWorldPartitionActorDescInstanceView& ActorDescView, const IWorldPartitionActorDescInstanceView& ReferenceActorDescView)
 {
 	TSharedRef<FTokenizedMessage> Message = FTokenizedMessage::Create(EMessageSeverity::Error);
 	Message->AddToken(FTextToken::Create(LOCTEXT("TokenMessage_WorldPartition_Actor", "Actor")))
@@ -71,7 +70,7 @@ void ITokenizedMessageErrorHandler::OnInvalidReferenceDataLayers(const FWorldPar
 	HandleTokenizedMessage(MoveTemp(Message));
 }
 
-void ITokenizedMessageErrorHandler::OnInvalidReferenceRuntimeGrid(const FWorldPartitionActorDescView& ActorDescView, const FWorldPartitionActorDescView& ReferenceActorDescView)
+void ITokenizedMessageErrorHandler::OnInvalidReferenceRuntimeGrid(const IWorldPartitionActorDescInstanceView& ActorDescView, const IWorldPartitionActorDescInstanceView& ReferenceActorDescView)
 {
 	TSharedRef<FTokenizedMessage> Message = FTokenizedMessage::Create(EMessageSeverity::Error);
 	Message->AddToken(FTextToken::Create(LOCTEXT("TokenMessage_WorldPartition_Actor", "Actor")))
@@ -83,7 +82,7 @@ void ITokenizedMessageErrorHandler::OnInvalidReferenceRuntimeGrid(const FWorldPa
 	HandleTokenizedMessage(MoveTemp(Message));
 }
 
-void ITokenizedMessageErrorHandler::OnInvalidReferenceLevelScriptStreamed(const FWorldPartitionActorDescView& ActorDescView)
+void ITokenizedMessageErrorHandler::OnInvalidReferenceLevelScriptStreamed(const IWorldPartitionActorDescInstanceView& ActorDescView)
 {
 	TSharedRef<FTokenizedMessage> Message = FTokenizedMessage::Create(EMessageSeverity::Error);
 	Message->AddToken(FTextToken::Create(LOCTEXT("TokenMessage_WorldPartition_LevelScriptBlueprintStreamedActorReference", "Level Script Blueprint references streamed actor")))
@@ -94,7 +93,7 @@ void ITokenizedMessageErrorHandler::OnInvalidReferenceLevelScriptStreamed(const
 	
 }
 
-void ITokenizedMessageErrorHandler::OnInvalidReferenceLevelScriptDataLayers(const FWorldPartitionActorDescView& ActorDescView)
+void ITokenizedMessageErrorHandler::OnInvalidReferenceLevelScriptDataLayers(const IWorldPartitionActorDescInstanceView& ActorDescView)
 {
 	TSharedRef<FTokenizedMessage> Message = FTokenizedMessage::Create(EMessageSeverity::Error);
 	Message->AddToken(FTextToken::Create(LOCTEXT("TokenMessage_WorldPartition_LevelScriptBlueprintActorReference", "Level Script Blueprint references actor")))
@@ -143,7 +142,7 @@ void ITokenizedMessageErrorHandler::OnDataLayerAssetConflict(const UDataLayerIns
 	HandleTokenizedMessage(MoveTemp(Message));
 }
 
-void ITokenizedMessageErrorHandler::OnActorNeedsResave(const FWorldPartitionActorDescView& ActorDescView)
+void ITokenizedMessageErrorHandler::OnActorNeedsResave(const IWorldPartitionActorDescInstanceView& ActorDescView)
 {
 	TSharedRef<FTokenizedMessage> Message = FTokenizedMessage::Create(EMessageSeverity::Info);
 	Message->AddToken(FTextToken::Create(LOCTEXT("TokenMessage_WorldPartition_ActorNeedsResave", "Actor needs resave")))
@@ -153,7 +152,7 @@ void ITokenizedMessageErrorHandler::OnActorNeedsResave(const FWorldPartitionActo
 	HandleTokenizedMessage(MoveTemp(Message));
 }
 
-void ITokenizedMessageErrorHandler::OnLevelInstanceInvalidWorldAsset(const FWorldPartitionActorDescView& ActorDescView, FName WorldAsset, ELevelInstanceInvalidReason Reason)
+void ITokenizedMessageErrorHandler::OnLevelInstanceInvalidWorldAsset(const IWorldPartitionActorDescInstanceView& ActorDescView, FName WorldAsset, ELevelInstanceInvalidReason Reason)
 {
 	FSoftObjectPath ActorPath(ActorDescView.GetActorSoftPath());
 	EMessageSeverity::Type MessageSeverity = EMessageSeverity::Info;
@@ -196,7 +195,7 @@ void ITokenizedMessageErrorHandler::OnLevelInstanceInvalidWorldAsset(const FWorl
 	HandleTokenizedMessage(MoveTemp(Message));
 }
 
-void ITokenizedMessageErrorHandler::OnInvalidActorFilterReference(const FWorldPartitionActorDescView& ActorDescView, const FWorldPartitionActorDescView& ReferenceActorDescView)
+void ITokenizedMessageErrorHandler::OnInvalidActorFilterReference(const IWorldPartitionActorDescInstanceView& ActorDescView, const IWorldPartitionActorDescInstanceView& ReferenceActorDescView)
 {
 	TSharedRef<FTokenizedMessage> Message = FTokenizedMessage::Create(EMessageSeverity::Error);
 	Message->AddToken(FTextToken::Create(LOCTEXT("TokenMessage_WorldPartition_Actor", "Actor")))
@@ -208,7 +207,7 @@ void ITokenizedMessageErrorHandler::OnInvalidActorFilterReference(const FWorldPa
 	HandleTokenizedMessage(MoveTemp(Message));
 }
 
-void ITokenizedMessageErrorHandler::OnInvalidHLODLayer(const FWorldPartitionActorDescView& ActorDescView)
+void ITokenizedMessageErrorHandler::OnInvalidHLODLayer(const IWorldPartitionActorDescInstanceView& ActorDescView)
 {
 	TSharedRef<FTokenizedMessage> Message = FTokenizedMessage::Create(EMessageSeverity::Error);
 	Message->AddToken(FTextToken::Create(LOCTEXT("TokenMessage_WorldPartition_Actor", "Actor")))
diff --git a/Engine/Source/Runtime/Engine/Private/WorldPartition/HLOD/HLODActorDesc.cpp b/Engine/Source/Runtime/Engine/Private/WorldPartition/HLOD/HLODActorDesc.cpp
--- a/Engine/Source/Runtime/Engine/Private/WorldPartition/HLOD/HLODActorDesc.cpp
+++ b/Engine/Source/Runtime/Engine/Private/WorldPartition/HLOD/HLODActorDesc.cpp
@@ -16,6 +16,7 @@
 #include "WorldPartition/HLOD/HLODActor.h"
 #include "WorldPartition/HLOD/HLODLayer.h"
 #include "WorldPartition/HLOD/HLODSourceActorsFromCell.h"
+#include "WorldPartition/WorldPartitionActorDescInstance.h"
 
 FHLODActorDesc::FHLODActorDesc()
 	: EditorBounds(ForceInit)
@@ -38,8 +39,8 @@ void FHLODActorDesc::Init(const AActor* InActor)
 		{
 			if (SubActor.ContainerID.IsMainContainer())
 			{
-				const FWorldPartitionActorDesc* SubActorDesc = WorldPartition->GetActorDesc(SubActor.ActorInstanceGuid);
-				if (SubActorDesc && SubActorDesc->GetActorNativeClass()->IsChildOf<AWorldPartitionHLOD>())
+				const FWorldPartitionActorDescInstance* SubActorDescInstance = WorldPartition->GetActorDescInstance(SubActor.ActorInstanceGuid);
+				if (SubActorDescInstance && SubActorDescInstance->GetActorNativeClass()->IsChildOf<AWorldPartitionHLOD>())
 				{
 					ChildHLODActors.Add(SubActor.ActorInstanceGuid);
 				}
@@ -198,4 +199,9 @@ int64 FHLODActorDesc::GetPackageSize(const AWorldPartitionHLOD* InHLODActor)
 	return ::GetPackageSize(PackageFileName);
 }
 
+bool FHLODActorDesc::IsRuntimeRelevant(const FWorldPartitionActorDescInstance* InActorDescInstance) const
+{
+	return !InActorDescInstance->GetForceNonSpatiallyLoaded();
+}
+
 #endif
diff --git a/Engine/Source/Runtime/Engine/Private/WorldPartition/HLOD/HLODLayer.cpp b/Engine/Source/Runtime/Engine/Private/WorldPartition/HLOD/HLODLayer.cpp
--- a/Engine/Source/Runtime/Engine/Private/WorldPartition/HLOD/HLODLayer.cpp
+++ b/Engine/Source/Runtime/Engine/Private/WorldPartition/HLOD/HLODLayer.cpp
@@ -16,7 +16,6 @@
 #include "WorldPartition/HLOD/IWorldPartitionHLODUtilities.h"
 #include "WorldPartition/HLOD/IWorldPartitionHLODUtilitiesModule.h"
 #include "WorldPartition/WorldPartition.h"
-#include "WorldPartition/WorldPartitionActorDescView.h"
 #endif
 
 #include UE_INLINE_GENERATED_CPP_BY_NAME(HLODLayer)
diff --git a/Engine/Source/Runtime/Engine/Private/WorldPartition/HLOD/HLODRuntimeSubsystem.cpp b/Engine/Source/Runtime/Engine/Private/WorldPartition/HLOD/HLODRuntimeSubsystem.cpp
--- a/Engine/Source/Runtime/Engine/Private/WorldPartition/HLOD/HLODRuntimeSubsystem.cpp
+++ b/Engine/Source/Runtime/Engine/Private/WorldPartition/HLOD/HLODRuntimeSubsystem.cpp
@@ -28,6 +28,7 @@
 #include "UObject/UObjectIterator.h"
 #include "WorldPartition/DataLayer/WorldDataLayers.h"
 #include "WorldPartition/HLOD/HLODStats.h"
+#include "WorldPartition/WorldPartitionActorDescInstance.h"
 
 #include UE_INLINE_GENERATED_CPP_BY_NAME(HLODRuntimeSubsystem)
 
@@ -850,11 +851,11 @@ bool UWorldPartitionHLODRuntimeSubsystem::WriteHLODStatsCSV(UWorld* InWorld, con
 		return false;
 	}
 	
-	typedef TFunction<FString(const FHLODActorDesc&)> FGetStatFunc;
+	typedef TFunction<FString(FWorldPartitionActorDescInstance* InActorDescInstance, const FHLODActorDesc&)> FGetStatFunc;
 
 	auto GetHLODStat = [](FName InStatName)
 	{
-		return TPair<FName, FGetStatFunc>(InStatName, [InStatName](const FHLODActorDesc& InActorDesc)
+		return TPair<FName, FGetStatFunc>(InStatName, [InStatName](FWorldPartitionActorDescInstance* InActorDescInstance, const FHLODActorDesc& InActorDesc)
 		{
 			return FString::Printf(TEXT("%lld"), InActorDesc.GetStat(InStatName));
 		});
@@ -870,11 +871,11 @@ bool UWorldPartitionHLODRuntimeSubsystem::WriteHLODStatsCSV(UWorld* InWorld, con
 
 	TArray<TPair<FName, FGetStatFunc>> StatsToWrite =
 	{	
-		{ "WorldPackage",		[InWorld](const FHLODActorDesc& InActorDesc) { return InWorld->GetPackage()->GetName(); } },
-		{ "Name",				[](const FHLODActorDesc& InActorDesc) { return InActorDesc.GetActorLabel().ToString(); } },
-		{ "HLODLayer",			[](const FHLODActorDesc& InActorDesc) { return InActorDesc.GetSourceHLODLayer().GetAssetName().ToString(); }},
-		{ "SpatiallyLoaded",	[](const FHLODActorDesc& InActorDesc) { return InActorDesc.GetIsSpatiallyLoaded() ? TEXT("true") : TEXT("false"); } },
-		{ "DataLayers",			[&GetDataLayerShortName](const FHLODActorDesc& InActorDesc) { return FString::JoinBy(InActorDesc.GetDataLayerInstanceNames(), TEXT(" | "), GetDataLayerShortName); } },
+		{ "WorldPackage",		[InWorld](FWorldPartitionActorDescInstance* InActorDescInstance, const FHLODActorDesc& InActorDesc) { return InWorld->GetPackage()->GetName(); } },
+		{ "Name",				[](FWorldPartitionActorDescInstance* InActorDescInstance, const FHLODActorDesc& InActorDesc) { return InActorDescInstance->GetActorLabel().ToString(); } },
+		{ "HLODLayer",			[](FWorldPartitionActorDescInstance* InActorDescInstance, const FHLODActorDesc& InActorDesc) { return InActorDesc.GetSourceHLODLayer().GetAssetName().ToString(); }},
+		{ "SpatiallyLoaded",	[](FWorldPartitionActorDescInstance* InActorDescInstance, const FHLODActorDesc& InActorDesc) { return InActorDescInstance->GetIsSpatiallyLoaded() ? TEXT("true") : TEXT("false"); } },
+		{ "DataLayers",			[&GetDataLayerShortName](FWorldPartitionActorDescInstance* InActorDescInstance, const FHLODActorDesc& InActorDesc) { return FString::JoinBy(InActorDescInstance->GetDataLayerInstanceNames(), TEXT(" | "), GetDataLayerShortName); } },
 
 		GetHLODStat(FWorldPartitionHLODStats::InputActorCount),
 		GetHLODStat(FWorldPartitionHLODStats::InputTriangleCount),
@@ -913,9 +914,13 @@ bool UWorldPartitionHLODRuntimeSubsystem::WriteHLODStatsCSV(UWorld* InWorld, con
 	}
 
 	// Write one line per HLOD actor desc
-	for (FActorDescContainerCollection::TIterator<AWorldPartitionHLOD> HLODIterator(WorldPartition); HLODIterator; ++HLODIterator)
+	for (FActorDescContainerInstanceCollection::TIterator<AWorldPartitionHLOD> HLODIterator(WorldPartition); HLODIterator; ++HLODIterator)
 	{
-		const FString StatLine = FString::JoinBy(StatsToWrite, TEXT(","), [&HLODIterator](const TPair<FName, FGetStatFunc>& Pair) { return Pair.Value(**HLODIterator); });
+		const FString StatLine = FString::JoinBy(StatsToWrite, TEXT(","), [&HLODIterator](const TPair<FName, FGetStatFunc>& Pair) 
+		{ 
+			const FHLODActorDesc& HLODActorDesc = *(FHLODActorDesc*)HLODIterator->GetActorDesc();
+			return Pair.Value(*HLODIterator, HLODActorDesc); 
+		});
 		Output.Logf(TEXT("%s" LINE_TERMINATOR_ANSI), *StatLine);
 	}
 
diff --git a/Engine/Source/Runtime/Engine/Private/WorldPartition/Landscape/LandscapeActorDesc.cpp b/Engine/Source/Runtime/Engine/Private/WorldPartition/Landscape/LandscapeActorDesc.cpp
--- a/Engine/Source/Runtime/Engine/Private/WorldPartition/Landscape/LandscapeActorDesc.cpp
+++ b/Engine/Source/Runtime/Engine/Private/WorldPartition/Landscape/LandscapeActorDesc.cpp
@@ -6,6 +6,7 @@
 #include "Landscape.h"
 #include "UObject/UE5MainStreamObjectVersion.h"
 #include "UObject/UE5ReleaseStreamObjectVersion.h"
+#include "WorldPartition/WorldPartitionActorDescInstance.h"
 
 void FLandscapeActorDesc::Init(const AActor* InActor)
 {
@@ -52,14 +53,14 @@ FBox FLandscapeActorDesc::GetEditorBounds() const
 	return FWorldPartitionActorDesc::GetEditorBounds();
 }
 
-void FLandscapeActorDesc::Unload()
+void FLandscapeActorDesc::OnUnloadingInstance(const FWorldPartitionActorDescInstance* InActorDescInstance) const
 {
-	if (ALandscapeProxy* LandscapeProxy = Cast<ALandscapeProxy>(GetActor()))
+	if (ALandscapeProxy* LandscapeProxy = Cast<ALandscapeProxy>(InActorDescInstance->GetActor()))
 	{
 		LandscapeProxy->ActorDescReferences.Empty();
 	}
 
-	FPartitionActorDesc::Unload();
+	FPartitionActorDesc::OnUnloadingInstance(InActorDescInstance);
 }
 
 bool FLandscapeActorDesc::Equals(const FWorldPartitionActorDesc* Other) const
diff --git a/Engine/Source/Runtime/Engine/Private/WorldPartition/LevelInstance/LevelInstanceActorDesc.cpp b/Engine/Source/Runtime/Engine/Private/WorldPartition/LevelInstance/LevelInstanceActorDesc.cpp
--- a/Engine/Source/Runtime/Engine/Private/WorldPartition/LevelInstance/LevelInstanceActorDesc.cpp
+++ b/Engine/Source/Runtime/Engine/Private/WorldPartition/LevelInstance/LevelInstanceActorDesc.cpp
@@ -1,19 +1,23 @@
 // Copyright Epic Games, Inc. All Rights Reserved.
 
 #include "WorldPartition/LevelInstance/LevelInstanceActorDesc.h"
-#include "WorldPartition/WorldPartitionActorDescView.h"
 
 #if WITH_EDITOR
+#include "Engine/Engine.h"
 #include "Engine/World.h"
 #include "Engine/Level.h"
 #include "LevelInstance/LevelInstanceInterface.h"
 #include "LevelInstance/LevelInstanceSubsystem.h"
 #include "LevelInstance/LevelInstanceComponent.h"
 #include "Misc/PackageName.h"
-#include "WorldPartition/ActorDescContainer.h"
+#include "WorldPartition/ActorDescContainerInstance.h"
+#include "WorldPartition/WorldPartitionActorDescInstance.h"
+#include "WorldPartition/WorldPartitionActorDescInstanceViewInterface.h"
 #include "WorldPartition/ErrorHandling/WorldPartitionStreamingGenerationErrorHandler.h"
 #include "WorldPartition/WorldPartitionActorDescArchive.h"
-#include "WorldPartition/WorldPartitionSubsystem.h"
+#include "WorldPartition/WorldPartitionHelpers.h"
+#include "WorldPartition/WorldPartition.h"
+#include "WorldPartition/ActorDescContainerSubsystem.h"
 #include "UObject/UE5ReleaseStreamObjectVersion.h"
 #include "UObject/UE5MainStreamObjectVersion.h"
 #include "UObject/FortniteMainBranchObjectVersion.h"
@@ -26,15 +30,17 @@ static FAutoConsoleVariableRef CVarForceLevelStreaming(
 
 FLevelInstanceActorDesc::FLevelInstanceActorDesc()
 	: DesiredRuntimeBehavior(ELevelInstanceRuntimeBehavior::Partitioned)
-	, LevelInstanceContainer(nullptr)
-	, bIsContainerInstance(false)
+	, ChildContainer(nullptr)
+	, bIsChildContainerInstance(false)
 {}
 
 FLevelInstanceActorDesc::~FLevelInstanceActorDesc()
 {
-	UnregisterContainerInstance();
-	check(!LevelInstanceContainer.IsValid());
-	check(!LevelInstanceContainerWorldContext.IsValid());
+	if (!IsEngineExitRequested())
+	{
+		UnregisterChildContainer();
+		check(!ChildContainer.IsValid());
+	}
 }
 
 void FLevelInstanceActorDesc::Init(const AActor* InActor)
@@ -47,7 +53,7 @@ void FLevelInstanceActorDesc::Init(const AActor* InActor)
 	DesiredRuntimeBehavior = LevelInstance->GetDesiredRuntimeBehavior();
 	Filter = LevelInstance->GetFilter();
 	
-	bIsContainerInstance = IsContainerInstanceInternal();
+	bIsChildContainerInstance = IsChildContainerInstanceInternal();
 }
 
 void FLevelInstanceActorDesc::Init(const FWorldPartitionActorDescInitData& DescData)
@@ -58,7 +64,12 @@ void FLevelInstanceActorDesc::Init(const FWorldPartitionActorDescInitData& DescD
 
 	FWorldPartitionActorDesc::Init(DescData);
 	
-	bIsContainerInstance = IsContainerInstanceInternal();
+	bIsChildContainerInstance = IsChildContainerInstanceInternal();
+}
+
+FString FLevelInstanceActorDesc::GetChildContainerName() const
+{
+	return GetChildContainerPackage().ToString();
 }
 
 bool FLevelInstanceActorDesc::Equals(const FWorldPartitionActorDesc* Other) const
@@ -78,78 +89,53 @@ bool FLevelInstanceActorDesc::Equals(const FWorldPartitionActorDesc* Other) cons
 
 void FLevelInstanceActorDesc::UpdateBounds()
 {
-	check(LevelInstanceContainerWorldContext.IsValid());
-
-	UWorldPartitionSubsystem* WorldPartitionSubsystem = UWorld::GetSubsystem<UWorldPartitionSubsystem>(LevelInstanceContainerWorldContext.Get());
-	check(WorldPartitionSubsystem);
-
-	FTransform LevelInstancePivotOffsetTransform = FTransform(ULevel::GetLevelInstancePivotOffsetFromPackage(GetContainerPackage()));
-	FTransform FinalLevelTransform = LevelInstancePivotOffsetTransform * LevelInstanceTransform;
-	FBox ContainerBounds = WorldPartitionSubsystem->GetContainerBounds(GetContainerPackage()).TransformBy(FinalLevelTransform);
+	if (UActorDescContainer* ChildContainerPtr = ChildContainer.Get())
+	{
+		FBox ContainerBounds = UActorDescContainerSubsystem::GetChecked().GetContainerBounds(GetChildContainerName()).TransformBy(GetChildContainerTransform());
 
-	ContainerBounds.GetCenterAndExtents(BoundsLocation, BoundsExtent);
+		ContainerBounds.GetCenterAndExtents(BoundsLocation, BoundsExtent);
+	}
 }
 
-void FLevelInstanceActorDesc::RegisterContainerInstance(UWorld* InWorldContext)
+void FLevelInstanceActorDesc::RegisterChildContainer()
 {
-	if (InWorldContext)
+	check(!ChildContainer.IsValid());
+	if (IsChildContainerInstance())
 	{
-		check(!LevelInstanceContainer.IsValid());
-		check(!LevelInstanceContainerWorldContext.IsValid());
-
-		if (IsContainerInstance())
-		{
-			check(InWorldContext);
-			LevelInstanceContainerWorldContext = InWorldContext;
-
-			UWorldPartitionSubsystem* WorldPartitionSubsystem = UWorld::GetSubsystem<UWorldPartitionSubsystem>(InWorldContext);
-			check(WorldPartitionSubsystem);
-
-			LevelInstanceContainer = WorldPartitionSubsystem->RegisterContainer(GetContainerPackage());
-			check(LevelInstanceContainer.IsValid());
-
-			// Should only be called on RegisterContainerInstance before ActorDesc is hashed
-			UpdateBounds();
-		}
+		ChildContainer = UActorDescContainerSubsystem::GetChecked().RegisterContainer({ GetChildContainerName(), GetChildContainerPackage() });
+		UpdateBounds();
 	}
 }
 
-void FLevelInstanceActorDesc::UnregisterContainerInstance()
+void FLevelInstanceActorDesc::UnregisterChildContainer()
 {
-	if (LevelInstanceContainer.IsValid())
+	if (UActorDescContainer* ChildContainerPtr = ChildContainer.Get())
 	{
-		check(LevelInstanceContainerWorldContext.IsValid());
-
-		UWorldPartitionSubsystem* WorldPartitionSubsystem = UWorld::GetSubsystem<UWorldPartitionSubsystem>(LevelInstanceContainerWorldContext.Get());
-		check(WorldPartitionSubsystem);
-
-		WorldPartitionSubsystem->UnregisterContainer(LevelInstanceContainer.Get());
-		LevelInstanceContainer.Reset();
+		UActorDescContainerSubsystem::GetChecked().UnregisterContainer(ChildContainerPtr);
 	}
-
-	LevelInstanceContainerWorldContext.Reset();
+	ChildContainer = nullptr;
 }
 
-void FLevelInstanceActorDesc::SetContainer(UActorDescContainer* InContainer, UWorld* InWorldContext)
+void FLevelInstanceActorDesc::SetContainer(UActorDescContainer* InContainer)
 {
-	FWorldPartitionActorDesc::SetContainer(InContainer, InWorldContext);
+	FWorldPartitionActorDesc::SetContainer(InContainer);
 
 	if (Container)
 	{
-		RegisterContainerInstance(InWorldContext);
+		RegisterChildContainer();
 	}
 	else
 	{
-		UnregisterContainerInstance();		
+		UnregisterChildContainer();
 	}
 }
 
-bool FLevelInstanceActorDesc::IsContainerInstance() const
+bool FLevelInstanceActorDesc::IsChildContainerInstance() const
 {
-	return bIsContainerInstance;
+	return bIsChildContainerInstance;
 }
 
-bool FLevelInstanceActorDesc::IsContainerInstanceInternal() const
+bool FLevelInstanceActorDesc::IsChildContainerInstanceInternal() const
 {
 	if (DesiredRuntimeBehavior != ELevelInstanceRuntimeBehavior::Partitioned)
 	{
@@ -166,7 +152,7 @@ bool FLevelInstanceActorDesc::IsContainerInstanceInternal() const
 		return false;
 	}
 	
-	if (!ULevel::GetIsLevelUsingExternalActorsFromPackage(GetContainerPackage()))
+	if (!ULevel::GetIsLevelUsingExternalActorsFromPackage(GetChildContainerPackage()))
 	{
 		return false;
 	}
@@ -174,39 +160,44 @@ bool FLevelInstanceActorDesc::IsContainerInstanceInternal() const
 	return true;
 }
 
-bool FLevelInstanceActorDesc::GetContainerInstance(FContainerInstance& OutContainerInstance) const
+FTransform FLevelInstanceActorDesc::GetChildContainerTransform() const
 {
-	if (LevelInstanceContainer.IsValid())
+	FTransform LevelInstancePivotOffsetTransform = FTransform(ULevel::GetLevelInstancePivotOffsetFromPackage(GetChildContainerPackage()));
+	return LevelInstancePivotOffsetTransform * LevelInstanceTransform;
+}
+
+bool FLevelInstanceActorDesc::GetChildContainerInstance(const FWorldPartitionActorDescInstance* InActorDescInstance, FContainerInstance& OutContainerInstance) const
+{
+	if (UActorDescContainerInstance* ContainerInstance = InActorDescInstance->GetChildContainerInstance())
 	{
-		const ILevelInstanceInterface* LevelInstance = ActorPtr.IsValid() ? CastChecked<const ILevelInstanceInterface>(ActorPtr.Get()) : nullptr;
-		OutContainerInstance.Container = LevelInstanceContainer.Get();
-		OutContainerInstance.LoadedLevel = LevelInstance ? LevelInstance->GetLoadedLevel() : nullptr;
-		OutContainerInstance.bSupportsPartialEditorLoading = LevelInstance ? LevelInstance->SupportsPartialEditorLoading() : false;
+		OutContainerInstance.ContainerInstance = ContainerInstance;
 		OutContainerInstance.ClusterMode = EContainerClusterMode::Partitioned;
 
-		// Apply level instance pivot offset
-		FTransform LevelInstancePivotOffsetTransform = FTransform(ULevel::GetLevelInstancePivotOffsetFromPackage(LevelInstanceContainer->GetContainerPackage()));
-		OutContainerInstance.Transform = LevelInstancePivotOffsetTransform * LevelInstanceTransform;		
+		// @todo_pat: this is to validate that new parenting of container instance code is equivalent
+		OutContainerInstance.Transform = GetChildContainerTransform();
+
 		return true;
 	}
 
 	return false;
 }
 
-void FLevelInstanceActorDesc::CheckForErrors(IStreamingGenerationErrorHandler* ErrorHandler) const
+void FLevelInstanceActorDesc::CheckForErrors(const IWorldPartitionActorDescInstanceView* InActorDescView, IStreamingGenerationErrorHandler* ErrorHandler) const
 {
-	FWorldPartitionActorDesc::CheckForErrors(ErrorHandler);
+	FWorldPartitionActorDesc::CheckForErrors(InActorDescView, ErrorHandler);
+
+	const FName ChildContainerPackage = InActorDescView->GetChildContainerPackage();
 
 	FPackagePath WorldAssetPath;
-	if (!FPackagePath::TryFromPackageName(GetContainerPackage(), WorldAssetPath) || !FPackageName::DoesPackageExist(WorldAssetPath))
+	if (!FPackagePath::TryFromPackageName(ChildContainerPackage, WorldAssetPath) || !FPackageName::DoesPackageExist(WorldAssetPath))
 	{
-		ErrorHandler->OnLevelInstanceInvalidWorldAsset(this, GetContainerPackage(), IStreamingGenerationErrorHandler::ELevelInstanceInvalidReason::WorldAssetNotFound);
+		ErrorHandler->OnLevelInstanceInvalidWorldAsset(*InActorDescView, ChildContainerPackage, IStreamingGenerationErrorHandler::ELevelInstanceInvalidReason::WorldAssetNotFound);
 	}
-	else if (!ULevel::GetIsLevelUsingExternalActorsFromPackage(GetContainerPackage()))
+	else if (!ULevel::GetIsLevelUsingExternalActorsFromPackage(ChildContainerPackage))
 	{
 		if (DesiredRuntimeBehavior != ELevelInstanceRuntimeBehavior::LevelStreaming)
 		{
-			ErrorHandler->OnLevelInstanceInvalidWorldAsset(this, GetContainerPackage(), IStreamingGenerationErrorHandler::ELevelInstanceInvalidReason::WorldAssetNotUsingExternalActors);
+			ErrorHandler->OnLevelInstanceInvalidWorldAsset(*InActorDescView, ChildContainerPackage, IStreamingGenerationErrorHandler::ELevelInstanceInvalidReason::WorldAssetNotUsingExternalActors);
 		}
 	}
 }
@@ -217,13 +208,39 @@ void FLevelInstanceActorDesc::TransferFrom(const FWorldPartitionActorDesc* From)
 
 	FLevelInstanceActorDesc* FromLevelInstanceActorDesc = (FLevelInstanceActorDesc*)From;
 
-	// Use the Register/Unregister so callbacks are added/removed
-	if (FromLevelInstanceActorDesc->LevelInstanceContainer.IsValid())
+	RegisterChildContainer();
+	FromLevelInstanceActorDesc->UnregisterChildContainer();
+}
+
+bool FLevelInstanceActorDesc::GetLoadedChildContainerInstance(const FWorldPartitionActorDescInstance* InActorDescInstance, FLoadedContainerInstance& OutContainerInstance) const
+{
+	if (const ILevelInstanceInterface* LevelInstance = Cast<ILevelInstanceInterface>(InActorDescInstance->GetActor()))
 	{
-		check(FromLevelInstanceActorDesc->LevelInstanceContainerWorldContext.IsValid());
-		RegisterContainerInstance(FromLevelInstanceActorDesc->LevelInstanceContainerWorldContext.Get());
-		FromLevelInstanceActorDesc->UnregisterContainerInstance();
+		if (ULevel* LoadedLevel = LevelInstance->GetLoadedLevel())
+		{
+			OutContainerInstance.LoadedLevel = LoadedLevel;
+			OutContainerInstance.bSupportsPartialEditorLoading = LevelInstance->SupportsPartialEditorLoading();
+			return true;
+		}
 	}
+
+	return false;
+}
+
+UActorDescContainerInstance* FLevelInstanceActorDesc::CreateChildContainerInstance(const FWorldPartitionActorDescInstance* InActorDescInstance) const
+{
+	UActorDescContainerInstance* ContainerInstance = InActorDescInstance->GetContainerInstance();
+	UActorDescContainerInstance* ChildContainerInstance = NewObject<UActorDescContainerInstance>(ContainerInstance, NAME_None, RF_Transient);
+	
+	// When a child container instance is created we create the whole hierarchy (for generate streaming)
+	const bool bInCreateContainerInstanceHierarchy = true;
+	UActorDescContainerInstance::FInitializeParams InitParams(InActorDescInstance->GetChildContainerPackage(), bInCreateContainerInstanceHierarchy);
+	InitParams.SetParent(ContainerInstance, InActorDescInstance->GetGuid())
+		.SetTransform(GetChildContainerTransform() * ContainerInstance->GetTransform());
+
+	ChildContainerInstance->Initialize(InitParams);
+		
+	return ChildContainerInstance;
 }
 
 void FLevelInstanceActorDesc::TransferWorldData(const FWorldPartitionActorDesc* From)
@@ -289,10 +306,10 @@ void FLevelInstanceActorDesc::Serialize(FArchive& Ar)
 			const ILevelInstanceInterface* LevelInstanceCDO = CastChecked<ILevelInstanceInterface>(CDO);
 			if (WorldAsset.IsValid() && (LevelInstanceCDO->IsLoadingEnabled() || bFixupOldVersion))
 			{
-				if (!IsContainerInstance())
+				if (!IsChildContainerInstance())
 				{
 					FBox OutBounds;
-					if (ULevelInstanceSubsystem::GetLevelInstanceBoundsFromPackage(LevelInstanceTransform, GetContainerPackage(), OutBounds))
+					if (ULevelInstanceSubsystem::GetLevelInstanceBoundsFromPackage(LevelInstanceTransform, GetChildContainerPackage(), OutBounds))
 					{
 						OutBounds.GetCenterAndExtents(BoundsLocation, BoundsExtent);
 					}
diff --git a/Engine/Source/Runtime/Engine/Private/WorldPartition/LoaderAdapter/LoaderAdapterActorList.cpp b/Engine/Source/Runtime/Engine/Private/WorldPartition/LoaderAdapter/LoaderAdapterActorList.cpp
--- a/Engine/Source/Runtime/Engine/Private/WorldPartition/LoaderAdapter/LoaderAdapterActorList.cpp
+++ b/Engine/Source/Runtime/Engine/Private/WorldPartition/LoaderAdapter/LoaderAdapterActorList.cpp
@@ -2,6 +2,7 @@
 
 #include "WorldPartition/LoaderAdapter/LoaderAdapterActorList.h"
 #include "WorldPartition/WorldPartition.h"
+#include "WorldPartition/WorldPartitionActorDescInstance.h"
 #include "Engine/Level.h"
 #include "Engine/World.h"
 
@@ -64,18 +65,18 @@ void FLoaderAdapterActorList::RemoveActors(const TArray<FWorldPartitionHandle>&
 
 	for (const FWorldPartitionHandle& ActorHandle : ActorHandles)
 	{
-		if (ActorHandle->IsContainerInstance())
+		if (ActorHandle.GetInstance()->IsChildContainerInstance())
 		{
-			FWorldPartitionActorDesc::FContainerInstance ContainerInstance;
-			if (ActorHandle->GetContainerInstance(ContainerInstance))
+			FWorldPartitionActorDesc::FLoadedContainerInstance ContainerInstance;
+			if (IWorldPartitionActorLoaderInterface::GetLoadedChildContainerInstance(ActorHandle, ContainerInstance))
 			{
 				if (ContainerInstance.bSupportsPartialEditorLoading)
 				{
 					if (UWorldPartition* ContainerWorldPartition = ContainerInstance.LoadedLevel ? ContainerInstance.LoadedLevel->GetWorldPartition() : nullptr)
 					{
-						for (FActorDescContainerCollection::TIterator<> ActorDescIterator(ContainerWorldPartition); ActorDescIterator; ++ActorDescIterator)
+						for (FActorDescContainerInstanceCollection::TIterator<> Iterator(ContainerWorldPartition); Iterator; ++Iterator)
 						{
-							FWorldPartitionHandle SubActorHandle(ContainerWorldPartition, ActorDescIterator->GetGuid());
+							FWorldPartitionHandle SubActorHandle(ContainerWorldPartition, Iterator->GetGuid());
 							ActorsToRemove.Add(SubActorHandle);
 						}
 					}
diff --git a/Engine/Source/Runtime/Engine/Private/WorldPartition/LoaderAdapter/LoaderAdapterList.cpp b/Engine/Source/Runtime/Engine/Private/WorldPartition/LoaderAdapter/LoaderAdapterList.cpp
--- a/Engine/Source/Runtime/Engine/Private/WorldPartition/LoaderAdapter/LoaderAdapterList.cpp
+++ b/Engine/Source/Runtime/Engine/Private/WorldPartition/LoaderAdapter/LoaderAdapterList.cpp
@@ -1,9 +1,10 @@
 // Copyright Epic Games, Inc. All Rights Reserved.
 
 #include "WorldPartition/LoaderAdapter/LoaderAdapterList.h"
-#include "WorldPartition/ActorDescContainerCollection.h"
+#include "WorldPartition/ActorDescContainerInstanceCollection.h"
 #include "WorldPartition/WorldPartition.h"
 #include "WorldPartition/WorldPartitionActorDesc.h"
+#include "WorldPartition/WorldPartitionActorDescInstance.h"
 #include "Engine/Level.h"
 
 #if WITH_EDITOR
@@ -25,18 +26,18 @@ void FLoaderAdapterList::ForEachActor(TFunctionRef<void(const FWorldPartitionHan
 
 void FLoaderAdapterList::HandleActorContainer(const FWorldPartitionHandle& InActor, TFunctionRef<void(const FWorldPartitionHandle&)> InOperation) const
 {
-	if (InActor->IsContainerInstance())
+	if (InActor.GetInstance()->IsChildContainerInstance())
 	{
-		FWorldPartitionActorDesc::FContainerInstance ContainerInstance;
-		if (InActor->GetContainerInstance(ContainerInstance))
+		FWorldPartitionActorDesc::FLoadedContainerInstance ContainerInstance;
+		if (IWorldPartitionActorLoaderInterface::GetLoadedChildContainerInstance(InActor, ContainerInstance))
 		{
 			if (ContainerInstance.bSupportsPartialEditorLoading)
 			{
 				if (UWorldPartition* ContainerWorldPartition = ContainerInstance.LoadedLevel ? ContainerInstance.LoadedLevel->GetWorldPartition() : nullptr)
 				{
-					for (FActorDescContainerCollection::TIterator<> ActorDescIterator(ContainerWorldPartition); ActorDescIterator; ++ActorDescIterator)
+					for (FActorDescContainerInstanceCollection::TIterator<> Iterator(ContainerWorldPartition); Iterator; ++Iterator)
 					{
-						FWorldPartitionHandle ActorHandle(ContainerWorldPartition, ActorDescIterator->GetGuid());
+						FWorldPartitionHandle ActorHandle(ContainerWorldPartition, Iterator->GetGuid());
 						InOperation(ActorHandle);
 						HandleActorContainer(ActorHandle, InOperation);
 					}
diff --git a/Engine/Source/Runtime/Engine/Private/WorldPartition/LoaderAdapter/LoaderAdapterPinnedActors.cpp b/Engine/Source/Runtime/Engine/Private/WorldPartition/LoaderAdapter/LoaderAdapterPinnedActors.cpp
--- a/Engine/Source/Runtime/Engine/Private/WorldPartition/LoaderAdapter/LoaderAdapterPinnedActors.cpp
+++ b/Engine/Source/Runtime/Engine/Private/WorldPartition/LoaderAdapter/LoaderAdapterPinnedActors.cpp
@@ -7,33 +7,34 @@
 #include "Engine/World.h"
 #include "Engine/Level.h"
 #include "WorldPartition/WorldPartition.h"
+#include "WorldPartition/WorldPartitionActorDescInstance.h"
 
 #define LOCTEXT_NAMESPACE "FLoaderAdapterPinnedActors"
 
 namespace LoaderAdapterPinnedActorsUtils
 {
-	static bool SupportsPinning(FWorldPartitionActorDesc* InActorDesc, bool bCheckIsMainWorldPartition)
+	static bool SupportsPinning(FWorldPartitionActorDescInstance* InActorDescInstance, bool bCheckIsMainWorldPartition)
 	{
-		if (!InActorDesc)
+		if (!InActorDescInstance)
 		{
 			return false;
 		}
 
 		// Only Spatially loaded actors can be pinned with the exception of non spatially loaded runtime only actors (ex: HLODs)
-		if (!InActorDesc->GetIsSpatiallyLoaded() && !InActorDesc->GetActorIsRuntimeOnly())
+		if (!InActorDescInstance->GetIsSpatiallyLoaded() && !InActorDescInstance->GetActorIsRuntimeOnly())
 		{
 			return false;
 		}
 
 		// This allows skipping actors that can never be loaded in current context: ex: bActorShouldSkipLevelInstance
-		if (!InActorDesc->IsEditorRelevant() && !InActorDesc->GetActorIsRuntimeOnly())
+		if (!InActorDescInstance->IsEditorRelevant() && !InActorDescInstance->GetActorIsRuntimeOnly())
 		{
 			return false;
 		}
 
-		if (UActorDescContainer* Container = InActorDesc->GetContainer())
+		if (UActorDescContainerInstance* ContainerInstance = InActorDescInstance->GetContainerInstance())
 		{
-			const UWorldPartition* ContainerWorldPartition = Container->GetWorldPartition();
+			const UWorldPartition* ContainerWorldPartition = ContainerInstance->GetWorldPartition();
 			return ContainerWorldPartition && (ContainerWorldPartition->IsMainWorldPartition() || !bCheckIsMainWorldPartition);
 		}
 
@@ -46,16 +47,16 @@ bool FLoaderAdapterPinnedActors::PassActorDescFilter(const FWorldPartitionHandle
 	// We want to be able to pin any type of actors (HLODs, etc).
 	// Allow recursive pinning by setting bCheckIsMainWorldPartition = false
 	const bool bCheckIsMainWorldPartition = false;
-	return ActorHandle.IsValid() && !ActorsToRemove.Contains(ActorHandle) && LoaderAdapterPinnedActorsUtils::SupportsPinning(ActorHandle.Get(), bCheckIsMainWorldPartition);
+	return ActorHandle.IsValid() && !ActorsToRemove.Contains(ActorHandle) && LoaderAdapterPinnedActorsUtils::SupportsPinning(ActorHandle.GetInstance(), bCheckIsMainWorldPartition);
 }
 
 
 
-bool FLoaderAdapterPinnedActors::SupportsPinning(FWorldPartitionActorDesc* InActorDesc)
+bool FLoaderAdapterPinnedActors::SupportsPinning(FWorldPartitionActorDescInstance* InActorDescInstance)
 {
 	// Public api doesn't allow pinning of non main world partition actors
 	const bool bCheckIsMainWorldPartition = true;
-	return LoaderAdapterPinnedActorsUtils::SupportsPinning(InActorDesc, bCheckIsMainWorldPartition);
+	return LoaderAdapterPinnedActorsUtils::SupportsPinning(InActorDescInstance, bCheckIsMainWorldPartition);
 }
 
 bool FLoaderAdapterPinnedActors::SupportsPinning(AActor* InActor)
diff --git a/Engine/Source/Runtime/Engine/Private/WorldPartition/LoaderAdapter/LoaderAdapterSpatial.cpp b/Engine/Source/Runtime/Engine/Private/WorldPartition/LoaderAdapter/LoaderAdapterSpatial.cpp
--- a/Engine/Source/Runtime/Engine/Private/WorldPartition/LoaderAdapter/LoaderAdapterSpatial.cpp
+++ b/Engine/Source/Runtime/Engine/Private/WorldPartition/LoaderAdapter/LoaderAdapterSpatial.cpp
@@ -3,6 +3,7 @@
 #include "WorldPartition/LoaderAdapter/LoaderAdapterSpatial.h"
 #include "WorldPartition/WorldPartition.h"
 #include "WorldPartition/WorldPartitionEditorHash.h"
+#include "WorldPartition/WorldPartitionActorDescInstance.h"
 #include "Engine/Level.h"
 #include "Engine/World.h"
 
@@ -19,35 +20,36 @@ void ILoaderAdapterSpatial::ForEachActor(TFunctionRef<void(const FWorldPartition
 	{
 		if (UWorldPartition* WorldPartition = GetWorld()->GetWorldPartition())
 		{
-			HandleIntersectingContainer(WorldPartition, WorldPartition->GetInstanceTransform(), *GetBoundingBox(), InOperation);
+			HandleIntersectingContainer(WorldPartition, *GetBoundingBox(), InOperation);
 		}
 	}
 }
 
-void ILoaderAdapterSpatial::HandleIntersectingContainer(UWorldPartition* InWorldPartition, const FTransform InInstanceTransform, const FBox& InBoundingBox, TFunctionRef<void(const FWorldPartitionHandle&)> InOperation) const
+void ILoaderAdapterSpatial::HandleIntersectingContainer(UWorldPartition* InWorldPartition, const FBox& InBoundingBox, TFunctionRef<void(const FWorldPartitionHandle&)> InOperation) const
 {
 	UWorldPartitionEditorHash::FForEachIntersectingActorParams ForEachIntersectingActorParams = UWorldPartitionEditorHash::FForEachIntersectingActorParams()
 		.SetIncludeSpatiallyLoadedActors(bIncludeSpatiallyLoadedActors)
 		.SetIncludeNonSpatiallyLoadedActors(bIncludeNonSpatiallyLoadedActors);
 
-	const FBox LocalBoundingBox = InBoundingBox.InverseTransformBy(InInstanceTransform);
-	InWorldPartition->EditorHash->ForEachIntersectingActor(LocalBoundingBox, [this, InWorldPartition, &InInstanceTransform, &InOperation](FWorldPartitionActorDesc* ActorDesc)
+	const FTransform InstanceTransform = InWorldPartition->GetInstanceTransform();
+	const FBox LocalBoundingBox = InBoundingBox.InverseTransformBy(InstanceTransform);
+	InWorldPartition->EditorHash->ForEachIntersectingActor(LocalBoundingBox, [this, InWorldPartition, &InstanceTransform, &InOperation](FWorldPartitionActorDescInstance* ActorDescInstance)
 	{
-		const FBox WorldActorEditorBox = ActorDesc->GetEditorBounds().TransformBy(InInstanceTransform);
+		const FBox WorldActorEditorBox = ActorDescInstance->GetEditorBounds().TransformBy(InstanceTransform);
 		if (Intersect(WorldActorEditorBox))
 		{
-			FWorldPartitionHandle ActorHandle(InWorldPartition, ActorDesc->GetGuid());
+			FWorldPartitionHandle ActorHandle(InWorldPartition, ActorDescInstance->GetGuid());
 			InOperation(ActorHandle);
 
-			if (ActorHandle->GetIsSpatiallyLoaded() && ActorHandle->IsContainerInstance())
+			if (ActorDescInstance->GetIsSpatiallyLoaded() && ActorDescInstance->IsChildContainerInstance())
 			{
-				FWorldPartitionActorDesc::FContainerInstance ContainerInstance;
-				if (ActorHandle->GetContainerInstance(ContainerInstance))
+				FWorldPartitionActorDesc::FLoadedContainerInstance ContainerInstance;
+				if (IWorldPartitionActorLoaderInterface::GetLoadedChildContainerInstance(ActorHandle, ContainerInstance))
 				{
 					if (UWorldPartition* ContainerWorldPartition = ContainerInstance.LoadedLevel ? ContainerInstance.LoadedLevel->GetWorldPartition() : nullptr)
 					{
 						const FBox InnerBoundingBox = ContainerInstance.bSupportsPartialEditorLoading ? *GetBoundingBox() : FBox(FVector(-HALF_WORLD_MAX), FVector(HALF_WORLD_MAX));
-						HandleIntersectingContainer(ContainerWorldPartition, ContainerWorldPartition->GetInstanceTransform(), InnerBoundingBox, InOperation);
+						HandleIntersectingContainer(ContainerWorldPartition, InnerBoundingBox, InOperation);
 					}
 				}
 			}
diff --git a/Engine/Source/Runtime/Engine/Private/WorldPartition/RuntimeHashSet/WorldPartitionRuntimeHashSetConversions.cpp b/Engine/Source/Runtime/Engine/Private/WorldPartition/RuntimeHashSet/WorldPartitionRuntimeHashSetConversions.cpp
--- a/Engine/Source/Runtime/Engine/Private/WorldPartition/RuntimeHashSet/WorldPartitionRuntimeHashSetConversions.cpp
+++ b/Engine/Source/Runtime/Engine/Private/WorldPartition/RuntimeHashSet/WorldPartitionRuntimeHashSetConversions.cpp
@@ -18,7 +18,7 @@ UWorldPartitionRuntimeHashSet* UWorldPartitionRuntimeHashSet::CreateFrom(const U
 
 	FStreamingGenerationNullErrorHandler NullErrorHandler;
 	UWorldPartition::FGenerateStreamingParams Params = UWorldPartition::FGenerateStreamingParams()
-		.SetActorDescContainer(WorldPartition->GetActorDescContainer())
+		.SetActorDescContainerInstance(WorldPartition->GetActorDescContainerInstance())
 		.SetErrorHandler(&NullErrorHandler);
 
 	UWorldPartition::FGenerateStreamingContext Context;
@@ -38,7 +38,7 @@ UWorldPartitionRuntimeHashSet* UWorldPartitionRuntimeHashSet::CreateFrom(const U
 		{
 			ActorSetInstance.ForEachActor([&GridHLODLayersMap, SpatialHash, DefaultHLODLayer, &ActorSetInstance](const FGuid& ActorGuid)
 			{
-				const FWorldPartitionActorDescView& ActorDescView = ActorSetInstance.ContainerInstance->ActorDescViewMap->FindByGuidChecked(ActorGuid);
+				const IWorldPartitionActorDescInstanceView& ActorDescView = ActorSetInstance.ActorSetContainerInstance->ActorDescViewMap->FindByGuidChecked(ActorGuid);
 
 				if (const UHLODLayer* HLODLayer = ActorDescView.GetHLODLayer().IsValid() ? Cast<UHLODLayer>(ActorDescView.GetHLODLayer().TryLoad()) : DefaultHLODLayer)
 				{
diff --git a/Engine/Source/Runtime/Engine/Private/WorldPartition/RuntimeHashSet/WorldPartitionRuntimeHashSetHLODGeneration.cpp b/Engine/Source/Runtime/Engine/Private/WorldPartition/RuntimeHashSet/WorldPartitionRuntimeHashSetHLODGeneration.cpp
--- a/Engine/Source/Runtime/Engine/Private/WorldPartition/RuntimeHashSet/WorldPartitionRuntimeHashSetHLODGeneration.cpp
+++ b/Engine/Source/Runtime/Engine/Private/WorldPartition/RuntimeHashSet/WorldPartitionRuntimeHashSetHLODGeneration.cpp
@@ -6,6 +6,7 @@
 #include "WorldPartition/WorldPartition.h"
 #include "WorldPartition/WorldPartitionLog.h"
 #include "WorldPartition/WorldPartitionHelpers.h"
+#include "WorldPartition/WorldPartitionActorDescInstance.h"
 
 #include "WorldPartition/HLOD/HLODLayer.h"
 #include "WorldPartition/HLOD/HLODActor.h"
@@ -35,9 +36,9 @@ public:
 		return WorldBounds;
 	}
 
-	virtual const FActorSetContainer* GetMainWorldContainer() const override
+	virtual const FActorSetContainerInstance* GetMainWorldContainerInstance() const override
 	{
-		return &ActorSetContainer;
+		return &ActorSetContainerInstance;
 	}
 
 	virtual void ForEachActorSetInstance(TFunctionRef<void(const FActorSetInstance&)> Func) const override
@@ -48,13 +49,13 @@ public:
 		}
 	}
 
-	virtual void ForEachActorSetContainer(TFunctionRef<void(const FActorSetContainer&)> Func) const override
+	virtual void ForEachActorSetContainerInstance(TFunctionRef<void(const FActorSetContainerInstance&)> Func) const override
 	{
-		Func(ActorSetContainer);
+		Func(ActorSetContainerInstance);
 	}
 
 	FBox WorldBounds;
-	FActorSetContainer ActorSetContainer;
+	FActorSetContainerInstance ActorSetContainerInstance;
 	FStreamingGenerationActorDescViewMap ActorDescViewMap;
 	FActorSetInstanceList ActorSetInstanceList;
 };
@@ -127,17 +128,17 @@ namespace PrivateUtils
 		}
 	}
 
-	static void DeletePackage(UWorldPartition* WorldPartition, FWorldPartitionActorDesc* ActorDesc, ISourceControlHelper* SourceControlHelper)
+	static void DeletePackage(UWorldPartition* WorldPartition, const FWorldPartitionHandle& Handle, ISourceControlHelper* SourceControlHelper)
 	{
-		if (ActorDesc->IsLoaded())
+		if (Handle.IsLoaded())
 		{
-			DeletePackage(ActorDesc->GetActor()->GetPackage(), SourceControlHelper);
-			WorldPartition->OnPackageDeleted(ActorDesc->GetActor()->GetPackage());
+			DeletePackage(Handle.GetActor()->GetPackage(), SourceControlHelper);
+			WorldPartition->OnPackageDeleted(Handle.GetActor()->GetPackage());
 		}
 		else
 		{
-			DeletePackage(ActorDesc->GetActorPackage().ToString(), SourceControlHelper);
-			WorldPartition->RemoveActor(ActorDesc->GetGuid());
+			DeletePackage(Handle.GetInstance()->GetActorPackage().ToString(), SourceControlHelper);
+			WorldPartition->RemoveActor(Handle.GetInstance()->GetGuid());
 		}
 	}
 }
@@ -165,12 +166,12 @@ bool UWorldPartitionRuntimeHashSet::SetupHLODActors(const IStreamingGenerationCo
 
 	UWorldPartition* WorldPartition = GetOuterUWorldPartition();
 	const UDataLayerManager* DataLayerManager = WorldPartition->GetDataLayerManager();
-	IStreamingGenerationContext::FActorSetContainer* MainActorSetContainer = const_cast<IStreamingGenerationContext::FActorSetContainer*>(StreamingGenerationContext->GetMainWorldContainer());
-	const FStreamingGenerationActorDescCollection* MainContainerCollection = MainActorSetContainer->ActorDescCollection;
+	IStreamingGenerationContext::FActorSetContainerInstance* MainActorSetContainer = const_cast<IStreamingGenerationContext::FActorSetContainerInstance*>(StreamingGenerationContext->GetMainWorldContainerInstance());
+	const FStreamingGenerationContainerInstanceCollection* MainContainerCollection = MainActorSetContainer->ContainerInstanceCollection;
 
 	// Create the HLOD creation context
 	FHLODCreationContext HLODCreationContext;
-	for (FActorDescList::TConstIterator<AWorldPartitionHLOD> HLODIterator(MainContainerCollection->GetMainActorDescContainer()); HLODIterator; ++HLODIterator)
+	for (UActorDescContainerInstance::TConstIterator<AWorldPartitionHLOD> HLODIterator(MainContainerCollection->GetMainContainer()); HLODIterator; ++HLODIterator)
 	{
 		FWorldPartitionHandle HLODActorHandle(WorldPartition, HLODIterator->GetGuid());
 		HLODCreationContext.HLODActorDescs.Emplace(HLODIterator->GetActorName(), MoveTemp(HLODActorHandle));
@@ -272,19 +273,19 @@ bool UWorldPartitionRuntimeHashSet::SetupHLODActors(const IStreamingGenerationCo
 			FHLODStreamingGenerationContext* HLODStreamingGenerationContext = (FHLODStreamingGenerationContext*)CurrentHLODStreamingGenerationContext.Get();
 
 			HLODStreamingGenerationContext->ActorSetInstanceList.Reserve(HLODActorGuids.Num());
-			HLODStreamingGenerationContext->ActorSetContainer.ActorDescViewMap = &HLODStreamingGenerationContext->ActorDescViewMap;
+			HLODStreamingGenerationContext->ActorSetContainerInstance.ActorDescViewMap = &HLODStreamingGenerationContext->ActorDescViewMap;
 
 			UE_LOG(LogWorldPartition, Log, TEXT("Creating HLOD context:"));
 			for (const FGuid& HLODActorGuid : HLODActorGuids)
 			{
-				FWorldPartitionActorDesc* HLODActorDesc = WorldPartition->GetActorDesc(HLODActorGuid);
-				check(HLODActorDesc);
+				FWorldPartitionActorDescInstance* HLODActorDescInstance = WorldPartition->GetActorDescInstance(HLODActorGuid);
+				check(HLODActorDescInstance);
 
-				FStreamingGenerationActorDescView* HLODActorDescView = HLODStreamingGenerationContext->ActorDescViewMap.Emplace(HLODActorDesc);
+				FStreamingGenerationActorDescView* HLODActorDescView = HLODStreamingGenerationContext->ActorDescViewMap.Emplace(HLODActorDescInstance);
 				HLODStreamingGenerationContext->WorldBounds += HLODActorDescView->GetRuntimeBounds();
 			
 				// Create actor set instances
-				IStreamingGenerationContext::FActorSet* ActorSet = HLODStreamingGenerationContext->ActorSetContainer.ActorSets.Emplace_GetRef(MakeUnique<IStreamingGenerationContext::FActorSet>()).Get();
+				IStreamingGenerationContext::FActorSet* ActorSet = HLODStreamingGenerationContext->ActorSetContainerInstance.ActorSets.Emplace_GetRef(MakeUnique<IStreamingGenerationContext::FActorSet>()).Get();
 				ActorSet->Actors.Add(HLODActorDescView->GetGuid());
 
 				IStreamingGenerationContext::FActorSetInstance& ActorSetInstance = HLODStreamingGenerationContext->ActorSetInstanceList.Emplace_GetRef();
@@ -293,20 +294,17 @@ bool UWorldPartitionRuntimeHashSet::SetupHLODActors(const IStreamingGenerationCo
 				ActorSetInstance.RuntimeGrid = HLODActorDescView->GetRuntimeGrid();
 				ActorSetInstance.bIsSpatiallyLoaded = HLODActorDescView->GetIsSpatiallyLoaded();
 				ActorSetInstance.ContentBundleID = MainContainerCollection->GetContentBundleGuid();
-				ActorSetInstance.ContainerInstance = &HLODStreamingGenerationContext->ActorSetContainer;
+				ActorSetInstance.ActorSetContainerInstance = &HLODStreamingGenerationContext->ActorSetContainerInstance;
 				ActorSetInstance.ActorSet = ActorSet;
 
-				TArray<const FWorldPartitionActorDescView*> WorldDataLayerViews;
-				Algo::Transform(MainActorSetContainer->ActorDescViewMap->FindByExactNativeClass<AWorldDataLayers>(), WorldDataLayerViews, [](const FStreamingGenerationActorDescView* ActorDescView) { return ActorDescView; });
-
 				TArray<FName> RuntimeDataLayerInstanceNames;
-				if (FDataLayerUtils::ResolveRuntimeDataLayerInstanceNames(DataLayerManager, *HLODActorDescView, WorldDataLayerViews, RuntimeDataLayerInstanceNames))
+				if (FDataLayerUtils::ResolveRuntimeDataLayerInstanceNames(DataLayerManager, *HLODActorDescView, *MainActorSetContainer->ActorDescViewMap, RuntimeDataLayerInstanceNames))
 				{
 					HLODActorDescView->SetRuntimeDataLayerInstanceNames(RuntimeDataLayerInstanceNames);
 					ActorSetInstance.DataLayers = DataLayerManager->GetRuntimeDataLayerInstances(RuntimeDataLayerInstanceNames);
 				}
 
-				UE_LOG(LogWorldPartition, Log, TEXT("\t- %s"), *HLODActorDesc->ToString());
+				UE_LOG(LogWorldPartition, Log, TEXT("\t- %s"), *HLODActorDescInstance->ToString());
 			}
 
 			HLODLevel++;
@@ -318,10 +316,8 @@ bool UWorldPartitionRuntimeHashSet::SetupHLODActors(const IStreamingGenerationCo
 	{
 		for (const auto& HLODActorPair : HLODCreationContext.HLODActorDescs)
 		{
-			FWorldPartitionActorDesc* HLODActorDesc = HLODActorPair.Value.Get();
-			check(HLODActorDesc);
-
-			PrivateUtils::DeletePackage(WorldPartition, HLODActorDesc, Params.SourceControlHelper);
+			check(HLODActorPair.Value.IsValid());
+			PrivateUtils::DeletePackage(WorldPartition, HLODActorPair.Value, Params.SourceControlHelper);
 		}
 	}
 
diff --git a/Engine/Source/Runtime/Engine/Private/WorldPartition/RuntimeSpatialHash/RuntimeSpatialHashGridHelper.cpp b/Engine/Source/Runtime/Engine/Private/WorldPartition/RuntimeSpatialHash/RuntimeSpatialHashGridHelper.cpp
--- a/Engine/Source/Runtime/Engine/Private/WorldPartition/RuntimeSpatialHash/RuntimeSpatialHashGridHelper.cpp
+++ b/Engine/Source/Runtime/Engine/Private/WorldPartition/RuntimeSpatialHash/RuntimeSpatialHashGridHelper.cpp
@@ -181,7 +181,7 @@ FSquare2DGridHelper GetPartitionedActors(const FBox& WorldBounds, const FSpatial
 			bool bUseLocation = true;
 			InActorSetInstance->ForEachActor([InActorSetInstance, &bUseLocation](const FGuid& ActorGuid)
 			{
-				const FWorldPartitionActorDescView& ActorDescView = InActorSetInstance->ContainerInstance->ActorDescViewMap->FindByGuidChecked(ActorGuid);
+				const IWorldPartitionActorDescInstanceView& ActorDescView = InActorSetInstance->ActorSetContainerInstance->ActorDescViewMap->FindByGuidChecked(ActorGuid);
 
 				if (!ActorDescView.GetActorNativeClass()->IsChildOf<APartitionActor>())
 				{
diff --git a/Engine/Source/Runtime/Engine/Private/WorldPartition/RuntimeSpatialHash/RuntimeSpatialHashHLOD.cpp b/Engine/Source/Runtime/Engine/Private/WorldPartition/RuntimeSpatialHash/RuntimeSpatialHashHLOD.cpp
--- a/Engine/Source/Runtime/Engine/Private/WorldPartition/RuntimeSpatialHash/RuntimeSpatialHashHLOD.cpp
+++ b/Engine/Source/Runtime/Engine/Private/WorldPartition/RuntimeSpatialHash/RuntimeSpatialHashHLOD.cpp
@@ -18,6 +18,7 @@
 #include "WorldPartition/WorldPartition.h"
 #include "WorldPartition/WorldPartitionHelpers.h"
 #include "WorldPartition/ContentBundle/ContentBundleActivationScope.h"
+#include "WorldPartition/WorldPartitionActorDescInstance.h"
 
 #include "UObject/GCObjectScopeGuard.h"
 #include "UObject/SavePackage.h"
@@ -79,17 +80,17 @@ static void DeletePackage(UPackage* Package, ISourceControlHelper* SourceControl
 	}
 }
 
-static void DeletePackage(UWorldPartition* WorldPartition, FWorldPartitionActorDesc* ActorDesc, ISourceControlHelper* SourceControlHelper)
+static void DeletePackage(UWorldPartition* WorldPartition, FWorldPartitionActorDescInstance* ActorDescInstance, ISourceControlHelper* SourceControlHelper)
 {
-	if (ActorDesc->IsLoaded())
+	if (ActorDescInstance->IsLoaded())
 	{
-		DeletePackage(ActorDesc->GetActor()->GetPackage(), SourceControlHelper);
-		WorldPartition->OnPackageDeleted(ActorDesc->GetActor()->GetPackage());
+		DeletePackage(ActorDescInstance->GetActor()->GetPackage(), SourceControlHelper);
+		WorldPartition->OnPackageDeleted(ActorDescInstance->GetActor()->GetPackage());
 	}
 	else
 	{
-		DeletePackage(ActorDesc->GetActorPackage().ToString(), SourceControlHelper);
-		WorldPartition->RemoveActor(ActorDesc->GetGuid());
+		DeletePackage(ActorDescInstance->GetActorPackage().ToString(), SourceControlHelper);
+		WorldPartition->RemoveActor(ActorDescInstance->GetGuid());
 	}
 }
 
@@ -113,7 +114,7 @@ static void GameTick(UWorld* InWorld)
 
 static TArray<FGuid> GenerateHLODActorsForGrid(UWorldPartition* WorldPartition, const IStreamingGenerationContext* StreamingGenerationContext, const UWorldPartition::FSetupHLODActorsParams& Params, const FSpatialHashRuntimeGrid& RuntimeGrid, uint32 HLODLevel, FHLODCreationContext& Context, ISourceControlHelper* SourceControlHelper, const TArray<const IStreamingGenerationContext::FActorSetInstance*>& ActorSetInstances, TArray<FName>& NewActors, TArray<FName>& DirtyActors)
 {
-	const IStreamingGenerationContext::FActorSetContainer* MainActorSetContainer = StreamingGenerationContext->GetMainWorldContainer();
+	const IStreamingGenerationContext::FActorSetContainerInstance* MainActorSetContainer = StreamingGenerationContext->GetMainWorldContainerInstance();
 	const FBox WorldBounds = StreamingGenerationContext->GetWorldBounds();
 
 	const FSquare2DGridHelper PartitionedActors = GetPartitionedActors(WorldBounds, RuntimeGrid, ActorSetInstances);
@@ -272,9 +273,9 @@ static TMap<UHLODLayer*, int32> GatherHLODLayers(const IStreamingGenerationConte
 	// Gather up all HLODLayers referenced by the actors, along with the HLOD level at which it was used
 	TMap<UHLODLayer*, int32> HLODLayersLevel;
 
-	StreamingGenerationContext->ForEachActorSetContainer([&HLODLayersLevel, WorldPartition](const IStreamingGenerationContext::FActorSetContainer& ActorSetContainer)
+	StreamingGenerationContext->ForEachActorSetContainerInstance([&HLODLayersLevel, WorldPartition](const IStreamingGenerationContext::FActorSetContainerInstance& ActorSetContainerInstance)
 	{
-		const FStreamingGenerationActorDescViewMap* ActorDescViewMap = ActorSetContainer.ActorDescViewMap;
+		const FStreamingGenerationActorDescViewMap* ActorDescViewMap = ActorSetContainerInstance.ActorDescViewMap;
 		ActorDescViewMap->ForEachActorDescView([&HLODLayersLevel, WorldPartition](const FStreamingGenerationActorDescView& ActorDescView)
 		{
 			if (!ActorDescView.GetActorNativeClass()->IsChildOf<AWorldPartitionHLOD>())
@@ -352,25 +353,24 @@ static TMap<FName, FSpatialHashRuntimeGrid> CreateHLODGrids(TMap<UHLODLayer*, in
 }
 
 // Create/destroy HLOD grid actors
-static void UpdateHLODGridsActors(UWorld* World, const UActorDescContainer* Container, const TMap<FName, FSpatialHashRuntimeGrid>& HLODGrids, ISourceControlHelper* SourceControlHelper)
+static void UpdateHLODGridsActors(UWorld* World, const UActorDescContainerInstance* ContainerInstance, const TMap<FName, FSpatialHashRuntimeGrid>& HLODGrids, ISourceControlHelper* SourceControlHelper)
 {
 	static const FName HLODGridTag = TEXT("HLOD");
 	static const uint32 HLODGridTagLen = HLODGridTag.ToString().Len();
 
 	// Gather all existing HLOD grid actors, see if some are unused and needs to be deleted
 	TMap<FName, ASpatialHashRuntimeGridInfo*> ExistingGridActors;
-	for (FActorDescList::TConstIterator<> ItActorDesc(Container); ItActorDesc; ++ItActorDesc)
+	for (UActorDescContainerInstance::TConstIterator<> Iterator(ContainerInstance); Iterator; ++Iterator)
 	{
-		if(ItActorDesc->GetActorNativeClass()->IsChildOf<ASpatialHashRuntimeGridInfo>())
+		if(Iterator->GetActorNativeClass()->IsChildOf<ASpatialHashRuntimeGridInfo>())
 		{
-			const FWorldPartitionActorDesc* GridActorDesc = *ItActorDesc;
-			ASpatialHashRuntimeGridInfo* GridActor = CastChecked<ASpatialHashRuntimeGridInfo>(GridActorDesc->GetActor());
+			ASpatialHashRuntimeGridInfo* GridActor = CastChecked<ASpatialHashRuntimeGridInfo>(Iterator->GetActor());
 			if (GridActor->ActorHasTag(HLODGridTag))
 			{
 				const FSpatialHashRuntimeGrid* HLODGrid = HLODGrids.Find(GridActor->GridSettings.GridName);
 				if (HLODGrid && GridActor->GridSettings.Priority && GridActor->GridSettings.HLODLayer)
 				{
-					check(GridActor->GetContentBundleGuid() == Container->GetContentBundleGuid());
+					check(GridActor->GetContentBundleGuid() == ContainerInstance->GetContentBundleGuid());
 
 					ExistingGridActors.Emplace(GridActor->GridSettings.GridName, GridActor);
 				}
@@ -395,7 +395,7 @@ static void UpdateHLODGridsActors(UWorld* World, const UActorDescContainer* Cont
 		ASpatialHashRuntimeGridInfo* GridActor = ExistingGridActors.FindRef(GridName);
 		if (!GridActor)
 		{
-			FContentBundleActivationScope ContentBndleScope(Container->GetContentBundleGuid());
+			FContentBundleActivationScope ContentBndleScope(ContainerInstance->GetContentBundleGuid());
 
 			FActorSpawnParameters SpawnParams;
 			SpawnParams.bCreateActorPackage = true;
@@ -403,7 +403,7 @@ static void UpdateHLODGridsActors(UWorld* World, const UActorDescContainer* Cont
 			GridActor->Tags.Add(HLODGridTag);
 			bDirty = true;
 
-			check(GridActor->GetContentBundleGuid() == Container->GetContentBundleGuid());
+			check(GridActor->GetContentBundleGuid() == ContainerInstance->GetContentBundleGuid());
 		}
 
 		const FString ActorLabel = GridSettings.GridName.ToString();
@@ -448,8 +448,8 @@ bool UWorldPartitionRuntimeSpatialHash::SetupHLODActors(const IStreamingGenerati
 	}
 
 	UWorldPartition* WorldPartition = GetOuterUWorldPartition();
-	IStreamingGenerationContext::FActorSetContainer* MainActorSetContainer = const_cast<IStreamingGenerationContext::FActorSetContainer*>(StreamingGenerationContext->GetMainWorldContainer());
-	const FStreamingGenerationActorDescCollection* MainContainerCollection = MainActorSetContainer->ActorDescCollection;
+	IStreamingGenerationContext::FActorSetContainerInstance* MainActorSetContainer = const_cast<IStreamingGenerationContext::FActorSetContainerInstance*>(StreamingGenerationContext->GetMainWorldContainerInstance());
+	const FStreamingGenerationContainerInstanceCollection* MainContainerCollection = MainActorSetContainer->ContainerInstanceCollection;
 
 	// Find all used HLOD layers
 	TMap<UHLODLayer*, int32> HLODLayersLevels = GatherHLODLayers(StreamingGenerationContext, WorldPartition);
@@ -472,7 +472,7 @@ bool UWorldPartitionRuntimeSpatialHash::SetupHLODActors(const IStreamingGenerati
 
 	// HLOD creation context
 	FHLODCreationContext Context;
-	for (FActorDescList::TConstIterator<AWorldPartitionHLOD> HLODIterator(MainContainerCollection->GetMainActorDescContainer()); HLODIterator; ++HLODIterator)
+	for (UActorDescContainerInstance::TConstIterator<AWorldPartitionHLOD> HLODIterator(MainContainerCollection->GetMainContainer()); HLODIterator; ++HLODIterator)
 	{
 		FWorldPartitionHandle HLODActorHandle(WorldPartition, HLODIterator->GetGuid());
 		Context.HLODActorDescs.Emplace(HLODIterator->GetActorName(), MoveTemp(HLODActorHandle));
@@ -496,11 +496,10 @@ bool UWorldPartitionRuntimeSpatialHash::SetupHLODActors(const IStreamingGenerati
 
 	// Keep track of all valid HLOD actors, along with which runtime grid they live in
 	TMap<FName, TArray<FGuid>> GridsHLODActors;
-	TArray<TUniquePtr<FStreamingGenerationActorDescView>> ActorDescViews;
 
 	const UDataLayerManager* DataLayerManager = WorldPartition->GetDataLayerManager();
 
-	auto GenerateHLODActors = [&GridsHLODActors, &ActorDescViews, &MainActorSetContainer, StreamingGenerationContext, &Params, WorldPartition, DataLayerManager, &Context, SourceControlHelper, &NewActors, &DirtyActors]
+	auto GenerateHLODActors = [&GridsHLODActors, &MainActorSetContainer, StreamingGenerationContext, &Params, WorldPartition, DataLayerManager, &Context, SourceControlHelper, &NewActors, &DirtyActors]
 	(const FSpatialHashRuntimeGrid& RuntimeGrid, uint32 HLODLevel, const TArray<const IStreamingGenerationContext::FActorSetInstance*>& ActorSetInstances)
 	{
 		// Generate HLODs for this grid
@@ -510,23 +509,18 @@ bool UWorldPartitionRuntimeSpatialHash::SetupHLODActors(const IStreamingGenerati
 		{
 			if (!Params.bReportOnly)
 			{
-				FWorldPartitionActorDesc* HLODActorDesc = WorldPartition->GetActorDesc(HLODActorGuid);
-				check(HLODActorDesc);
-				FStreamingGenerationActorDescView* ActorDescView = ActorDescViews.Emplace_GetRef(MakeUnique<FStreamingGenerationActorDescView>(HLODActorDesc)).Get();
+				FWorldPartitionActorDescInstance* HLODActorDescInstance = WorldPartition->GetActorDescInstance(HLODActorGuid);
+				check(HLODActorDescInstance);
+				FStreamingGenerationActorDescViewMap* NonConstActorDescViewMap = const_cast<FStreamingGenerationActorDescViewMap*>(MainActorSetContainer->ActorDescViewMap);
+				FStreamingGenerationActorDescView* ActorDescView = NonConstActorDescViewMap->Emplace(HLODActorDescInstance);
 				
-				TArray<const FWorldPartitionActorDescView*> WorldDataLayerViews;
-				Algo::Transform(MainActorSetContainer->ActorDescViewMap->FindByExactNativeClass<AWorldDataLayers>(), WorldDataLayerViews, [](const FStreamingGenerationActorDescView* ActorDescView) { return ActorDescView; });
-
 				TArray<FName> RuntimeDataLayerInstanceNames;
-				if (FDataLayerUtils::ResolveRuntimeDataLayerInstanceNames(DataLayerManager, *ActorDescView, WorldDataLayerViews, RuntimeDataLayerInstanceNames))
+				if (FDataLayerUtils::ResolveRuntimeDataLayerInstanceNames(DataLayerManager, *ActorDescView, *MainActorSetContainer->ActorDescViewMap, RuntimeDataLayerInstanceNames))
 				{
 					ActorDescView->SetRuntimeDataLayerInstanceNames(RuntimeDataLayerInstanceNames);
 				}
-				FStreamingGenerationActorDescViewMap* NonConstActorDescViewMap = const_cast<FStreamingGenerationActorDescViewMap*>(MainActorSetContainer->ActorDescViewMap);			
-				NonConstActorDescViewMap->Emplace(HLODActorGuid, *ActorDescView);
 
-				const FStreamingGenerationActorDescView& HLODActorDescView = MainActorSetContainer->ActorDescViewMap->FindByGuidChecked(HLODActorGuid);
-				GridsHLODActors.FindOrAdd(HLODActorDescView.GetRuntimeGrid()).Add(HLODActorGuid);
+				GridsHLODActors.FindOrAdd(ActorDescView->GetRuntimeGrid()).Add(HLODActorGuid);
 			}
 		}
 	};
@@ -571,7 +565,7 @@ bool UWorldPartitionRuntimeSpatialHash::SetupHLODActors(const IStreamingGenerati
 				NewHLODActorSetInstance.DataLayers = DataLayerManager->GetRuntimeDataLayerInstances(HLODActorDescView.GetRuntimeDataLayerInstanceNames());
 			}
 			NewHLODActorSetInstance.ContentBundleID = MainContainerCollection->GetContentBundleGuid();
-			NewHLODActorSetInstance.ContainerInstance = MainActorSetContainer;
+			NewHLODActorSetInstance.ActorSetContainerInstance = MainActorSetContainer;
 			NewHLODActorSetInstance.ActorSet = MainActorSetContainer->ActorSets.Add_GetRef(MakeUnique<IStreamingGenerationContext::FActorSet>()).Get();
 			const_cast<IStreamingGenerationContext::FActorSet*>(NewHLODActorSetInstance.ActorSet)->Actors.Add(HLODActorDescView.GetGuid());
 
@@ -585,10 +579,10 @@ bool UWorldPartitionRuntimeSpatialHash::SetupHLODActors(const IStreamingGenerati
 
 	auto DeleteHLODActor = [&SourceControlHelper, WorldPartition](FWorldPartitionHandle ActorHandle)
 	{
-		FWorldPartitionActorDesc* HLODActorDesc = ActorHandle.Get();
-		check(HLODActorDesc);
+		FWorldPartitionActorDescInstance* ActorDescInstance = ActorHandle.GetInstance();
+		check(ActorDescInstance);
 
-		DeletePackage(WorldPartition, HLODActorDesc, SourceControlHelper);
+		DeletePackage(WorldPartition, ActorDescInstance, SourceControlHelper);
 	};
 
 	// Destroy all unreferenced HLOD actors
@@ -605,7 +599,7 @@ bool UWorldPartitionRuntimeSpatialHash::SetupHLODActors(const IStreamingGenerati
 	if (!Params.bReportOnly)
 	{
 		// Create/destroy HLOD grid actors
-		UpdateHLODGridsActors(GetWorld(), MainContainerCollection->GetMainActorDescContainer(), HLODGrids, SourceControlHelper);
+		UpdateHLODGridsActors(GetWorld(), MainContainerCollection->GetMainContainer(), HLODGrids, SourceControlHelper);
 	}
 
 	auto DumpActorsStats = [](const TCHAR* StatType, const TArray<FName>& Actors)
