WorldPartition: Add UActorDescContainerInstance/FWorldPartitionActorDescInstance which take transient/context data from UActorDescContainer/FWorldPartitionActorDesc leaving them as pure on disk repres

diff --git a/Engine/Source/Runtime/Engine/Private/WorldPartition/WorldPartition.cpp b/Engine/Source/Runtime/Engine/Private/WorldPartition/WorldPartition.cpp
--- a/Engine/Source/Runtime/Engine/Private/WorldPartition/WorldPartition.cpp
+++ b/Engine/Source/Runtime/Engine/Private/WorldPartition/WorldPartition.cpp
@@ -46,6 +46,8 @@
 #include "WorldPartition/LoaderAdapter/LoaderAdapterPinnedActors.h"
 #include "WorldPartition/HLOD/HLODLayer.h"
 #include "WorldPartition/Cook/WorldPartitionCookPackageContextInterface.h"
+#include "WorldPartition/WorldPartitionActorDescInstance.h"
+#include "WorldPartition/ActorDescContainerInstance.h"
 #include "WorldPartition/ErrorHandling/WorldPartitionStreamingGenerationLogErrorHandler.h"
 #include "WorldPartition/ErrorHandling/WorldPartitionStreamingGenerationMapCheckErrorHandler.h"
 #include "Modules/ModuleManager.h"
@@ -141,7 +143,7 @@ TMap<FName, FString> GetDataLayersDumpString(const UWorldPartition* WorldPartiti
 	return DataLayersDumpString;
 }
 
-FString GetActorDescDumpString(const FWorldPartitionActorDesc* ActorDesc, const TMap<FName, FString>& DataLayersDumpString)
+FString GetActorDescDumpString(const FWorldPartitionActorDescInstance* ActorDescInstance, const TMap<FName, FString>& DataLayersDumpString)
 {
 	auto GetDataLayerString = [&DataLayersDumpString](const TArray<FName>& DataLayerNames)
 	{
@@ -160,11 +162,11 @@ FString GetActorDescDumpString(const FWorldPartitionActorDesc* ActorDesc, const
 		});
 	};
 
-	check(ActorDesc);
+	check(ActorDescInstance);
 	return FString::Printf(
 		TEXT("%s DataLayerNames:%s") LINE_TERMINATOR, 
-		*ActorDesc->ToString(FWorldPartitionActorDesc::EToStringMode::Full), 
-		*GetDataLayerString(ActorDesc->GetDataLayerInstanceNames())
+		*ActorDescInstance->ToString(FWorldPartitionActorDesc::EToStringMode::Full),
+		*GetDataLayerString(ActorDescInstance->GetDataLayerInstanceNames())
 	);
 }
 
@@ -200,9 +202,9 @@ static FAutoConsoleCommand DumpActorDesc(
 						TMap<FName, FString> DataLayersDumpString = GetDataLayersDumpString(WorldPartition);
 						for (const FString& ActorPath : ActorPaths)
 						{
-							if (const FWorldPartitionActorDesc* ActorDesc = WorldPartition->GetActorDescByPath(ActorPath))
+							if (const FWorldPartitionActorDescInstance* ActorDescInstance = WorldPartition->GetActorDescInstanceByPath(ActorPath))
 							{
-								UE_LOG(LogWorldPartition, Log, TEXT("%s"), *GetActorDescDumpString(ActorDesc, DataLayersDumpString));
+								UE_LOG(LogWorldPartition, Log, TEXT("%s"), *GetActorDescDumpString(ActorDescInstance, DataLayersDumpString));
 							}
 						}
 					}
@@ -615,55 +617,13 @@ void UWorldPartition::Initialize(UWorld* InWorld, const FTransform& InTransform)
 	TArray<FGuid> ForceLoadedActorGuids;
 	if (bIsEditor || bIsGame || bIsPIEWorldTravel || bIsDedicatedServer)
 	{
-		UPackage* LevelPackage = OuterWorld->PersistentLevel->GetOutermost();
+		FName ContainerPackageName = UActorDescContainerInstance::GetContainerPackageNameFromWorld(OuterWorld);
 
-		// Duplicated worlds (ex: WorldPartitionRenameDuplicateBuilder) will not have a loaded path 
-		const FName PackageName = LevelPackage->GetLoadedPath().GetPackageFName().IsNone() ? LevelPackage->GetFName() : LevelPackage->GetLoadedPath().GetPackageFName();
+		ActorDescContainerInstance = RegisterActorDescContainerInstance(UActorDescContainerInstance::FInitializeParams(ContainerPackageName));
 
-		// Currently known Instancing use cases:
-		//	- World Partition map template (New Level)
-		//	- PIE World Travel
-		FString SourceWorldPath, RemappedWorldPath;
-		const bool bIsInstanced = OuterWorld->GetSoftObjectPathMapping(SourceWorldPath, RemappedWorldPath);
-
-		if (bIsInstanced)
+		ForEachActorDescContainer([this, bIsEditor, bIsCooking](UActorDescContainerInstance* InActorDescContainerInstance)
 		{
-			InstancingContext.AddPackageMapping(PackageName, LevelPackage->GetFName());
-
-			// SoftObjectPaths: Specific case for new maps (/Temp/Untitled) where we need to remap the AssetPath and not just the Package name because the World gets renamed (See UWorld::PostLoad)
-			InstancingContext.AddPathMapping(
-				FSoftObjectPath(*FString::Format(TEXT("{0}.{1}"), {PackageName.ToString(), FPackageName::GetShortName(PackageName)})),
-				FSoftObjectPath(OuterWorld)
-			);
-		}
-
-		FContainerRegistrationParams ContainerInitParams(PackageName);
-		ActorDescContainer = RegisterActorDescContainer(ContainerInitParams);
-
-		// If a Valid Actor references an Invalid Actor:
-		// Make sure Invalid Actors do not load their imports (ex: outer non instanced world).
-		if (bIsInstanced)
-		{
-			TRACE_CPUPROFILER_EVENT_SCOPE(ContainerInstancing);
-			for (const FAssetData& InvalidActor : ActorDescContainer->InvalidActors)
-			{
-				InstancingContext.AddPackageMapping(InvalidActor.PackageName, NAME_None);
-			}
-
-			for (FActorDescContainerCollection::TIterator<> ActorDescIterator(this); ActorDescIterator; ++ActorDescIterator)
-			{
-				const FString LongActorPackageName = ActorDescIterator->GetActorPackage().ToString();
-				const FString InstancedName = ULevel::GetExternalActorPackageInstanceName(LevelPackage->GetName(), LongActorPackageName);
-
-				InstancingContext.AddPackageMapping(*LongActorPackageName, *InstancedName);
-
-				ActorDescIterator->TransformInstance(SourceWorldPath, RemappedWorldPath);
-			}
-		}
-
-		ForEachActorDescContainer([this, bIsEditor, bIsCooking](UActorDescContainer* InActorDescContainer)
-		{
-			InitializeActorDescContainerEditorStreaming(InActorDescContainer, bIsEditor && !bIsCooking);
+			InitializeActorDescContainerEditorStreaming(InActorDescContainerInstance, bIsEditor && !bIsCooking);
 		});
 	}
 #endif
@@ -811,7 +771,7 @@ void UWorldPartition::Uninitialize()
 		DirtyActors.Empty();
 
 		UninitializeActorDescContainers();
-		ActorDescContainer = nullptr;
+		ActorDescContainerInstance = nullptr;
 
 		EditorHash = nullptr;
 		bIsPIE = false;
@@ -1248,23 +1208,23 @@ bool UWorldPartition::IsSimulating(bool bIncludeTestEnableSimulationStreamingSou
 }
 
 #if WITH_EDITOR
-void UWorldPartition::OnActorDescAdded(FWorldPartitionActorDesc* NewActorDesc)
+void UWorldPartition::OnActorDescInstanceAdded(FWorldPartitionActorDescInstance* NewActorDescInstance)
 {
-	NewActorDesc->bIsForcedNonSpatiallyLoaded = !IsStreamingEnabledInEditor();
+	NewActorDescInstance->SetForceNonSpatiallyLoadded(!IsStreamingEnabledInEditor());
 
-	HashActorDesc(NewActorDesc);
+	HashActorDescInstance(NewActorDescInstance);
 
-	if (AActor* NewActor = NewActorDesc->GetActor())
+	if (AActor* NewActor = NewActorDescInstance->GetActor())
 	{
-		DirtyActors.Add(FWorldPartitionReference(NewActorDesc->GetContainer(), NewActorDesc->GetGuid()), NewActor);
+		DirtyActors.Add(FWorldPartitionReference(NewActorDescInstance->GetContainerInstance(), NewActorDescInstance->GetGuid()), NewActor);
 	}
 
 	if (ForceLoadedActors)
 	{
-		ForceLoadedActors->AddActors({ NewActorDesc->GetGuid() });
+		ForceLoadedActors->AddActors({ NewActorDescInstance->GetGuid() });
 	}
 
-	if (AlwaysLoadedActors && !NewActorDesc->GetIsSpatiallyLoaded())
+	if (AlwaysLoadedActors && !NewActorDescInstance->GetIsSpatiallyLoaded())
 	{
 		AlwaysLoadedActors->RefreshLoadedState();
 	}
@@ -1275,39 +1235,39 @@ void UWorldPartition::OnActorDescAdded(FWorldPartitionActorDesc* NewActorDesc)
 	}
 }
 
-void UWorldPartition::OnActorDescRemoved(FWorldPartitionActorDesc* ActorDesc)
+void UWorldPartition::OnActorDescInstanceRemoved(FWorldPartitionActorDescInstance* ActorDescInstance)
 {
 	if (PinnedActors)
 	{
-		PinnedActors->RemoveActors({ FWorldPartitionHandle(ActorDesc->GetContainer(), ActorDesc->GetGuid()) });
+		PinnedActors->RemoveActors({ FWorldPartitionHandle(ActorDescInstance->GetContainerInstance(), ActorDescInstance->GetGuid()) });
 	}
 
-	UnhashActorDesc(ActorDesc);
+	UnhashActorDescInstance(ActorDescInstance);
 
 	if (ForceLoadedActors)
 	{
-		ForceLoadedActors->RemoveActors({ ActorDesc->GetGuid() });
+		ForceLoadedActors->RemoveActors({ ActorDescInstance->GetGuid() });
 	}
 
-	if (AlwaysLoadedActors && !ActorDesc->GetIsSpatiallyLoaded())
+	if (AlwaysLoadedActors && !ActorDescInstance->GetIsSpatiallyLoaded())
 	{
 		AlwaysLoadedActors->RefreshLoadedState();
 	}
-		
+
 	if (WorldPartitionEditor)
 	{
 		WorldPartitionEditor->Refresh();
 	}
 }
 
-void UWorldPartition::OnActorDescUpdating(FWorldPartitionActorDesc* ActorDesc)
+void UWorldPartition::OnActorDescInstanceUpdating(FWorldPartitionActorDescInstance* ActorDescInstance)
 {
-	UnhashActorDesc(ActorDesc);
+	UnhashActorDescInstance(ActorDescInstance);
 }
 
-void UWorldPartition::OnActorDescUpdated(FWorldPartitionActorDesc* ActorDesc)
+void UWorldPartition::OnActorDescInstanceUpdated(FWorldPartitionActorDescInstance* ActorDescInstance)
 {
-	HashActorDesc(ActorDesc);
+	HashActorDescInstance(ActorDescInstance);
 
 	if (WorldPartitionEditor)
 	{
@@ -1315,26 +1275,16 @@ void UWorldPartition::OnActorDescUpdated(FWorldPartitionActorDesc* ActorDesc)
 	}
 }
 
-bool UWorldPartition::GetInstancingContext(const FLinkerInstancingContext*& OutInstancingContext) const
-{
-	if (InstancingContext.IsInstanced())
-	{
-		OutInstancingContext = &InstancingContext;
-		return true;
-	}
-	return false;
-}
-
-void UWorldPartition::InitializeActorDescContainerEditorStreaming(UActorDescContainer* InActorDescContainer, bool bInHashActorDescs)
+void UWorldPartition::InitializeActorDescContainerEditorStreaming(UActorDescContainerInstance* InActorDescContainerInstance, bool bInHashActorDescs)
 {
 	TRACE_CPUPROFILER_EVENT_SCOPE(InitializeActorDescContainerEditorStreaming);
 
 	const bool bIsStreamingEnabled = IsStreamingEnabledInEditor();
 
 	TArray<FGuid> ForceLoadedActorGuids;
-	for (UActorDescContainer::TIterator<> It(InActorDescContainer); It; ++It)
+	for (UActorDescContainerInstance::TIterator<> It(InActorDescContainerInstance); It; ++It)
 	{
-		It->bIsForcedNonSpatiallyLoaded = !bIsStreamingEnabled;
+		It->SetForceNonSpatiallyLoadded(!bIsStreamingEnabled);
 
 		if (ForceLoadedActors)
 		{
@@ -1343,7 +1293,7 @@ void UWorldPartition::InitializeActorDescContainerEditorStreaming(UActorDescCont
 
 		if (bInHashActorDescs)
 		{
-			HashActorDesc(*It);
+			HashActorDescInstance(*It);
 		}
 	}
 
@@ -1376,11 +1326,11 @@ void UWorldPartition::SetEnableStreaming(bool bInEnableStreaming)
 
 void UWorldPartition::OnEnableStreamingChanged()
 {
-	for (FActorDescContainerCollection::TIterator<> ActorDescIterator(this); ActorDescIterator; ++ActorDescIterator)
+	for (FActorDescContainerInstanceCollection::TIterator<> Iterator(this); Iterator; ++Iterator)
 	{
-		UnhashActorDesc(*ActorDescIterator);
-		ActorDescIterator->bIsForcedNonSpatiallyLoaded = !IsStreamingEnabledInEditor();
-		HashActorDesc(*ActorDescIterator);
+		UnhashActorDescInstance(*Iterator);
+		Iterator->SetForceNonSpatiallyLoadded(!IsStreamingEnabledInEditor());
+		HashActorDescInstance(*Iterator);
 	}
 
 	FLoaderAdapterAlwaysLoadedActors* OldAlwaysLoadedActors = AlwaysLoadedActors;
@@ -1415,9 +1365,9 @@ void UWorldPartition::OnEnableLoadingInEditorChanged()
 		ForceLoadedActors = new FLoaderAdapterActorList(OuterWorld);
 
 		TArray<FGuid> ForceLoadedActorGuids;
-		for (FActorDescContainerCollection::TIterator<> ActorDescIterator(this); ActorDescIterator; ++ActorDescIterator)
+		for (FActorDescContainerInstanceCollection::TIterator<> Iterator(this); Iterator; ++Iterator)
 		{
-			ForceLoadedActorGuids.Add(ActorDescIterator->GetGuid());
+			ForceLoadedActorGuids.Add(Iterator->GetGuid());
 		}
 
 		if (ForceLoadedActorGuids.Num())
@@ -1427,23 +1377,23 @@ void UWorldPartition::OnEnableLoadingInEditorChanged()
 	}
 }
 
-void UWorldPartition::HashActorDesc(FWorldPartitionActorDesc* ActorDesc)
+void UWorldPartition::HashActorDescInstance(FWorldPartitionActorDescInstance* ActorDescInstance)
 {
-	check(ActorDesc);
+	check(ActorDescInstance);
 	check(EditorHash);
 
-	FWorldPartitionHandle ActorHandle(this, ActorDesc->GetGuid());
+	FWorldPartitionHandle ActorHandle(ActorDescInstance);
 	EditorHash->HashActor(ActorHandle);
 
 	bShouldCheckEnableStreamingWarning = IsMainWorldPartition();
 }
 
-void UWorldPartition::UnhashActorDesc(FWorldPartitionActorDesc* ActorDesc)
+void UWorldPartition::UnhashActorDescInstance(FWorldPartitionActorDescInstance* ActorDescInstance)
 {
-	check(ActorDesc);
+	check(ActorDescInstance);
 	check(EditorHash);
 
-	FWorldPartitionHandle ActorHandle(this, ActorDesc->GetGuid());
+	FWorldPartitionHandle ActorHandle(ActorDescInstance);
 	EditorHash->UnhashActor(ActorHandle);
 }
 
@@ -1531,11 +1481,11 @@ bool UWorldPartition::ResolveSubobject(const TCHAR* SubObjectPath, UObject*& Out
 				SubObjectName = SubObjectPath;
 			}
 
-			if (const FWorldPartitionActorDesc* ActorDesc = GetActorDescByPath(SubObjectName))
+			if (const FWorldPartitionActorDescInstance* ActorDescInstance = GetActorDescInstanceByPath(SubObjectName))
 			{
 				if (bLoadIfExists)
 				{
-					LoadedSubobjects.Emplace(this, ActorDesc->GetGuid());
+					LoadedSubobjects.Emplace(this, ActorDescInstance->GetGuid());
 				}
 
 				OutObject = StaticFindObject(UObject::StaticClass(), GetWorld()->PersistentLevel, SubObjectPath);
@@ -1573,9 +1523,9 @@ void UWorldPartition::AddReferencedObjects(UObject* InThis, FReferenceCollector&
 		Collector.AllowEliminatingReferences(true);
 	}
 
-	for (TObjectPtr<UActorDescContainer>& Container : This->ActorDescContainerCollection)
+	for (TObjectPtr<UActorDescContainerInstance>& ContainerInstance : This->ActorDescContainerInstanceCollection)
 	{
-		Collector.AddReferencedObject(Container);
+		Collector.AddReferencedObject(ContainerInstance);
 	}
 #endif
 
@@ -1590,7 +1540,7 @@ void UWorldPartition::Tick(float DeltaSeconds)
 		EditorHash->Tick(DeltaSeconds);
 	}
 
-		for (TMap<FWorldPartitionReference, TObjectPtr<AActor>>::TIterator DirtyActorIt(DirtyActors); DirtyActorIt; ++DirtyActorIt)
+	for (TMap<FWorldPartitionReference, TObjectPtr<AActor>>::TIterator DirtyActorIt(DirtyActors); DirtyActorIt; ++DirtyActorIt)
 	{
 		if (!DirtyActorIt.Key().IsValid())
 		{
@@ -1599,12 +1549,12 @@ void UWorldPartition::Tick(float DeltaSeconds)
 		else if (!DirtyActorIt.Value()->GetPackage()->IsDirty())
 		{
 			// Transfer ownership of the last ref if actor can be pinned
-			if (DirtyActorIt.Key()->GetHardRefCount() <= 1 && PinnedActors && FLoaderAdapterPinnedActors::SupportsPinning(DirtyActorIt.Key().Get()))
+			if (DirtyActorIt.Key().GetInstance()->GetHardRefCount() <= 1 && PinnedActors && FLoaderAdapterPinnedActors::SupportsPinning(DirtyActorIt.Key().GetInstance()))
 			{
 				PinnedActors->AddActors({ DirtyActorIt.Key().ToHandle() });
 				DirtyActorIt.RemoveCurrent();
 			} // Clean up if a loader took a reference on the actor (ILoaderAdapter::RefreshLoadedState was called since)
-			else if (DirtyActorIt.Key()->GetHardRefCount() > 1)
+			else if (DirtyActorIt.Key().GetInstance()->GetHardRefCount() > 1)
 			{
 				DirtyActorIt.RemoveCurrent();
 			}
@@ -1628,11 +1578,11 @@ void UWorldPartition::Tick(float DeltaSeconds)
 			bEnablingStreamingJustified = false;
 
 			FBox AllActorsBounds(ForceInit);
-			for (FActorDescContainerCollection::TConstIterator<> ActorDescIterator(this); ActorDescIterator; ++ActorDescIterator)
+			for (FActorDescContainerInstanceCollection::TConstIterator<> Iterator(this); Iterator; ++Iterator)
 			{
-				if (ActorDescIterator->GetIsSpatiallyLoadedRaw() || ActorDescIterator->GetActorNativeClass()->IsChildOf<ALandscapeProxy>())
+				if (Iterator->GetActorDesc()->GetIsSpatiallyLoadedRaw() || Iterator->GetActorNativeClass()->IsChildOf<ALandscapeProxy>())
 				{
-					const FBox EditorBounds = ActorDescIterator->GetEditorBounds();
+					const FBox EditorBounds = Iterator->GetEditorBounds();
 					if (EditorBounds.IsValid)
 					{
 						AllActorsBounds += EditorBounds;
@@ -1833,7 +1783,7 @@ void UWorldPartition::EndCook(IWorldPartitionCookPackageContext& CookContext)
 bool UWorldPartition::GatherPackagesToCook(IWorldPartitionCookPackageContext& CookContext)
 {
 	FGenerateStreamingParams Params = FGenerateStreamingParams()
-		.SetActorDescContainer(ActorDescContainer);
+		.SetActorDescContainerInstance(ActorDescContainerInstance);
 
 	TArray<FString> PackagesToCook;
 	FGenerateStreamingContext Context = FGenerateStreamingContext()
@@ -1903,9 +1853,9 @@ void UWorldPartition::SavePerUserSettings()
 		GetMutableDefault<UWorldPartitionEditorPerProjectUserSettings>()->SetEditorLoadedRegions(GetWorld(), GetUserLoadedEditorRegions());
 
 		TArray<FName> EditorLoadedLocationVolumes;
-		for (FActorDescContainerCollection::TConstIterator<> ActorDescIterator(this); ActorDescIterator; ++ActorDescIterator)
+		for (FActorDescContainerInstanceCollection::TConstIterator<> Iterator(this); Iterator; ++Iterator)
 		{
-			if (ALocationVolume* LocationVolume = Cast<ALocationVolume>(ActorDescIterator->GetActor()); IsValid(LocationVolume))
+			if (ALocationVolume* LocationVolume = Cast<ALocationVolume>(Iterator->GetActor()); IsValid(LocationVolume))
 			{
 				check(LocationVolume->GetClass()->ImplementsInterface(UWorldPartitionActorLoaderInterface::StaticClass()));
 
@@ -1926,19 +1876,19 @@ void UWorldPartition::DumpActorDescs(const FString& Path)
 {
 	if (FArchive* LogFile = IFileManager::Get().CreateFileWriter(*Path))
 	{
-		TArray<const FWorldPartitionActorDesc*> ActorDescs;
+		TArray<const FWorldPartitionActorDescInstance*> ActorDescInstances;
 		TMap<FName, FString> DataLayersDumpString = GetDataLayersDumpString(this);
-		for (FActorDescContainerCollection::TConstIterator<> ActorDescIterator(this); ActorDescIterator; ++ActorDescIterator)
+		for (FActorDescContainerInstanceCollection::TConstIterator<> Iterator(this); Iterator; ++Iterator)
 		{
-			ActorDescs.Add(*ActorDescIterator);
+			ActorDescInstances.Add(*Iterator);
 		}
-		ActorDescs.Sort([](const FWorldPartitionActorDesc& A, const FWorldPartitionActorDesc& B)
+		ActorDescInstances.Sort([](const FWorldPartitionActorDescInstance& A, const FWorldPartitionActorDescInstance& B)
 		{
 			return A.GetGuid() < B.GetGuid();
 		});
-		for (const FWorldPartitionActorDesc* ActorDescIterator : ActorDescs)
+		for (const FWorldPartitionActorDescInstance* Iterator : ActorDescInstances)
 		{
-			FString LineEntry = GetActorDescDumpString(ActorDescIterator, DataLayersDumpString);
+			FString LineEntry = GetActorDescDumpString(Iterator, DataLayersDumpString);
 			LogFile->Serialize(TCHAR_TO_ANSI(*LineEntry), LineEntry.Len());
 		}
 
@@ -1991,18 +1941,24 @@ void UWorldPartition::AppendAssetRegistryTags(FAssetRegistryTagsContext Context)
 	}
 }
 
-UActorDescContainer* UWorldPartition::RegisterActorDescContainer(const FContainerRegistrationParams& InRegistrationParameters)
+UActorDescContainerInstance* UWorldPartition::RegisterActorDescContainerInstance(const UActorDescContainerInstance::FInitializeParams& InParams)
 {
 	TRACE_CPUPROFILER_EVENT_SCOPE(UWorldPartition::RegisterActorDescContainer);
-
-	if (!Contains(InRegistrationParameters.PackageName))
+	const bool bIsEditor = !World->IsGameWorld();
+	const bool bIsGameWorld = !bIsEditor;
+	const bool bIsCooking = IsRunningCookCommandlet();
+	const bool bIsStreamedLevel = ULevelStreaming::FindStreamingLevel(GetTypedOuter<ULevel>()) != nullptr;
+	
+	if (!Contains(InParams.ContainerPackageName))
 	{	
-		UActorDescContainer::FInitializeParams ContainerInitParams(GetWorld(), InRegistrationParameters.PackageName);
-
+		// Initialize ContainerInstance hierarchy if we are the main world partition or if we are a game streamed world partition which means we have our own generate streaming
+		const bool bCreateContainerInstanceHierarchy = IsMainWorldPartition() || (bIsGameWorld && bIsStreamedLevel) || InParams.bCreateContainerInstanceHierarchy;
+		UActorDescContainerInstance::FInitializeParams InitParams(InParams.ContainerPackageName, bCreateContainerInstanceHierarchy);
+		
 		const FWorldDataLayersActorDesc* WorldDataLayerActorsDesc = nullptr;
-		ContainerInitParams.FilterActorDesc = [this, &WorldDataLayerActorsDesc, &InRegistrationParameters](const FWorldPartitionActorDesc* ActorDesc)
+		InitParams.FilterActorDescFunc = [this, &WorldDataLayerActorsDesc, &InParams](const FWorldPartitionActorDesc* ActorDesc)
 		{
-			if (InRegistrationParameters.FilterActorDescFunc && !InRegistrationParameters.FilterActorDescFunc(ActorDesc))
+			if (InParams.FilterActorDescFunc && !InParams.FilterActorDescFunc(ActorDesc))
 			{
 				return false;
 			}
@@ -2021,9 +1977,9 @@ UActorDescContainer* UWorldPartition::RegisterActorDescContainer(const FContaine
 			}
 
 			// Filter actors with duplicated GUID in WorldPartition (across containers):
-			// difference with the duplicate check in UActorDescContainer is that WorldPartition is a collection of containers so same Guid could exist across those containers
+			// difference with the duplicate check in UActorDescContainerInstance is that WorldPartition is a collection of containers so same Guid could exist across those containers
 			// which wouldn't be validated by the container itself.
-			if (GetActorDesc(ActorDesc->GetGuid()))
+			if (GetActorDescInstance(ActorDesc->GetGuid()))
 			{
 				UE_LOG(LogWorldPartition, Warning, TEXT("Found existing actor descriptor guid `%s`: Actor: '%s' from package '%s'"),
 					*ActorDesc->GetGuid().ToString(),
@@ -2035,40 +1991,44 @@ UActorDescContainer* UWorldPartition::RegisterActorDescContainer(const FContaine
 			return true;
 		};
 
-		UActorDescContainer* ContainerToRegister = NewObject<UActorDescContainer>(this);
-		ContainerToRegister->Initialize(ContainerInitParams);
+		UActorDescContainerInstance* ContainerInstanceToRegister = NewObject<UActorDescContainerInstance>(this, UActorDescContainerInstance::StaticClass(), NAME_None, RF_Transient);
+		
+		OnActorDescContainerInstancePreInitialize.ExecuteIfBound(InitParams, ContainerInstanceToRegister);
 
-		AddContainer(ContainerToRegister);
+		ContainerInstanceToRegister->Initialize(InitParams);
+			
+		AddContainer(ContainerInstanceToRegister);
 
 		if (IsInitialized() && EditorHash != nullptr)
 		{
 			FWorldPartitionReference WDLReference;
-			for (FActorDescList::TIterator<> ActorDescIterator(ContainerToRegister); ActorDescIterator; ++ActorDescIterator)
+			for (UActorDescContainerInstance::TIterator<> Iterator(ContainerInstanceToRegister); Iterator; ++Iterator)
 			{
-				if (ActorDescIterator->GetActorNativeClass()->IsChildOf<AWorldDataLayers>())
+				if (Iterator->GetActorNativeClass()->IsChildOf<AWorldDataLayers>())
 				{
-					WDLReference = FWorldPartitionReference(this, ActorDescIterator->GetGuid());
+					WDLReference = FWorldPartitionReference(this, Iterator->GetGuid());
 					break;
 				}
 			}
 
-			InitializeActorDescContainerEditorStreaming(ContainerToRegister, true);
+			const bool bHashActorDescs = bIsEditor && !bIsCooking;
+			InitializeActorDescContainerEditorStreaming(ContainerInstanceToRegister, bHashActorDescs);
 		}
 
-		OnActorDescContainerRegistered.Broadcast(ContainerToRegister);
+		OnActorDescContainerInstanceRegistered.Broadcast(ContainerInstanceToRegister);
 
-		return ContainerToRegister;
+		return ContainerInstanceToRegister;
 	}
 	
 	return nullptr;
 }
 
-bool UWorldPartition::UnregisterActorDescContainer(UActorDescContainer* InActorDescContainer)
+bool UWorldPartition::UnregisterActorDescContainerInstance(UActorDescContainerInstance* InActorDescContainerInstance)
 {
-	if (Contains(InActorDescContainer->GetContainerPackage()))
+	if (Contains(InActorDescContainerInstance->GetContainerPackage()))
 	{
 		TArray<FGuid> ActorGuids;
-		for (UActorDescContainer::TIterator<> It(InActorDescContainer); It; ++It)
+		for (UActorDescContainerInstance::TIterator<> It(InActorDescContainerInstance); It; ++It)
 		{
 			FWorldPartitionHandle ActorHandle(this, It->GetGuid());
 			if (ActorHandle.IsValid())
@@ -2092,19 +2052,19 @@ bool UWorldPartition::UnregisterActorDescContainer(UActorDescContainer* InActorD
 			ForceLoadedActors->RemoveActors(ActorGuids);
 		}
 
-		OnActorDescContainerUnregistered.Broadcast(InActorDescContainer);
+		OnActorDescContainerInstanceUnregistered.Broadcast(InActorDescContainerInstance);
 
 		if (IsInitialized() && EditorHash != nullptr)
 		{
-			for (UActorDescContainer::TIterator<> It(InActorDescContainer); It; ++It)
+			for (UActorDescContainerInstance::TIterator<> It(InActorDescContainerInstance); It; ++It)
 			{
-				UnhashActorDesc(*It);
+				UnhashActorDescInstance(*It);
 			}
 		}
 
-		InActorDescContainer->Uninitialize();
+		InActorDescContainerInstance->Uninitialize();
 
-		verify(RemoveContainer(InActorDescContainer));
+		verify(RemoveContainer(InActorDescContainerInstance));
 
 		return true;
 	}
@@ -2114,9 +2074,9 @@ bool UWorldPartition::UnregisterActorDescContainer(UActorDescContainer* InActorD
 
 void UWorldPartition::UninitializeActorDescContainers()
 {
-	for (UActorDescContainer* Container : ActorDescContainerCollection)
+	for (UActorDescContainerInstance* ContainerInstance : ActorDescContainerInstanceCollection)
 	{
-		Container->Uninitialize();
+		ContainerInstance->Uninitialize();
 	}
 
 	Empty();
@@ -2193,10 +2153,7 @@ void UWorldPartition::OnWorldRenamed(UWorld* RenamedWorld)
 {
 	if (GetWorld() == RenamedWorld)
 	{
-		ActorDescContainer->SetContainerPackage(GetWorld()->GetPackage()->GetFName());
-
-		// World was renamed so existing context is invalid.
-		InstancingContext = FLinkerInstancingContext();
+		ActorDescContainerInstance->SetContainerPackage(GetWorld()->GetPackage()->GetFName());
 	}
 }
 
diff --git a/Engine/Source/Runtime/Engine/Private/WorldPartition/WorldPartitionActorDesc.cpp b/Engine/Source/Runtime/Engine/Private/WorldPartition/WorldPartitionActorDesc.cpp
--- a/Engine/Source/Runtime/Engine/Private/WorldPartition/WorldPartitionActorDesc.cpp
+++ b/Engine/Source/Runtime/Engine/Private/WorldPartition/WorldPartitionActorDesc.cpp
@@ -16,7 +16,8 @@
 #include "UObject/FortniteSeasonBranchObjectVersion.h"
 #include "UObject/FortniteMainBranchObjectVersion.h"
 #include "WorldPartition/WorldPartitionHelpers.h"
-#include "WorldPartition/WorldPartitionActorDescView.h"
+#include "WorldPartition/WorldPartitionActorDescInstance.h"
+#include "WorldPartition/WorldPartitionActorDescInstanceViewInterface.h"
 #include "WorldPartition/WorldPartitionLog.h"
 #include "WorldPartition/ActorDescContainer.h"
 #include "WorldPartition/HLOD/HLODLayer.h"
@@ -49,13 +50,9 @@ FWorldPartitionActorDesc::FWorldPartitionActorDesc()
 	, bActorIsListedInSceneOutliner(true)
 	, bIsUsingDataLayerAsset(false)
 	, bIsBoundsValid(false)
-	, SoftRefCount(0)
-	, HardRefCount(0)
 	, ActorNativeClass(nullptr)
 	, Container(nullptr)
-	, bIsForcedNonSpatiallyLoaded(false)
 	, bIsDefaultActorDesc(false)
-	, UnloadedReason(nullptr)
 {}
 
 void FWorldPartitionActorDesc::Init(const AActor* InActor)
@@ -137,36 +134,12 @@ void FWorldPartitionActorDesc::Init(const AActor* InActor)
 
 			// Init DataLayers persistent info
 			DataLayers = bIsUsingDataLayerAsset ? MoveTemp(LocalDataLayerAssetPaths) : MoveTemp(LocalDataLayerInstanceNames);
-
-			// Init DataLayers transient info
-			UWorld* ActorOwningWorld = InActor->GetWorld();
-			// In Editor, always use owning world to resolve data layers.
-			// In game (PIE), always use actor's world partition DataLayerManager to resolve data layers.
-			const UDataLayerManager* ResolvingDataLayerManager = !ActorOwningWorld->IsGameWorld() ? ActorOwningWorld->GetDataLayerManager() : ActorWorldPartition->GetDataLayerManager();
-			// The only case where it's normal not to have a valid ResolvingDataLayerManager is if OwningWorld is not game and not partitioned
-			// Also, consider that we don't need the DataLayerManager for any running commandlet except for cook
-			// If Actor's World Partition isn't initialized we don't need to resolve (ex: Disaster Recovery Plugin will save Actor Packages in temp location which will Create a temp ActorDesc)
-			const bool bCanSkipDataLayerManagerValidation = (IsRunningCommandlet() && !IsRunningCookCommandlet()) || (!ActorOwningWorld->IsGameWorld() && (!ActorOwningWorld->IsPartitionedWorld() || !ActorWorldPartition->IsInitialized()));
-			check(ResolvingDataLayerManager || bCanSkipDataLayerManagerValidation);
-			if (ResolvingDataLayerManager)
-			{
-				if (ensure(ResolvingDataLayerManager->CanResolveDataLayers()))
-				{
-					// Here we process a loaded actor, so resolving makes sense as long as the actordesc is not reused as a template
-					ResolvedDataLayerInstanceNames = FDataLayerUtils::ResolvedDataLayerInstanceNames(ResolvingDataLayerManager, this);
-				}
-				else
-				{
-					check(!HasResolvedDataLayerInstanceNames());
-				}
-			}
 		}
 		else
 		{
 			// Possible there is no World Partition for regular OFPA levels that haven't been converted to support Data Layers
 			bIsUsingDataLayerAsset = true;
 			DataLayers.Empty();
-			ResolvedDataLayerInstanceNames = TArray<FName>();
 		}
 	}
 
@@ -219,7 +192,6 @@ void FWorldPartitionActorDesc::Init(const AActor* InActor)
 	}
 
 	Container = nullptr;
-	ActorPtr = const_cast<AActor*>(InActor);
 }
 
 void FWorldPartitionActorDesc::Init(const FWorldPartitionActorDescInitData& DescData)
@@ -364,15 +336,17 @@ bool FWorldPartitionActorDesc::ShouldResave(const FWorldPartitionActorDesc* Othe
 	return (bActorIsHLODRelevant == Other->bActorIsHLODRelevant) || (bActorIsHLODRelevant && !Other->bActorIsHLODRelevant);
 }
 
-void FWorldPartitionActorDesc::SerializeTo(TArray<uint8>& OutData)
+void FWorldPartitionActorDesc::SerializeTo(TArray<uint8>& OutData) const
 {
+	FWorldPartitionActorDesc* MutableThis = const_cast<FWorldPartitionActorDesc*>(this);
+
 	// Serialize to archive and gather custom versions
 	TArray<uint8> PayloadData;
 	FMemoryWriter PayloadAr(PayloadData, true);
-	FActorDescArchive ActorDescAr(PayloadAr, this);
+	FActorDescArchive ActorDescAr(PayloadAr, MutableThis);
 	ActorDescAr.Init();
 
-	Serialize(ActorDescAr);
+	MutableThis->Serialize(ActorDescAr);
 
 	// Serialize custom versions
 	TArray<uint8> HeaderData;
@@ -385,22 +359,9 @@ void FWorldPartitionActorDesc::SerializeTo(TArray<uint8>& OutData)
 	OutData.Append(PayloadData);
 }
 
-const TArray<FName>& FWorldPartitionActorDesc::GetDataLayerInstanceNames() const
-{
-	static TArray<FName> EmptyDataLayers;
-	if (ensure(HasResolvedDataLayerInstanceNames()))
-	{
-		return ResolvedDataLayerInstanceNames.GetValue();
-	}
-	return EmptyDataLayers;
-}
-
 void FWorldPartitionActorDesc::TransferFrom(const FWorldPartitionActorDesc* From)
 {
 	Container = From->Container;
-	SoftRefCount = From->SoftRefCount;
-	HardRefCount = From->HardRefCount;
-	bIsForcedNonSpatiallyLoaded = From->bIsForcedNonSpatiallyLoaded;
 }
 
 void FWorldPartitionActorDesc::RegisterActorDescDeprecator(TSubclassOf<AActor> ActorClass, const FActorDescDeprecator& Deprecator)
@@ -409,12 +370,6 @@ void FWorldPartitionActorDesc::RegisterActorDescDeprecator(TSubclassOf<AActor> A
 	Deprecators.Add(ActorClass, Deprecator);
 }
 
-void FWorldPartitionActorDesc::TransformInstance(const FString& From, const FString& To)
-{
-	check(!HardRefCount);
-	ActorPath = *ActorPath.ToString().Replace(*From, *To);
-}
-
 FString FWorldPartitionActorDesc::ToString(EToStringMode Mode) const
 {
 	auto GetBoolStr = [](bool bValue) -> const TCHAR*
@@ -778,11 +733,11 @@ FGuid FWorldPartitionActorDesc::GetContentBundleGuid() const
 	return ContentBundleGuid;
 }
 
-void FWorldPartitionActorDesc::CheckForErrors(IStreamingGenerationErrorHandler* ErrorHandler) const
+void FWorldPartitionActorDesc::CheckForErrors(const IWorldPartitionActorDescInstanceView* InActorDescView, IStreamingGenerationErrorHandler* ErrorHandler) const
 {
 	if (IsResaveNeeded())
 	{
-		ErrorHandler->OnActorNeedsResave(this);
+		ErrorHandler->OnActorNeedsResave(*InActorDescView);
 	}
 }
 
@@ -796,7 +751,7 @@ bool FWorldPartitionActorDesc::IsListedInSceneOutliner() const
 	return bActorIsListedInSceneOutliner;
 }
 
-bool FWorldPartitionActorDesc::IsEditorRelevant() const
+bool FWorldPartitionActorDesc::IsEditorRelevant(const FWorldPartitionActorDescInstance* InActorDescInstance) const
 {
 	if (GetActorIsRuntimeOnly())
 	{
@@ -805,118 +760,15 @@ bool FWorldPartitionActorDesc::IsEditorRelevant() const
 
 	if (IsMainWorldOnly())
 	{
-		return GetContainer() && !GetContainer()->IsTemplateContainer() && GetContainer()->GetWorldPartition()->IsMainWorldPartition();
+		return InActorDescInstance->GetContainerInstance()->GetContainerID().IsMainContainer();
 	}
 
 	return true;
 }
 
-bool FWorldPartitionActorDesc::IsRuntimeRelevant(const FActorContainerID& InContainerID) const
-{
-	return InContainerID.IsMainContainer() || !IsMainWorldOnly();
-}
-
-bool FWorldPartitionActorDesc::IsLoaded(bool bEvenIfPendingKill) const
-{
-	if (ActorPtr.IsExplicitlyNull() || ActorPtr.IsStale())
-	{
-		ActorPtr = FindObject<AActor>(nullptr, *ActorPath.ToString());
-	}
-
-	return ActorPtr.IsValid(bEvenIfPendingKill);
-}
-
-AActor* FWorldPartitionActorDesc::GetActor(bool bEvenIfPendingKill, bool bEvenIfUnreachable) const
-{
-	if (ActorPtr.IsExplicitlyNull() || ActorPtr.IsStale())
-	{
-		ActorPtr = FindObject<AActor>(nullptr, *ActorPath.ToString());
-	}
-
-	return bEvenIfUnreachable ? ActorPtr.GetEvenIfUnreachable() : ActorPtr.Get(bEvenIfPendingKill);
-}
-
-TWeakObjectPtr<AActor>* FWorldPartitionActorDesc::GetActorPtr(bool bEvenIfPendingKill, bool bEvenIfUnreachable) const
-{
-	return GetActor(bEvenIfPendingKill, bEvenIfUnreachable) ? &ActorPtr : nullptr;
-}
-
-const FText& FWorldPartitionActorDesc::GetUnloadedReason() const
-{
-	static FText Unloaded(LOCTEXT("UnloadedReason", "Unloaded"));
-	return UnloadedReason ? *UnloadedReason : Unloaded;
-}
-
-AActor* FWorldPartitionActorDesc::Load() const
+bool FWorldPartitionActorDesc::IsRuntimeRelevant(const FWorldPartitionActorDescInstance* InActorDescInstance) const
 {
-	static FText FailedToLoad(LOCTEXT("FailedToLoadReason", "Failed to load"));
-	UnloadedReason = nullptr;
-
-	if (ActorPtr.IsExplicitlyNull() || ActorPtr.IsStale())
-	{
-		// First, try to find the existing actor which could have been loaded by another actor (through standard serialization)
-		ActorPtr = FindObject<AActor>(nullptr, *ActorPath.ToString());
-	}
-
-	// Then, if the actor isn't loaded, load it
-	if (ActorPtr.IsExplicitlyNull())
-	{
-		const FLinkerInstancingContext* InstancingContext = Container ? Container->GetInstancingContext() : nullptr;
-
-		UPackage* Package = nullptr;
-
-		if (InstancingContext)
-		{
-			FName RemappedPackageName = InstancingContext->RemapPackage(ActorPackage);
-			check(RemappedPackageName != ActorPath.GetLongPackageFName());
-
-			Package = CreatePackage(*RemappedPackageName.ToString());
-		}
-
-		Package = LoadPackage(Package, *ActorPackage.ToString(), LOAD_None, nullptr, InstancingContext);
-
-		if (Package)
-		{
-			ActorPtr = FindObject<AActor>(nullptr, *ActorPath.ToString());
-			if (!ActorPtr.IsValid())
-			{
-				UE_LOG(LogWorldPartition, Warning, TEXT("Can't load actor guid `%s` ('%s') from package '%s'"), *Guid.ToString(), *GetActorName().ToString(), *ActorPackage.ToString());
-				UnloadedReason = &FailedToLoad;
-			}
-		}
-	}
-
-	return ActorPtr.Get();
-}
-
-void FWorldPartitionActorDesc::Unload()
-{
-	if (AActor* Actor = GetActor())
-	{
-		// At this point, it can happen that an actor isn't in an external package:
-		//
-		// PIE travel: 
-		//		in this case, actors referenced by the world package (an example is the level script) will be duplicated as part of the PIE world duplication and will end up
-		//		not being using an external package, which is fine because in that case they are considered as always loaded.
-		//
-		// FWorldPartitionCookPackageSplitter:
-		//		should mark each FWorldPartitionActorDesc as moved, and the splitter should take responsbility for calling ClearFlags on every object in 
-		//		the package when it does the move
-
-		if (Actor->IsPackageExternal())
-		{
-			ForEachObjectWithPackage(Actor->GetPackage(), [](UObject* Object)
-			{
-				if (Object->HasAnyFlags(RF_Public | RF_Standalone))
-				{
-					CastChecked<UMetaData>(Object)->ClearFlags(RF_Public | RF_Standalone);
-				}
-				return true;
-			}, false);
-		}
-
-		ActorPtr = nullptr;
-	}
+	return !IsMainWorldOnly() || InActorDescInstance->GetContainerInstance()->GetContainerID().IsMainContainer();
 }
 
 #undef LOCTEXT_NAMESPACE
diff --git a/Engine/Source/Runtime/Engine/Private/WorldPartition/WorldPartitionActorDescInstance.cpp b/Engine/Source/Runtime/Engine/Private/WorldPartition/WorldPartitionActorDescInstance.cpp
new file mode 100644
--- /dev/null
+++ b/Engine/Source/Runtime/Engine/Private/WorldPartition/WorldPartitionActorDescInstance.cpp
@@ -0,0 +1,230 @@
+// Copyright Epic Games, Inc. All Rights Reserved.
+
+#include "WorldPartition/WorldPartitionActorDescInstance.h"
+
+#if WITH_EDITOR
+
+#include "WorldPartition/ActorDescContainerInstance.h"
+#include "WorldPartition/WorldPartition.h"
+#include "WorldPartition/WorldPartitionActorDesc.h"
+#include "UObject/MetaData.h"
+
+#define LOCTEXT_NAMESPACE "FWorldPartitionActorDescInstance"
+
+FWorldPartitionActorDescInstance::FWorldPartitionActorDescInstance()
+	: ContainerInstance(nullptr)
+	, SoftRefCount(0)
+	, HardRefCount(0)
+	, bIsForcedNonSpatiallyLoaded(false)
+	, UnloadedReason(nullptr)
+	, ActorDesc(nullptr)
+	, ChildContainerInstance(nullptr)
+{
+}
+
+FWorldPartitionActorDescInstance::FWorldPartitionActorDescInstance(UActorDescContainerInstance* InContainerInstance, FWorldPartitionActorDesc* InActorDesc)
+	: FWorldPartitionActorDescInstance()
+{
+	check(InContainerInstance);
+	ContainerInstance = InContainerInstance;
+
+	check(InActorDesc);
+	ActorDesc = InActorDesc;
+}
+
+void FWorldPartitionActorDescInstance::UpdateActorDesc(FWorldPartitionActorDesc* InActorDesc)
+{
+	ActorDesc = InActorDesc;
+}
+
+bool FWorldPartitionActorDescInstance::IsLoaded(bool bEvenIfPendingKill) const
+{
+	if (ActorPtr.IsExplicitlyNull() || ActorPtr.IsStale())
+	{
+		ActorPtr = FindObject<AActor>(nullptr, *GetActorSoftPath().ToString());
+	}
+
+	return ActorPtr.IsValid(bEvenIfPendingKill);
+}
+
+AActor* FWorldPartitionActorDescInstance::GetActor(bool bEvenIfPendingKill, bool bEvenIfUnreachable) const
+{
+	if (ActorPtr.IsExplicitlyNull() || ActorPtr.IsStale())
+	{
+		ActorPtr = FindObject<AActor>(nullptr, *GetActorSoftPath().ToString());
+	}
+
+	return bEvenIfUnreachable ? ActorPtr.GetEvenIfUnreachable() : ActorPtr.Get(bEvenIfPendingKill);
+}
+
+TWeakObjectPtr<AActor>* FWorldPartitionActorDescInstance::GetActorPtr(bool bEvenIfPendingKill, bool bEvenIfUnreachable) const
+{
+	return GetActor(bEvenIfPendingKill, bEvenIfUnreachable) ? &ActorPtr : nullptr;
+}
+
+FSoftObjectPath FWorldPartitionActorDescInstance::GetActorSoftPath() const
+{
+	return ActorPath.IsSet() ? ActorPath.GetValue() : GetActorDesc()->GetActorSoftPath();
+}
+
+FName FWorldPartitionActorDescInstance::GetActorName() const
+{
+	return *FPaths::GetExtension(GetActorSoftPath().ToString());
+}
+
+bool FWorldPartitionActorDescInstance::IsValid() const
+{
+	return !!GetActorDesc();
+}
+
+bool FWorldPartitionActorDescInstance::IsEditorRelevant() const
+{
+	return GetActorDesc()->IsEditorRelevant(this);
+}
+
+bool FWorldPartitionActorDescInstance::IsRuntimeRelevant() const
+{
+	return GetActorDesc()->IsRuntimeRelevant(this);
+}
+
+AActor* FWorldPartitionActorDescInstance::Load()
+{
+	static FText FailedToLoad(LOCTEXT("FailedToLoadReason", "Failed to load"));
+	UnloadedReason = nullptr;
+
+	if (ActorPtr.IsExplicitlyNull() || ActorPtr.IsStale())
+	{
+		// First, try to find the existing actor which could have been loaded by another actor (through standard serialization)
+		ActorPtr = FindObject<AActor>(nullptr, *GetActorSoftPath().ToString());
+	}
+
+	// Then, if the actor isn't loaded, load it
+	if (ActorPtr.IsExplicitlyNull())
+	{
+		const FLinkerInstancingContext* InstancingContext = GetContainerInstance()->GetInstancingContext();
+		FSoftObjectPath LocalActorPath = GetActorSoftPath();
+
+		UPackage* Package = nullptr;
+
+		if (InstancingContext)
+		{
+			FName RemappedPackageName = InstancingContext->RemapPackage(GetActorPackage());
+			check(RemappedPackageName != LocalActorPath.GetLongPackageFName());
+
+			Package = CreatePackage(*RemappedPackageName.ToString());
+		}
+
+		Package = LoadPackage(Package, *GetActorPackage().ToString(), LOAD_None, nullptr, InstancingContext);
+
+		if (Package)
+		{
+			ActorPtr = FindObject<AActor>(nullptr, *LocalActorPath.ToString());
+			if (!ActorPtr.IsValid())
+			{
+				UE_LOG(LogWorldPartition, Warning, TEXT("Can't load actor guid `%s` ('%s') from package '%s'"), *GetGuid().ToString(), *GetActorName().ToString(), *GetActorPackage().ToString());
+				UnloadedReason = &FailedToLoad;
+			}
+		}
+	}
+
+	return ActorPtr.Get();
+}
+
+void FWorldPartitionActorDescInstance::Unload()
+{
+	// Notify Desc as it can have some custom code to run on the actor depending on type
+	GetActorDesc()->OnUnloadingInstance(this);
+
+	if (AActor* Actor = GetActor())
+	{
+		// At this point, it can happen that an actor isn't in an external package:
+		//
+		// PIE travel: 
+		//		in this case, actors referenced by the world package (an example is the level script) will be duplicated as part of the PIE world duplication and will end up
+		//		not being using an external package, which is fine because in that case they are considered as always loaded.
+		//
+		// FWorldPartitionCookPackageSplitter:
+		//		should mark each FWorldPartitionActorDesc as moved, and the splitter should take responsbility for calling ClearFlags on every object in 
+		//		the package when it does the move
+
+		if (Actor->IsPackageExternal())
+		{
+			ForEachObjectWithPackage(Actor->GetPackage(), [](UObject* Object)
+				{
+					if (Object->HasAnyFlags(RF_Public | RF_Standalone))
+					{
+						CastChecked<UMetaData>(Object)->ClearFlags(RF_Public | RF_Standalone);
+					}
+					return true;
+				}, false);
+		}
+
+		ActorPtr = nullptr;
+	}
+}
+
+void FWorldPartitionActorDescInstance::Invalidate()
+{
+	check(!ChildContainerInstance);
+	ContainerInstance = nullptr;
+}
+
+const TArray<FName>& FWorldPartitionActorDescInstance::GetDataLayerInstanceNames() const
+{
+	static TArray<FName> EmptyDataLayers;
+	if (ensure(HasResolvedDataLayerInstanceNames()))
+	{
+		return ResolvedDataLayerInstanceNames.GetValue();
+	}
+	return EmptyDataLayers;
+}
+
+const FText& FWorldPartitionActorDescInstance::GetUnloadedReason() const
+{
+	static FText Unloaded(LOCTEXT("UnloadedReason", "Unloaded"));
+	return UnloadedReason ? *UnloadedReason : Unloaded;
+}
+
+FString FWorldPartitionActorDescInstance::ToString(FWorldPartitionActorDesc::EToStringMode Mode) const
+{
+	return GetActorDesc()->ToString(Mode);
+}
+
+void FWorldPartitionActorDescInstance::RegisterChildContainerInstance()
+{
+	check(IsChildContainerInstance());
+	check(!ChildContainerInstance);
+
+	ChildContainerInstance = GetActorDesc()->CreateChildContainerInstance(this);
+	check(ChildContainerInstance);
+	ContainerInstance->OnRegisterChildContainerInstance(GetGuid(), ChildContainerInstance);
+}
+
+void FWorldPartitionActorDescInstance::UnregisterChildContainerInstance()
+{
+	check(ChildContainerInstance);
+	ContainerInstance->OnUnregisterChildContainerInstance(GetGuid());
+	ChildContainerInstance->Uninitialize();
+	ChildContainerInstance = nullptr;
+}
+
+void FWorldPartitionActorDescInstance::UpdateChildContainerInstance()
+{
+	check(ChildContainerInstance);
+		
+	// Create before unregistering so that we benefit from shared containers
+	UActorDescContainerInstance* NewChildContainerInstance = IsChildContainerInstance() ? GetActorDesc()->CreateChildContainerInstance(this) : nullptr;
+		
+	ContainerInstance->OnUnregisterChildContainerInstance(GetGuid());
+	ChildContainerInstance->Uninitialize();
+	
+	if (NewChildContainerInstance)
+	{
+		ChildContainerInstance = NewChildContainerInstance;
+		ContainerInstance->OnRegisterChildContainerInstance(GetGuid(), ChildContainerInstance);
+	}
+}
+
+#undef LOCTEXT_NAMESPACE
+
+#endif
\ No newline at end of file
diff --git a/Engine/Source/Runtime/Engine/Private/WorldPartition/WorldPartitionActorDescUtils.cpp b/Engine/Source/Runtime/Engine/Private/WorldPartition/WorldPartitionActorDescUtils.cpp
--- a/Engine/Source/Runtime/Engine/Private/WorldPartition/WorldPartitionActorDescUtils.cpp
+++ b/Engine/Source/Runtime/Engine/Private/WorldPartition/WorldPartitionActorDescUtils.cpp
@@ -15,6 +15,8 @@
 #include "Misc/Base64.h"
 #include "WorldPartition/WorldPartitionActorDesc.h"
 #include "WorldPartition/WorldPartitionHelpers.h"
+#include "WorldPartition/ActorDescContainerInstance.h"
+#include "WorldPartition/WorldPartitionActorDescInstance.h"
 
 static FName NAME_ActorMetaDataClass(TEXT("ActorMetaDataClass"));
 static FName NAME_ActorMetaData(TEXT("ActorMetaData"));
@@ -191,20 +193,20 @@ void FWorldPartitionActorDescUtils::UpdateActorDescriptorFromActorDescriptor(TUn
 	OutActorDesc = MoveTemp(InActorDesc);
 }
 
-void FWorldPartitionActorDescUtils::ReplaceActorDescriptorPointerFromActor(const AActor* InOldActor, AActor* InNewActor, FWorldPartitionActorDesc* InActorDesc)
+void FWorldPartitionActorDescUtils::ReplaceActorDescriptorPointerFromActor(const AActor* InOldActor, AActor* InNewActor, FWorldPartitionActorDescInstance* InActorDescInstance)
 {
 	if (InNewActor)
 	{
 		checkf(InOldActor->GetActorGuid() == InNewActor->GetActorGuid(), TEXT("Mismatching new actor GUID: old=%s new=%s"), *InOldActor->GetActorGuid().ToString(), *InNewActor->GetActorGuid().ToString());
-		checkf(InNewActor->GetActorGuid() == InActorDesc->GetGuid(), TEXT("Mismatching desc actor GUID: desc=%s new=%s"), *InActorDesc->GetGuid().ToString(), *InNewActor->GetActorGuid().ToString());
+		checkf(InNewActor->GetActorGuid() == InActorDescInstance->GetGuid(), TEXT("Mismatching desc actor GUID: desc=%s new=%s"), *InActorDescInstance->GetGuid().ToString(), *InNewActor->GetActorGuid().ToString());
 	}
 
-	if (InActorDesc->ActorPtr.IsValid())
+	if (InActorDescInstance->ActorPtr.IsValid())
 	{
-		checkf(InActorDesc->ActorPtr == InOldActor, TEXT("Mismatching old desc actor: desc=%s old=%s"), *InActorDesc->ActorPtr->GetActorNameOrLabel(), *InOldActor->GetActorNameOrLabel());
+		checkf(InActorDescInstance->ActorPtr == InOldActor, TEXT("Mismatching old desc actor: desc=%s old=%s"), *InActorDescInstance->ActorPtr->GetActorNameOrLabel(), *InOldActor->GetActorNameOrLabel());
 	}
 
-	InActorDesc->ActorPtr = InNewActor;
+	InActorDescInstance->ActorPtr = InNewActor;
 }
 
 bool FWorldPartitionActorDescUtils::FixupRedirectedAssetPath(FName& InOutAssetPath)
diff --git a/Engine/Source/Runtime/Engine/Private/WorldPartition/WorldPartitionActorDescView.cpp b/Engine/Source/Runtime/Engine/Private/WorldPartition/WorldPartitionActorDescView.cpp
--- a/Engine/Source/Runtime/Engine/Private/WorldPartition/WorldPartitionActorDescView.cpp
+++ b/Engine/Source/Runtime/Engine/Private/WorldPartition/WorldPartitionActorDescView.cpp
@@ -3,47 +3,346 @@
 #include "WorldPartition/WorldPartitionActorDescView.h"
 
 #if WITH_EDITOR
+#include "WorldPartition/WorldPartitionLog.h"
 #include "WorldPartition/WorldPartitionActorDesc.h"
 
-FWorldPartitionActorDescView::FWorldPartitionActorDescView() : FWorldPartitionActorDescView(nullptr) {}
-FWorldPartitionActorDescView::FWorldPartitionActorDescView(const FWorldPartitionActorDesc* InActorDesc)	: ActorDesc(InActorDesc) {}
-const FGuid& FWorldPartitionActorDescView::GetGuid() const { return ActorDesc->GetGuid(); }
-FTopLevelAssetPath FWorldPartitionActorDescView::GetBaseClass() const {	return ActorDesc->GetBaseClass(); }
-FTopLevelAssetPath FWorldPartitionActorDescView::GetNativeClass() const { return ActorDesc->GetNativeClass(); }
-UClass* FWorldPartitionActorDescView::GetActorNativeClass() const {	return ActorDesc->GetActorNativeClass(); }
-FName FWorldPartitionActorDescView::GetRuntimeGrid() const { return ActorDesc->GetRuntimeGrid(); }
-bool FWorldPartitionActorDescView::GetIsSpatiallyLoaded() const { return ActorDesc->GetIsSpatiallyLoaded(); }
-bool FWorldPartitionActorDescView::GetActorIsEditorOnly() const { return ActorDesc->GetActorIsEditorOnly(); }
-bool FWorldPartitionActorDescView::GetActorIsRuntimeOnly() const { return ActorDesc->GetActorIsRuntimeOnly(); }
-bool FWorldPartitionActorDescView::GetActorIsHLODRelevant() const {	return ActorDesc->GetActorIsHLODRelevant(); }
-FSoftObjectPath FWorldPartitionActorDescView::GetHLODLayer() const { return ActorDesc->GetHLODLayer(); }
-bool FWorldPartitionActorDescView::HasResolvedDataLayerInstanceNames() const { return ActorDesc->HasResolvedDataLayerInstanceNames(); }
-const TArray<FName>& FWorldPartitionActorDescView::GetDataLayerInstanceNames() const { return ActorDesc->GetDataLayerInstanceNames(); }
-const TArray<FName>& FWorldPartitionActorDescView::GetTags() const { return ActorDesc->GetTags(); }
-FName FWorldPartitionActorDescView::GetActorPackage() const { return ActorDesc->GetActorPackage(); }
-FSoftObjectPath FWorldPartitionActorDescView::GetActorSoftPath() const { return ActorDesc->GetActorSoftPath(); }
-FName FWorldPartitionActorDescView::GetActorLabel() const {	return ActorDesc->GetActorLabel(); }
-FBox FWorldPartitionActorDescView::GetEditorBounds() const { return ActorDesc->GetEditorBounds(); }
-FBox FWorldPartitionActorDescView::GetRuntimeBounds() const { return ActorDesc->GetRuntimeBounds(); }
-const TArray<FGuid>& FWorldPartitionActorDescView::GetReferences() const { return ActorDesc->GetReferences(); }
-FString FWorldPartitionActorDescView::ToString() const { return ActorDesc->ToString(); }
-const FGuid& FWorldPartitionActorDescView::GetParentActor() const {	return ActorDesc->GetParentActor(); }
-FName FWorldPartitionActorDescView::GetActorName() const { return ActorDesc->GetActorName(); }
-const FGuid& FWorldPartitionActorDescView::GetFolderGuid() const { return ActorDesc->GetFolderGuid(); }
-FGuid FWorldPartitionActorDescView::GetContentBundleGuid() const { return ActorDesc->GetContentBundleGuid(); }
-FName FWorldPartitionActorDescView::GetContainerPackage() const { return ActorDesc->GetContainerPackage(); }
-bool FWorldPartitionActorDescView::IsContainerInstance() const { return ActorDesc->IsContainerInstance(); }
-bool FWorldPartitionActorDescView::GetContainerInstance(FWorldPartitionActorDesc::FContainerInstance& OutContainerInstance) const {	return ActorDesc->GetContainerInstance(OutContainerInstance); }
-EWorldPartitionActorFilterType FWorldPartitionActorDescView::GetContainerFilterType() const { return ActorDesc->GetContainerFilterType(); }
-const FWorldPartitionActorFilter* FWorldPartitionActorDescView::GetContainerFilter() const { return ActorDesc->GetContainerFilter(); }
-void FWorldPartitionActorDescView::CheckForErrors(IStreamingGenerationErrorHandler* ErrorHandler) const { return ActorDesc->CheckForErrors(ErrorHandler); }
-FName FWorldPartitionActorDescView::GetActorLabelOrName() const { return ActorDesc->GetActorLabelOrName(); }
-AActor* FWorldPartitionActorDescView::GetActor() const { return ActorDesc->GetActor(); }
-bool FWorldPartitionActorDescView::IsEditorOnlyReference(const FGuid& ReferenceGuid) const { return ActorDesc->IsEditorOnlyReference(ReferenceGuid); }
-bool FWorldPartitionActorDescView::GetProperty(FName PropertyName, FName* PropertyValue) const { return ActorDesc->GetProperty(PropertyName, PropertyValue); }
-bool FWorldPartitionActorDescView::HasProperty(FName PropertyName) const { return ActorDesc->HasProperty(PropertyName); }
-
-const TArray<FName>& FWorldPartitionActorDescView::GetRuntimeDataLayerInstanceNames() const { static TArray<FName> EmptyDataLayers; return EmptyDataLayers; }
-const TArray<FGuid>& FWorldPartitionActorDescView::GetEditorReferences() const { static TArray<FGuid> EmptyEditorReferences; return EmptyEditorReferences; }
-
-#endif
+PRAGMA_DISABLE_DEPRECATION_WARNINGS
+
+FWorldPartitionActorDescView::FWorldPartitionActorDescView()
+	: FWorldPartitionActorDescView(nullptr)
+{}
+
+FWorldPartitionActorDescView::FWorldPartitionActorDescView(const FWorldPartitionActorDesc* InActorDesc)
+	: ActorDesc(InActorDesc)
+	, ParentView(nullptr)
+	, bIsForcedNonSpatiallyLoaded(false)
+	, bIsForcedNoRuntimeGrid(false)
+	, bIsForcedNoDataLayers(false)
+	, bIsForceNoHLODLayer(false)
+{}
+
+const FGuid& FWorldPartitionActorDescView::GetGuid() const
+{
+	return ActorDesc->GetGuid();
+}
+
+FTopLevelAssetPath FWorldPartitionActorDescView::GetBaseClass() const
+{
+	return ActorDesc->GetBaseClass();
+}
+
+FTopLevelAssetPath FWorldPartitionActorDescView::GetNativeClass() const
+{
+	return ActorDesc->GetNativeClass();
+}
+
+UClass* FWorldPartitionActorDescView::GetActorNativeClass() const
+{
+	return ActorDesc->GetActorNativeClass();
+}
+
+FVector FWorldPartitionActorDescView::GetOrigin() const
+{
+	return ActorDesc->GetOrigin();
+}
+
+FName FWorldPartitionActorDescView::GetRuntimeGrid() const
+{
+	if (bIsForcedNoRuntimeGrid)
+	{
+		return NAME_None;
+	}
+
+	if (ParentView)
+	{
+		return ParentView->GetRuntimeGrid();
+	}
+
+	return ActorDesc->GetRuntimeGrid();
+}
+
+bool FWorldPartitionActorDescView::GetActorIsEditorOnly() const
+{
+	return ActorDesc->GetActorIsEditorOnly();
+}
+
+bool FWorldPartitionActorDescView::GetActorIsRuntimeOnly() const
+{
+	return ActorDesc->GetActorIsRuntimeOnly();
+}
+
+bool FWorldPartitionActorDescView::GetIsSpatiallyLoaded() const
+{
+	if (bIsForcedNonSpatiallyLoaded)
+	{
+		return false;
+	}
+
+	bool bIsSpatiallyLoaded = ActorDesc->GetIsSpatiallyLoaded();
+	if (bIsSpatiallyLoaded && ParentView)
+	{
+		bIsSpatiallyLoaded = ParentView->GetIsSpatiallyLoaded();
+	}
+
+	return bIsSpatiallyLoaded;
+}
+
+bool FWorldPartitionActorDescView::GetActorIsHLODRelevant() const
+{
+	return ActorDesc->GetActorIsHLODRelevant();
+}
+
+FSoftObjectPath FWorldPartitionActorDescView::GetHLODLayer() const
+{
+	if (bIsForceNoHLODLayer)
+	{
+		return FSoftObjectPath();
+	}
+
+	if (RuntimedHLODLayer.IsSet())
+	{
+		return RuntimedHLODLayer.GetValue();
+	}
+
+	return ActorDesc->GetHLODLayer();
+}
+
+void FWorldPartitionActorDescView::SetDataLayerInstanceNames(const TArray<FName>& InDataLayerInstanceNames)
+{
+	check(!ActorDesc->HasResolvedDataLayerInstanceNames());
+	ResolvedDataLayerInstanceNames = InDataLayerInstanceNames;
+}
+
+const TArray<FName>& FWorldPartitionActorDescView::GetDataLayerInstanceNames() const
+{
+	static TArray<FName> EmptyDataLayers;
+	if (bIsForcedNoDataLayers)
+	{
+		return EmptyDataLayers;
+	}
+
+	if (ParentView)
+	{
+		return ParentView->GetDataLayerInstanceNames();
+	}
+
+	if (ResolvedDataLayerInstanceNames.IsSet())
+	{
+		return ResolvedDataLayerInstanceNames.GetValue();
+	}
+
+	return EmptyDataLayers;
+}
+
+const TArray<FName>& FWorldPartitionActorDescView::GetRuntimeDataLayerInstanceNames() const
+{
+	if (bIsForcedNoDataLayers || !ensure(RuntimeDataLayerInstanceNames.IsSet()))
+	{
+		static TArray<FName> EmptyDataLayers;
+		return EmptyDataLayers;
+	}
+
+	if (ParentView)
+	{
+		return ParentView->GetRuntimeDataLayerInstanceNames();
+	}
+
+	return RuntimeDataLayerInstanceNames.GetValue();
+}
+
+const TArray<FName>& FWorldPartitionActorDescView::GetTags() const
+{
+	return ActorDesc->GetTags();
+}
+
+FName FWorldPartitionActorDescView::GetActorPackage() const
+{
+	return ActorDesc->GetActorPackage();
+}
+
+FSoftObjectPath FWorldPartitionActorDescView::GetActorSoftPath() const
+{
+	return ActorDesc->GetActorSoftPath();
+}
+
+FName FWorldPartitionActorDescView::GetActorLabel() const
+{
+	return ActorDesc->GetActorLabel();
+}
+
+FName FWorldPartitionActorDescView::GetActorName() const
+{
+	return ActorDesc->GetActorName();
+}
+
+FBox FWorldPartitionActorDescView::GetBounds() const
+{
+	return ActorDesc->GetBounds();
+}
+
+FBox FWorldPartitionActorDescView::GetEditorBounds() const
+{
+	return ActorDesc->GetEditorBounds();
+}
+
+FBox FWorldPartitionActorDescView::GetRuntimeBounds() const
+{
+	return ActorDesc->GetRuntimeBounds();
+}
+
+const TArray<FGuid>& FWorldPartitionActorDescView::GetReferences() const
+{
+	return RuntimeReferences.IsSet() ? RuntimeReferences.GetValue() : ActorDesc->GetReferences();
+}
+
+const TArray<FGuid>& FWorldPartitionActorDescView::GetEditorReferences() const
+{
+	return EditorReferences;
+}
+
+FString FWorldPartitionActorDescView::ToString() const
+{
+	return ActorDesc->ToString();
+}
+
+const FGuid& FWorldPartitionActorDescView::GetParentActor() const
+{
+	return ActorDesc->GetParentActor();
+}
+
+const FGuid& FWorldPartitionActorDescView::GetFolderGuid() const
+{
+	return ActorDesc->GetFolderGuid();
+}
+
+FGuid FWorldPartitionActorDescView::GetContentBundleGuid() const
+{
+	return ActorDesc->GetContentBundleGuid();
+}
+
+bool FWorldPartitionActorDescView::IsContainerInstance() const
+{
+	return ActorDesc->IsChildContainerInstance();
+}
+
+EWorldPartitionActorFilterType FWorldPartitionActorDescView::GetContainerFilterType() const
+{
+	return ActorDesc->GetChildContainerFilterType();
+}
+
+FName FWorldPartitionActorDescView::GetContainerPackage() const
+{
+	return ActorDesc->GetChildContainerPackage();
+}
+
+bool FWorldPartitionActorDescView::GetContainerInstance(FWorldPartitionActorDesc::FContainerInstance& OutContainerInstance) const
+{
+	return ActorDesc->GetContainerInstance(OutContainerInstance);
+}
+
+const FWorldPartitionActorFilter* FWorldPartitionActorDescView::GetContainerFilter() const
+{
+	check(GetContainerFilterType() != EWorldPartitionActorFilterType::None);
+	return ActorDesc->GetContainerFilter();
+}
+
+void FWorldPartitionActorDescView::CheckForErrors(IStreamingGenerationErrorHandler* ErrorHandler) const
+{
+	ActorDesc->CheckForErrors(ErrorHandler);
+}
+
+FName FWorldPartitionActorDescView::GetActorLabelOrName() const
+{
+	return ActorDesc->GetActorLabelOrName();
+}
+
+bool FWorldPartitionActorDescView::ShouldValidateRuntimeGrid() const
+{
+	return ActorDesc->ShouldValidateRuntimeGrid();
+}
+
+void FWorldPartitionActorDescView::SetForcedNonSpatiallyLoaded()
+{
+	if (!bIsForcedNonSpatiallyLoaded)
+	{
+		bIsForcedNonSpatiallyLoaded = true;
+		UE_LOG(LogWorldPartition, Verbose, TEXT("Actor '%s' forced to be non-spatially loaded"), *GetActorLabelOrName().ToString());
+	}
+}
+
+void FWorldPartitionActorDescView::SetForcedNoRuntimeGrid()
+{
+	if (!bIsForcedNoRuntimeGrid)
+	{
+		bIsForcedNoRuntimeGrid = true;
+		UE_LOG(LogWorldPartition, Verbose, TEXT("Actor '%s' runtime grid invalidated"), *GetActorLabelOrName().ToString());
+	}
+}
+
+void FWorldPartitionActorDescView::SetForcedNoDataLayers()
+{
+	if (!bIsForcedNoDataLayers)
+	{
+		bIsForcedNoDataLayers = true;
+		UE_LOG(LogWorldPartition, Verbose, TEXT("Actor '%s' data layers invalidated"), *GetActorLabelOrName().ToString());
+	}
+}
+
+void FWorldPartitionActorDescView::SetRuntimeDataLayerInstanceNames(const TArray<FName>& InRuntimeDataLayerInstanceNames)
+{
+	RuntimeDataLayerInstanceNames = InRuntimeDataLayerInstanceNames;
+}
+
+void FWorldPartitionActorDescView::SetRuntimeReferences(const TArray<FGuid>& InRuntimeReferences)
+{
+	RuntimeReferences = InRuntimeReferences;
+}
+
+void FWorldPartitionActorDescView::SetEditorReferences(const TArray<FGuid>& InEditorReferences)
+{
+	EditorReferences = InEditorReferences;
+}
+
+void FWorldPartitionActorDescView::SetForcedNoHLODLayer()
+{
+	if (!bIsForceNoHLODLayer)
+	{
+		bIsForceNoHLODLayer = true;
+		UE_LOG(LogWorldPartition, Verbose, TEXT("Actor '%s' HLOD layer invalidated"), *GetActorLabelOrName().ToString());
+	}
+}
+
+void FWorldPartitionActorDescView::SetRuntimeHLODLayer(const FSoftObjectPath& InHLODLayer)
+{
+	RuntimedHLODLayer = InHLODLayer;
+}
+
+AActor* FWorldPartitionActorDescView::GetActor() const
+{
+	return ActorDesc->GetActor();
+}
+
+bool FWorldPartitionActorDescView::IsEditorOnlyReference(const FGuid& ReferenceGuid) const
+{
+	return ActorDesc->IsEditorOnlyReference(ReferenceGuid);
+}
+
+bool FWorldPartitionActorDescView::GetProperty(FName PropertyName, FName* PropertyValue) const
+{
+	return ActorDesc->GetProperty(PropertyName, PropertyValue);
+}
+
+bool FWorldPartitionActorDescView::HasProperty(FName PropertyName) const
+{
+	return ActorDesc->HasProperty(PropertyName);
+}
+
+void FWorldPartitionActorDescView::SetParentView(const FWorldPartitionActorDescView* InParentView)
+{
+	check(!ParentView);
+	check(GetParentActor().IsValid());
+	ParentView = InParentView;
+}
+
+PRAGMA_ENABLE_DEPRECATION_WARNINGS
+#endif
\ No newline at end of file
diff --git a/Engine/Source/Runtime/Engine/Private/WorldPartition/WorldPartitionActorDescViewProxy.cpp b/Engine/Source/Runtime/Engine/Private/WorldPartition/WorldPartitionActorDescViewProxy.cpp
--- a/Engine/Source/Runtime/Engine/Private/WorldPartition/WorldPartitionActorDescViewProxy.cpp
+++ b/Engine/Source/Runtime/Engine/Private/WorldPartition/WorldPartitionActorDescViewProxy.cpp
@@ -3,12 +3,12 @@
 #include "WorldPartition/WorldPartitionActorDescViewProxy.h"
 #include "GameFramework/Actor.h"
 #include "UObject/Package.h"
-#include "UObject/Package.h"
-#include "WorldPartition/WorldPartitionActorDescView.h"
 
 #if WITH_EDITOR
 #include "WorldPartition/WorldPartitionActorDesc.h"
+#include "WorldPartition/WorldPartitionActorDescView.h"
 
+PRAGMA_DISABLE_DEPRECATION_WARNINGS
 FWorldPartitionActorViewProxy::FWorldPartitionActorViewProxy(const FWorldPartitionActorDesc* InActorDesc)
 	: FWorldPartitionActorDescView(InActorDesc)
 {
@@ -21,4 +21,5 @@ FWorldPartitionActorViewProxy::FWorldPartitionActorViewProxy(const FWorldPartiti
 		}
 	}
 }
+PRAGMA_ENABLE_DEPRECATION_WARNINGS
 #endif
diff --git a/Engine/Source/Runtime/Engine/Private/WorldPartition/WorldPartitionActorLoaderInterface.cpp b/Engine/Source/Runtime/Engine/Private/WorldPartition/WorldPartitionActorLoaderInterface.cpp
--- a/Engine/Source/Runtime/Engine/Private/WorldPartition/WorldPartitionActorLoaderInterface.cpp
+++ b/Engine/Source/Runtime/Engine/Private/WorldPartition/WorldPartitionActorLoaderInterface.cpp
@@ -10,6 +10,8 @@
 #include "Containers/Ticker.h"
 #include "Misc/ScopedSlowTask.h"
 #include "WorldPartition/WorldPartition.h"
+#include "WorldPartition/ActorDescContainerInstance.h"
+#include "WorldPartition/WorldPartitionActorDescInstance.h"
 #endif
 
 #define LOCTEXT_NAMESPACE "WorldPartition"
@@ -102,8 +104,8 @@ void IWorldPartitionActorLoaderInterface::ILoaderAdapter::RegisterDelegates()
 
 	if (UWorldPartition* WorldPartition = World->GetWorldPartition())
 	{
-		WorldPartition->OnActorDescContainerRegistered.AddRaw(this, &IWorldPartitionActorLoaderInterface::ILoaderAdapter::ILoaderAdapter::OnActorDescContainerInitialize);
-		WorldPartition->OnActorDescContainerUnregistered.AddRaw(this, &IWorldPartitionActorLoaderInterface::ILoaderAdapter::ILoaderAdapter::OnActorDescContainerUninitialize);
+		WorldPartition->OnActorDescContainerInstanceRegistered.AddRaw(this, &IWorldPartitionActorLoaderInterface::ILoaderAdapter::ILoaderAdapter::OnActorDescContainerInstanceInitialize);
+		WorldPartition->OnActorDescContainerInstanceUnregistered.AddRaw(this, &IWorldPartitionActorLoaderInterface::ILoaderAdapter::ILoaderAdapter::OnActorDescContainerInstanceUninitialize);
 	}
 }
 
@@ -113,14 +115,14 @@ void IWorldPartitionActorLoaderInterface::ILoaderAdapter::UnregisterDelegates()
 
 	if (UWorldPartition* WorldPartition = World->GetWorldPartition())
 	{
-		WorldPartition->OnActorDescContainerRegistered.RemoveAll(this);
-		WorldPartition->OnActorDescContainerUnregistered.RemoveAll(this);
+		WorldPartition->OnActorDescContainerInstanceRegistered.RemoveAll(this);
+		WorldPartition->OnActorDescContainerInstanceUnregistered.RemoveAll(this);
 	}
 }
 
 bool IWorldPartitionActorLoaderInterface::ILoaderAdapter::PassActorDescFilter(const FWorldPartitionHandle& Actor) const
 {
-	return Actor->IsEditorRelevant();
+	return Actor.GetInstance()->IsEditorRelevant();
 }
 
 void IWorldPartitionActorLoaderInterface::ILoaderAdapter::RefreshLoadedState()
@@ -131,20 +133,20 @@ void IWorldPartitionActorLoaderInterface::ILoaderAdapter::RefreshLoadedState()
 		{
 			TArray<FWorldPartitionHandle> ActorsToLoad;
 			TArray<FWorldPartitionHandle> ActorsToUnload;			
-			ForEachActor([this, &ActorsToLoad, &ActorsToUnload](const FWorldPartitionHandle& Actor)
+			ForEachActor([this, &ActorsToLoad, &ActorsToUnload](const FWorldPartitionHandle& ActorHandle)
 			{
-				const FActorReferenceMap* ActorReferences = GetContainerReferencesConst(Actor->GetContainer());
+				const FActorReferenceMap* ActorReferences = GetContainerReferencesConst(ActorHandle.GetContainerInstance());
 
-				if (ShouldActorBeLoaded(Actor))
+				if (ShouldActorBeLoaded(ActorHandle))
 				{
-					if (!ActorReferences || !ActorReferences->Contains(Actor->GetGuid()))
+					if (!ActorReferences || !ActorReferences->Contains(ActorHandle.GetInstance()->GetGuid()))
 					{
-						ActorsToLoad.Add(Actor);
+						ActorsToLoad.Add(ActorHandle);
 					}
 				}
-				else if (ActorReferences && ActorReferences->Contains(Actor->GetGuid()))
+				else if (ActorReferences && ActorReferences->Contains(ActorHandle.GetInstance()->GetGuid()))
 				{
-					ActorsToUnload.Add(Actor);
+					ActorsToUnload.Add(ActorHandle);
 				}
 			});
 
@@ -203,23 +205,23 @@ void IWorldPartitionActorLoaderInterface::ILoaderAdapter::RefreshLoadedState()
 	}
 }
 
-void IWorldPartitionActorLoaderInterface::ILoaderAdapter::OnActorDescContainerInitialize(UActorDescContainer* Container)
+void IWorldPartitionActorLoaderInterface::ILoaderAdapter::OnActorDescContainerInstanceInitialize(UActorDescContainerInstance* ContainerInstance)
 {
 	RefreshLoadedState();
 }
 
-void IWorldPartitionActorLoaderInterface::ILoaderAdapter::OnActorDescContainerUninitialize(UActorDescContainer* Container)
+void IWorldPartitionActorLoaderInterface::ILoaderAdapter::OnActorDescContainerInstanceUninitialize(UActorDescContainerInstance* ContainerInstance)
 {
 	// @todo_ow :
 	// If the container is removed before broadcasting its uninitialize event we could simply call RefreshLoadedState(), but it would break call symetry
 	TArray<FWorldPartitionHandle> ActorsToUnload;
-	for(UActorDescContainer::TConstIterator<> It(Container); It; ++It)
+	for(UActorDescContainerInstance::TConstIterator<> It(ContainerInstance); It; ++It)
 	{
-		const FActorReferenceMap* ActorReferences = GetContainerReferencesConst(It->GetContainer());
+		const FActorReferenceMap* ActorReferences = GetContainerReferencesConst(It->GetContainerInstance());
 
 		if (ActorReferences && ActorReferences->Find(It->GetGuid()))
 		{
-			ActorsToUnload.Emplace(Container, It->GetGuid());
+			ActorsToUnload.Emplace(ContainerInstance, It->GetGuid());
 		}
 	}
 	
@@ -244,7 +246,7 @@ bool IWorldPartitionActorLoaderInterface::ILoaderAdapter::ShouldActorBeLoaded(co
 {
 	check(Actor.IsValid());
 
-	Actor->SetUnloadedReason(nullptr);
+	Actor.GetInstance()->SetUnloadedReason(nullptr);
 
 	if (!PassActorDescFilter(Actor))
 	{
@@ -255,7 +257,7 @@ bool IWorldPartitionActorLoaderInterface::ILoaderAdapter::ShouldActorBeLoaded(co
 	{
 		if (!ActorDescFilter.Get().PassFilter(World, Actor))
 		{
-			Actor->SetUnloadedReason(ActorDescFilter.Get().GetFilterReason());
+			Actor.GetInstance()->SetUnloadedReason(ActorDescFilter.Get().GetFilterReason());
 			return false;
 		}
 	}
@@ -287,13 +289,14 @@ void IWorldPartitionActorLoaderInterface::ILoaderAdapter::AddReferenceToActor(FW
 {
 	TFunction<void(const FWorldPartitionHandle&, FReferenceMap&)> AddReferences = [this, &AddReferences](const FWorldPartitionHandle& Handle, FReferenceMap& ReferenceMap)
 	{
-		if (!ReferenceMap.Contains(Handle->GetGuid()))
+		const FWorldPartitionActorDescInstance* ActorDescInstance = Handle.GetInstance();
+		if (!ReferenceMap.Contains(ActorDescInstance->GetGuid()))
 		{
-			ReferenceMap.Emplace(Handle->GetGuid(), Handle.ToReference());
+			ReferenceMap.Emplace(ActorDescInstance->GetGuid(), Handle.ToReference());
 			
-			for (const FGuid& ReferencedActorGuid : Handle->GetReferences())
+			for (const FGuid& ReferencedActorGuid : ActorDescInstance->GetReferences())
 			{
-				FWorldPartitionHandle ReferenceActorHandle(Handle->GetContainer(), ReferencedActorGuid);
+				FWorldPartitionHandle ReferenceActorHandle(Handle.GetContainerInstance(), ReferencedActorGuid);
 
 				if (ReferenceActorHandle.IsValid())
 				{
@@ -303,14 +306,14 @@ void IWorldPartitionActorLoaderInterface::ILoaderAdapter::AddReferenceToActor(FW
 		}
 	};
 
-	FActorReferenceMap& ActorReferences = GetContainerReferences(ActorHandle->GetContainer());
-	AddReferences(ActorHandle, ActorReferences.Emplace(ActorHandle->GetGuid()));
+	FActorReferenceMap& ActorReferences = GetContainerReferences(ActorHandle.GetContainerInstance());
+	AddReferences(ActorHandle, ActorReferences.Emplace(ActorHandle.GetInstance()->GetGuid()));
 }
 
 void IWorldPartitionActorLoaderInterface::ILoaderAdapter::RemoveReferenceToActor(FWorldPartitionHandle& ActorHandle)
 {
-	FActorReferenceMap& ActorReferences = GetContainerReferences(ActorHandle->GetContainer());
-	ActorReferences.Remove(ActorHandle->GetGuid());
+	FActorReferenceMap& ActorReferences = GetContainerReferences(ActorHandle.GetContainerInstance());
+	ActorReferences.Remove(ActorHandle.GetInstance()->GetGuid());
 }
 
 void IWorldPartitionActorLoaderInterface::ILoaderAdapter::OnRefreshLoadedState(bool bFromUserOperation)
@@ -338,14 +341,24 @@ void IWorldPartitionActorLoaderInterface::RefreshLoadedState(bool bIsFromUserCha
 	}
 }
 
-IWorldPartitionActorLoaderInterface::FActorReferenceMap& IWorldPartitionActorLoaderInterface::ILoaderAdapter::GetContainerReferences(UActorDescContainer* InContainer)
+IWorldPartitionActorLoaderInterface::FActorReferenceMap& IWorldPartitionActorLoaderInterface::ILoaderAdapter::GetContainerReferences(UActorDescContainerInstance* InContainerInstance)
 {
-	return ContainerActorReferences.FindOrAdd(InContainer);
+	return ContainerActorReferences.FindOrAdd(InContainerInstance);
 }
 
-const IWorldPartitionActorLoaderInterface::FActorReferenceMap* IWorldPartitionActorLoaderInterface::ILoaderAdapter::GetContainerReferencesConst(UActorDescContainer* InContainer) const
+const IWorldPartitionActorLoaderInterface::FActorReferenceMap* IWorldPartitionActorLoaderInterface::ILoaderAdapter::GetContainerReferencesConst(UActorDescContainerInstance* InContainerInstance) const
 {
-	return ContainerActorReferences.Find(InContainer);
+	return ContainerActorReferences.Find(InContainerInstance);
+}
+
+bool IWorldPartitionActorLoaderInterface::GetLoadedChildContainerInstance(const FWorldPartitionHandle& Handle, FWorldPartitionActorDesc::FLoadedContainerInstance& OutContainerInstance)
+{
+	if (FWorldPartitionActorDescInstance* Instance = Handle.GetInstance())
+	{
+		return Instance->GetLoadedChildContainerInstance(OutContainerInstance);
+	}
+
+	return false;
 }
 #endif
 
diff --git a/Engine/Source/Runtime/Engine/Private/WorldPartition/WorldPartitionEditorSpatialHash.cpp b/Engine/Source/Runtime/Engine/Private/WorldPartition/WorldPartitionEditorSpatialHash.cpp
--- a/Engine/Source/Runtime/Engine/Private/WorldPartition/WorldPartitionEditorSpatialHash.cpp
+++ b/Engine/Source/Runtime/Engine/Private/WorldPartition/WorldPartitionEditorSpatialHash.cpp
@@ -3,7 +3,7 @@
 #include "WorldPartition/WorldPartitionEditorSpatialHash.h"
 #include "HAL/IConsoleManager.h"
 #include "WorldPartition/WorldPartitionActorDesc.h"
-#include "WorldPartition/WorldPartitionActorDesc.h"
+#include "WorldPartition/WorldPartitionActorDescInstance.h"
 
 #include UE_INLINE_GENERATED_CPP_BY_NAME(WorldPartitionEditorSpatialHash)
 
@@ -67,16 +67,16 @@ void UWorldPartitionEditorSpatialHash::Tick(float DeltaSeconds)
 		{
 			NewEditorBounds += Cell->Bounds;
 
-			for (const FWorldPartitionHandle& Actor : Cell->Actors)
+			for (const FWorldPartitionHandle& ActorHandle : Cell->Actors)
 			{
-				RuntimeBounds += Actor->GetEditorBounds();
+				RuntimeBounds += ActorHandle.GetInstance()->GetEditorBounds();
 			}
 		}
 
 		NonSpatialBounds.Init();
-		for (const FWorldPartitionHandle& Actor : AlwaysLoadedCell->Actors)
+		for (const FWorldPartitionHandle& ActorHandle : AlwaysLoadedCell->Actors)
 		{
-			NonSpatialBounds += Actor->GetEditorBounds();
+			NonSpatialBounds += ActorHandle.GetInstance()->GetEditorBounds();
 		}
 
 		const int32 OldLevel = GetLevelForBox(EditorBounds);
@@ -126,18 +126,19 @@ void UWorldPartitionEditorSpatialHash::HashActor(FWorldPartitionHandle& InActorH
 {
 	check(InActorHandle.IsValid());
 
-	const bool bConsiderActorSpatiallyLoaded = InActorHandle->GetIsSpatiallyLoaded();
-	const FBox ActorBounds = bConsiderActorSpatiallyLoaded ? InActorHandle->GetEditorBounds() : FBox(ForceInit);
+	const FWorldPartitionActorDescInstance* ActorDescInstance = InActorHandle.GetInstance();
+	const bool bConsiderActorSpatiallyLoaded = ActorDescInstance->GetIsSpatiallyLoaded();
+	const FBox ActorBounds = bConsiderActorSpatiallyLoaded ? ActorDescInstance->GetEditorBounds() : FBox(ForceInit);
 
 #if DO_CHECK
-	check(!HashedActors.Contains(InActorHandle->GetGuid()));
-	HashedActors.Add(InActorHandle->GetGuid(), ActorBounds);
+	check(!HashedActors.Contains(ActorDescInstance->GetGuid()));
+	HashedActors.Add(ActorDescInstance->GetGuid(), ActorBounds);
 #endif
 
 	if (!bConsiderActorSpatiallyLoaded)
 	{
 		AlwaysLoadedCell->Actors.Add(InActorHandle);
-		NonSpatialBounds += InActorHandle->GetEditorBounds();
+		NonSpatialBounds += ActorDescInstance->GetEditorBounds();
 	}
 	else
 	{
@@ -230,7 +231,7 @@ void UWorldPartitionEditorSpatialHash::HashActor(FWorldPartitionHandle& InActorH
 			}
 		}
 
-		RuntimeBounds += InActorHandle->GetEditorBounds();
+		RuntimeBounds += ActorDescInstance->GetEditorBounds();
 	}
 }
 
@@ -238,11 +239,12 @@ void UWorldPartitionEditorSpatialHash::UnhashActor(FWorldPartitionHandle& InActo
 {
 	check(InActorHandle.IsValid());
 
-	const bool bConsiderActorSpatiallyLoaded = InActorHandle->GetIsSpatiallyLoaded();
-	const FBox ActorBounds = bConsiderActorSpatiallyLoaded ? InActorHandle->GetEditorBounds() : FBox(ForceInit);
+	const FWorldPartitionActorDescInstance* ActorDescInstance = InActorHandle.GetInstance();
+	const bool bConsiderActorSpatiallyLoaded = ActorDescInstance->GetIsSpatiallyLoaded();
+	const FBox ActorBounds = bConsiderActorSpatiallyLoaded ? ActorDescInstance->GetEditorBounds() : FBox(ForceInit);
 
 #if DO_CHECK
-	const FBox OldActorBounds = HashedActors.FindAndRemoveChecked(InActorHandle->GetGuid());
+	const FBox OldActorBounds = HashedActors.FindAndRemoveChecked(ActorDescInstance->GetGuid());
 	check(ActorBounds == OldActorBounds);
 #endif
 
@@ -297,9 +299,9 @@ void UWorldPartitionEditorSpatialHash::UnhashActor(FWorldPartitionHandle& InActo
 	bBoundsDirty = true;
 }
 
-int32 UWorldPartitionEditorSpatialHash::ForEachIntersectingActor(const FBox& Box, TFunctionRef<void(FWorldPartitionActorDesc*)> InOperation, const FForEachIntersectingActorParams& Params)
+int32 UWorldPartitionEditorSpatialHash::ForEachIntersectingActor(const FBox& Box, TFunctionRef<void(FWorldPartitionActorDescInstance*)> InOperation, const FForEachIntersectingActorParams& Params)
 {
-	TSet<FWorldPartitionActorDesc*> IntersectedActorDescs;
+	TSet<FWorldPartitionHandle> IntersectedActors;
 
 	const int MininumumLevel = Params.MinimumBox.IsSet() ? GetLevelForBox(*Params.MinimumBox) : 0;
 
@@ -307,18 +309,18 @@ int32 UWorldPartitionEditorSpatialHash::ForEachIntersectingActor(const FBox& Box
 	{
 		ForEachIntersectingCell(Box, [&](FCell* Cell)
 		{
-			for(FWorldPartitionHandle& ActorDesc: Cell->Actors)
+			for(FWorldPartitionHandle& ActorHandle: Cell->Actors)
 			{
-				if (ActorDesc.IsValid())
+				if (ActorHandle.IsValid())
 				{
 					bool bWasAlreadyInSet;
-					IntersectedActorDescs.Add(*ActorDesc, &bWasAlreadyInSet);
+					IntersectedActors.Add(ActorHandle, &bWasAlreadyInSet);
 
 					if (!bWasAlreadyInSet)
 					{
-						if (Box.Intersect(ActorDesc->GetEditorBounds()))
+						if (Box.Intersect(ActorHandle.GetInstance()->GetEditorBounds()))
 						{
-							InOperation(*ActorDesc);
+							InOperation(ActorHandle.GetInstance());
 						}
 					}
 				}
@@ -328,25 +330,25 @@ int32 UWorldPartitionEditorSpatialHash::ForEachIntersectingActor(const FBox& Box
 
 	if (Params.bIncludeNonSpatiallyLoadedActors)
 	{
-		for(FWorldPartitionHandle& ActorDesc: AlwaysLoadedCell->Actors)
+		for(FWorldPartitionHandle& ActorHandle : AlwaysLoadedCell->Actors)
 		{
-			if (ActorDesc.IsValid())
+			if (ActorHandle.IsValid())
 			{
-				if (Box.Intersect(ActorDesc->GetEditorBounds()))
+				if (Box.Intersect(ActorHandle.GetInstance()->GetEditorBounds()))
 				{
 					bool bWasAlreadyInSet;
-					IntersectedActorDescs.Add(*ActorDesc, &bWasAlreadyInSet);
+					IntersectedActors.Add(ActorHandle, &bWasAlreadyInSet);
 				
 					if (!bWasAlreadyInSet)
 					{
-						InOperation(*ActorDesc);
+						InOperation(ActorHandle.GetInstance());
 					}
 				}
 			}
 		}
 	}
 
-	return IntersectedActorDescs.Num();
+	return IntersectedActors.Num();
 }
 
 int32 UWorldPartitionEditorSpatialHash::ForEachIntersectingCellInner(const FBox& Box, const FCellCoord& CellCoord, TFunctionRef<void(FCell*)> InOperation, int32 MinimumLevel)
diff --git a/Engine/Source/Runtime/Engine/Private/WorldPartition/WorldPartitionHandle.cpp b/Engine/Source/Runtime/Engine/Private/WorldPartition/WorldPartitionHandle.cpp
--- a/Engine/Source/Runtime/Engine/Private/WorldPartition/WorldPartitionHandle.cpp
+++ b/Engine/Source/Runtime/Engine/Private/WorldPartition/WorldPartitionHandle.cpp
@@ -1,32 +1,52 @@
 // Copyright Epic Games, Inc. All Rights Reserved.
 
 #include "WorldPartition/WorldPartitionHandle.h"
-#include "WorldPartition/ActorDescContainerCollection.h"
+
 #include "Engine/Level.h"
 #include "Engine/World.h"
 
 #if WITH_EDITOR
+
+#include "WorldPartition/ActorDescContainerInstanceCollection.h"
+#include "WorldPartition/ActorDescContainerInstance.h"
+#include "WorldPartition/WorldPartitionActorDescInstance.h"
+
 /**
 * FWorldPartitionImplBase
 */
-TUniquePtr<FWorldPartitionActorDesc>* FWorldPartitionImplBase::GetActorDesc(UActorDescContainer* Container, const FGuid& ActorGuid)
+UActorDescContainerInstance* FWorldPartitionImplBase::GetActorDescContainerInstance(TUniquePtr<FWorldPartitionActorDescInstance>* InActorDescIntance)
 {
-	if (TUniquePtr<FWorldPartitionActorDesc>** ActorDescPtr = Container->ActorsByGuid.Find(ActorGuid))
-	{
-		return *ActorDescPtr;
-	}
+	return InActorDescIntance ? InActorDescIntance->Get()->GetContainerInstance() : nullptr;
+}
 
-	return nullptr;
+UActorDescContainerInstance* FWorldPartitionImplBase::GetActorDescContainerInstance(FWorldPartitionActorDescInstance* InActorDescIntance)
+{
+	return InActorDescIntance ? InActorDescIntance->GetContainerInstance() : nullptr;
 }
 
-UActorDescContainer* FWorldPartitionImplBase::GetActorDescContainer(TUniquePtr<FWorldPartitionActorDesc>* ActorDesc)
+FGuid FWorldPartitionImplBase::GetActorDescInstanceGuid(const FWorldPartitionActorDescInstance* InActorDescIntance)
 {
-	return ActorDesc ? ActorDesc->Get()->GetContainer() : nullptr;
+	return InActorDescIntance ? InActorDescIntance->GetGuid() : FGuid();
 }
 
-bool FWorldPartitionImplBase::IsActorDescLoaded(FWorldPartitionActorDesc* ActorDesc)
+TUniquePtr<FWorldPartitionActorDescInstance>* FWorldPartitionImplBase::GetActorDescInstance(UActorDescContainerInstance* InContainerInstance, const FGuid& InActorGuid)
 {
-	return ActorDesc->IsLoaded();
+	return InContainerInstance ? InContainerInstance->GetActorDescInstancePtr(InActorGuid) : nullptr;
+}
+
+FWorldPartitionActorDesc* FWorldPartitionImplBase::GetActorDesc(FWorldPartitionActorDescInstance* InActorDescInstance)
+{
+	return InActorDescInstance ? const_cast<FWorldPartitionActorDesc*>(InActorDescInstance->GetActorDesc()) : nullptr;
+}
+
+bool FWorldPartitionImplBase::IsLoaded(FWorldPartitionActorDescInstance* InActorDescInstance)
+{
+	return InActorDescInstance->IsLoaded();
+}
+
+AActor* FWorldPartitionImplBase::GetActor(FWorldPartitionActorDescInstance* InActorDescInstance)
+{
+	return InActorDescInstance->GetActor();
 }
 
 /**
@@ -35,22 +55,22 @@ bool FWorldPartitionImplBase::IsActorDescLoaded(FWorldPartitionActorDesc* ActorD
 FWorldPartitionLoadingContext::FImmediate FWorldPartitionLoadingContext::DefaultContext;
 FWorldPartitionLoadingContext::IContext* FWorldPartitionLoadingContext::ActiveContext = &DefaultContext;
 
-void FWorldPartitionLoadingContext::LoadAndRegisterActor(FWorldPartitionActorDesc* ActorDesc)
+void FWorldPartitionLoadingContext::LoadAndRegisterActor(FWorldPartitionActorDescInstance* InActorDescInstance)
 {
 #if DO_CHECK
-	FWorldPartitionActorDesc::FRegisteringUnregisteringGuard Guard(ActorDesc);
+	FWorldPartitionActorDescInstance::FRegisteringUnregisteringGuard Guard(InActorDescInstance);
 #endif
 
-	ActiveContext->RegisterActor(ActorDesc);
+	ActiveContext->RegisterActor(InActorDescInstance);
 }
 
-void FWorldPartitionLoadingContext::UnloadAndUnregisterActor(FWorldPartitionActorDesc* ActorDesc)
+void FWorldPartitionLoadingContext::UnloadAndUnregisterActor(FWorldPartitionActorDescInstance* InActorDescInstance)
 {
 #if DO_CHECK
-	FWorldPartitionActorDesc::FRegisteringUnregisteringGuard Guard(ActorDesc);
+	FWorldPartitionActorDescInstance::FRegisteringUnregisteringGuard Guard(InActorDescInstance);
 #endif
 
-	ActiveContext->UnregisterActor(ActorDesc);
+	ActiveContext->UnregisterActor(InActorDescInstance);
 }
 
 /**
@@ -71,40 +91,40 @@ FWorldPartitionLoadingContext::IContext::~IContext()
 /**
 * FImmediate
 */
-void FWorldPartitionLoadingContext::FImmediate::RegisterActor(FWorldPartitionActorDesc* ActorDesc)
+void FWorldPartitionLoadingContext::FImmediate::RegisterActor(FWorldPartitionActorDescInstance* InActorDescInstance)
 {
 	TGuardValue<bool> IsEditorLoadingPackageGuard(GIsEditorLoadingPackage, true);
 
-	if (AActor* Actor = ActorDesc->Load())
+	if (AActor* Actor = InActorDescInstance->Load())
 	{
-		UActorDescContainer* Container = ActorDesc->GetContainer();
+		UActorDescContainerInstance* Container = InActorDescInstance->GetContainerInstance();
 		check(Container);
 
-		const FTransform& ContainerTransform = Container->GetInstanceTransform();
+		const FTransform& ContainerTransform = Container->GetTransform();
 		const FTransform* ContainerTransformPtr = ContainerTransform.Equals(FTransform::Identity) ? nullptr : &ContainerTransform;
 
 		Actor->GetLevel()->AddLoadedActor(Actor, ContainerTransformPtr);
 	}
 }
 
-void FWorldPartitionLoadingContext::FImmediate::UnregisterActor(FWorldPartitionActorDesc* ActorDesc)
+void FWorldPartitionLoadingContext::FImmediate::UnregisterActor(FWorldPartitionActorDescInstance* InActorDescInstance)
 {
 	// Set GIsEditorLoadingPackage to avoid dirtying the Actor package if Modify() is called during the unload sequence
 	TGuardValue<bool> IsEditorLoadingPackageGuard(GIsEditorLoadingPackage, true);
 
 	// When cleaning up worlds, actors are already marked as garbage at this point, so no need to remove them from the world
-	if (AActor* Actor = ActorDesc->GetActor(); IsValid(Actor))
+	if (AActor* Actor = InActorDescInstance->GetActor(); IsValid(Actor))
 	{
-		UActorDescContainer* Container = ActorDesc->GetContainer();
+		UActorDescContainerInstance* Container = InActorDescInstance->GetContainerInstance();
 		check(Container);
 
-		const FTransform& ContainerTransform = Container->GetInstanceTransform();
+		const FTransform& ContainerTransform = Container->GetTransform();
 		const FTransform* ContainerTransformPtr = ContainerTransform.Equals(FTransform::Identity) ? nullptr : &ContainerTransform;
 
 		Actor->GetLevel()->RemoveLoadedActor(Actor, ContainerTransformPtr);
 
-		ActorDesc->Unload();
-	}	
+		InActorDescInstance->Unload();
+	}
 }
 
 /**
@@ -112,24 +132,24 @@ void FWorldPartitionLoadingContext::FImmediate::UnregisterActor(FWorldPartitionA
 */
 FWorldPartitionLoadingContext::FDeferred::~FDeferred()
 {
-	TMap<UActorDescContainer*, FContainerOps> LocalContainerOps = MoveTemp(ContainerOps);
+	TMap<UActorDescContainerInstance*, FContainerInstanceOps> LocalContainerInstanceOps = MoveTemp(ContainerInstanceOps);
 
-	while(LocalContainerOps.Num())
+	while (LocalContainerInstanceOps.Num())
 	{
-		for (auto& [Container, ContainerOp] : LocalContainerOps)
+		for (auto& [Container, ContainerOp] : LocalContainerInstanceOps)
 		{
-			const FTransform& ContainerTransform = Container->GetInstanceTransform();
+			const FTransform& ContainerTransform = Container->GetTransform();
 			const FTransform* ContainerTransformPtr = ContainerTransform.Equals(FTransform::Identity) ? nullptr : &ContainerTransform;
 
-			auto CreateActorList = [](TArray<AActor*>& ActorList, const TSet<FWorldPartitionActorDesc*>& SourceList) -> ULevel*
+			auto CreateActorList = [](TArray<AActor*>& ActorList, const TSet<FWorldPartitionActorDescInstance*>& SourceList) -> ULevel*
 			{
 				ActorList.Empty(SourceList.Num());
 
 				ULevel* Level = nullptr;
-				for (FWorldPartitionActorDesc* ActorDesc : SourceList)
+				for (FWorldPartitionActorDescInstance* ActorDescInstance : SourceList)
 				{
 					// When cleaning up worlds, actors are already marked as garbage at this point, so no need to remove them from the world
-					if (AActor* Actor = ActorDesc->GetActor(); IsValid(Actor))
+					if (AActor* Actor = ActorDescInstance->GetActor(); IsValid(Actor))
 					{
 						ActorList.Add(Actor);
 
@@ -164,51 +184,52 @@ FWorldPartitionLoadingContext::FDeferred::~FDeferred()
 					Level->RemoveLoadedActors(ActorList, ContainerTransformPtr);
 				}
 
-				for (FWorldPartitionActorDesc* ActorDesc : ContainerOp.Unregistrations)
+				for (FWorldPartitionActorDescInstance* ActorDescInstance : ContainerOp.Unregistrations)
 				{
-					ActorDesc->Unload();
+					ActorDescInstance->Unload();
 				}
 			}
 		}
 
 		// Continue with potentially new registrations/unregistrations that may have happenned during previous cycle
-		LocalContainerOps = MoveTemp(ContainerOps);
+		LocalContainerInstanceOps = MoveTemp(ContainerInstanceOps);
 	}
+	
 }
 
-void FWorldPartitionLoadingContext::FDeferred::RegisterActor(FWorldPartitionActorDesc* ActorDesc)
+void FWorldPartitionLoadingContext::FDeferred::RegisterActor(FWorldPartitionActorDescInstance* InActorDescInstance)
 {
 	TGuardValue<bool> IsEditorLoadingPackageGuard(GIsEditorLoadingPackage, true);
 
-	check(ActorDesc);
-	if (ActorDesc->Load())
+	check(InActorDescInstance);
+	if (InActorDescInstance->Load())
 	{
-		UActorDescContainer* Container = ActorDesc->GetContainer();
+		UActorDescContainerInstance* Container = InActorDescInstance->GetContainerInstance();
 		check(Container);
 
 		bool bIsAlreadyInSetPtr;
-		ContainerOps.FindOrAdd(Container).Registrations.Add(ActorDesc, &bIsAlreadyInSetPtr);
+		ContainerInstanceOps.FindOrAdd(Container).Registrations.Add(InActorDescInstance, &bIsAlreadyInSetPtr);
 		check(!bIsAlreadyInSetPtr);
 		NumRegistrations++;
 
-		check(!ContainerOps.FindChecked(Container).Unregistrations.Contains(ActorDesc));
+		check(!ContainerInstanceOps.FindChecked(Container).Unregistrations.Contains(InActorDescInstance));
 	}
 }
 
-void FWorldPartitionLoadingContext::FDeferred::UnregisterActor(FWorldPartitionActorDesc* ActorDesc)
+void FWorldPartitionLoadingContext::FDeferred::UnregisterActor(FWorldPartitionActorDescInstance* InActorDescInstance)
 {
-	check(ActorDesc);
-	if (AActor* Actor = ActorDesc->GetActor(); IsValid(Actor))
+	check(InActorDescInstance);
+	if (AActor* Actor = InActorDescInstance->GetActor(); IsValid(Actor))
 	{
-		UActorDescContainer* Container = ActorDesc->GetContainer();
+		UActorDescContainerInstance* Container = InActorDescInstance->GetContainerInstance();
 		check(Container);
 
 		bool bIsAlreadyInSetPtr;
-		ContainerOps.FindOrAdd(Container).Unregistrations.Add(ActorDesc, &bIsAlreadyInSetPtr);
+		ContainerInstanceOps.FindOrAdd(Container).Unregistrations.Add(InActorDescInstance, &bIsAlreadyInSetPtr);
 		check(!bIsAlreadyInSetPtr);
 		NumUnregistrations++;
 
-		check(!ContainerOps.FindChecked(Container).Registrations.Contains(ActorDesc));
+		check(!ContainerInstanceOps.FindChecked(Container).Registrations.Contains(InActorDescInstance));
 
 		bNeedsClearTransactions |= Actor->GetTypedOuter<UWorld>()->HasAnyFlags(RF_Transactional);
 	}
@@ -217,14 +238,14 @@ void FWorldPartitionLoadingContext::FDeferred::UnregisterActor(FWorldPartitionAc
 /**
 * FWorldPartitionHandleImpl
 */
-void FWorldPartitionHandleImpl::IncRefCount(FWorldPartitionActorDesc* ActorDesc)
+void FWorldPartitionHandleImpl::IncRefCount(FWorldPartitionActorDescInstance* InActorDescInstance)
 {
-	ActorDesc->IncSoftRefCount();
+	InActorDescInstance->IncSoftRefCount();
 }
 
-void FWorldPartitionHandleImpl::DecRefCount(FWorldPartitionActorDesc* ActorDesc)
+void FWorldPartitionHandleImpl::DecRefCount(FWorldPartitionActorDescInstance* InActorDescInstance)
 {
-	ActorDesc->DecSoftRefCount();
+	InActorDescInstance->DecSoftRefCount();
 }
 
 TWorldPartitionHandle<FWorldPartitionReferenceImpl> FWorldPartitionHandleImpl::ToReference(const TWorldPartitionHandle<FWorldPartitionHandleImpl>& Source)
@@ -237,19 +258,19 @@ TWorldPartitionHandle<FWorldPartitionReferenceImpl> FWorldPartitionHandleImpl::T
 /**
 * FWorldPartitionReferenceImpl
 */
-void FWorldPartitionReferenceImpl::IncRefCount(FWorldPartitionActorDesc* ActorDesc)
+void FWorldPartitionReferenceImpl::IncRefCount(FWorldPartitionActorDescInstance* InActorDescInstance)
 {
-	if (ActorDesc->IncHardRefCount() == 1)
+	if (InActorDescInstance->IncHardRefCount() == 1)
 	{
-		FWorldPartitionLoadingContext::LoadAndRegisterActor(ActorDesc);
+		FWorldPartitionLoadingContext::LoadAndRegisterActor(InActorDescInstance);
 	}
 }
 
-void FWorldPartitionReferenceImpl::DecRefCount(FWorldPartitionActorDesc* ActorDesc)
+void FWorldPartitionReferenceImpl::DecRefCount(FWorldPartitionActorDescInstance* InActorDescInstance)
 {
-	if (ActorDesc->DecHardRefCount() == 0)
+	if (InActorDescInstance->DecHardRefCount() == 0)
 	{
-		FWorldPartitionLoadingContext::UnloadAndUnregisterActor(ActorDesc);
+		FWorldPartitionLoadingContext::UnloadAndUnregisterActor(InActorDescInstance);
 	}
 }
 
diff --git a/Engine/Source/Runtime/Engine/Private/WorldPartition/WorldPartitionHelpers.cpp b/Engine/Source/Runtime/Engine/Private/WorldPartition/WorldPartitionHelpers.cpp
--- a/Engine/Source/Runtime/Engine/Private/WorldPartition/WorldPartitionHelpers.cpp
+++ b/Engine/Source/Runtime/Engine/Private/WorldPartition/WorldPartitionHelpers.cpp
@@ -6,6 +6,8 @@
 #include "WorldPartition/WorldPartitionLog.h"
 #include "WorldPartition/WorldPartitionEditorHash.h"
 #include "WorldPartition/WorldPartitionRuntimeCell.h"
+#include "WorldPartition/ActorDescContainerInstance.h"
+#include "WorldPartition/WorldPartitionActorDescInstance.h"
 #include "GameFramework/PlayerController.h"
 #include "Engine/Engine.h"
 #include "Algo/AnyOf.h"
@@ -92,26 +94,26 @@ bool FWorldPartitionHelpers::IsActorDescClassCompatibleWith(const FWorldPartitio
 	return ActorBaseClass->IsChildOf(Class);
 }
 
-void FWorldPartitionHelpers::ForEachIntersectingActorDesc(UWorldPartition* WorldPartition, const FBox& Box, TSubclassOf<AActor> ActorClass, TFunctionRef<bool(const FWorldPartitionActorDesc*)> Func)
+void FWorldPartitionHelpers::ForEachIntersectingActorDescInstance(UWorldPartition* WorldPartition, const FBox& Box, TSubclassOf<AActor> ActorClass, TFunctionRef<bool(const FWorldPartitionActorDescInstance*)> Func)
 {
 	bool bProcessNextActors = true;
 
-	WorldPartition->EditorHash->ForEachIntersectingActor(Box, [&ActorClass, Func, &bProcessNextActors](const FWorldPartitionActorDesc* ActorDesc)
+	WorldPartition->EditorHash->ForEachIntersectingActor(Box, [&ActorClass, Func, &bProcessNextActors](const FWorldPartitionActorDescInstance* ActorDescInstance)
 	{
-		if (bProcessNextActors && IsActorDescClassCompatibleWith(ActorDesc, ActorClass))
+		if (bProcessNextActors && IsActorDescClassCompatibleWith(ActorDescInstance->GetActorDesc(), ActorClass))
 		{
-			bProcessNextActors = Func(ActorDesc);
+			bProcessNextActors = Func(ActorDescInstance);
 		}
 	});
 }
 
-void FWorldPartitionHelpers::ForEachActorDesc(UWorldPartition* WorldPartition, TSubclassOf<AActor> ActorClass, TFunctionRef<bool(const FWorldPartitionActorDesc*)> Func)
+void FWorldPartitionHelpers::ForEachActorDescInstance(UWorldPartition* WorldPartition, TSubclassOf<AActor> ActorClass, TFunctionRef<bool(const FWorldPartitionActorDescInstance*)> Func)
 {
-	for (FActorDescContainerCollection::TConstIterator<> ActorDescIterator(WorldPartition); ActorDescIterator; ++ActorDescIterator)
+	for (FActorDescContainerInstanceCollection::TConstIterator<> Iterator(WorldPartition); Iterator; ++Iterator)
 	{
-		if (IsActorDescClassCompatibleWith(*ActorDescIterator, ActorClass))
+		if (IsActorDescClassCompatibleWith(Iterator->GetActorDesc(), ActorClass))
 		{
-			if (!Func(*ActorDescIterator))
+			if (!Func(*Iterator))
 			{
 				return;
 			}
@@ -128,11 +130,11 @@ namespace WorldPartitionHelpers
 			return;
 		}
 
-		if (const FWorldPartitionActorDesc* ActorDesc = WorldPartition->GetActorDesc(ActorGuid))
+		if (const FWorldPartitionActorDescInstance* ActorDescInstance = WorldPartition->GetActorDescInstance(ActorGuid))
 		{
 			InOutActorReferences.Emplace(ActorGuid);
 
-			for (FGuid ReferenceGuid : ActorDesc->GetReferences())
+			for (FGuid ReferenceGuid : ActorDescInstance->GetReferences())
 			{
 				LoadReferencesInternal(WorldPartition, ReferenceGuid, InOutActorReferences);
 			}
@@ -154,13 +156,13 @@ FWorldPartitionHelpers::FForEachActorWithLoadingParams::FForEachActorWithLoading
 	, ActorClasses({ AActor::StaticClass() })
 {}
 
-void FWorldPartitionHelpers::ForEachActorWithLoading(UWorldPartition* WorldPartition, TFunctionRef<bool(const FWorldPartitionActorDesc*)> Func, const FForEachActorWithLoadingParams& Params)
+void FWorldPartitionHelpers::ForEachActorWithLoading(UWorldPartition* WorldPartition, TFunctionRef<bool(const FWorldPartitionActorDescInstance*)> Func, const FForEachActorWithLoadingParams& Params)
 {
 	FForEachActorWithLoadingResult Result;
 	ForEachActorWithLoading(WorldPartition, Func, Params, Result);
 }
 
-void FWorldPartitionHelpers::ForEachActorWithLoading(UWorldPartition* WorldPartition, TFunctionRef<bool(const FWorldPartitionActorDesc*)> Func, const FForEachActorWithLoadingParams& Params, FForEachActorWithLoadingResult& Result)
+void FWorldPartitionHelpers::ForEachActorWithLoading(UWorldPartition* WorldPartition, TFunctionRef<bool(const FWorldPartitionActorDescInstance*)> Func, const FForEachActorWithLoadingParams& Params, FForEachActorWithLoadingResult& Result)
 {
 	check(Result.ActorReferences.IsEmpty());
 	auto CallGarbageCollect = [&Params, &Result]()
@@ -174,16 +176,16 @@ void FWorldPartitionHelpers::ForEachActorWithLoading(UWorldPartition* WorldParti
 		DoCollectGarbage();
 	};
 
-	auto ForEachActorWithLoadingImpl = [&](const FWorldPartitionActorDesc* ActorDesc)
+	auto ForEachActorWithLoadingImpl = [&](const FWorldPartitionActorDescInstance* ActorDescInstance)
 	{
-		if (Algo::AnyOf(Params.ActorClasses, [ActorDesc](UClass* ActorClass) { return IsActorDescClassCompatibleWith(ActorDesc, ActorClass); }))
+		if (Algo::AnyOf(Params.ActorClasses, [ActorDescInstance](UClass* ActorClass) { return IsActorDescClassCompatibleWith(ActorDescInstance->GetActorDesc(), ActorClass); }))
 		{
-			if (!Params.FilterActorDesc || Params.FilterActorDesc(ActorDesc))
+			if (!Params.FilterActorDesc || Params.FilterActorDesc(ActorDescInstance->GetActorDesc()))
 			{
-				WorldPartitionHelpers::LoadReferences(WorldPartition, ActorDesc->GetGuid(), Result.ActorReferences);
+				WorldPartitionHelpers::LoadReferences(WorldPartition, ActorDescInstance->GetGuid(), Result.ActorReferences);
 
-				FWorldPartitionReference ActorReference(WorldPartition, ActorDesc->GetGuid());
-				if (!Func(ActorReference.Get()))
+				FWorldPartitionReference ActorReference(WorldPartition, ActorDescInstance->GetGuid());
+				if (!Func(ActorReference.GetInstance()))
 				{
 					return false;
 				}
@@ -200,9 +202,9 @@ void FWorldPartitionHelpers::ForEachActorWithLoading(UWorldPartition* WorldParti
 
 	if (Params.ActorGuids.IsEmpty())
 	{
-		for (FActorDescContainerCollection::TConstIterator<> ActorDescIterator(WorldPartition); ActorDescIterator; ++ActorDescIterator)
+		for (FActorDescContainerInstanceCollection::TConstIterator<> ActorDescInstanceIterator(WorldPartition); ActorDescInstanceIterator; ++ActorDescInstanceIterator)
 		{
-			if (const FWorldPartitionActorDesc* ActorDesc = *ActorDescIterator)
+			if (const FWorldPartitionActorDescInstance* ActorDesc = *ActorDescInstanceIterator)
 			{
 				if (!ForEachActorWithLoadingImpl(ActorDesc))
 				{
@@ -216,7 +218,7 @@ void FWorldPartitionHelpers::ForEachActorWithLoading(UWorldPartition* WorldParti
 	{
 		for (const FGuid& ActorGuid : Params.ActorGuids)
 		{
-			if (const FWorldPartitionActorDesc* ActorDesc = WorldPartition->GetActorDesc(ActorGuid))
+			if (const FWorldPartitionActorDescInstance* ActorDesc = WorldPartition->GetActorDescInstance(ActorGuid))
 			{
 				if (!ForEachActorWithLoadingImpl(ActorDesc))
 				{
