WorldPartition: Add UActorDescContainerInstance/FWorldPartitionActorDescInstance which take transient/context data from UActorDescContainer/FWorldPartitionActorDesc leaving them as pure on disk repres

diff --git a/Engine/Source/Runtime/Engine/Private/WorldPartition/WorldPartitionRuntimeHash.cpp b/Engine/Source/Runtime/Engine/Private/WorldPartition/WorldPartitionRuntimeHash.cpp
--- a/Engine/Source/Runtime/Engine/Private/WorldPartition/WorldPartitionRuntimeHash.cpp
+++ b/Engine/Source/Runtime/Engine/Private/WorldPartition/WorldPartitionRuntimeHash.cpp
@@ -5,7 +5,6 @@
 #include "Engine/World.h"
 #include "Misc/HierarchicalLogArchive.h"
 #include "WorldPartition/WorldPartition.h"
-#include "WorldPartition/WorldPartitionActorDescView.h"
 #include "WorldPartition/WorldPartitionStreamingSource.h"
 #include "WorldPartition/WorldPartitionRuntimeLevelStreamingCell.h"
 #include "WorldPartition/WorldPartitionLevelStreamingDynamic.h"
@@ -166,8 +165,6 @@ void UWorldPartitionRuntimeHash::OnEndPlay()
 		FWorldPartitionLoadingContext::FDeferred LoadingContext;
 		AlwaysLoadedActorsForPIE.Empty();
 	}
-
-	ModifiedActorDescListForPIE.Empty();
 }
 
 bool UWorldPartitionRuntimeHash::GenerateStreaming(class UWorldPartitionStreamingPolicy* StreamingPolicy, const IStreamingGenerationContext* StreamingGenerationContext, TArray<FString>* OutPackagesToGenerate)
@@ -187,7 +184,7 @@ bool UWorldPartitionRuntimeHash::PopulateCellActorInstances(const TArray<const I
 	// Then, duplication of world for PIE will duplicate only these actors. 
 	// When stopping PIE, WorldPartition will release these FWorldPartitionReferences which 
 	// will unload actors that were not already loaded in the non PIE world.
-	TArray<TPair<FWorldPartitionReference, const FWorldPartitionActorDescView*>> AlwaysLoadedReferences;
+	TArray<TPair<FWorldPartitionReference, const IWorldPartitionActorDescInstanceView*>> AlwaysLoadedReferences;
 	{
 		FWorldPartitionLoadingContext::FDeferred LoadingContext;
 
@@ -198,7 +195,7 @@ bool UWorldPartitionRuntimeHash::PopulateCellActorInstances(const TArray<const I
 			ActorSetInstance->ForEachActor([this, ActorSetInstance, &AlwaysLoadedReferences, &OutCellActorInstances, bForceLoadAlwaysLoadedReferences](const FGuid& ActorGuid)
 			{
 				IStreamingGenerationContext::FActorInstance ActorInstance(ActorGuid, ActorSetInstance);
-				const FWorldPartitionActorDescView& ActorDescView = ActorInstance.GetActorDescView();
+				const IWorldPartitionActorDescInstanceView& ActorDescView = ActorInstance.GetActorDescView();
 
 				// Instanced world partition, ContainerID is the main container, but it's not the main world partition,
 				// so the always loaded actors don't be part of the process of ForceExternalActorLevelReference/AlwaysLoadedActorsForPIE.
@@ -244,10 +241,10 @@ void UWorldPartitionRuntimeHash::PopulateRuntimeCell(UWorldPartitionRuntimeCell*
 	{
 		if (ActorInstance.GetContainerID().IsMainContainer())
 		{
-			const FWorldPartitionActorDescView& ActorDescView = ActorInstance.GetActorDescView();
+			const FStreamingGenerationActorDescView& ActorDescView = ActorInstance.GetActorDescView();
 			if (AActor* Actor = FindObject<AActor>(nullptr, *ActorDescView.GetActorSoftPath().ToString()))
 			{
-				if (ModifiedActorDescListForPIE.GetActorDesc(ActorDescView.GetGuid()))
+				if (ActorDescView.IsUnsaved())
 				{
 					// Create an actor container to make sure duplicated actors will share an outer to properly remap inter-actors references
 					RuntimeCell->UnsavedActorsContainer = NewObject<UActorContainer>(RuntimeCell);
@@ -260,9 +257,8 @@ void UWorldPartitionRuntimeHash::PopulateRuntimeCell(UWorldPartitionRuntimeCell*
 	FBox CellContentBounds(ForceInit);
 	for (const IStreamingGenerationContext::FActorInstance& ActorInstance : ActorInstances)
 	{
-		const FWorldPartitionActorDescView& ActorDescView = ActorInstance.GetActorDescView();
-		const UActorDescContainer* ActorDescContainer = ActorDescView.GetActorDesc()->GetContainer();
-		RuntimeCell->AddActorToCell(ActorDescView, ActorInstance.GetContainerID(), ActorInstance.GetTransform(), ActorDescContainer);
+		const FStreamingGenerationActorDescView& ActorDescView = ActorInstance.GetActorDescView();
+		RuntimeCell->AddActorToCell(ActorDescView);
 		const FBox RuntimeBounds = ActorDescView.GetRuntimeBounds();
 		if (RuntimeBounds.IsValid)
 		{
diff --git a/Engine/Source/Runtime/Engine/Private/WorldPartition/WorldPartitionRuntimeLevelStreamingCell.cpp b/Engine/Source/Runtime/Engine/Private/WorldPartition/WorldPartitionRuntimeLevelStreamingCell.cpp
--- a/Engine/Source/Runtime/Engine/Private/WorldPartition/WorldPartitionRuntimeLevelStreamingCell.cpp
+++ b/Engine/Source/Runtime/Engine/Private/WorldPartition/WorldPartitionRuntimeLevelStreamingCell.cpp
@@ -2,14 +2,15 @@
 
 #include "WorldPartition/WorldPartitionRuntimeLevelStreamingCell.h"
 #include "WorldPartition/WorldPartitionLevelStreamingDynamic.h"
-#include "WorldPartition/WorldPartitionActorDescView.h"
 #include "WorldPartition/WorldPartitionLevelStreamingPolicy.h"
 #include "WorldPartition/WorldPartitionDebugHelper.h"
+#include "WorldPartition/WorldPartitionStreamingGenerationContext.h"
 #include "Engine/LevelStreaming.h"
 #include "Engine/Level.h"
 #include "Misc/HierarchicalLogArchive.h"
 #include "Misc/Paths.h"
 #include "AssetRegistry/IAssetRegistry.h"
+#include "UObject/Package.h"
 #include "UObject/AssetRegistryTagsContext.h"
 
 #include UE_INLINE_GENERATED_CPP_BY_NAME(WorldPartitionRuntimeLevelStreamingCell)
@@ -250,24 +251,31 @@ void UWorldPartitionRuntimeLevelStreamingCell::SetIsAlwaysLoaded(bool bInIsAlway
 }
 
 #if WITH_EDITOR
-void UWorldPartitionRuntimeLevelStreamingCell::AddActorToCell(const FWorldPartitionActorDescView& ActorDescView, const FActorContainerID& InContainerID, const FTransform& InContainerTransform, const UActorDescContainer* InContainer)
+void UWorldPartitionRuntimeLevelStreamingCell::AddActorToCell(const FStreamingGenerationActorDescView& ActorDescView)
 {
 	check(!ActorDescView.GetActorIsEditorOnly());
 
+	const UActorDescContainerInstance* ContainerInstance = ActorDescView.GetContainerInstance();
+	check(ContainerInstance);
+
+	const FActorContainerID& ContainerID = ContainerInstance->GetContainerID();
+	const FTransform ContainerTransform = ContainerInstance->GetTransform();
+	const FName ContainerPackage = ContainerInstance->GetContainerPackage();
+
 	for (const FGuid& EditorReferenceGuid : ActorDescView.GetEditorReferences())
 	{
-		const FWorldPartitionActorDesc& ReferenceActorDesc = InContainer->GetActorDescChecked(EditorReferenceGuid);
+		const FWorldPartitionActorDescInstance& ReferenceActorDesc = ContainerInstance->GetActorDescInstanceChecked(EditorReferenceGuid);
 
 		Packages.Emplace(
-			ReferenceActorDesc.GetActorPackage(), 
-			*ReferenceActorDesc.GetActorSoftPath().ToString(), 
+			ReferenceActorDesc.GetActorPackage(),
+			*ReferenceActorDesc.GetActorSoftPath().ToString(),
 			ReferenceActorDesc.GetBaseClass(),
 			ReferenceActorDesc.GetNativeClass(),
-			InContainerID, 
-			InContainerTransform, 
-			InContainer->GetContainerPackage(), 
-			GetWorld()->GetPackage()->GetFName(), 
-			InContainerID.GetActorGuid(ReferenceActorDesc.GetGuid()),
+			ContainerID,
+			ContainerTransform,
+			ContainerPackage,
+			GetWorld()->GetPackage()->GetFName(),
+			ContainerID.GetActorGuid(EditorReferenceGuid),
 			true
 		);
 	}
@@ -277,11 +285,11 @@ void UWorldPartitionRuntimeLevelStreamingCell::AddActorToCell(const FWorldPartit
 		*ActorDescView.GetActorSoftPath().ToString(), 
 		ActorDescView.GetBaseClass(),
 		ActorDescView.GetNativeClass(),
-		InContainerID, 
-		InContainerTransform, 
-		InContainer->GetContainerPackage(), 
+		ContainerID,
+		ContainerTransform,
+		ContainerPackage, 
 		GetWorld()->GetPackage()->GetFName(), 
-		InContainerID.GetActorGuid(ActorDescView.GetGuid()),
+		ContainerID.GetActorGuid(ActorDescView.GetGuid()),
 		false
 	);
 }
diff --git a/Engine/Source/Runtime/Engine/Private/WorldPartition/WorldPartitionRuntimeSpatialHash.cpp b/Engine/Source/Runtime/Engine/Private/WorldPartition/WorldPartitionRuntimeSpatialHash.cpp
--- a/Engine/Source/Runtime/Engine/Private/WorldPartition/WorldPartitionRuntimeSpatialHash.cpp
+++ b/Engine/Source/Runtime/Engine/Private/WorldPartition/WorldPartitionRuntimeSpatialHash.cpp
@@ -34,6 +34,7 @@
 #include "Misc/HierarchicalLogArchive.h"
 #include "UObject/UnrealType.h"
 #include "UObject/ObjectSaveContext.h"
+#include "UObject/Package.h"
 #include "UObject/FortniteMainBranchObjectVersion.h"
 #include "Misc/HashBuilder.h"
 
@@ -1211,9 +1212,8 @@ bool UWorldPartitionRuntimeSpatialHash::GenerateStreaming(UWorldPartitionStreami
 	TArray<FSpatialHashRuntimeGrid> AllGrids;
 	AllGrids.Append(Grids);
 
-	TArray<const FWorldPartitionActorDescView*> SpatialHashRuntimeGridInfos;
-	Algo::Transform(StreamingGenerationContext->GetMainWorldContainer()->ActorDescViewMap->FindByExactNativeClass<ASpatialHashRuntimeGridInfo>(), SpatialHashRuntimeGridInfos, [](const FStreamingGenerationActorDescView* ActorDescView) { return ActorDescView; });
-	for (const FWorldPartitionActorDescView* SpatialHashRuntimeGridInfo : SpatialHashRuntimeGridInfos)
+	TArray<const FStreamingGenerationActorDescView*> SpatialHashRuntimeGridInfos = StreamingGenerationContext->GetMainWorldContainerInstance()->ActorDescViewMap->FindByExactNativeClass<ASpatialHashRuntimeGridInfo>();
+	for (const FStreamingGenerationActorDescView* SpatialHashRuntimeGridInfo : SpatialHashRuntimeGridInfos)
 	{
 		FWorldPartitionReference Ref(WorldPartition, SpatialHashRuntimeGridInfo->GetGuid());
 		ASpatialHashRuntimeGridInfo* RuntimeGridActor = CastChecked<ASpatialHashRuntimeGridInfo>(SpatialHashRuntimeGridInfo->GetActor());
