WorldPartition: Add UActorDescContainerInstance/FWorldPartitionActorDescInstance which take transient/context data from UActorDescContainer/FWorldPartitionActorDesc leaving them as pure on disk repres

diff --git a/Engine/Source/Runtime/Engine/Private/WorldPartition/WorldPartitionStreamingGeneration.cpp b/Engine/Source/Runtime/Engine/Private/WorldPartition/WorldPartitionStreamingGeneration.cpp
--- a/Engine/Source/Runtime/Engine/Private/WorldPartition/WorldPartitionStreamingGeneration.cpp
+++ b/Engine/Source/Runtime/Engine/Private/WorldPartition/WorldPartitionStreamingGeneration.cpp
@@ -53,6 +53,91 @@ static FAutoConsoleCommand DumpStreamingGenerationLog(
 	})
 );
 
+class FGCUnsavedDirtyActorContainerInstances : public FGCObject
+{
+	friend class FStreamingGenerationUnsavedDirtyActorDescInstance;
+public:
+	//~ Begin FGCObject
+	virtual void AddReferencedObjects(FReferenceCollector& Collector)
+	{
+		Collector.AddReferencedObjects(ContainerInstances);
+	}
+	virtual FString GetReferencerName() const { return TEXT("FUnsavedDirtyActorContainer"); }
+	//~ End
+
+protected:
+	TSet<TObjectPtr<UActorDescContainerInstance>> ContainerInstances;
+};
+
+class FStreamingGenerationUnsavedDirtyActorDescInstance : public FWorldPartitionActorDescInstance
+{
+public:
+	FStreamingGenerationUnsavedDirtyActorDescInstance(UActorDescContainerInstance* InContainerInstance, AActor* InActor)
+	{
+		ContainerInstance = InContainerInstance;
+
+		ActorDescPtr = InActor->CreateActorDesc();
+		ActorDescPtr->SetContainer(InContainerInstance->GetContainer());
+		ActorDesc = ActorDescPtr.Get();
+
+		if (IsChildContainerInstance())
+		{
+			RegisterChildContainerInstance();
+		}
+	}
+
+	FStreamingGenerationUnsavedDirtyActorDescInstance(FStreamingGenerationUnsavedDirtyActorDescInstance&& InOther) = delete;
+	FStreamingGenerationUnsavedDirtyActorDescInstance& operator=(FStreamingGenerationUnsavedDirtyActorDescInstance&& InOther) = delete;
+
+	virtual ~FStreamingGenerationUnsavedDirtyActorDescInstance()
+	{
+		if (IsChildContainerInstance())
+		{
+			UnregisterChildContainerInstance();
+		}
+	}
+
+	static TUniquePtr<FStreamingGenerationUnsavedDirtyActorDescInstance> Create(AActor* InActor, const FStreamingGenerationContainerInstanceCollection& InContainerInstanceCollection)
+	{
+		UActorDescContainerInstance* HandlingContainer = const_cast<UActorDescContainerInstance*>(InContainerInstanceCollection.FindHandlingContainerInstance(InActor).Get());
+		check(HandlingContainer);
+		return TUniquePtr<FStreamingGenerationUnsavedDirtyActorDescInstance>(new FStreamingGenerationUnsavedDirtyActorDescInstance(HandlingContainer, InActor));
+	}
+
+	//~ Begin FWorldPartitionActorDescInstance	
+	virtual void RegisterChildContainerInstance() override
+	{
+		if (!UnsavedDirtyActorContainerInstances.IsValid())
+		{
+			UnsavedDirtyActorContainerInstances = MakeUnique<FGCUnsavedDirtyActorContainerInstances>();
+		}
+
+		ChildContainerInstance = GetActorDesc()->CreateChildContainerInstance(this);
+		UnsavedDirtyActorContainerInstances->ContainerInstances.Add(ChildContainerInstance);
+	}
+	virtual void UnregisterChildContainerInstance() override
+	{
+		check(UnsavedDirtyActorContainerInstances.IsValid())
+		{
+			UnsavedDirtyActorContainerInstances->ContainerInstances.Remove(ChildContainerInstance);
+			if (UnsavedDirtyActorContainerInstances->ContainerInstances.IsEmpty())
+			{
+				UnsavedDirtyActorContainerInstances.Reset();
+			}
+		}
+
+		ChildContainerInstance->Uninitialize();
+		ChildContainerInstance = nullptr;
+	}
+	//~ End
+protected:
+
+	TUniquePtr<FWorldPartitionActorDesc> ActorDescPtr;
+	static TUniquePtr<FGCUnsavedDirtyActorContainerInstances> UnsavedDirtyActorContainerInstances;
+};
+
+TUniquePtr<FGCUnsavedDirtyActorContainerInstances> FStreamingGenerationUnsavedDirtyActorDescInstance::UnsavedDirtyActorContainerInstances;
+
 template <class T>
 class TErrorHandlerSelector
 {
@@ -76,25 +161,6 @@ private:
 	IStreamingGenerationErrorHandler* ErrorHandler;
 };
 
-UWorldPartition::FCheckForErrorsParams::FCheckForErrorsParams()
-	: ErrorHandler(nullptr)
-	, ActorDescContainerCollection(nullptr)
-	, bEnableStreaming(false)
-{}
-
-FStreamingGenerationActorDescView::FStreamingGenerationActorDescView()
-	: FWorldPartitionActorDescView(nullptr)
-{}
-
-FStreamingGenerationActorDescView::FStreamingGenerationActorDescView(const FWorldPartitionActorDesc* InActorDesc)
-	: FWorldPartitionActorDescView(InActorDesc)
-	, ParentView(nullptr)
-	, bIsForcedNonSpatiallyLoaded(false)
-	, bIsForcedNoRuntimeGrid(false)
-	, bIsForcedNoDataLayers(false)	
-	, bIsForceNoHLODLayer(false)
-{}
-
 FName FStreamingGenerationActorDescView::GetRuntimeGrid() const
 {
 	if (bIsForcedNoRuntimeGrid)
@@ -107,7 +173,7 @@ FName FStreamingGenerationActorDescView::GetRuntimeGrid() const
 		return ParentView->GetRuntimeGrid();
 	}
 
-	return FWorldPartitionActorDescView::GetRuntimeGrid();
+	return ActorDescInstance->GetRuntimeGrid();
 }
 
 bool FStreamingGenerationActorDescView::GetIsSpatiallyLoaded() const
@@ -117,8 +183,7 @@ bool FStreamingGenerationActorDescView::GetIsSpatiallyLoaded() const
 		return false;
 	}
 
-	bool bIsSpatiallyLoaded = FWorldPartitionActorDescView::GetIsSpatiallyLoaded();
-
+	bool bIsSpatiallyLoaded = ActorDescInstance->GetIsSpatiallyLoaded();
 	if (bIsSpatiallyLoaded && ParentView)
 	{
 		bIsSpatiallyLoaded = ParentView->GetIsSpatiallyLoaded();
@@ -133,13 +198,13 @@ FSoftObjectPath FStreamingGenerationActorDescView::GetHLODLayer() const
 	{
 		return FSoftObjectPath();
 	}
-	
+
 	if (RuntimedHLODLayer.IsSet())
 	{
 		return RuntimedHLODLayer.GetValue();
 	}
 
-	return FWorldPartitionActorDescView::GetHLODLayer();
+	return ActorDescInstance->GetHLODLayer();
 }
 
 const TArray<FName>& FStreamingGenerationActorDescView::GetDataLayerInstanceNames() const
@@ -154,18 +219,18 @@ const TArray<FName>& FStreamingGenerationActorDescView::GetDataLayerInstanceName
 	{
 		return ParentView->GetDataLayerInstanceNames();
 	}
-	
+
 	if (ResolvedDataLayerInstanceNames.IsSet())
 	{
 		return ResolvedDataLayerInstanceNames.GetValue();
 	}
 
-	return FWorldPartitionActorDescView::GetDataLayerInstanceNames();
+	return ActorDescInstance->GetDataLayerInstanceNames();
 }
 
 const TArray<FGuid>& FStreamingGenerationActorDescView::GetReferences() const
 {
-	return RuntimeReferences.IsSet() ? RuntimeReferences.GetValue() : FWorldPartitionActorDescView::GetReferences();
+	return RuntimeReferences.IsSet() ? RuntimeReferences.GetValue() : ActorDescInstance->GetReferences();
 }
 
 const TArray<FGuid>& FStreamingGenerationActorDescView::GetEditorReferences() const
@@ -173,14 +238,7 @@ const TArray<FGuid>& FStreamingGenerationActorDescView::GetEditorReferences() co
 	return EditorReferences;
 }
 
-bool FStreamingGenerationActorDescView::ShouldValidateRuntimeGrid() const
-{
-PRAGMA_DISABLE_DEPRECATION_WARNINGS
-	return ActorDesc->ShouldValidateRuntimeGrid();
-PRAGMA_ENABLE_DEPRECATION_WARNINGS
-}
-
-void FStreamingGenerationActorDescView::SetParentView(const FWorldPartitionActorDescView* InParentView)
+void FStreamingGenerationActorDescView::SetParentView(const FStreamingGenerationActorDescView* InParentView)
 {
 	check(!ParentView);
 	check(GetParentActor().IsValid());
@@ -189,7 +247,7 @@ void FStreamingGenerationActorDescView::SetParentView(const FWorldPartitionActor
 
 void FStreamingGenerationActorDescView::SetDataLayerInstanceNames(const TArray<FName>& InDataLayerInstanceNames)
 {
-	check(!FWorldPartitionActorDescView::HasResolvedDataLayerInstanceNames());
+	check(!ActorDescInstance->HasResolvedDataLayerInstanceNames());
 	ResolvedDataLayerInstanceNames = InDataLayerInstanceNames;
 }
 
@@ -206,7 +264,7 @@ void FStreamingGenerationActorDescView::SetForcedNoRuntimeGrid()
 {
 	if (!bIsForcedNoRuntimeGrid)
 	{
-		bIsForcedNoRuntimeGrid = true;	
+		bIsForcedNoRuntimeGrid = true;
 		UE_LOG(LogWorldPartition, Verbose, TEXT("Actor '%s' runtime grid invalidated"), *GetActorLabelOrName().ToString());
 	}
 }
@@ -253,7 +311,8 @@ const TArray<FName>& FStreamingGenerationActorDescView::GetRuntimeDataLayerInsta
 {
 	if (bIsForcedNoDataLayers || !RuntimeDataLayerInstanceNames.IsSet())
 	{
-		return FWorldPartitionActorDescView::GetRuntimeDataLayerInstanceNames();
+		static TArray<FName> EmptyDataLayers;
+		return EmptyDataLayers;
 	}
 
 	if (ParentView)
@@ -264,26 +323,63 @@ const TArray<FName>& FStreamingGenerationActorDescView::GetRuntimeDataLayerInsta
 	return RuntimeDataLayerInstanceNames.GetValue();
 }
 
+FStreamingGenerationActorDescViewMap::FStreamingGenerationActorDescViewMap()
+{}
+
+TArray<const FStreamingGenerationActorDescView*> FStreamingGenerationActorDescViewMap::FindByExactNativeClass(UClass* InExactNativeClass) const
+{
+	check(InExactNativeClass->IsNative());
+	const FName NativeClassName = InExactNativeClass->GetFName();
+	TArray<const FStreamingGenerationActorDescView*> Result;
+	ActorDescViewsByClass.MultiFind(NativeClassName, Result);
+	return Result;
+}
+
+FStreamingGenerationActorDescView* FStreamingGenerationActorDescViewMap::Emplace(const FGuid& InGuid, const FStreamingGenerationActorDescView& InActorDescView)
+{
+	FStreamingGenerationActorDescView* NewActorDescView = ActorDescViewList.Emplace_GetRef(MakeUnique<FStreamingGenerationActorDescView>(InActorDescView)).Get();
+	NewActorDescView->ActorDescViewMap = this;
+
+	const UClass* NativeClass = NewActorDescView->GetActorNativeClass();
+	const FName NativeClassName = NativeClass->GetFName();
+
+	ActorDescViewsByGuid.Emplace(InGuid, NewActorDescView);
+	ActorDescViewsByClass.Add(NativeClassName, NewActorDescView);
+
+	return NewActorDescView;
+}
+
+FStreamingGenerationActorDescView* FStreamingGenerationActorDescViewMap::Emplace(const FWorldPartitionActorDescInstance* InActorDescInstance)
+{
+	return Emplace(InActorDescInstance->GetGuid(), FStreamingGenerationActorDescView(*this, InActorDescInstance));
+}
+
+UWorldPartition::FCheckForErrorsParams::FCheckForErrorsParams()
+	: ErrorHandler(nullptr)
+	, bEnableStreaming(false)
+	, ActorDescContainerInstanceCollection(nullptr)
+{}
+
 class FWorldPartitionStreamingGenerator
 {
 	class FStreamingGenerationContext : public IStreamingGenerationContext
 	{
 	public:
-		FStreamingGenerationContext(const FWorldPartitionStreamingGenerator* StreamingGenerator, const FStreamingGenerationActorDescCollection& TopLevelActorDescCollection)
+		FStreamingGenerationContext(const FWorldPartitionStreamingGenerator* StreamingGenerator, const FStreamingGenerationContainerInstanceCollection& TopLevelActorDescCollection)
 		{
 			// Create the dataset required for IStreamingGenerationContext interface
 			MainWorldActorSetContainerIndex = INDEX_NONE;
-			ActorSetContainers.Empty(StreamingGenerator->ContainerCollectionDescriptorsMap.Num());
-			
-			TMap<TWeakPtr<FStreamingGenerationActorDescCollection>, int32> ActorSetContainerMap;
-			for (const auto& [LevelName, ContainerDescriptor] : StreamingGenerator->ContainerCollectionDescriptorsMap)
+			ActorSetContainerInstances.Empty(StreamingGenerator->ContainerCollectionInstanceDescriptorsMap.Num());
+
+			TMap<TWeakPtr<FStreamingGenerationContainerInstanceCollection>, int32> ActorSetContainerMap;
+			for (const auto& [ContainerID, ContainerDescriptor] : StreamingGenerator->ContainerCollectionInstanceDescriptorsMap)
 			{
-				const int32 ContainerIndex = ActorSetContainers.AddDefaulted();
-				ActorSetContainerMap.Add(ContainerDescriptor.ActorDescCollection, ContainerIndex);
+				const int32 ContainerIndex = ActorSetContainerInstances.AddDefaulted();
+				ActorSetContainerMap.Add(ContainerDescriptor.ContainerInstanceCollection, ContainerIndex);
 
-				FActorSetContainer& ActorSetContainer = ActorSetContainers[ContainerIndex];
-				ActorSetContainer.ActorDescViewMap = &ContainerDescriptor.ActorDescViewMap;
-				ActorSetContainer.ActorDescCollection = ContainerDescriptor.ActorDescCollection.Get();
+				FActorSetContainerInstance& ActorSetContainer = ActorSetContainerInstances[ContainerIndex];
+				ActorSetContainer.ActorDescViewMap = &ContainerDescriptor.ActorDescViewMap.Get();
+				ActorSetContainer.ContainerInstanceCollection = ContainerDescriptor.ContainerInstanceCollection.Get();
 
 				ActorSetContainer.ActorSets.Empty(ContainerDescriptor.Clusters.Num());
 				for (const TArray<FGuid>& Cluster : ContainerDescriptor.Clusters)
@@ -292,7 +388,7 @@ class FWorldPartitionStreamingGenerator
 					ActorSet.Actors = Cluster;
 				}
 
-				if (ContainerDescriptor.ActorDescCollection->GetMainContainerPackageName() == TopLevelActorDescCollection.GetMainContainerPackageName())
+				if (ContainerDescriptor.ContainerInstanceCollection->GetMainContainerPackageName() == TopLevelActorDescCollection.GetMainContainerPackageName())
 				{
 					check(MainWorldActorSetContainerIndex == INDEX_NONE);
 					MainWorldActorSetContainerIndex = ContainerIndex;
@@ -302,13 +398,12 @@ class FWorldPartitionStreamingGenerator
 			ActorSetInstances.Empty();
 			for (const auto& [ContainerID, ContainerCollectionInstanceDescriptor] : StreamingGenerator->ContainerCollectionInstanceDescriptorsMap)
 			{
-				const FContainerCollectionDescriptor& ContainerCollectionDescriptor = StreamingGenerator->ContainerCollectionDescriptorsMap.FindChecked(ContainerCollectionInstanceDescriptor.ActorDescCollection->GetMainContainerPackageName());
-				const FActorSetContainer& ActorSetContainer = ActorSetContainers[ActorSetContainerMap.FindChecked(ContainerCollectionDescriptor.ActorDescCollection)];
+				const FActorSetContainerInstance& ActorSetContainer = ActorSetContainerInstances[ActorSetContainerMap.FindChecked(ContainerCollectionInstanceDescriptor.ContainerInstanceCollection)];
 				const TSet<FGuid>* FilteredActors = StreamingGenerator->ContainerFilteredActors.Find(ContainerID);
 				for (const TUniquePtr<FActorSet>& ActorSetPtr : ActorSetContainer.ActorSets)
 				{
 					const FActorSet& ActorSet = *ActorSetPtr;
-					const FWorldPartitionActorDescView& ReferenceActorDescView = ActorSetContainer.ActorDescViewMap->FindByGuidChecked(ActorSet.Actors[0]);
+					const FStreamingGenerationActorDescView& ReferenceActorDescView = ActorSetContainer.ActorDescViewMap->FindByGuidChecked(ActorSet.Actors[0]);
 
 					bool bContainsUnfilteredActors = !FilteredActors;
 
@@ -329,8 +424,8 @@ class FWorldPartitionStreamingGenerator
 					{
 						FActorSetInstance& ActorSetInstance = ActorSetInstances.Emplace_GetRef();
 						const FContainerCollectionInstanceDescriptor::FPerInstanceData& PerInstanceData = ContainerCollectionInstanceDescriptor.GetInstanceData(ReferenceActorDescView.GetGuid());
-				
-						ActorSetInstance.ContainerInstance = &ActorSetContainer;
+
+						ActorSetInstance.ActorSetContainerInstance = &ActorSetContainer;
 						ActorSetInstance.ActorSet = &ActorSet;
 						ActorSetInstance.FilteredActors = FilteredActors;
 						ActorSetInstance.ContainerID = ContainerCollectionInstanceDescriptor.ID;
@@ -342,9 +437,9 @@ class FWorldPartitionStreamingGenerator
 
 						ActorSetInstance.Bounds.Init();
 						const FTransform& ContainerCollectionInstanceDescriptorTransform = ContainerCollectionInstanceDescriptor.Transform;
-						ActorSetInstance.ForEachActor([this, &ActorSetContainer, &ActorSetInstance, &ContainerCollectionInstanceDescriptorTransform](const FGuid& ActorGuid)
+						ActorSetInstance.ForEachActor([this, &ActorSetInstance, &ContainerCollectionInstanceDescriptorTransform](const FGuid& ActorGuid)
 						{
-							const FWorldPartitionActorDescView& ActorDescView = ActorSetContainer.ActorDescViewMap->FindByGuidChecked(ActorGuid);
+							const FStreamingGenerationActorDescView& ActorDescView = ActorSetInstance.ActorSetContainerInstance->ActorDescViewMap->FindByGuidChecked(ActorGuid);
 							const FBox RuntimeBounds = ActorDescView.GetRuntimeBounds();
 							if (RuntimeBounds.IsValid)
 							{
@@ -367,9 +462,9 @@ class FWorldPartitionStreamingGenerator
 			return WorldBounds;
 		}
 
-		virtual const FActorSetContainer* GetMainWorldContainer() const override
+		virtual const FActorSetContainerInstance* GetMainWorldContainerInstance() const override
 		{
-			return ActorSetContainers.IsValidIndex(MainWorldActorSetContainerIndex) ? &ActorSetContainers[MainWorldActorSetContainerIndex] : nullptr;
+			return ActorSetContainerInstances.IsValidIndex(MainWorldActorSetContainerIndex) ? &ActorSetContainerInstances[MainWorldActorSetContainerIndex] : nullptr;
 		}
 
 		virtual void ForEachActorSetInstance(TFunctionRef<void(const FActorSetInstance&)> Func) const override
@@ -380,11 +475,11 @@ class FWorldPartitionStreamingGenerator
 			}
 		}
 
-		virtual void ForEachActorSetContainer(TFunctionRef<void(const FActorSetContainer&)> Func) const override
+		virtual void ForEachActorSetContainerInstance(TFunctionRef<void(const FActorSetContainerInstance&)> Func) const override
 		{
-			for (const FActorSetContainer& ActorSetContainer : ActorSetContainers)
+			for (const FActorSetContainerInstance& ActorSetContainerInstance : ActorSetContainerInstances)
 			{
-				Func(ActorSetContainer);
+				Func(ActorSetContainerInstance);
 			}
 		}
 		//~End IStreamingGenerationContext interface};
@@ -392,32 +487,11 @@ class FWorldPartitionStreamingGenerator
 	private:
 		FBox WorldBounds;
 		int32 MainWorldActorSetContainerIndex;
-		TArray<FActorSetContainer> ActorSetContainers;
+		TArray<FActorSetContainerInstance> ActorSetContainerInstances;
 		TArray<FActorSetInstance> ActorSetInstances;
 	};
 
-	/** 
-	 * An actor container descriptor, one for the main world and one for every unique level used by level instances.
-	 */
-	struct FContainerCollectionDescriptor
-	{
-		/** The unique actor descriptor container collection for this descriptor */
-		TSharedPtr<FStreamingGenerationActorDescCollection> ActorDescCollection;
-
-		/** The actor descriptor views for for this descriptor */
-		FStreamingGenerationActorDescViewMap ActorDescViewMap;
-
-		/** Set of editor-only actors that are not part of the actor descriptor views */
-		TSet<FGuid> EditorOnlyActorDescMap;
-
-		/** List of actor descriptor views that are containers (mainly level instances) */
-		TArray<FStreamingGenerationActorDescView> ContainerCollectionInstanceViews;
-
-		/** List of actor clusters for this descriptor */
-		TArray<TArray<FGuid>> Clusters;
-	};
-
-	/** 
+	/**
 	 * An actor container instance descriptor, one for the main world and one for every actor container instance
 	 */
 	struct FContainerCollectionInstanceDescriptor
@@ -428,13 +502,34 @@ class FWorldPartitionStreamingGenerator
 
 		FBox Bounds;
 		FTransform Transform;
-		TSharedPtr<FStreamingGenerationActorDescCollection> ActorDescCollection;
+		TSharedPtr<FStreamingGenerationContainerInstanceCollection> ContainerInstanceCollection;
 		EContainerClusterMode ClusterMode;
 		FString OwnerName;
 		FActorContainerID ID;
 		FActorContainerID ParentID;
 		FGuid ContentBundleID;
 
+		FContainerCollectionInstanceDescriptor(const FContainerCollectionInstanceDescriptor& Other) = delete;
+		FContainerCollectionInstanceDescriptor(FContainerCollectionInstanceDescriptor&& Other) = default;
+
+		FContainerCollectionInstanceDescriptor& operator=(const FContainerCollectionInstanceDescriptor& Other) = delete;
+		FContainerCollectionInstanceDescriptor& operator=(FContainerCollectionInstanceDescriptor&& Other) = default;
+				
+		/** The actor descriptor views for for this descriptor (TUniqueObj so it is moveable without having to update FStreamingGenerationActorDescView::ActorDescViewMap pointer */
+		TUniqueObj<FStreamingGenerationActorDescViewMap> ActorDescViewMap;
+
+		/** Set of editor-only actors that are not part of the actor descriptor views */
+		TSet<FGuid> EditorOnlyActorDescMap;
+
+		/** List of actor descriptor views that are containers (mainly level instances) */
+		TArray<FStreamingGenerationActorDescView> ContainerCollectionInstanceViews;
+
+		/** List of unsaved/dirty descriptor views, unique ptrs so that they are moveable without having to update references to them from FStreamingGenerationActorDescView::ActorDescView */
+		TArray<TUniquePtr<FStreamingGenerationUnsavedDirtyActorDescInstance>> UnsavedDirtyInstances;
+
+		/** List of actor clusters for this descriptor */
+		TArray<TArray<FGuid>> Clusters;
+
 		struct FPerInstanceData
 		{
 			bool bIsSpatiallyLoaded;
@@ -508,21 +603,18 @@ class FWorldPartitionStreamingGenerator
 
 	void ResolveRuntimeDataLayers(FStreamingGenerationActorDescView& ActorDescView, const FStreamingGenerationActorDescViewMap& ActorDescViewMap)
 	{
-		TArray<const FWorldPartitionActorDescView*> WorldDataLayerViews;
-		Algo::Transform(ActorDescViewMap.FindByExactNativeClass<AWorldDataLayers>(), WorldDataLayerViews, [](const FStreamingGenerationActorDescView* ActorDescView) { return ActorDescView; });
-
 		// Resolve DataLayerInstanceNames of ActorDescView only when necessary (i.e. when container is a template)
-		if (!ActorDescView.GetActorDesc()->HasResolvedDataLayerInstanceNames())
+		if (!ActorDescView.HasResolvedDataLayerInstanceNames())
 		{
 			// Build a WorldDataLayerActorDescs if DataLayerManager can't resolve Data Layers (i.e. when validating changelists and World is not loaded)
 			const bool bDataLayerManagerCanResolve = DataLayerManager && DataLayerManager->CanResolveDataLayers();
-			const TArray<const FWorldDataLayersActorDesc*> WorldDataLayerActorDescs = !bDataLayerManagerCanResolve ? FDataLayerUtils::FindWorldDataLayerActorDescs(WorldDataLayerViews) : TArray<const FWorldDataLayersActorDesc*>();
+			const TArray<const FWorldDataLayersActorDesc*> WorldDataLayerActorDescs = !bDataLayerManagerCanResolve ? FDataLayerUtils::FindWorldDataLayerActorDescs(ActorDescViewMap) : TArray<const FWorldDataLayersActorDesc*>();
 			const TArray<FName> DataLayerInstanceNames = FDataLayerUtils::ResolvedDataLayerInstanceNames(DataLayerManager, ActorDescView.GetActorDesc(), WorldDataLayerActorDescs);
 			ActorDescView.SetDataLayerInstanceNames(DataLayerInstanceNames);
 		}
 
 		TArray<FName> RuntimeDataLayerInstanceNames;
-		if (FDataLayerUtils::ResolveRuntimeDataLayerInstanceNames(DataLayerManager, ActorDescView, WorldDataLayerViews, RuntimeDataLayerInstanceNames))
+		if (FDataLayerUtils::ResolveRuntimeDataLayerInstanceNames(DataLayerManager, ActorDescView, ActorDescViewMap, RuntimeDataLayerInstanceNames))
 		{
 			ActorDescView.SetRuntimeDataLayerInstanceNames(RuntimeDataLayerInstanceNames);
 		}
@@ -546,68 +638,55 @@ class FWorldPartitionStreamingGenerator
 		}
 	}
 
-	void CreateActorDescViewMap(const FStreamingGenerationActorDescCollection& InActorDescCollection, FStreamingGenerationActorDescViewMap& OutActorDescViewMap, TSet<FGuid>& OutEditorOnlyActorDescMap, const FActorContainerID& InContainerID, TArray<FStreamingGenerationActorDescView>& OutContainerInstances)
+	void CreateActorDescViewMap(const FStreamingGenerationContainerInstanceCollection& InActorDescCollection, FStreamingGenerationActorDescViewMap& OutActorDescViewMap, TSet<FGuid>& OutEditorOnlyActorDescMap, const FActorContainerID& InContainerID, TArray<FStreamingGenerationActorDescView>& OutContainerInstances, TArray<TUniquePtr<FStreamingGenerationUnsavedDirtyActorDescInstance>>& OutUnsavedDirtyInstances)
 	{
 		// Should we handle unsaved or newly created actors?
-		const bool bHandleUnsavedActors = ModifiedActorsDescList && InContainerID.IsMainContainer();
+		const bool bShouldHandleUnsavedActors = bHandleUnsavedActors && InContainerID.IsMainContainer();
 
 		// Consider all actors of a /Temp/ container package as Unsaved because loading them from disk will fail (Outer world name mismatch)
 		const bool bIsTempContainerPackage = FPackageName::IsTempPackage((InActorDescCollection.GetMainContainerPackageName().ToString()));
 
-		// Create an actor descriptor view for the specified actor (modified or unsaved actors)
-		auto GetModifiedActorDesc = [this](AActor* InActor, const FStreamingGenerationActorDescCollection& InActorDescCollection) -> FWorldPartitionActorDesc*
-		{
-			FWorldPartitionActorDesc* ModifiedActorDesc = ModifiedActorsDescList->AddActor(InActor);
-
-			const UActorDescContainer* HandlingContainer = InActorDescCollection.FindHandlingContainer(InActor);
-			check(HandlingContainer != nullptr);
-
-			// Pretend that this actor descriptor belongs to the original container, even if it's not present. It's essentially a proxy
-			// descriptor on top an existing one and at this point no code should require to access the container to resolve it anyways.
-			ModifiedActorDesc->SetContainer(const_cast<UActorDescContainer*>(HandlingContainer), InActor->GetWorld());
-
-			return ModifiedActorDesc;
-		};
-
 		// Test whether an actor should be included in the ActorDescViewMap.
-		auto ShouldRegisterActorDesc = [this](const FWorldPartitionActorDesc* ActorDesc, const FActorContainerID& ContainerID)
+		auto ShouldRegisterActorDesc = [this](const IWorldPartitionActorDescInstance* InActorDescView)
 		{
 			for (UClass* FilteredClass : FilteredClasses)
 			{
-				if (ActorDesc->GetActorNativeClass()->IsChildOf(FilteredClass))
+				if (InActorDescView->GetActorNativeClass()->IsChildOf(FilteredClass))
 				{
 					return false;
 				}
 			}
 
-			if (!ActorDesc->IsRuntimeRelevant(ContainerID))
+			if (!InActorDescView->IsRuntimeRelevant())
 			{
 				return false;
 			}
 
-			return ActorDesc->IsLoaded() ? !ActorDesc->GetActor()->IsEditorOnly() : !ActorDesc->GetActorIsEditorOnly();
+			return InActorDescView->IsLoaded() ? !InActorDescView->GetActor()->IsEditorOnly() : !InActorDescView->GetActorIsEditorOnly();
 		};
 
 		// Register the actor descriptor view
-		auto RegisterActorDescView = [this, &OutActorDescViewMap, &OutContainerInstances](const FGuid& ActorGuid, FStreamingGenerationActorDescView& InActorDescView)
+		auto RegisterActorDescView = [this, &OutActorDescViewMap, &OutContainerInstances](FStreamingGenerationActorDescView&& InActorDescView)
 		{
-			if (InActorDescView.IsContainerInstance())
+			if (InActorDescView.IsChildContainerInstance())
 			{
 				OutContainerInstances.Add(InActorDescView);
 			}
 			else
 			{
-				OutActorDescViewMap.Emplace(ActorGuid, InActorDescView);
+				const FGuid ActorGuid = InActorDescView.GetGuid();
+				OutActorDescViewMap.Emplace(ActorGuid, MoveTemp(InActorDescView));
 			}
 		};
-		
-		TMap<FGuid, FGuid> ContainerGuidsRemap;
-		for (FStreamingGenerationActorDescCollection::TConstIterator<> ActorDescIt(&InActorDescCollection); ActorDescIt; ++ActorDescIt)
+
+		for (FStreamingGenerationContainerInstanceCollection::TConstIterator<> Iterator(&InActorDescCollection); Iterator; ++Iterator)
 		{
-			if (ShouldRegisterActorDesc(*ActorDescIt, InContainerID))
+			// @todo_pat: this is to validate that new parenting of container instance code is equivalent
+			check(Iterator->GetContainerInstance()->GetContainerID() == InContainerID);
+			if (ShouldRegisterActorDesc(*Iterator))
 			{
 				// Handle unsaved actors
-				if (AActor* Actor = ActorDescIt->GetActor())
+				if (AActor* Actor = Iterator->GetActor())
 				{
 					// Deleted actors
 					if (!IsValid(Actor))
@@ -616,63 +695,64 @@ class FWorldPartitionStreamingGenerator
 					}
 
 					// Dirty actors
-					if (bHandleUnsavedActors && (bIsTempContainerPackage || Actor->GetPackage()->IsDirty()))
+					if (bShouldHandleUnsavedActors && (bIsTempContainerPackage || Actor->GetPackage()->IsDirty()))
 					{
 						// Dirty, unsaved actor for PIE
-						FStreamingGenerationActorDescView ModifiedActorDescView = GetModifiedActorDesc(Actor, InActorDescCollection);
-						RegisterActorDescView(ActorDescIt->GetGuid(), ModifiedActorDescView);
+						TUniquePtr<FStreamingGenerationUnsavedDirtyActorDescInstance>& UnsavedDirtyRef = OutUnsavedDirtyInstances.Add_GetRef(FStreamingGenerationUnsavedDirtyActorDescInstance::Create(Actor, InActorDescCollection));
+						RegisterActorDescView(FStreamingGenerationActorDescView(OutActorDescViewMap, UnsavedDirtyRef.Get(), true));
 						continue;
 					}
 				}
 
 				// Non-dirty actor
-				FStreamingGenerationActorDescView ActorDescView(*ActorDescIt);
-				RegisterActorDescView(ActorDescIt->GetGuid(), ActorDescView);
+				RegisterActorDescView(FStreamingGenerationActorDescView(OutActorDescViewMap, *Iterator));
 			}
 			else
 			{
-				OutEditorOnlyActorDescMap.Add(ActorDescIt->GetGuid());
+				OutEditorOnlyActorDescMap.Add(Iterator->GetGuid());
 			}
 		}
 
 		// Append new unsaved actors for the persistent level
-		if (bHandleUnsavedActors)
+		if (bShouldHandleUnsavedActors)
 		{
 			for (AActor* Actor : InActorDescCollection.GetWorld()->PersistentLevel->Actors)
 			{
 				if (IsValid(Actor) && Actor->IsPackageExternal() && Actor->IsMainPackageActor() && !Actor->IsEditorOnly()
 					&& (Actor->GetContentBundleGuid() == InActorDescCollection.GetContentBundleGuid())
-					&& !InActorDescCollection.GetActorDesc(Actor->GetActorGuid()))
+					&& !InActorDescCollection.GetActorDescInstance(Actor->GetActorGuid()))
 				{
-					FStreamingGenerationActorDescView ModifiedActorDescView = GetModifiedActorDesc(Actor, InActorDescCollection);
-					if (ShouldRegisterActorDesc(ModifiedActorDescView.GetActorDesc(), InContainerID))
+					TUniquePtr<FStreamingGenerationUnsavedDirtyActorDescInstance> UnsavedViewPtr = FStreamingGenerationUnsavedDirtyActorDescInstance::Create(Actor, InActorDescCollection);
+					FStreamingGenerationActorDescView ModifiedActorDescView(OutActorDescViewMap, UnsavedViewPtr.Get(), true);
+					if (ShouldRegisterActorDesc(&ModifiedActorDescView))
 					{
-						RegisterActorDescView(Actor->GetActorGuid(), ModifiedActorDescView);
+						OutUnsavedDirtyInstances.Add(MoveTemp(UnsavedViewPtr));
+						RegisterActorDescView(MoveTemp(ModifiedActorDescView));
 					}
 				}
 			}
 		}
 	}
 
-	void CreateActorDescriptorViewsRecursive(const FContainerCollectionInstanceDescriptor& InContainerCollectionInstanceDescriptor)
+	void CreateActorDescriptorViewsRecursive(FContainerCollectionInstanceDescriptor&& InContainerCollectionInstanceDescriptor)
 	{
 		// Inherited parent per-instance data logic
-		auto InheritParentContainerPerInstanceData = [&InContainerCollectionInstanceDescriptor](const FContainerCollectionInstanceDescriptor::FPerInstanceData& InParentPerInstanceData, const FWorldPartitionActorDescView& InActorDescView)
+		auto InheritParentContainerPerInstanceData = [](const FContainerCollectionInstanceDescriptor& ParentContainerCollectionInstanceDescriptor, const FStreamingGenerationActorDescView& InActorDescView)
 		{
 			FContainerCollectionInstanceDescriptor::FPerInstanceData ResultPerInstanceData;
 
 			// Apply AND logic on spatially loaded flag
-			ResultPerInstanceData.bIsSpatiallyLoaded = InActorDescView.GetIsSpatiallyLoaded() && InParentPerInstanceData.bIsSpatiallyLoaded;
+			ResultPerInstanceData.bIsSpatiallyLoaded = InActorDescView.GetIsSpatiallyLoaded() && ParentContainerCollectionInstanceDescriptor.InstanceData.bIsSpatiallyLoaded;
 
 			// Runtime grid is inherited from the main world if the actor has its runtime grid set to none.
-			ResultPerInstanceData.RuntimeGrid = (InContainerCollectionInstanceDescriptor.ID.IsMainContainer() || InParentPerInstanceData.RuntimeGrid.IsNone()) ? InActorDescView.GetRuntimeGrid() : InParentPerInstanceData.RuntimeGrid;
+			ResultPerInstanceData.RuntimeGrid = (ParentContainerCollectionInstanceDescriptor.ID.IsMainContainer() || ParentContainerCollectionInstanceDescriptor.InstanceData.RuntimeGrid.IsNone()) ? InActorDescView.GetRuntimeGrid() : ParentContainerCollectionInstanceDescriptor.InstanceData.RuntimeGrid;
 
 			// Data layers are accumulated down the hierarchy chain, since level instances supports data layers assignation on actors
 			ResultPerInstanceData.DataLayers = InActorDescView.GetRuntimeDataLayerInstanceNames();
-			ResultPerInstanceData.DataLayers.Append(InParentPerInstanceData.DataLayers);
+			ResultPerInstanceData.DataLayers.Append(ParentContainerCollectionInstanceDescriptor.InstanceData.DataLayers);
 			ResultPerInstanceData.DataLayers.Sort(FNameFastLess());
 
-			if (InParentPerInstanceData.DataLayers.Num())
+			if (ParentContainerCollectionInstanceDescriptor.InstanceData.DataLayers.Num())
 			{
 				// Remove potential duplicates from sorted data layers array
 				ResultPerInstanceData.DataLayers.SetNum(Algo::Unique(ResultPerInstanceData.DataLayers));
@@ -681,38 +761,32 @@ class FWorldPartitionStreamingGenerator
 			return ResultPerInstanceData;
 		};
 
-		// ContainerDescriptor may be reallocated after this scope
+		// Hold on to ID
+		const FActorContainerID ContainerID = InContainerCollectionInstanceDescriptor.ID;
 		{
-			// Create or resolve container descriptor
-			FContainerCollectionDescriptor& ContainerCollectionDescriptor = ContainerCollectionDescriptorsMap.FindOrAdd(InContainerCollectionInstanceDescriptor.ActorDescCollection->GetMainContainerPackageName());
-		
+			TArray<FStreamingGenerationActorDescView> ContainerCollectionInstanceViews;
+
 			// ContainerInstanceDescriptor may be reallocated after this scope
 			{
 				// Create container instance descriptor
-				check(!ContainerCollectionInstanceDescriptorsMap.Contains(InContainerCollectionInstanceDescriptor.ID));
-
-				FContainerCollectionInstanceDescriptor& ContainerCollectionInstanceDescriptor = ContainerCollectionInstanceDescriptorsMap.Add(InContainerCollectionInstanceDescriptor.ID, InContainerCollectionInstanceDescriptor);
-
-				if (!ContainerCollectionDescriptor.ActorDescCollection)
-				{
-					ContainerCollectionDescriptor.ActorDescCollection = ContainerCollectionInstanceDescriptor.ActorDescCollection;
+				check(!ContainerCollectionInstanceDescriptorsMap.Contains(ContainerID));
 
-					// Gather actor descriptor views for this container
-					CreateActorDescViewMap(*ContainerCollectionInstanceDescriptor.ActorDescCollection, ContainerCollectionDescriptor.ActorDescViewMap, ContainerCollectionDescriptor.EditorOnlyActorDescMap, ContainerCollectionInstanceDescriptor.ID, ContainerCollectionDescriptor.ContainerCollectionInstanceViews);
+				FContainerCollectionInstanceDescriptor& ContainerCollectionInstanceDescriptor = ContainerCollectionInstanceDescriptorsMap.Add(ContainerID, MoveTemp(InContainerCollectionInstanceDescriptor));
+				// Gather actor descriptor views for this container
+				CreateActorDescViewMap(*ContainerCollectionInstanceDescriptor.ContainerInstanceCollection, ContainerCollectionInstanceDescriptor.ActorDescViewMap.Get(), ContainerCollectionInstanceDescriptor.EditorOnlyActorDescMap, ContainerCollectionInstanceDescriptor.ID, ContainerCollectionInstanceDescriptor.ContainerCollectionInstanceViews, ContainerCollectionInstanceDescriptor.UnsavedDirtyInstances);
 
-					// Resolve actor descriptor views before validation
-					ResolveContainerDescriptor(ContainerCollectionDescriptor);
+				// Resolve actor descriptor views before validation
+				ResolveContainerDescriptor(ContainerCollectionInstanceDescriptor);
 
-					// Validate container, fixing anything illegal, etc.
-					ValidateContainerDescriptor(ContainerCollectionDescriptor, ContainerCollectionInstanceDescriptor.ID.IsMainContainer());
+				// Validate container, fixing anything illegal, etc.
+				ValidateContainerDescriptor(ContainerCollectionInstanceDescriptor, ContainerCollectionInstanceDescriptor.ID.IsMainContainer());
 
-					// Update container, computing cluster, bounds, etc.
-					UpdateContainerDescriptor(ContainerCollectionDescriptor);
-				}
+				// Update container, computing cluster, bounds, etc.
+				UpdateContainerDescriptor(ContainerCollectionInstanceDescriptor);
 
 				// Calculate Bounds of non-container actor descriptor views
 				check(!ContainerCollectionInstanceDescriptor.Bounds.IsValid);
-				ContainerCollectionDescriptor.ActorDescViewMap.ForEachActorDescView([&ContainerCollectionInstanceDescriptor](const FStreamingGenerationActorDescView& ActorDescView)
+				ContainerCollectionInstanceDescriptor.ActorDescViewMap->ForEachActorDescView([&ContainerCollectionInstanceDescriptor](const FStreamingGenerationActorDescView& ActorDescView)
 				{
 					if (ActorDescView.GetIsSpatiallyLoaded())
 					{
@@ -722,68 +796,79 @@ class FWorldPartitionStreamingGenerator
 						ContainerCollectionInstanceDescriptor.Bounds += RuntimeBounds.TransformBy(ContainerCollectionInstanceDescriptor.Transform);
 					}
 				});
+
+				// Copy list as descriptor might get reallocated after this scope
+				ContainerCollectionInstanceViews.Append(ContainerCollectionInstanceDescriptor.ContainerCollectionInstanceViews);
 			}
 
+
 			// Parse actor containers
-			TArray<FStreamingGenerationActorDescView> ContainerCollectionInstanceViews = ContainerCollectionDescriptor.ContainerCollectionInstanceViews;
-			for (const FWorldPartitionActorDescView& ContainerCollectionInstanceView : ContainerCollectionInstanceViews)
+			for (const FStreamingGenerationActorDescView& ContainerCollectionInstanceView : ContainerCollectionInstanceViews)
 			{
 				FWorldPartitionActorDesc::FContainerInstance SubContainerInstance;
-				if (!ContainerCollectionInstanceView.GetContainerInstance(SubContainerInstance) || !SubContainerInstance.Container)
+				if (!ContainerCollectionInstanceView.GetChildContainerInstance(SubContainerInstance) || !SubContainerInstance.ContainerInstance)
 				{
-					ErrorHandler->OnLevelInstanceInvalidWorldAsset(ContainerCollectionInstanceView, ContainerCollectionInstanceView.GetContainerPackage(), IStreamingGenerationErrorHandler::ELevelInstanceInvalidReason::WorldAssetHasInvalidContainer);
+					ErrorHandler->OnLevelInstanceInvalidWorldAsset(ContainerCollectionInstanceView, ContainerCollectionInstanceView.GetChildContainerPackage(), IStreamingGenerationErrorHandler::ELevelInstanceInvalidReason::WorldAssetHasInvalidContainer);
 					continue;
 				}
 
 				bool bContainerWasAlreadyInSet;
-				ContainerInstancesStack.Add(SubContainerInstance.Container->ContainerPackageName, &bContainerWasAlreadyInSet);
+				ContainerInstancesStack.Add(SubContainerInstance.ContainerInstance->GetContainerPackage(), &bContainerWasAlreadyInSet);
 
 				if (bContainerWasAlreadyInSet)
 				{
-					ErrorHandler->OnLevelInstanceInvalidWorldAsset(ContainerCollectionInstanceView, ContainerCollectionInstanceView.GetContainerPackage(), IStreamingGenerationErrorHandler::ELevelInstanceInvalidReason::CirculalReference);
+					ErrorHandler->OnLevelInstanceInvalidWorldAsset(ContainerCollectionInstanceView, ContainerCollectionInstanceView.GetChildContainerPackage(), IStreamingGenerationErrorHandler::ELevelInstanceInvalidReason::CirculalReference);
 					continue;
 				}
 
+				FContainerCollectionInstanceDescriptor& ContainerCollectionInstanceDescriptor = ContainerCollectionInstanceDescriptorsMap.FindChecked(ContainerID);
 				FContainerCollectionInstanceDescriptor SubContainerInstanceDescriptor;
-				SubContainerInstanceDescriptor.ID = FActorContainerID(InContainerCollectionInstanceDescriptor.ID, ContainerCollectionInstanceView.GetGuid());
-				SubContainerInstanceDescriptor.ActorDescCollection = MakeShared<FStreamingGenerationActorDescCollection>(FStreamingGenerationActorDescCollection{SubContainerInstance.Container});
-				SubContainerInstanceDescriptor.Transform = SubContainerInstance.Transform * InContainerCollectionInstanceDescriptor.Transform;
-				SubContainerInstanceDescriptor.ParentID = InContainerCollectionInstanceDescriptor.ID;
+
+				SubContainerInstanceDescriptor.ID = SubContainerInstance.ContainerInstance->GetContainerID();
+				check(SubContainerInstanceDescriptor.ID == FActorContainerID(ContainerCollectionInstanceDescriptor.ID, ContainerCollectionInstanceView.GetGuid()));
+
+				SubContainerInstanceDescriptor.ContainerInstanceCollection = MakeShared<FStreamingGenerationContainerInstanceCollection>(FStreamingGenerationContainerInstanceCollection{ SubContainerInstance.ContainerInstance });
+				SubContainerInstanceDescriptor.Transform = SubContainerInstance.ContainerInstance->GetTransform();
+
+				// @todo_pat: this is to validate that new parenting of container instance code is equivalent
+				const FTransform ValidationTransform = SubContainerInstance.Transform * ContainerCollectionInstanceDescriptor.Transform;
+				check(SubContainerInstanceDescriptor.Transform.Equals(ValidationTransform));
+
+				SubContainerInstanceDescriptor.ParentID = ContainerCollectionInstanceDescriptor.ID;
 				SubContainerInstanceDescriptor.OwnerName = *ContainerCollectionInstanceView.GetActorLabelOrName().ToString();
 				// Since Content Bundles streaming generation happens in its own context, all actor set instances must have the same content bundle GUID for now, so Level Instances
 				// placed inside a Content Bundle will propagate their Content Bundle GUID to child instances.
-				SubContainerInstanceDescriptor.ContentBundleID = InContainerCollectionInstanceDescriptor.ContentBundleID;
-				SubContainerInstanceDescriptor.InstanceData = InheritParentContainerPerInstanceData(InContainerCollectionInstanceDescriptor.InstanceData, ContainerCollectionInstanceView);
+				SubContainerInstanceDescriptor.ContentBundleID = ContainerCollectionInstanceDescriptor.ContentBundleID;
+				SubContainerInstanceDescriptor.InstanceData = InheritParentContainerPerInstanceData(ContainerCollectionInstanceDescriptor, ContainerCollectionInstanceView);
 
-				if (WorldPartitionSubsystem && InContainerCollectionInstanceDescriptor.ID.IsMainContainer() && ContainerCollectionInstanceView.GetContainerFilterType() == EWorldPartitionActorFilterType::Loading)
+				if (WorldPartitionSubsystem && ContainerID.IsMainContainer() && ContainerCollectionInstanceView.GetChildContainerFilterType() == EWorldPartitionActorFilterType::Loading)
 				{
-					if (const FWorldPartitionActorFilter* ContainerFilter = ContainerCollectionInstanceView.GetContainerFilter())
-					{						
-						ContainerFilteredActors.Append(WorldPartitionSubsystem->GetFilteredActorsPerContainer(SubContainerInstanceDescriptor.ID, ContainerCollectionInstanceView.GetContainerPackage().ToString(), *ContainerFilter));
+					if (const FWorldPartitionActorFilter* ContainerFilter = ContainerCollectionInstanceView.GetChildContainerFilter())
+					{
+						ContainerFilteredActors.Append(WorldPartitionSubsystem->GetFilteredActorsPerContainer(SubContainerInstanceDescriptor.ID, ContainerCollectionInstanceView.GetChildContainerPackage().ToString(), *ContainerFilter));
 					}
 				}
 
-				CreateActorDescriptorViewsRecursive(SubContainerInstanceDescriptor);
+				CreateActorDescriptorViewsRecursive(MoveTemp(SubContainerInstanceDescriptor));
 
-				verify(ContainerInstancesStack.Remove(SubContainerInstance.Container->ContainerPackageName));
+				verify(ContainerInstancesStack.Remove(SubContainerInstance.ContainerInstance->GetContainerPackage()));
 			}
 		}
 
 		// Fetch the versions stored in the map as it can have been reallocated during recursion
-		FContainerCollectionDescriptor& ContainerCollectionDescriptor = ContainerCollectionDescriptorsMap.FindOrAdd(InContainerCollectionInstanceDescriptor.ActorDescCollection->GetMainContainerPackageName());
-		FContainerCollectionInstanceDescriptor& ContainerCollectionInstanceDescriptor = ContainerCollectionInstanceDescriptorsMap.FindChecked(InContainerCollectionInstanceDescriptor.ID);
+		FContainerCollectionInstanceDescriptor& ContainerCollectionInstanceDescriptor = ContainerCollectionInstanceDescriptorsMap.FindChecked(ContainerID);
 
-		if (!InContainerCollectionInstanceDescriptor.ID.IsMainContainer())
+		if (!ContainerID.IsMainContainer())
 		{
-			FContainerCollectionInstanceDescriptor& ParentContainerCollection = ContainerCollectionInstanceDescriptorsMap.FindChecked(InContainerCollectionInstanceDescriptor.ParentID);
+			FContainerCollectionInstanceDescriptor& ParentContainerCollection = ContainerCollectionInstanceDescriptorsMap.FindChecked(ContainerCollectionInstanceDescriptor.ParentID);
 			ParentContainerCollection.Bounds += ContainerCollectionInstanceDescriptor.Bounds;
 		}
 
 		// Apply per-instance data
-		ContainerCollectionInstanceDescriptor.PerInstanceData.Reserve(ContainerCollectionDescriptor.ActorDescViewMap.Num());
-		ContainerCollectionDescriptor.ActorDescViewMap.ForEachActorDescView([&ContainerCollectionInstanceDescriptor, &InheritParentContainerPerInstanceData](FStreamingGenerationActorDescView& ActorDescView)
+		ContainerCollectionInstanceDescriptor.PerInstanceData.Reserve(ContainerCollectionInstanceDescriptor.ActorDescViewMap->Num());
+		ContainerCollectionInstanceDescriptor.ActorDescViewMap->ForEachActorDescView([&ContainerCollectionInstanceDescriptor, &InheritParentContainerPerInstanceData](FStreamingGenerationActorDescView& ActorDescView)
 		{
-			const FContainerCollectionInstanceDescriptor::FPerInstanceData PerInstanceData = InheritParentContainerPerInstanceData(ContainerCollectionInstanceDescriptor.InstanceData, ActorDescView);
+			const FContainerCollectionInstanceDescriptor::FPerInstanceData PerInstanceData = InheritParentContainerPerInstanceData(ContainerCollectionInstanceDescriptor, ActorDescView);
 
 			if (PerInstanceData != ContainerCollectionInstanceDescriptor.InstanceData)
 			{
@@ -796,74 +881,78 @@ class FWorldPartitionStreamingGenerator
 		ValidateContainerInstanceDescriptor(ContainerCollectionInstanceDescriptor, ContainerCollectionInstanceDescriptor.ID.IsMainContainer());
 	}
 
-	/** 
+	/**
 	 * Creates the actor descriptor views for the specified container.
 	 */
-	void CreateActorContainers(const FStreamingGenerationActorDescCollection& InActorDescCollection)
+	void CreateActorContainers(const FStreamingGenerationContainerInstanceCollection& InContainerInstanceCollection)
 	{
 		TRACE_CPUPROFILER_EVENT_SCOPE(FWorldPartitionStreamingGenerator::CreateActorContainers);
 
 		FContainerCollectionInstanceDescriptor MainContainerCollectionInstance;
-		MainContainerCollectionInstance.ActorDescCollection = MakeShared<FStreamingGenerationActorDescCollection>(InActorDescCollection);
+		MainContainerCollectionInstance.ContainerInstanceCollection = MakeShared<FStreamingGenerationContainerInstanceCollection>(InContainerInstanceCollection);
 		MainContainerCollectionInstance.ClusterMode = EContainerClusterMode::Partitioned;
 		MainContainerCollectionInstance.OwnerName = TEXT("MainContainer");
-		MainContainerCollectionInstance.ContentBundleID = InActorDescCollection.GetContentBundleGuid();
+		MainContainerCollectionInstance.ContentBundleID = InContainerInstanceCollection.GetContentBundleGuid();
 		MainContainerCollectionInstance.InstanceData.bIsSpatiallyLoaded = true; // Since we apply AND logic on spatially loaded flag recursively, startup value must be true
 
-		CreateActorDescriptorViewsRecursive(MainContainerCollectionInstance);
+		// Create Child Containers
+		CreateActorDescriptorViewsRecursive(MoveTemp(MainContainerCollectionInstance));
 	}
 
-	/** 
+	/**
 	 * Creates the actor descriptor container resolver
 	 */
-	void CreateContainerResolver(const FStreamingGenerationActorDescCollection& InActorDescCollection)
+	void CreateContainerResolver(const FStreamingGenerationContainerInstanceCollection& InContainerInstanceCollection)
 	{
-		ContainerResolver.SetMainContainerPackage(InActorDescCollection.GetMainContainerPackageName());
+		ContainerResolver.SetMainContainerPackage(InContainerInstanceCollection.GetMainContainerPackageName());
 
-		for (const auto& [LevelName, ContainerDescriptor] : ContainerCollectionDescriptorsMap)
+		for (const auto& [ContainerID, ContainerDescriptor] : ContainerCollectionInstanceDescriptorsMap)
 		{
-			FWorldPartitionRuntimeContainer& Container = ContainerResolver.AddContainer(ContainerDescriptor.ActorDescCollection->GetMainContainerPackageName());
-
-			for (const FWorldPartitionActorDescView& ActorDescView : ContainerDescriptor.ContainerCollectionInstanceViews)
+			if (!ContainerResolver.ContainsContainer(ContainerDescriptor.ContainerInstanceCollection->GetMainContainerPackageName()))
 			{
-				Container.AddContainerInstance(ActorDescView.GetActorName(), { ActorDescView.GetGuid(), ActorDescView.GetContainerPackage() });
+				FWorldPartitionRuntimeContainer& Container = ContainerResolver.AddContainer(ContainerDescriptor.ContainerInstanceCollection->GetMainContainerPackageName());
+
+				for (const FStreamingGenerationActorDescView& ActorDescView : ContainerDescriptor.ContainerCollectionInstanceViews)
+				{
+					Container.AddContainerInstance(ActorDescView.GetActorName(), { ActorDescView.GetGuid(), ActorDescView.GetChildContainerPackage() });
+				}
 			}
 		}
 
 		ContainerResolver.BuildContainerIDToEditorPathMap();
 	}
 
-	/** 
+	/**
 	 * Perform various validations on the container descriptor, and adjust it based on different requirements. This needs to happen before updating
 	 * containers bounds because some actor descriptor views might change grid placement, etc.
 	 */
-	void ResolveContainerDescriptor(FContainerCollectionDescriptor& ContainerCollectionDescriptor)
+	void ResolveContainerDescriptor(FContainerCollectionInstanceDescriptor& ContainerCollectionInstanceDescriptor)
 	{
-		auto ResolveActorDescView = [this, &ContainerCollectionDescriptor](FStreamingGenerationActorDescView& ActorDescView)
+		auto ResolveActorDescView = [this, &ContainerCollectionInstanceDescriptor](FStreamingGenerationActorDescView& ActorDescView)
 		{
 			ResolveRuntimeSpatiallyLoaded(ActorDescView);
 			ResolveRuntimeGrid(ActorDescView);
-			ResolveRuntimeDataLayers(ActorDescView, ContainerCollectionDescriptor.ActorDescViewMap);
+			ResolveRuntimeDataLayers(ActorDescView, ContainerCollectionInstanceDescriptor.ActorDescViewMap.Get());
 			ResolveHLODLayer(ActorDescView, WorldPartitionContext ? FSoftObjectPath(WorldPartitionContext->GetDefaultHLODLayer()) : FSoftObjectPath());
-			ResolveParentView(ActorDescView, ContainerCollectionDescriptor.ActorDescViewMap);
+			ResolveParentView(ActorDescView, ContainerCollectionInstanceDescriptor.ActorDescViewMap.Get());
 		};
 
-		ContainerCollectionDescriptor.ActorDescViewMap.ForEachActorDescView([this, &ResolveActorDescView](FStreamingGenerationActorDescView& ActorDescView)
+		ContainerCollectionInstanceDescriptor.ActorDescViewMap->ForEachActorDescView([this, &ResolveActorDescView](FStreamingGenerationActorDescView& ActorDescView)
 		{
 			ResolveActorDescView(ActorDescView);
 		});
 
-		for (FStreamingGenerationActorDescView& ContainerCollectionInstanceView : ContainerCollectionDescriptor.ContainerCollectionInstanceViews)
+		for (FStreamingGenerationActorDescView& ContainerCollectionInstanceView : ContainerCollectionInstanceDescriptor.ContainerCollectionInstanceViews)
 		{
 			ResolveActorDescView(ContainerCollectionInstanceView);
 		}
 	}
 
-	/** 
+	/**
 	 * Perform various validations on the container descriptor, and adjust it based on different requirements. This needs to happen before updating
 	 * containers bounds because some actor descriptor views might change grid placement, etc.
 	 */
-	void ValidateContainerDescriptor(FContainerCollectionDescriptor& ContainerCollectionDescriptor, bool bIsMainContainer)
+	void ValidateContainerDescriptor(FContainerCollectionInstanceDescriptor& ContainerCollectionInstanceDescriptor, bool bIsMainContainer)
 	{
 		if (bIsMainContainer)
 		{
@@ -891,12 +980,12 @@ class FWorldPartitionStreamingGenerator
 			}
 			else
 			{
-				ULevel::GetLevelScriptExternalActorsReferencesFromPackage(ContainerCollectionDescriptor.ActorDescCollection->GetMainContainerPackageName(), LevelScriptReferences);
+				ULevel::GetLevelScriptExternalActorsReferencesFromPackage(ContainerCollectionInstanceDescriptor.ContainerInstanceCollection->GetMainContainerPackageName(), LevelScriptReferences);
 			}
 
 			for (const FGuid& LevelScriptReferenceActorGuid : LevelScriptReferences)
 			{
-				if (FStreamingGenerationActorDescView* ActorDescView = (FStreamingGenerationActorDescView*)ContainerCollectionDescriptor.ActorDescViewMap.FindByGuid(LevelScriptReferenceActorGuid))
+				if (FStreamingGenerationActorDescView* ActorDescView = ContainerCollectionInstanceDescriptor.ActorDescViewMap->FindByGuid(LevelScriptReferenceActorGuid))
 				{
 					if (ActorDescView->GetIsSpatiallyLoaded())
 					{
@@ -914,12 +1003,12 @@ class FWorldPartitionStreamingGenerator
 		}
 
 		// Route standard CheckForErrors calls which should not modify actor descriptors in any ways
-		ContainerCollectionDescriptor.ActorDescViewMap.ForEachActorDescView([this](FStreamingGenerationActorDescView& ActorDescView)
+		ContainerCollectionInstanceDescriptor.ActorDescViewMap->ForEachActorDescView([this](FStreamingGenerationActorDescView& ActorDescView)
 		{
 			ActorDescView.CheckForErrors(ErrorHandler);
 		});
 
-		for (FWorldPartitionActorDescView& ContainerCollectionInstanceView : ContainerCollectionDescriptor.ContainerCollectionInstanceViews)
+		for (FStreamingGenerationActorDescView& ContainerCollectionInstanceView : ContainerCollectionInstanceDescriptor.ContainerCollectionInstanceViews)
 		{
 			ContainerCollectionInstanceView.CheckForErrors(ErrorHandler);
 		}
@@ -931,7 +1020,7 @@ class FWorldPartitionStreamingGenerator
 		// 
 		// This works because fixes are deterministic and always apply the same way to both Actors being modified, so there's no ordering issues possible.
 		int32 NbErrorsDetected = INDEX_NONE;
-		for(uint32 NbValidationPasses = 0; NbErrorsDetected; NbValidationPasses++)
+		for (uint32 NbValidationPasses = 0; NbErrorsDetected; NbValidationPasses++)
 		{
 			// Type of work performed in this pass, for clarity
 			enum class EPassType
@@ -943,10 +1032,10 @@ class FWorldPartitionStreamingGenerator
 
 			NbErrorsDetected = 0;
 
-			ContainerCollectionDescriptor.ActorDescViewMap.ForEachActorDescView([this, &ContainerCollectionDescriptor, &NbErrorsDetected, PassType](FStreamingGenerationActorDescView& ActorDescView)
+			ContainerCollectionInstanceDescriptor.ActorDescViewMap->ForEachActorDescView([this, &ContainerCollectionInstanceDescriptor, &NbErrorsDetected, PassType](FStreamingGenerationActorDescView& ActorDescView)
 			{
 				// Validate grid placement
-				auto IsReferenceGridPlacementValid = [](const FWorldPartitionActorDescView& RefererActorDescView, const FWorldPartitionActorDescView& ReferenceActorDescView)
+				auto IsReferenceGridPlacementValid = [](const FStreamingGenerationActorDescView& RefererActorDescView, const FStreamingGenerationActorDescView& ReferenceActorDescView)
 				{
 					const bool bIsActorDescSpatiallyLoaded = RefererActorDescView.GetIsSpatiallyLoaded();
 					const bool bIsActorDescRefSpatiallyLoaded = ReferenceActorDescView.GetIsSpatiallyLoaded();
@@ -954,7 +1043,7 @@ class FWorldPartitionStreamingGenerator
 				};
 
 				// Validate data layers
-				auto IsReferenceDataLayersValid = [](const FWorldPartitionActorDescView& RefererActorDescView, const FWorldPartitionActorDescView& ReferenceActorDescView)
+				auto IsReferenceDataLayersValid = [](const FStreamingGenerationActorDescView& RefererActorDescView, const FStreamingGenerationActorDescView& ReferenceActorDescView)
 				{
 					if (RefererActorDescView.GetRuntimeDataLayerInstanceNames().Num() == ReferenceActorDescView.GetRuntimeDataLayerInstanceNames().Num())
 					{
@@ -968,7 +1057,7 @@ class FWorldPartitionStreamingGenerator
 				};
 
 				// Validate runtime grid references
-				auto IsReferenceRuntimeGridValid = [](const FWorldPartitionActorDescView& RefererActorDescView, const FWorldPartitionActorDescView& ReferenceActorDescView)
+				auto IsReferenceRuntimeGridValid = [](const FStreamingGenerationActorDescView& RefererActorDescView, const FStreamingGenerationActorDescView& ReferenceActorDescView)
 				{
 					return RefererActorDescView.GetRuntimeGrid() == ReferenceActorDescView.GetRuntimeGrid();
 				};
@@ -977,9 +1066,9 @@ class FWorldPartitionStreamingGenerator
 				struct FActorReferenceInfo
 				{
 					FGuid ActorGuid;
-					FStreamingGenerationActorDescView* ActorDescView;
+					FStreamingGenerationActorDescView* ActorDesc;
 					FGuid ReferenceGuid;
-					FStreamingGenerationActorDescView* ReferenceActorDescView;
+					FStreamingGenerationActorDescView* ReferenceActorDesc;
 				};
 
 				TArray<FActorReferenceInfo> References;
@@ -990,10 +1079,10 @@ class FWorldPartitionStreamingGenerator
 					if (ReferenceGuid != ActorDescView.GetParentActor()) // References to the parent are inversed in their handling 
 					{
 						// Filter out parent back references
-						FStreamingGenerationActorDescView* ReferenceActorDesc = (FStreamingGenerationActorDescView*)ContainerCollectionDescriptor.ActorDescViewMap.FindByGuid(ReferenceGuid);
+						FStreamingGenerationActorDescView* ReferenceActorDesc = ContainerCollectionInstanceDescriptor.ActorDescViewMap->FindByGuid(ReferenceGuid);
 						if (!ReferenceActorDesc || (ReferenceActorDesc->GetParentActor() != ActorDescView.GetGuid()))
 						{
-							References.Emplace(FActorReferenceInfo { ActorDescView.GetGuid(), &ActorDescView, ReferenceGuid, ReferenceActorDesc });
+							References.Emplace(FActorReferenceInfo{ ActorDescView.GetGuid(), &ActorDescView, ReferenceGuid, ReferenceActorDesc });
 						}
 					}
 				}
@@ -1006,8 +1095,8 @@ class FWorldPartitionStreamingGenerator
 
 					while (ParentGuid.IsValid())
 					{
-						FStreamingGenerationActorDescView* ParentDescView = (FStreamingGenerationActorDescView*)ContainerCollectionDescriptor.ActorDescViewMap.FindByGuid(ParentGuid);
-					
+						FStreamingGenerationActorDescView* ParentDescView = ContainerCollectionInstanceDescriptor.ActorDescViewMap->FindByGuid(ParentGuid);
+
 						if (ParentDescView)
 						{
 							TopParentDescView = ParentDescView;
@@ -1021,20 +1110,20 @@ class FWorldPartitionStreamingGenerator
 								References.Emplace(FActorReferenceInfo{ ActorDescView.GetGuid(), &ActorDescView, ParentGuid, nullptr });
 							}
 
-							break; 
+							break;
 						}
 					}
 
 					if (TopParentDescView)
 					{
-						References.Emplace(FActorReferenceInfo { TopParentDescView->GetGuid(), TopParentDescView, ActorDescView.GetGuid(), &ActorDescView });
+						References.Emplace(FActorReferenceInfo{ TopParentDescView->GetGuid(), TopParentDescView, ActorDescView.GetGuid(), &ActorDescView });
 					}
 				}
 
 				for (FActorReferenceInfo& Info : References)
 				{
-					FStreamingGenerationActorDescView* RefererActorDescView = Info.ActorDescView;
-					FStreamingGenerationActorDescView* ReferenceActorDescView = Info.ReferenceActorDescView;
+					FStreamingGenerationActorDescView* RefererActorDescView = Info.ActorDesc;
+					FStreamingGenerationActorDescView* ReferenceActorDescView = Info.ReferenceActorDesc;
 
 					if (ReferenceActorDescView)
 					{
@@ -1046,7 +1135,7 @@ class FWorldPartitionStreamingGenerator
 							{
 								if (PassType == EPassType::ErrorReporting)
 								{
-									ErrorHandler->OnInvalidReferenceGridPlacement(*RefererActorDescView, *ReferenceActorDescView);									
+									ErrorHandler->OnInvalidReferenceGridPlacement(*RefererActorDescView, *ReferenceActorDescView);
 								}
 								else
 								{
@@ -1061,7 +1150,7 @@ class FWorldPartitionStreamingGenerator
 							{
 								if (PassType == EPassType::ErrorReporting)
 								{
-									ErrorHandler->OnInvalidReferenceDataLayers(*RefererActorDescView, *ReferenceActorDescView);									
+									ErrorHandler->OnInvalidReferenceDataLayers(*RefererActorDescView, *ReferenceActorDescView);
 								}
 								else
 								{
@@ -1090,23 +1179,26 @@ class FWorldPartitionStreamingGenerator
 					}
 					else
 					{
-						if (!ContainerCollectionDescriptor.EditorOnlyActorDescMap.Contains(Info.ReferenceGuid))
+						if (!ContainerCollectionInstanceDescriptor.EditorOnlyActorDescMap.Contains(Info.ReferenceGuid))
 						{
 							if (PassType == EPassType::ErrorReporting)
 							{
-								FWorldPartitionActorDescView ReferendceActorDescView;
-								FWorldPartitionActorDescView* ReferendceActorDescViewPtr = nullptr;
-
-								if (const UActorDescContainer** ExistingReferenceContainerPtr = ActorGuidsToContainerMap.Find((Info.ReferenceGuid)))
-								{ 
-									if (const FWorldPartitionActorDesc* ReferenceActorDesc = (*ExistingReferenceContainerPtr)->GetActorDesc(Info.ReferenceGuid))
-									{
-										ReferendceActorDescView = FWorldPartitionActorDescView(ReferenceActorDesc);
-										ReferendceActorDescViewPtr = &ReferendceActorDescView;
-									}
+								IWorldPartitionActorDescInstance* ReferencedActorDescInstance = nullptr;
+
+								if (const UActorDescContainerInstance** ExistingReferencedContainerPtr = ActorGuidsToContainerInstanceMap.Find((Info.ReferenceGuid)))
+								{
+									ReferencedActorDescInstance = (*ExistingReferencedContainerPtr)->GetActorDescInstance(Info.ReferenceGuid);
 								}
 
-								ErrorHandler->OnInvalidReference(*RefererActorDescView, Info.ReferenceGuid, ReferendceActorDescViewPtr);
+								if (ReferencedActorDescInstance)
+								{
+									FStreamingGenerationActorDescView InvalidReference(ReferencedActorDescInstance);
+									ErrorHandler->OnInvalidReference(*RefererActorDescView, Info.ReferenceGuid, &InvalidReference);
+								}
+								else
+								{
+									ErrorHandler->OnInvalidReference(*RefererActorDescView, Info.ReferenceGuid, nullptr);
+								}
 
 								NbErrorsDetected++;
 							}
@@ -1117,7 +1209,7 @@ class FWorldPartitionStreamingGenerator
 		}
 
 		// Split runtime and editor references
-		ContainerCollectionDescriptor.ActorDescViewMap.ForEachActorDescView([this, &ContainerCollectionDescriptor](FStreamingGenerationActorDescView& ActorDescView)
+		ContainerCollectionInstanceDescriptor.ActorDescViewMap->ForEachActorDescView([this, &ContainerCollectionInstanceDescriptor](FStreamingGenerationActorDescView& ActorDescView)
 		{
 			TArray<FGuid> RuntimeReferences;
 			TArray<FGuid> EditorReferences;
@@ -1127,7 +1219,7 @@ class FWorldPartitionStreamingGenerator
 
 			for (const FGuid& ReferenceGuid : ActorDescView.GetReferences())
 			{
-				if (FStreamingGenerationActorDescView* ReferenceActorDescView = (FStreamingGenerationActorDescView*)ContainerCollectionDescriptor.ActorDescViewMap.FindByGuid(ReferenceGuid))
+				if (FStreamingGenerationActorDescView* ReferenceActorDescView = ContainerCollectionInstanceDescriptor.ActorDescViewMap->FindByGuid(ReferenceGuid))
 				{
 					// The actor reference is not editor-only, but we are referencing it through an editor-only property
 					if (ActorDescView.IsEditorOnlyReference(ReferenceGuid))
@@ -1139,7 +1231,7 @@ class FWorldPartitionStreamingGenerator
 						RuntimeReferences.Add(ReferenceGuid);
 					}
 				}
-				else if (ContainerCollectionDescriptor.EditorOnlyActorDescMap.Contains(ReferenceGuid))
+				else if (ContainerCollectionInstanceDescriptor.EditorOnlyActorDescMap.Contains(ReferenceGuid))
 				{
 					EditorReferences.Add(ReferenceGuid);
 				}
@@ -1153,14 +1245,12 @@ class FWorldPartitionStreamingGenerator
 		});
 	}
 
-	/** 
+	/**
 	 * Perform various validations on the container descriptor instance, and adjust it based on different requirements. This needs to happen before updating
 	 * containers bounds because some actor descriptor views might change grid placement, etc.
 	 */
 	void ValidateContainerInstanceDescriptor(FContainerCollectionInstanceDescriptor& ContainerCollectionInstanceDescriptor, bool bIsMainContainer)
 	{
-		FContainerCollectionDescriptor& ContainerCollectionDescriptor = ContainerCollectionDescriptorsMap.FindChecked(ContainerCollectionInstanceDescriptor.ActorDescCollection->GetMainContainerPackageName());
-
 		// Perform various adjustements based on validations and report errors
 		//
 		// The first validation pass is used to report errors, subsequent passes are used to make corrections to the actor descriptor views.
@@ -1168,7 +1258,7 @@ class FWorldPartitionStreamingGenerator
 		// 
 		// This works because fixes are deterministic and always apply the same way to both Actors being modified, so there's no ordering issues possible.
 		int32 NbErrorsDetected = INDEX_NONE;
-		for(uint32 NbValidationPasses = 0; NbErrorsDetected; NbValidationPasses++)
+		for (uint32 NbValidationPasses = 0; NbErrorsDetected; NbValidationPasses++)
 		{
 			// Type of work performed in this pass, for clarity
 			enum class EPassType
@@ -1180,12 +1270,12 @@ class FWorldPartitionStreamingGenerator
 
 			NbErrorsDetected = 0;
 
-			ContainerCollectionDescriptor.ActorDescViewMap.ForEachActorDescView([this, &ContainerCollectionDescriptor, &ContainerCollectionInstanceDescriptor, &NbErrorsDetected, PassType](FStreamingGenerationActorDescView& ActorDescView)
+			ContainerCollectionInstanceDescriptor.ActorDescViewMap->ForEachActorDescView([this, &ContainerCollectionInstanceDescriptor, &NbErrorsDetected, PassType](FStreamingGenerationActorDescView& ActorDescView)
 			{
 				FContainerCollectionInstanceDescriptor::FPerInstanceData& PerInstanceData = ContainerCollectionInstanceDescriptor.GetInstanceData(ActorDescView.GetGuid());
 
 PRAGMA_DISABLE_DEPRECATION_WARNINGS
-				if (ActorDescView.ShouldValidateRuntimeGrid())
+				if (ActorDescView.GetActorDesc()->ShouldValidateRuntimeGrid())
 PRAGMA_ENABLE_DEPRECATION_WARNINGS
 				{
 					if (!IsValidGrid(PerInstanceData.RuntimeGrid))
@@ -1201,7 +1291,7 @@ PRAGMA_ENABLE_DEPRECATION_WARNINGS
 
 						NbErrorsDetected++;
 					}
-					
+
 					if (ActorDescView.GetHLODLayer().IsValid() && !IsValidHLODLayer(PerInstanceData.RuntimeGrid, ActorDescView.GetHLODLayer()))
 					{
 						if (PassType == EPassType::ErrorReporting)
@@ -1227,7 +1317,7 @@ PRAGMA_ENABLE_DEPRECATION_WARNINGS
 						{
 							if (PassType == EPassType::ErrorReporting)
 							{
-								const FWorldPartitionActorDescView& ReferenceActorDesc = ContainerCollectionDescriptor.ActorDescViewMap.FindByGuidChecked(ReferenceGuid);
+								const FStreamingGenerationActorDescView& ReferenceActorDesc = ContainerCollectionInstanceDescriptor.ActorDescViewMap->FindByGuidChecked(ReferenceGuid);
 								ErrorHandler->OnInvalidActorFilterReference(ActorDescView, ReferenceActorDesc);
 							}
 							else
@@ -1243,15 +1333,15 @@ PRAGMA_ENABLE_DEPRECATION_WARNINGS
 		}
 	}
 
-	/** 
+	/**
 	 * Update the container descriptor containers to adjust their bounds from actor descriptor views.
 	 */
-	void UpdateContainerDescriptor(FContainerCollectionDescriptor& ContainerCollectionDescriptor)
+	void UpdateContainerDescriptor(FContainerCollectionInstanceDescriptor& ContainerCollectionInstanceDescriptor)
 	{
 		// Build clusters for this container - at this point, all actors references should be in the same data layers, grid, etc because of actor descriptors validation.
 		TArray<TPair<FGuid, TArray<FGuid>>> ActorsWithRefs;
-		ContainerCollectionDescriptor.ActorDescViewMap.ForEachActorDescView([&ActorsWithRefs](const FStreamingGenerationActorDescView& ActorDescView) { ActorsWithRefs.Emplace(ActorDescView.GetGuid(), ActorDescView.GetReferences()); });
-		ContainerCollectionDescriptor.Clusters = GenerateObjectsClusters(ActorsWithRefs);
+		ContainerCollectionInstanceDescriptor.ActorDescViewMap->ForEachActorDescView([&ActorsWithRefs](const FStreamingGenerationActorDescView& ActorDescView) { ActorsWithRefs.Emplace(ActorDescView.GetGuid(), ActorDescView.GetReferences()); });
+		ContainerCollectionInstanceDescriptor.Clusters = GenerateObjectsClusters(ActorsWithRefs);
 	}
 
 	/**
@@ -1262,15 +1352,13 @@ PRAGMA_ENABLE_DEPRECATION_WARNINGS
 		// Validate that all actors part of the same actor set share the same set of values
 		for (const auto& [ContainerID, ContainerCollectionInstanceDescriptor] : ContainerCollectionInstanceDescriptorsMap)
 		{
-			const FContainerCollectionDescriptor& ContainerCollectionDescriptor = ContainerCollectionDescriptorsMap.FindChecked(ContainerCollectionInstanceDescriptor.ActorDescCollection->GetMainContainerPackageName());
-					
-			for (const TArray<FGuid>& Cluster : ContainerCollectionDescriptor.Clusters)
+			for (const TArray<FGuid>& Cluster : ContainerCollectionInstanceDescriptor.Clusters)
 			{
-				const FWorldPartitionActorDescView& ReferenceActorDescView = ContainerCollectionDescriptor.ActorDescViewMap.FindByGuidChecked(Cluster[0]);
+				const FStreamingGenerationActorDescView& ReferenceActorDescView = ContainerCollectionInstanceDescriptor.ActorDescViewMap->FindByGuidChecked(Cluster[0]);
 
 				for (const FGuid& ActorGuid : Cluster)
 				{
-					const FWorldPartitionActorDescView& ActorDescView = ContainerCollectionDescriptor.ActorDescViewMap.FindByGuidChecked(ActorGuid);
+					const FStreamingGenerationActorDescView& ActorDescView = ContainerCollectionInstanceDescriptor.ActorDescViewMap->FindByGuidChecked(ActorGuid);
 					check(ActorDescView.GetRuntimeGrid() == ReferenceActorDescView.GetRuntimeGrid());
 					check(ActorDescView.GetIsSpatiallyLoaded() == ReferenceActorDescView.GetIsSpatiallyLoaded());
 					check(ActorDescView.GetContentBundleGuid() == ReferenceActorDescView.GetContentBundleGuid());
@@ -1284,32 +1372,32 @@ public:
 	{
 		FWorldPartitionStreamingGeneratorParams()
 			: WorldPartitionContext(nullptr)
-			, ModifiedActorsDescList(nullptr)
 			, ErrorHandler(&NullErrorHandler)
 			, bEnableStreaming(false)
 			, bCreateContainerResolver(false)
+			, bHandleUnsavedActors(false)
 		{}
 
 		const UWorldPartition* WorldPartitionContext;
-		FActorDescList* ModifiedActorsDescList;
 		IStreamingGenerationErrorHandler* ErrorHandler;
 		bool bEnableStreaming;
 		bool bCreateContainerResolver;
-		TMap<FGuid, const UActorDescContainer*> ActorGuidsToContainerMap;
+		bool bHandleUnsavedActors;
+		TMap<FGuid, const UActorDescContainerInstance*> ActorGuidsToContainerInstanceMap;
 		TArray<TSubclassOf<AActor>> FilteredClasses;
 		TFunction<bool(FName)> IsValidGrid;
 		TFunction<bool(FName GridName, const FSoftObjectPath&)> IsValidHLODLayer;
 
 		FWorldPartitionStreamingGeneratorParams& SetWorldPartitionContext(const UWorldPartition* InWorldPartitionContext) { WorldPartitionContext = InWorldPartitionContext; return *this; }
-		FWorldPartitionStreamingGeneratorParams& SetModifiedActorsDescList(FActorDescList* InModifiedActorsDescList) { ModifiedActorsDescList = InModifiedActorsDescList; return *this; }
 		FWorldPartitionStreamingGeneratorParams& SetErrorHandler(IStreamingGenerationErrorHandler* InErrorHandler) { ErrorHandler = InErrorHandler; return *this; }
 		FWorldPartitionStreamingGeneratorParams& SetEnableStreaming(bool bInEnableStreaming) { bEnableStreaming = bInEnableStreaming; return *this; }
 		FWorldPartitionStreamingGeneratorParams& SetCreateContainerResolver(bool bInCreateContainerResolver) { bCreateContainerResolver = bInCreateContainerResolver; return *this; }
-		FWorldPartitionStreamingGeneratorParams& SetActorGuidsToContainerMap(const TMap<FGuid, const UActorDescContainer*>& InActorGuidsToContainerMap) { ActorGuidsToContainerMap = InActorGuidsToContainerMap; return *this; }
+		FWorldPartitionStreamingGeneratorParams& SetHandleUnsavedActors(bool bInHandleUnsavedActors) { bHandleUnsavedActors = bInHandleUnsavedActors; return *this; }
+		FWorldPartitionStreamingGeneratorParams& SetActorGuidsToContainerInstanceMap(const TMap<FGuid, const UActorDescContainerInstance*>& InActorGuidsToContainerInstanceMap) { ActorGuidsToContainerInstanceMap = InActorGuidsToContainerInstanceMap; return *this; }
 		FWorldPartitionStreamingGeneratorParams& SetFilteredClasses(const TArray<TSubclassOf<AActor>>& InFilteredClasses) { FilteredClasses = InFilteredClasses; return *this; }
 		FWorldPartitionStreamingGeneratorParams& SetIsValidGrid(TFunction<bool(FName)> InIsValidGrid) { IsValidGrid = InIsValidGrid; return *this; }
 		FWorldPartitionStreamingGeneratorParams& SetIsValidHLODLayer(TFunction<bool(FName GridName, const FSoftObjectPath&)> InIsValidHLODLayer) { IsValidHLODLayer = InIsValidHLODLayer; return *this; }
-
+				
 		inline static FStreamingGenerationNullErrorHandler NullErrorHandler;
 	};
 
@@ -1317,25 +1405,26 @@ public:
 		: WorldPartitionContext(Params.WorldPartitionContext)
 		, bEnableStreaming(Params.bEnableStreaming)
 		, bCreateContainerResolver(Params.bCreateContainerResolver)
-		, ModifiedActorsDescList(Params.ModifiedActorsDescList)
+		, bHandleUnsavedActors(Params.bHandleUnsavedActors)
 		, FilteredClasses(Params.FilteredClasses)
 		, IsValidGrid(Params.IsValidGrid)
 		, IsValidHLODLayer(Params.IsValidHLODLayer)
 		, ErrorHandler(Params.ErrorHandler)
-		, ActorGuidsToContainerMap(Params.ActorGuidsToContainerMap)
+		, ActorGuidsToContainerInstanceMap(Params.ActorGuidsToContainerInstanceMap)
 	{
 		UWorld* OwningWorld = WorldPartitionContext ? WorldPartitionContext->GetWorld() : nullptr;
 		WorldPartitionSubsystem = OwningWorld ? UWorld::GetSubsystem<UWorldPartitionSubsystem>(OwningWorld) : nullptr;
 		DataLayerManager = OwningWorld ? WorldPartitionContext->GetDataLayerManager() : nullptr;
 	}
 
-	void PreparationPhase(const FStreamingGenerationActorDescCollection& ActorDescCollection)
+	
+	void PreparationPhase(const FStreamingGenerationContainerInstanceCollection& ContainerInstanceCollection)
 	{
-		CreateActorContainers(ActorDescCollection);
+		CreateActorContainers(ContainerInstanceCollection);
 
 		if (bCreateContainerResolver)
 		{
-			CreateContainerResolver(ActorDescCollection);
+			CreateContainerResolver(ContainerInstanceCollection);
 		}
 
 		ValidateInternalState();
@@ -1346,7 +1435,7 @@ public:
 		FString StateLogOutputFilename = FPaths::ProjectLogDir() / TEXT("WorldPartition") / FString::Printf(TEXT("StreamingGeneration-%s"), Suffix);
 
 		if (bTimeStamped)
-		{			
+		{
 			StateLogOutputFilename += FString::Printf(TEXT("-%08x-%s"), FPlatformProcess::GetCurrentProcessId(), *FDateTime::Now().ToIso8601().Replace(TEXT(":"), TEXT(".")));
 		}
 
@@ -1368,29 +1457,39 @@ public:
 			}
 		}
 
+		TSet<FString> UniqueContainerNames;
+
 		UE_SCOPED_INDENT_LOG_ARCHIVE(Ar.PrintfIndent(TEXT("Containers:")));
-		for (auto& [LevelName, ContainerCollectionDescriptor] : ContainerCollectionDescriptorsMap)
+		for (auto& [ContainerID, ContainerCollectionInstanceDescriptor] : ContainerCollectionInstanceDescriptorsMap)
 		{
-			UE_SCOPED_INDENT_LOG_ARCHIVE(Ar.PrintfIndent(TEXT("Container: %s"), *ContainerCollectionDescriptor.ActorDescCollection->GetMainContainerPackageName().ToString()));
+			FString ContainerPackageName = ContainerCollectionInstanceDescriptor.ContainerInstanceCollection->GetMainContainerPackageName().ToString();
+			bool bIsAlreadySet = false;
+			UniqueContainerNames.Add(ContainerPackageName, &bIsAlreadySet);
+			if (bIsAlreadySet)
+			{
+				continue;
+			}
+
+			UE_SCOPED_INDENT_LOG_ARCHIVE(Ar.PrintfIndent(TEXT("Container: %s"), *ContainerPackageName));
 
-			if (ContainerCollectionDescriptor.ActorDescViewMap.Num())
+			if (ContainerCollectionInstanceDescriptor.ActorDescViewMap->Num())
 			{
 				UE_SCOPED_INDENT_LOG_ARCHIVE(Ar.PrintfIndent(TEXT("ActorDescs:")));
 
-				TMap<FGuid, FStreamingGenerationActorDescView*> SortedActorDescViewMap = ContainerCollectionDescriptor.ActorDescViewMap.ActorDescViewsByGuid;
+				TMap<FGuid, FStreamingGenerationActorDescView*> SortedActorDescViewMap = ContainerCollectionInstanceDescriptor.ActorDescViewMap->ActorDescViewsByGuid;
 				SortedActorDescViewMap.KeySort([](const FGuid& GuidA, const FGuid& GuidB) { return GuidA < GuidB; });
 
 				for (auto& [ActorGuid, ActorDescView] : SortedActorDescViewMap)
 				{
-					Ar.Print(*ActorDescView->ToString());
+					Ar.Print(*ActorDescView->ToString(FWorldPartitionActorDesc::EToStringMode::Compact));
 				}
 			}
 
-			if (ContainerCollectionDescriptor.Clusters.Num())
+			if (ContainerCollectionInstanceDescriptor.Clusters.Num())
 			{
 				UE_SCOPED_INDENT_LOG_ARCHIVE(Ar.PrintfIndent(TEXT("Clusters:")));
 
-				TArray<TArray<FGuid>> SortedClusters = ContainerCollectionDescriptor.Clusters;
+				TArray<TArray<FGuid>> SortedClusters = ContainerCollectionInstanceDescriptor.Clusters;
 				for (TArray<FGuid>& ActorGuids : SortedClusters)
 				{
 					ActorGuids.Sort([](const FGuid& GuidA, const FGuid& GuidB) { return GuidA < GuidB; });
@@ -1403,8 +1502,8 @@ public:
 					UE_SCOPED_INDENT_LOG_ARCHIVE(Ar.PrintfIndent(TEXT("[%3d]"), ClusterIndex++));
 					for (const FGuid& ActorGuid : ActorGuids)
 					{
-						const FWorldPartitionActorDescView& ActorDescView = ContainerCollectionDescriptor.ActorDescViewMap.FindByGuidChecked(ActorGuid);
-						Ar.Print(*ActorDescView.ToString());
+						const FStreamingGenerationActorDescView& ActorDescView = ContainerCollectionInstanceDescriptor.ActorDescViewMap->FindByGuidChecked(ActorGuid);
+						Ar.Print(*ActorDescView.ToString(FWorldPartitionActorDesc::EToStringMode::Compact));
 					}
 				}
 			}
@@ -1416,7 +1515,7 @@ public:
 			auto DumpContainerInstancesRecursive = [this, &InvertedContainersHierarchy, &Ar](const FActorContainerID& ContainerID, auto& RecursiveFunc) -> void
 			{
 				const FContainerCollectionInstanceDescriptor& ContainerInstanceDescriptor = ContainerCollectionInstanceDescriptorsMap.FindChecked(ContainerID);
-				
+
 				{
 					UE_SCOPED_INDENT_LOG_ARCHIVE(Ar.PrintfIndent(TEXT("%s:"), *ContainerInstanceDescriptor.OwnerName));
 
@@ -1424,7 +1523,7 @@ public:
 					Ar.Printf(TEXT(" ParentID: %s"), *ContainerInstanceDescriptor.ParentID.ToString());
 					Ar.Printf(TEXT("   Bounds: %s"), *ContainerInstanceDescriptor.Bounds.ToString());
 					Ar.Printf(TEXT("Transform: %s"), *ContainerInstanceDescriptor.Transform.ToString());
-					Ar.Printf(TEXT("Container: %s"), *ContainerInstanceDescriptor.ActorDescCollection->GetMainContainerPackageName().ToString());
+					Ar.Printf(TEXT("Container: %s"), *ContainerInstanceDescriptor.ContainerInstanceCollection->GetMainContainerPackageName().ToString());
 				}
 
 				TArray<FActorContainerID> ChildContainersIDs;
@@ -1434,7 +1533,7 @@ public:
 				if (ChildContainersIDs.Num())
 				{
 					UE_SCOPED_INDENT_LOG_ARCHIVE(Ar.PrintfIndent(TEXT("SubContainers:")));
-						
+
 					for (const FActorContainerID& ChildContainerID : ChildContainersIDs)
 					{
 						RecursiveFunc(ChildContainerID, RecursiveFunc);
@@ -1448,12 +1547,12 @@ public:
 		DumpContainerInstances(FActorContainerID());
 	}
 
-	const FStreamingGenerationContext* GetStreamingGenerationContext(const FStreamingGenerationActorDescCollection& ActorDescCollection)
+	const FStreamingGenerationContext* GetStreamingGenerationContext(const FStreamingGenerationContainerInstanceCollection& ContainerInstanceCollection)
 	{
 		if (!StreamingGenerationContext)
 		{
 			// Construct the streaming generation context
-			StreamingGenerationContext = MakeUnique<FStreamingGenerationContext>(this, ActorDescCollection);
+			StreamingGenerationContext = MakeUnique<FStreamingGenerationContext>(this, ContainerInstanceCollection);
 		}
 
 		return StreamingGenerationContext.Get();
@@ -1475,24 +1574,21 @@ private:
 	const UDataLayerManager* DataLayerManager;
 	bool bEnableStreaming;
 	bool bCreateContainerResolver;
-	FActorDescList* ModifiedActorsDescList;
+	bool bHandleUnsavedActors;
 	TArray<TSubclassOf<AActor>> FilteredClasses;
 	TFunction<bool(FName)> IsValidGrid;
 	TFunction<bool(FName, const FSoftObjectPath&)> IsValidHLODLayer;
 	IStreamingGenerationErrorHandler* ErrorHandler;
 	FWorldPartitionRuntimeContainerResolver ContainerResolver;
 
-	// Maps a level to its descriptor
-	TMap<FName, FContainerCollectionDescriptor> ContainerCollectionDescriptorsMap;
-	
 	/** Maps containers IDs to their container collection instance descriptor */
 	TMap<FActorContainerID, FContainerCollectionInstanceDescriptor> ContainerCollectionInstanceDescriptorsMap;
 
 	/** Data required for streaming generation interface */
 	TUniquePtr<FStreamingGenerationContext> StreamingGenerationContext;
 
-	/** List of containers participating in this streaming generation step */
-	TMap<FGuid, const UActorDescContainer*> ActorGuidsToContainerMap;
+	/** List of container instances participating in this streaming generation step */
+	TMap<FGuid, const UActorDescContainerInstance*> ActorGuidsToContainerInstanceMap;
 
 	/** List of current container instances on the stack to detect circular references */
 	TSet<FName> ContainerInstancesStack;
@@ -1501,24 +1597,10 @@ private:
 	TMap<FActorContainerID, TSet<FGuid>> ContainerFilteredActors;
 };
 
-bool UWorldPartition::GenerateStreaming(TArray<FString>* OutPackagesToGenerate)
-{
-	FGenerateStreamingParams Params = FGenerateStreamingParams().SetActorDescContainer(nullptr);
-	FGenerateStreamingContext Context = FGenerateStreamingContext().SetPackagesToGenerate(OutPackagesToGenerate);
-	return GenerateStreaming(Params, Context);
-}
-
-bool UWorldPartition::GenerateContainerStreaming(const UActorDescContainer* InActorDescContainer, TArray<FString>* OutPackagesToGenerate)
-{
-	FGenerateStreamingParams Params = FGenerateStreamingParams().SetActorDescContainer(InActorDescContainer);
-	FGenerateStreamingContext Context = FGenerateStreamingContext().SetPackagesToGenerate(OutPackagesToGenerate);
-	return GenerateContainerStreaming(Params, Context);
-}
-
 bool UWorldPartition::GenerateStreaming(const FGenerateStreamingParams& InParams, FGenerateStreamingContext& InContext)
 {
-	check(InParams.ActorDescCollection.IsEmpty());
-	FGenerateStreamingParams Params = FGenerateStreamingParams(InParams).SetActorDescContainer(ActorDescContainer);
+	check(InParams.ContainerInstanceCollection.IsEmpty());
+	FGenerateStreamingParams Params = FGenerateStreamingParams(InParams).SetActorDescContainerInstance(ActorDescContainerInstance);
 
 	OnPreGenerateStreaming.Broadcast(InContext.PackagesToGenerate);
 
@@ -1529,14 +1611,7 @@ bool UWorldPartition::GenerateContainerStreaming(const FGenerateStreamingParams&
 {
 	TRACE_CPUPROFILER_EVENT_SCOPE(UWorldPartition::GenerateContainerStreaming);
 
-	FActorDescList* ModifiedActorsDescList = nullptr;
-
-	if (bIsPIE)
-	{
-		ModifiedActorsDescList = &RuntimeHash->ModifiedActorDescListForPIE;
-	}
-
-	const FString ContainerPackageName = InParams.ActorDescCollection.GetMainContainerPackageName().ToString();
+	const FString ContainerPackageName = InParams.ContainerInstanceCollection.GetMainContainerPackageName().ToString();
 	FString ContainerShortName = FPackageName::GetShortName(ContainerPackageName);
 	if (!ContainerPackageName.StartsWith(TEXT("/Game/")))
 	{
@@ -1564,7 +1639,7 @@ bool UWorldPartition::GenerateContainerStreaming(const FGenerateStreamingParams&
 		StateLogSuffix += TEXT("_");
 		StateLogSuffix += ContainerShortName;
 		LogFileAr = FWorldPartitionStreamingGenerator::CreateDumpStateLogArchive(*StateLogSuffix, !InParams.OutputLogPath);
-		
+
 		if (LogFileAr.IsValid())
 		{
 			InContext.OutputLogFilename = LogFileAr->GetArchiveName();
@@ -1576,7 +1651,7 @@ bool UWorldPartition::GenerateContainerStreaming(const FGenerateStreamingParams&
 
 	FWorldPartitionStreamingGenerator::FWorldPartitionStreamingGeneratorParams StreamingGeneratorParams = FWorldPartitionStreamingGenerator::FWorldPartitionStreamingGeneratorParams()
 		.SetWorldPartitionContext(this)
-		.SetModifiedActorsDescList( ModifiedActorsDescList)
+		.SetHandleUnsavedActors(bIsPIE)
 		.SetIsValidGrid([this](FName GridName) { return RuntimeHash->IsValidGrid(GridName); })
 		.SetIsValidHLODLayer([this](FName GridName, const FSoftObjectPath& HLODLayerPath) { return RuntimeHash->IsValidHLODLayer(GridName, HLODLayerPath); })
 		.SetErrorHandler(ErrorHandlerSelector.Get())
@@ -1586,7 +1661,7 @@ bool UWorldPartition::GenerateContainerStreaming(const FGenerateStreamingParams&
 	FWorldPartitionStreamingGenerator StreamingGenerator(StreamingGeneratorParams);
 
 	// Preparation Phase
-	StreamingGenerator.PreparationPhase(InParams.ActorDescCollection);
+	StreamingGenerator.PreparationPhase(InParams.ContainerInstanceCollection);
 
 	if (HierarchicalLogAr.IsValid())
 	{
@@ -1598,7 +1673,7 @@ bool UWorldPartition::GenerateContainerStreaming(const FGenerateStreamingParams&
 	StreamingPolicy = NewObject<UWorldPartitionStreamingPolicy>(const_cast<UWorldPartition*>(this), WorldPartitionStreamingPolicyClass.Get(), NAME_None, bIsPIE ? RF_Transient : RF_NoFlags);
 
 	check(RuntimeHash);
-	const IStreamingGenerationContext* StreamingGenerationContext = StreamingGenerator.GetStreamingGenerationContext(InParams.ActorDescCollection);
+	const IStreamingGenerationContext* StreamingGenerationContext = StreamingGenerator.GetStreamingGenerationContext(InParams.ContainerInstanceCollection);
 	check(StreamingGenerationContext);
 	if (RuntimeHash->GenerateStreaming(StreamingPolicy, StreamingGenerationContext, InContext.PackagesToGenerate))
 	{
@@ -1647,14 +1722,14 @@ TUniquePtr<IStreamingGenerationContext> UWorldPartition::GenerateStreamingGenera
 		.SetCreateContainerResolver(FEditorPathHelper::IsEnabled());
 
 	TUniquePtr<FStreamingGenerationContextCopy> StreamingGenerationContextCopy = MakeUnique<FStreamingGenerationContextCopy>(StreamingGeneratorParams);
-	
-	StreamingGenerationContextCopy->StreamingGenerator.PreparationPhase(InParams.ActorDescCollection);
-	
-	const IStreamingGenerationContext* StreamingGenerationContext = StreamingGenerationContextCopy->StreamingGenerator.GetStreamingGenerationContext(InParams.ActorDescCollection);	
+
+	StreamingGenerationContextCopy->StreamingGenerator.PreparationPhase(InParams.ContainerInstanceCollection);
+
+	const IStreamingGenerationContext* StreamingGenerationContext = StreamingGenerationContextCopy->StreamingGenerator.GetStreamingGenerationContext(InParams.ContainerInstanceCollection);
 	check(StreamingGenerationContext);
 
 	StreamingGenerationContextCopy->SetSourceContext(StreamingGenerationContext);
-	
+
 	return MoveTemp(StreamingGenerationContextCopy);
 }
 
@@ -1678,7 +1753,7 @@ URuntimeHashExternalStreamingObjectBase* UWorldPartition::FlushStreamingToExtern
 
 void UWorldPartition::SetupHLODActors(const FSetupHLODActorsParams& Params)
 {
-	ForEachActorDescContainer([this, &Params](UActorDescContainer* InActorDescContainer)
+	ForEachActorDescContainer([this, &Params](UActorDescContainerInstance* InContainerInstance)
 	{
 		TErrorHandlerSelector<FStreamingGenerationLogErrorHandler> ErrorHandlerSelector;
 
@@ -1691,7 +1766,7 @@ void UWorldPartition::SetupHLODActors(const FSetupHLODActorsParams& Params)
 			.SetIsValidHLODLayer([this](FName GridName, const FSoftObjectPath& HLODLayerPath) { return RuntimeHash->IsValidHLODLayer(GridName, HLODLayerPath); });
 
 		FWorldPartitionStreamingGenerator StreamingGenerator(StreamingGeneratorParams);
-		FStreamingGenerationActorDescCollection ActorDescCollection{InActorDescContainer};
+		FStreamingGenerationContainerInstanceCollection ActorDescCollection{ InContainerInstance };
 		StreamingGenerator.PreparationPhase(ActorDescCollection);
 
 		TUniquePtr<FArchive> LogFileAr = FWorldPartitionStreamingGenerator::CreateDumpStateLogArchive(TEXT("HLOD"));
@@ -1703,24 +1778,24 @@ void UWorldPartition::SetupHLODActors(const FSetupHLODActorsParams& Params)
 		}
 
 		RuntimeHash->SetupHLODActors(StreamingGenerator.GetStreamingGenerationContext(ActorDescCollection), Params);
-	});	
+	});
 }
 
-FStreamingGenerationActorDescCollection::FStreamingGenerationActorDescCollection(std::initializer_list<TObjectPtr<const UActorDescContainer>> ActorDescContainerArray)
-	: TActorDescContainerCollection<TObjectPtr<const UActorDescContainer>>(ActorDescContainerArray)
+FStreamingGenerationContainerInstanceCollection::FStreamingGenerationContainerInstanceCollection(std::initializer_list<TObjectPtr<const UActorDescContainerInstance>> ActorDescContainerArray)
+	: TActorDescContainerInstanceCollection<TObjectPtr<const UActorDescContainerInstance>>(ActorDescContainerArray)
 {
 	SortCollection();
 }
 
-FStreamingGenerationActorDescCollection::FStreamingGenerationActorDescCollection(const TArray<const UActorDescContainer*>& ActorDescContainers)
-	: TActorDescContainerCollection<TObjectPtr<const UActorDescContainer>>(ActorDescContainers)
+FStreamingGenerationContainerInstanceCollection::FStreamingGenerationContainerInstanceCollection(const TArray<const UActorDescContainerInstance*>& ActorDescContainers)
+	: TActorDescContainerInstanceCollection<TObjectPtr<const UActorDescContainerInstance>>(ActorDescContainers)
 {
 	SortCollection();
 }
 
-UWorld* FStreamingGenerationActorDescCollection::GetWorld() const
+UWorld* FStreamingGenerationContainerInstanceCollection::GetWorld() const
 {
-	if (const UActorDescContainer* MainContainer = GetMainActorDescContainer())
+	if (const UActorDescContainerInstance* MainContainer = GetMainContainer())
 	{
 		if (UWorldPartition* WorldPartition = MainContainer->GetWorldPartition())
 		{
@@ -1731,9 +1806,9 @@ UWorld* FStreamingGenerationActorDescCollection::GetWorld() const
 	return nullptr;
 }
 
-FGuid FStreamingGenerationActorDescCollection::GetContentBundleGuid() const
+FGuid FStreamingGenerationContainerInstanceCollection::GetContentBundleGuid() const
 {
-	if (const UActorDescContainer* MainContainer = GetMainActorDescContainer())
+	if (const UActorDescContainerInstance* MainContainer = GetMainContainer())
 	{
 		return MainContainer->GetContentBundleGuid();
 	}
@@ -1741,19 +1816,19 @@ FGuid FStreamingGenerationActorDescCollection::GetContentBundleGuid() const
 	return FGuid();
 }
 
-const UActorDescContainer* FStreamingGenerationActorDescCollection::GetMainActorDescContainer() const
+const UActorDescContainerInstance* FStreamingGenerationContainerInstanceCollection::GetMainContainer() const
 {
 	if (!IsEmpty())
 	{
-		return ActorDescContainerCollection[MainContainerIdx];
+		return ActorDescContainerInstanceCollection[MainContainerIdx];
 	}
 
 	return nullptr;
 }
 
-FName FStreamingGenerationActorDescCollection::GetMainContainerPackageName() const
+FName FStreamingGenerationContainerInstanceCollection::GetMainContainerPackageName() const
 {
-	if (const UActorDescContainer* MainContainer = GetMainActorDescContainer())
+	if (const UActorDescContainerInstance* MainContainer = GetMainContainer())
 	{
 		return MainContainer->GetContainerPackage();
 	}
@@ -1761,36 +1836,36 @@ FName FStreamingGenerationActorDescCollection::GetMainContainerPackageName() con
 	return FName();
 }
 
-TArrayView<const UActorDescContainer* const> FStreamingGenerationActorDescCollection::GetExternalDataLayerContainers()
+TArrayView<const UActorDescContainerInstance* const> FStreamingGenerationContainerInstanceCollection::GetExternalDataLayerContainers()
 {
-	if (ActorDescContainerCollection.Num() <= 1)
+	if (ActorDescContainerInstanceCollection.Num() <= 1)
 	{
-		return TArrayView<const UActorDescContainer*>();
+		return TArrayView<const UActorDescContainerInstance*>();
 	}
 
-	uint32 NumExternalDataLayerContainers = ActorDescContainerCollection.Num() - ExternalDataLayerContainerStartIdx;
-	return MakeArrayView(&ActorDescContainerCollection[ExternalDataLayerContainerStartIdx], NumExternalDataLayerContainers);
+	uint32 NumExternalDataLayerContainers = ActorDescContainerInstanceCollection.Num() - ExternalDataLayerContainerStartIdx;
+	return MakeArrayView(&ActorDescContainerInstanceCollection[ExternalDataLayerContainerStartIdx], NumExternalDataLayerContainers);
 }
 
-void FStreamingGenerationActorDescCollection::OnCollectionChanged()
+void FStreamingGenerationContainerInstanceCollection::OnCollectionChanged()
 {
 	SortCollection();
 }
 
-void FStreamingGenerationActorDescCollection::SortCollection()
+void FStreamingGenerationContainerInstanceCollection::SortCollection()
 {
 	if (IsEmpty())
 	{
 		return;
 	}
 
-	auto IsMainPartitionContainer = [this](const UActorDescContainer* ActorDescContainer)
+	auto IsMainPartitionContainer = [this](const UActorDescContainerInstance* ActorDescContainerInstance)
 	{
-		bool bIsContentBundleContainer = ActorDescContainer->GetContentBundleGuid().IsValid() && GetActorDescContainerCount() == 1 && !ActorDescContainer->GetContainerPackage().ToString().StartsWith(TEXT("/Game/"), ESearchCase::IgnoreCase);
-		return !ActorDescContainer->GetContentBundleGuid().IsValid() || bIsContentBundleContainer;
+		bool bIsContentBundleContainer = ActorDescContainerInstance->GetContentBundleGuid().IsValid() && GetActorDescContainerCount() == 1 && !ActorDescContainerInstance->GetContainerPackage().ToString().StartsWith(TEXT("/Game/"), ESearchCase::IgnoreCase);
+		return !ActorDescContainerInstance->GetContentBundleGuid().IsValid() || bIsContentBundleContainer;
 	};
 
-	Algo::Sort(ActorDescContainerCollection, [&IsMainPartitionContainer](const UActorDescContainer* A, const UActorDescContainer* B)
+	Algo::Sort(ActorDescContainerInstanceCollection, [&IsMainPartitionContainer](const UActorDescContainerInstance* A, const UActorDescContainerInstance* B)
 	{
 		if (IsMainPartitionContainer(A) && !IsMainPartitionContainer(B))
 		{
@@ -1804,88 +1879,72 @@ void FStreamingGenerationActorDescCollection::SortCollection()
 		return A->GetContainerPackage().Compare(B->GetContainerPackage()) < 0;
 	});
 
-	check(IsMainPartitionContainer(ActorDescContainerCollection[MainContainerIdx]));
+	check(IsMainPartitionContainer(ActorDescContainerInstanceCollection[MainContainerIdx]));
 
 #if DO_CHECK
-	TArrayView<const UActorDescContainer* const> ExternalDataLayerContainerView = GetExternalDataLayerContainers();
-	Algo::ForEach(ExternalDataLayerContainerView, [&IsMainPartitionContainer](const UActorDescContainer* ActorDescContainer) { check(!IsMainPartitionContainer(ActorDescContainer)); });
+	TArrayView<const UActorDescContainerInstance* const> ExternalDataLayerContainerView = GetExternalDataLayerContainers();
+	Algo::ForEach(ExternalDataLayerContainerView, [&IsMainPartitionContainer](const UActorDescContainerInstance* ActorDescContainer) { check(!IsMainPartitionContainer(ActorDescContainer)); });
 #endif
 }
 
 void UWorldPartition::CheckForErrors(IStreamingGenerationErrorHandler* ErrorHandler) const
 {
-	FActorDescList ModifiedActorDescList;
-
 	TErrorHandlerSelector<FStreamingGenerationLogErrorHandler> ErrorHandlerSelector(ErrorHandler);
 
 	FWorldPartitionStreamingGenerator::FWorldPartitionStreamingGeneratorParams StreamingGeneratorParams = FWorldPartitionStreamingGenerator::FWorldPartitionStreamingGeneratorParams()
 		.SetWorldPartitionContext(this)
-		.SetModifiedActorsDescList(&ModifiedActorDescList)
+		.SetHandleUnsavedActors(true)
 		.SetErrorHandler(ErrorHandlerSelector.Get())
 		.SetIsValidGrid([this](FName GridName) { return RuntimeHash->IsValidGrid(GridName); })
 		.SetIsValidHLODLayer([this](FName GridName, const FSoftObjectPath& HLODLayerPath) { return RuntimeHash->IsValidHLODLayer(GridName, HLODLayerPath); })
 		.SetEnableStreaming(IsStreamingEnabled());
 
-	ForEachActorDescContainer([&StreamingGeneratorParams](const UActorDescContainer* InActorDescContainer)
+	ForEachActorDescContainer([&StreamingGeneratorParams](const UActorDescContainerInstance* InContainerInstance)
 	{
-		for (FActorDescList::TConstIterator<> ActorDescIt(InActorDescContainer); ActorDescIt; ++ActorDescIt)
+		for (UActorDescContainerInstance::TConstIterator<> Iterator(InContainerInstance); Iterator; ++Iterator)
 		{
-			check(!StreamingGeneratorParams.ActorGuidsToContainerMap.Contains(ActorDescIt->GetGuid()));
-			StreamingGeneratorParams.ActorGuidsToContainerMap.Add(ActorDescIt->GetGuid(), InActorDescContainer);
+			check(!StreamingGeneratorParams.ActorGuidsToContainerInstanceMap.Contains(Iterator->GetGuid()));
+			StreamingGeneratorParams.ActorGuidsToContainerInstanceMap.Add(Iterator->GetGuid(), InContainerInstance);
 		}
 	});
 
-	ForEachActorDescContainer([this, &StreamingGeneratorParams](const UActorDescContainer* InActorDescContainer)
+	ForEachActorDescContainer([this, &StreamingGeneratorParams](const UActorDescContainerInstance* InContainerInstance)
 	{
-		check(StreamingGeneratorParams.WorldPartitionContext == InActorDescContainer->GetWorldPartition());
+		check(StreamingGeneratorParams.WorldPartitionContext == InContainerInstance->GetWorldPartition());
 		FWorldPartitionStreamingGenerator StreamingGenerator(StreamingGeneratorParams);
-		FStreamingGenerationActorDescCollection Collection{InActorDescContainer};
+		FStreamingGenerationContainerInstanceCollection Collection{ InContainerInstance };
 		StreamingGenerator.PreparationPhase(Collection);
-	});	
-}
-
-/* Deprecated */
-void UWorldPartition::CheckForErrors(IStreamingGenerationErrorHandler* ErrorHandler, const UActorDescContainer* ActorDescContainer, bool bEnableStreaming, bool)
-{
-	FActorDescContainerCollection ActorDescContainerCollection;
-	ActorDescContainerCollection.AddContainer(const_cast<UActorDescContainer*>(ActorDescContainer));
-
-	FCheckForErrorsParams Params = FCheckForErrorsParams()
-		.SetErrorHandler(ErrorHandler)
-		.SetActorDescContainerCollection(&ActorDescContainerCollection)
-		.SetEnableStreaming(bEnableStreaming);
-
-	CheckForErrors(Params);
+	});
 }
 
 /* Static version, mainly used by changelist validation */
 void UWorldPartition::CheckForErrors(const FCheckForErrorsParams& Params)
 {
 	check(Params.ErrorHandler);
-	check(Params.ActorDescContainerCollection);
+	check(Params.ActorDescContainerInstanceCollection);
 
 	FActorDescList ModifiedActorDescList;
 
-	TArray<const UActorDescContainer*> Containers;
-	Params.ActorDescContainerCollection->ForEachActorDescContainer([&Containers](UActorDescContainer* InActorDescContainer) { Containers.Add(InActorDescContainer); });
+	TArray<const UActorDescContainerInstance*> Containers;
+	Params.ActorDescContainerInstanceCollection->ForEachActorDescContainer([&Containers](UActorDescContainerInstance* InContainerInstance) { Containers.Add(InContainerInstance); });
 
-	FStreamingGenerationActorDescCollection Collection { Containers };
-	const UActorDescContainer* MainActorDescContainer = Collection.GetMainActorDescContainer();
+	FStreamingGenerationContainerInstanceCollection Collection{ Containers };
+	const UActorDescContainerInstance* MainContainerInstance = Collection.GetMainContainer();
 
 	TErrorHandlerSelector<FStreamingGenerationLogErrorHandler> ErrorHandlerSelector(Params.ErrorHandler);
 
 	FWorldPartitionStreamingGenerator::FWorldPartitionStreamingGeneratorParams StreamingGeneratorParams = FWorldPartitionStreamingGenerator::FWorldPartitionStreamingGeneratorParams()
-		.SetWorldPartitionContext(MainActorDescContainer->GetWorldPartition())
-		.SetModifiedActorsDescList(!MainActorDescContainer->IsTemplateContainer() ? &ModifiedActorDescList : nullptr)
+		.SetWorldPartitionContext(MainContainerInstance->GetWorldPartition())
+		.SetHandleUnsavedActors(MainContainerInstance->HasWorldPartition())
 		.SetErrorHandler(ErrorHandlerSelector.Get())
 		.SetIsValidGrid([](FName) { return true; })
 		.SetIsValidHLODLayer([](FName, const FSoftObjectPath&) { return true; })
 		.SetEnableStreaming(Params.bEnableStreaming)
-		.SetActorGuidsToContainerMap(Params.ActorGuidsToContainerMap);
+		.SetActorGuidsToContainerInstanceMap(Params.ActorGuidsToContainerInstanceMap);
 
 	FWorldPartitionStreamingGenerator StreamingGenerator(StreamingGeneratorParams);
 	StreamingGenerator.PreparationPhase(Collection);
 }
 #endif // WITH_EDITOR
 
-#undef LOCTEXT_NAMESPACE
+#undef LOCTEXT_NAMESPACE
\ No newline at end of file
diff --git a/Engine/Source/Runtime/Engine/Private/WorldPartition/WorldPartitionStreamingGenerationContext.cpp b/Engine/Source/Runtime/Engine/Private/WorldPartition/WorldPartitionStreamingGenerationContext.cpp
--- a/Engine/Source/Runtime/Engine/Private/WorldPartition/WorldPartitionStreamingGenerationContext.cpp
+++ b/Engine/Source/Runtime/Engine/Private/WorldPartition/WorldPartitionStreamingGenerationContext.cpp
@@ -4,9 +4,10 @@
 #include "WorldPartition/WorldPartitionStreamingGeneration.h"
 
 #if WITH_EDITOR
-const FWorldPartitionActorDescView& IStreamingGenerationContext::FActorInstance::GetActorDescView() const
+
+const FStreamingGenerationActorDescView& IStreamingGenerationContext::FActorInstance::GetActorDescView() const
 {
-	return ActorSetInstance->ContainerInstance->ActorDescViewMap->FindByGuidChecked(ActorGuid);
+	return ActorSetInstance->ActorSetContainerInstance->ActorDescViewMap->FindByGuidChecked(ActorGuid);
 }
 
 const FActorContainerID& IStreamingGenerationContext::FActorInstance::GetContainerID() const
