WorldPartition: Add UActorDescContainerInstance/FWorldPartitionActorDescInstance which take transient/context data from UActorDescContainer/FWorldPartitionActorDesc leaving them as pure on disk repres

diff --git a/Engine/Source/Runtime/Engine/Private/WorldPartition/WorldPartitionSubsystem.cpp b/Engine/Source/Runtime/Engine/Private/WorldPartition/WorldPartitionSubsystem.cpp
--- a/Engine/Source/Runtime/Engine/Private/WorldPartition/WorldPartitionSubsystem.cpp
+++ b/Engine/Source/Runtime/Engine/Private/WorldPartition/WorldPartitionSubsystem.cpp
@@ -12,6 +12,7 @@
 #include "WorldPartition/WorldPartitionReplay.h"
 #include "WorldPartition/WorldPartitionDraw2DContext.h"
 #include "WorldPartition/WorldPartitionStreamingPolicy.h"
+#include "WorldPartition/ActorDescContainerSubsystem.h"
 #include "WorldPartition/DataLayer/WorldDataLayersActorDesc.h"
 #include "WorldPartition/DataLayer/DataLayerManager.h"
 #include "WorldPartition/DataLayer/DataLayerAsset.h"
@@ -199,13 +200,6 @@ const UWorldPartition* UWorldPartitionSubsystem::GetWorldPartition() const
 }
 
 #if WITH_EDITOR
-void UWorldPartitionSubsystem::AddReferencedObjects(UObject* InThis, FReferenceCollector& Collector)
-{
-	UWorldPartitionSubsystem* This = CastChecked<UWorldPartitionSubsystem>(InThis);
-
-	This->ActorDescContainerInstanceManager.AddReferencedObjects(Collector);
-}
-
 FWorldPartitionActorFilter UWorldPartitionSubsystem::GetWorldPartitionActorFilter(const FString& InWorldPackage, EWorldPartitionActorFilterType InFilterTypes) const
 {
 	TSet<FString> VisitedPackages;
@@ -221,12 +215,12 @@ FWorldPartitionActorFilter UWorldPartitionSubsystem::GetWorldPartitionActorFilte
 	}
 
 	InOutVisitedPackages.Add(InWorldPackage);
-
+	
 	// Most of the time if this will return an existing Container but when loading a new LevelInstance (Content Browser Drag&Drop, Create LI) 
 	// This will make sure Container exists.
-	UActorDescContainer* LevelContainer = ActorDescContainerInstanceManager.RegisterContainer(*InWorldPackage, GetWorld());
+	UActorDescContainer* LevelContainer = UActorDescContainerSubsystem::GetChecked().RegisterContainer(UActorDescContainer::FInitializeParams{ *InWorldPackage });
 	check(LevelContainer);
-	ON_SCOPE_EXIT{ ActorDescContainerInstanceManager.UnregisterContainer(LevelContainer); };
+	ON_SCOPE_EXIT{ UActorDescContainerSubsystem::GetChecked().UnregisterContainer(LevelContainer); };
 
 	// Lazy create filter for now
 	TArray<const FWorldPartitionActorDesc*> ContainerActorDescs;
@@ -239,7 +233,7 @@ FWorldPartitionActorFilter UWorldPartitionSubsystem::GetWorldPartitionActorFilte
 			check(!WorldDataLayersActorDesc);
 			WorldDataLayersActorDesc = static_cast<const FWorldDataLayersActorDesc*>(*ActorDescIt);
 		}
-		else if ((ActorDescIt->GetContainerFilterType() & InFilterTypes) != EWorldPartitionActorFilterType::None)
+		else if ((ActorDescIt->GetChildContainerFilterType() & InFilterTypes) != EWorldPartitionActorFilterType::None)
 		{
 			ContainerActorDescs.Add(*ActorDescIt);
 		}
@@ -264,11 +258,11 @@ FWorldPartitionActorFilter UWorldPartitionSubsystem::GetWorldPartitionActorFilte
 		TSet<FString> VisitedPackagesCopy(InOutVisitedPackages);
 
 		// Get World Default Filter
-		FWorldPartitionActorFilter* ChildFilter = new FWorldPartitionActorFilter(GetWorldPartitionActorFilterInternal(ContainerActorDesc->GetContainerPackage().ToString(), InFilterTypes, VisitedPackagesCopy));
+		FWorldPartitionActorFilter* ChildFilter = new FWorldPartitionActorFilter(GetWorldPartitionActorFilterInternal(ContainerActorDesc->GetChildContainerPackage().ToString(), InFilterTypes, VisitedPackagesCopy));
 		ChildFilter->DisplayName = ContainerActorDesc->GetActorLabelOrName().ToString();
 
 		// Apply Filter to Default
-		if (const FWorldPartitionActorFilter* ContainerFilter = ContainerActorDesc->GetContainerFilter())
+		if (const FWorldPartitionActorFilter* ContainerFilter = ContainerActorDesc->GetChildContainerFilter())
 		{
 			ChildFilter->Override(*ContainerFilter);
 		}
@@ -315,8 +309,8 @@ TMap<FActorContainerID, TSet<FGuid>> UWorldPartitionSubsystem::GetFilteredActors
 		{
 			return *FoundContainer;
 		}
-		
-		UActorDescContainer* RegisteredContainer = RegisterContainer(ContainerPackage);
+	
+		UActorDescContainer* RegisteredContainer = UActorDescContainerSubsystem::GetChecked().RegisterContainer(UActorDescContainer::FInitializeParams{ ContainerPackage });
 		RegisteredContainers.Add(ContainerPackage, RegisteredContainer);
 		return RegisteredContainer;
 	};
@@ -352,11 +346,11 @@ TMap<FActorContainerID, TSet<FGuid>> UWorldPartitionSubsystem::GetFilteredActors
 				}
 			}
 
-			if ((ActorDescIt->GetContainerFilterType() & InFilterTypes) != EWorldPartitionActorFilterType::None)
+			if ((ActorDescIt->GetChildContainerFilterType() & InFilterTypes) != EWorldPartitionActorFilterType::None)
 			{
-				if (const FWorldPartitionActorFilter* ChildFilter = ActorDescIt->GetContainerFilter())
+				if (const FWorldPartitionActorFilter* ChildFilter = ActorDescIt->GetChildContainerFilter())
 				{
-					UActorDescContainer* ChildContainer = FindOrRegisterContainer(ActorDescIt->GetContainerPackage());
+					UActorDescContainer* ChildContainer = FindOrRegisterContainer(ActorDescIt->GetChildContainerPackage());
 					check(ChildContainer);
 					ProcessContainers(FActorContainerID(InContainerID, ActorDescIt->GetGuid()), ChildContainer);
 				}
@@ -368,9 +362,10 @@ TMap<FActorContainerID, TSet<FGuid>> UWorldPartitionSubsystem::GetFilteredActors
 	ProcessContainers(InContainerID, Container);
 	
 	// Unregister Containers
+	UActorDescContainerSubsystem& ContainerSubsystem = UActorDescContainerSubsystem::GetChecked();
 	for (auto& [Name, RegisteredContainer] : RegisteredContainers)
 	{
-		UnregisterContainer(RegisteredContainer);
+		ContainerSubsystem.UnregisterContainer(RegisteredContainer);
 	}
 
 	return FilteredActors;
@@ -384,87 +379,6 @@ bool UWorldPartitionSubsystem::IsRunningConvertWorldPartitionCommandlet()
 	return GetRunningCommandletClass() && GetRunningCommandletClass()->IsChildOf(WorldPartitionConvertCommandletClass);
 }
 
-void UWorldPartitionSubsystem::FActorDescContainerInstanceManager::FActorDescContainerInstance::AddReferencedObjects(FReferenceCollector& Collector)
-{
-	Collector.AddReferencedObject(Container);
-}
-
-void UWorldPartitionSubsystem::FActorDescContainerInstanceManager::FActorDescContainerInstance::UpdateBounds()
-{
-	Bounds.Init();
-	for (FActorDescList::TIterator<> ActorDescIt(Container); ActorDescIt; ++ActorDescIt)
-	{
-		// FActorDescContainerInstance are currently only used by FLevelInstanceActorDescs and so we don't consider actors that should exit only in the main world.
-		if (ActorDescIt->IsMainWorldOnly())
-		{
-			continue;
-		}
-
-		const FBox RuntimeBounds = ActorDescIt->GetRuntimeBounds();
-		if (RuntimeBounds.IsValid)
-		{
-			Bounds += RuntimeBounds;
-		}
-	}
-}
-
-void UWorldPartitionSubsystem::FActorDescContainerInstanceManager::AddReferencedObjects(FReferenceCollector& Collector)
-{
-	for (auto& [Name, ContainerInstance] : ActorDescContainers)
-	{
-		ContainerInstance.AddReferencedObjects(Collector);
-	}
-}
-
-UActorDescContainer* UWorldPartitionSubsystem::FActorDescContainerInstanceManager::RegisterContainer(FName PackageName, UWorld* InWorld)
-{
-	FActorDescContainerInstance* ExistingContainerInstance = &ActorDescContainers.FindOrAdd(PackageName);
-	UActorDescContainer* ActorDescContainer = ExistingContainerInstance->Container;
-
-	if (ExistingContainerInstance->RefCount++ == 0)
-	{
-		ActorDescContainer = NewObject<UActorDescContainer>(GetTransientPackage());
-		ExistingContainerInstance->Container = ActorDescContainer;
-
-		// This will potentially invalidate ExistingContainerInstance due to ActorDescContainers reallocation
-		ActorDescContainer->Initialize({ InWorld, PackageName });
-
-		ExistingContainerInstance = &ActorDescContainers.FindChecked(PackageName);
-		ExistingContainerInstance->UpdateBounds();
-	}
-
-	check(ActorDescContainer->IsTemplateContainer());
-	return ActorDescContainer;
-}
-
-void UWorldPartitionSubsystem::FActorDescContainerInstanceManager::UnregisterContainer(UActorDescContainer* Container)
-{
-	FName PackageName = Container->GetContainerPackage();
-	FActorDescContainerInstance& ExistingContainerInstance = ActorDescContainers.FindChecked(PackageName);
-
-	if (--ExistingContainerInstance.RefCount == 0)
-	{
-		ExistingContainerInstance.Container->Uninitialize();
-		ActorDescContainers.FindAndRemoveChecked(PackageName);
-	}
-}
-
-FBox UWorldPartitionSubsystem::FActorDescContainerInstanceManager::GetContainerBounds(FName PackageName) const
-{
-	if (const FActorDescContainerInstance* ActorDescContainerInstance = ActorDescContainers.Find(PackageName))
-	{
-		return ActorDescContainerInstance->Bounds;
-	}
-	return FBox(ForceInit);
-}
-
-void UWorldPartitionSubsystem::FActorDescContainerInstanceManager::UpdateContainerBounds(FName PackageName)
-{
-	if (FActorDescContainerInstance* ActorDescContainerInstance = ActorDescContainers.Find(PackageName))
-	{
-		ActorDescContainerInstance->UpdateBounds();
-	}
-}
 #endif
 
 void UWorldPartitionSubsystem::Initialize(FSubsystemCollectionBase& Collection)
diff --git a/Engine/Source/Runtime/Engine/Private/WorldPartitionBlueprintLibrary.cpp b/Engine/Source/Runtime/Engine/Private/WorldPartitionBlueprintLibrary.cpp
--- a/Engine/Source/Runtime/Engine/Private/WorldPartitionBlueprintLibrary.cpp
+++ b/Engine/Source/Runtime/Engine/Private/WorldPartitionBlueprintLibrary.cpp
@@ -10,6 +10,7 @@
 #include "WorldPartition/LoaderAdapter/LoaderAdapterActorList.h"
 #include "WorldPartition/WorldPartition.h"
 #include "WorldPartition/WorldPartitionHelpers.h"
+#include "WorldPartition/WorldPartitionActorDescInstance.h"
 #include "Editor.h"
 #endif
 
@@ -81,30 +82,30 @@ void UWorldPartitionBlueprintLibrary::OnWorldPartitionUninitialized(UWorldPartit
 	}
 }
 
-bool UWorldPartitionBlueprintLibrary::GetActorDescs(const UWorldPartition* WorldPartiton, TArray<FActorDesc>& OutActorDescs)
+bool UWorldPartitionBlueprintLibrary::GetActorDescs(const UWorldPartition* InWorldPartition, TArray<FActorDesc>& OutActorDescs)
 {
 	bool bResult = true;
 
-	WorldPartiton->ForEachActorDescContainer([&](const UActorDescContainer* ActorDescContainer)
+	InWorldPartition->ForEachActorDescContainer([&](const UActorDescContainerInstance* InContainerInstance)
 	{
-		bResult &= GetActorDescs(ActorDescContainer, FTransform::Identity, OutActorDescs);
+		bResult &= GetActorDescs(InContainerInstance, FTransform::Identity, OutActorDescs);
 	});
 	
 	return bResult;
 }
 
-bool UWorldPartitionBlueprintLibrary::GetActorDescs(const UActorDescContainer* InContainer, const FTransform& InTransform, TArray<FActorDesc>& OutActorDescs)
+bool UWorldPartitionBlueprintLibrary::GetActorDescs(const UActorDescContainerInstance* InContainerInstance, const FTransform& InTransform, TArray<FActorDesc>& OutActorDescs)
 {
 	bool bResult = true;
 
-	for (FActorDescList::TConstIterator<> ActorDescIt(InContainer); ActorDescIt; ++ActorDescIt)
+	for (UActorDescContainerInstance::TConstIterator<> Iterator(InContainerInstance); Iterator; ++Iterator)
 	{
-		if (ActorDescIt->IsContainerInstance())
+		if (Iterator->IsChildContainerInstance())
 		{
 			FWorldPartitionActorDesc::FContainerInstance ContainerInstance;
-			if (ActorDescIt->GetContainerInstance(ContainerInstance))
+			if (Iterator->GetChildContainerInstance(ContainerInstance))
 			{
-				bResult &= GetActorDescs(ContainerInstance.Container, ContainerInstance.Transform * InTransform, OutActorDescs);
+				bResult &= GetActorDescs(ContainerInstance.ContainerInstance, ContainerInstance.Transform * InTransform, OutActorDescs);
 			}
 			else
 			{
@@ -113,23 +114,23 @@ bool UWorldPartitionBlueprintLibrary::GetActorDescs(const UActorDescContainer* I
 		}
 		else
 		{
-			OutActorDescs.Emplace(**ActorDescIt, InTransform);
+			OutActorDescs.Emplace(*Iterator->GetActorDesc(), InTransform);
 		}
 	}
 
 	return bResult;
 }
 
-bool UWorldPartitionBlueprintLibrary::HandleIntersectingActorDesc(const FWorldPartitionActorDesc* ActorDesc, const FBox& InBox, const FTransform& InTransform, TArray<FActorDesc>& OutActorDescs)
+bool UWorldPartitionBlueprintLibrary::HandleIntersectingActorDesc(const FWorldPartitionActorDescInstance* ActorDescInstance, const FBox& InBox, const FTransform& InTransform, TArray<FActorDesc>& OutActorDescs)
 {
 	bool bResult = true;
 
-	if (ActorDesc->IsContainerInstance())
+	if (ActorDescInstance->IsChildContainerInstance())
 	{
 		FWorldPartitionActorDesc::FContainerInstance ContainerInstance;
-		if (ActorDesc->GetContainerInstance(ContainerInstance))
+		if (ActorDescInstance->GetChildContainerInstance(ContainerInstance))
 		{
-			bResult &= GetIntersectingActorDescs(ContainerInstance.Container, InBox, ContainerInstance.Transform * InTransform, OutActorDescs);
+			bResult &= GetIntersectingActorDescs(ContainerInstance.ContainerInstance, InBox, ContainerInstance.Transform * InTransform, OutActorDescs);
 		}
 		else
 		{
@@ -138,7 +139,7 @@ bool UWorldPartitionBlueprintLibrary::HandleIntersectingActorDesc(const FWorldPa
 	}
 	else
 	{
-		OutActorDescs.Emplace(*ActorDesc, InTransform);
+		OutActorDescs.Emplace(*ActorDescInstance->GetActorDesc(), InTransform);
 	}
 
 	return bResult;
@@ -148,7 +149,7 @@ bool UWorldPartitionBlueprintLibrary::GetIntersectingActorDescs(UWorldPartition*
 {
 	bool bResult = true;
 
-	FWorldPartitionHelpers::ForEachIntersectingActorDesc(WorldPartition, InBox, [&bResult, &InBox, &OutActorDescs](const FWorldPartitionActorDesc* ActorDesc)
+	FWorldPartitionHelpers::ForEachIntersectingActorDescInstance(WorldPartition, InBox, [&bResult, &InBox, &OutActorDescs](const FWorldPartitionActorDescInstance* ActorDesc)
 	{
 		bResult &= HandleIntersectingActorDesc(ActorDesc, InBox, FTransform::Identity, OutActorDescs);
 		return true;
@@ -157,15 +158,15 @@ bool UWorldPartitionBlueprintLibrary::GetIntersectingActorDescs(UWorldPartition*
 	return bResult;
 }
 
-bool UWorldPartitionBlueprintLibrary::GetIntersectingActorDescs(const UActorDescContainer* InContainer, const FBox& InBox, const FTransform& InTransform, TArray<FActorDesc>& OutActorDescs)
+bool UWorldPartitionBlueprintLibrary::GetIntersectingActorDescs(const UActorDescContainerInstance* InContainerInstance, const FBox& InBox, const FTransform& InTransform, TArray<FActorDesc>& OutActorDescs)
 {
 	bool bResult = true;
 
-	for (FActorDescList::TConstIterator<> ActorDescIt(InContainer); ActorDescIt; ++ActorDescIt)
+	for (UActorDescContainerInstance::TConstIterator<> Iterator(InContainerInstance); Iterator; ++Iterator)
 	{
-		if (ActorDescIt->GetEditorBounds().TransformBy(InTransform).Intersect(InBox))
+		if (Iterator->GetEditorBounds().TransformBy(InTransform).Intersect(InBox))
 		{
-			bResult &= HandleIntersectingActorDesc(*ActorDescIt, InBox, InTransform, OutActorDescs);
+			bResult &= HandleIntersectingActorDesc(*Iterator, InBox, InTransform, OutActorDescs);
 		}
 	}
 
diff --git a/Engine/Source/Runtime/Engine/Public/LevelInstance/LevelInstanceEditorLevelStreaming.h b/Engine/Source/Runtime/Engine/Public/LevelInstance/LevelInstanceEditorLevelStreaming.h
--- a/Engine/Source/Runtime/Engine/Public/LevelInstance/LevelInstanceEditorLevelStreaming.h
+++ b/Engine/Source/Runtime/Engine/Public/LevelInstance/LevelInstanceEditorLevelStreaming.h
@@ -6,6 +6,7 @@
 #include "UObject/ObjectMacros.h"
 #include "Engine/LevelStreamingAlwaysLoaded.h"
 #include "LevelInstance/LevelInstanceTypes.h"
+#include "WorldPartition/ActorDescContainerInstance.h"
 #include "LevelInstanceEditorLevelStreaming.generated.h"
 
 class ILevelInstanceInterface;
@@ -27,6 +28,7 @@ public:
 protected:
 	void OnLevelActorAdded(AActor* InActor);
 	void OnLoadedActorsAddedToLevelPreEvent(const TArray<AActor*>& InActors);
+	void OnPreInitializeContainerInstance(UActorDescContainerInstance::FInitializeParams& InInitParams, UActorDescContainerInstance* InContainerInstance);
 
 	friend class ULevelInstanceSubsystem;
 
diff --git a/Engine/Source/Runtime/Engine/Public/LevelInstance/LevelInstanceLevelStreaming.h b/Engine/Source/Runtime/Engine/Public/LevelInstance/LevelInstanceLevelStreaming.h
--- a/Engine/Source/Runtime/Engine/Public/LevelInstance/LevelInstanceLevelStreaming.h
+++ b/Engine/Source/Runtime/Engine/Public/LevelInstance/LevelInstanceLevelStreaming.h
@@ -6,6 +6,7 @@
 #include "UObject/ObjectMacros.h"
 #include "Engine/LevelStreamingDynamic.h"
 #include "LevelInstance/LevelInstanceTypes.h"
+#include "WorldPartition/ActorDescContainerInstance.h"
 #include "LevelInstanceLevelStreaming.generated.h"
 
 class ILevelInstanceInterface;
@@ -44,6 +45,7 @@ private:
 	ENGINE_API virtual void OnLoadedActorsAddedToLevelPostEvent(const TArray<AActor*>& InActors);
 	ENGINE_API virtual void OnLoadedActorsRemovedFromLevelPostEvent(const TArray<AActor*>& InActors);
 	ENGINE_API void OnLevelStreamingStateChanged(UWorld* InWorld, const ULevelStreaming* InLevelStreaming, ULevel* InLevelIfLoaded, ELevelStreamingState InPrevState, ELevelStreamingState InNewState);
+	ENGINE_API void OnPreInitializeContainerInstance(UActorDescContainerInstance::FInitializeParams& InInitParams, UActorDescContainerInstance* InContainerInstance);
 
 	bool IsEditorWorldMode() const;
 
diff --git a/Engine/Source/Runtime/Engine/Public/LevelInstance/LevelInstanceSubsystem.h b/Engine/Source/Runtime/Engine/Public/LevelInstance/LevelInstanceSubsystem.h
--- a/Engine/Source/Runtime/Engine/Public/LevelInstance/LevelInstanceSubsystem.h
+++ b/Engine/Source/Runtime/Engine/Public/LevelInstance/LevelInstanceSubsystem.h
@@ -19,7 +19,6 @@ class ULevelInstanceEditorObject;
 class ULevelStreamingLevelInstance;
 class ULevelStreamingLevelInstanceEditor;
 class UWorldPartitionSubsystem;
-class UActorDescContainer;
 class UBlueprint;
 
 /**
diff --git a/Engine/Source/Runtime/Engine/Public/WorldPartition/ActorDescContainer.h b/Engine/Source/Runtime/Engine/Public/WorldPartition/ActorDescContainer.h
--- a/Engine/Source/Runtime/Engine/Public/WorldPartition/ActorDescContainer.h
+++ b/Engine/Source/Runtime/Engine/Public/WorldPartition/ActorDescContainer.h
@@ -28,50 +28,53 @@ public:
 	/* Struct of parameters passed to Initialize function. */
 	struct FInitializeParams
 	{
+		UE_DEPRECATED(5.4, "Use constructor with no UWorld param instead")
 		FInitializeParams(UWorld* InWorld, FName InPackageName)
-			: World(InWorld)
-			, PackageName(InPackageName)
+			: PackageName(InPackageName)
+			, ContainerName(InPackageName.ToString())
 		{}
 
-		/* The world the actor descriptor container is associated with. */
-		UWorld* World;
-			
+		FInitializeParams(FName InPackageName)
+			: PackageName(InPackageName)
+			, ContainerName(InPackageName.ToString())
+		{}
+
+		FInitializeParams(const FString& InContainerName, FName InPackageName)
+			: PackageName(InPackageName)
+			, ContainerName(InContainerName)
+		{}
+		
 		/* The long package name of the container package on disk. */
 		FName PackageName;
 
+		/* The unique name for the container : defaults to PackageName */
+		FString ContainerName;
+
 		/* Custom filter function used to filter actors descriptors. */
 		TUniqueFunction<bool(const FWorldPartitionActorDesc*)> FilterActorDesc;
 	};
 
 	ENGINE_API void Initialize(const FInitializeParams& InitParams);
-	UE_DEPRECATED(5.4, "UActorDescContainer::Update is deprecated.")
-	ENGINE_API void Update() {}
 	ENGINE_API void Uninitialize();
 
 	bool IsInitialized() const { return bContainerInitialized; }
 
 	ENGINE_API void OnObjectPreSave(UObject* Object, FObjectPreSaveContext SaveContext);
 	ENGINE_API void OnPackageDeleted(UPackage* Package);
-	ENGINE_API void OnObjectsReplaced(const TMap<UObject*, UObject*>& OldToNewObjectMap);
 	ENGINE_API void OnClassDescriptorUpdated(const FWorldPartitionActorDesc* InClassDesc);
 
+	virtual FString GetContainerName() const { return ContainerPackageName.ToString(); }
 	FName GetContainerPackage() const { return ContainerPackageName; }
 	void SetContainerPackage(const FName& InContainerPackageName) { ContainerPackageName = InContainerPackageName; }
 
 	FGuid GetContentBundleGuid() const { return ContentBundleGuid; }
 	void SetContentBundleGuid(const FGuid& InGetContentBundleGuid) { ContentBundleGuid = InGetContentBundleGuid; }
 
-	ENGINE_API bool IsTemplateContainer() const;
-	ENGINE_API bool IsMainPartitionContainer() const;
-	ENGINE_API UWorldPartition* GetWorldPartition() const;
-
 	ENGINE_API FString GetExternalActorPath() const;
 
 	/** Removes an actor desc without the need to load a package */
 	ENGINE_API bool RemoveActor(const FGuid& ActorGuid);
 
-	ENGINE_API void LoadAllActors(TArray<FWorldPartitionReference>& OutReferences);
-
 	ENGINE_API bool IsActorDescHandled(const AActor* Actor) const;
 
 	DECLARE_EVENT_OneParam(UActorDescContainer, FActorDescAddedEvent, FWorldPartitionActorDesc*);
@@ -80,17 +83,20 @@ public:
 	DECLARE_EVENT_OneParam(UActorDescContainer, FActorDescRemovedEvent, FWorldPartitionActorDesc*);
 	FActorDescRemovedEvent OnActorDescRemovedEvent;
 
-	DECLARE_MULTICAST_DELEGATE_OneParam(FActorDescContainerInitializeDelegate, UActorDescContainer*);
-	static ENGINE_API FActorDescContainerInitializeDelegate OnActorDescContainerInitialized;
+	DECLARE_EVENT_OneParam(UActorDescContainer, FActorDescUpdatingEvent, FWorldPartitionActorDesc*);
+	FActorDescUpdatingEvent OnActorDescUpdatingEvent;
 
-	ENGINE_API const FLinkerInstancingContext* GetInstancingContext() const;
-	ENGINE_API const FTransform& GetInstanceTransform() const;
+	DECLARE_EVENT_OneParam(UActorDescContainer, FActorDescUpdatedEvent, FWorldPartitionActorDesc*);
+	FActorDescUpdatedEvent OnActorDescUpdatedEvent;
 
+	DECLARE_MULTICAST_DELEGATE_OneParam(FActorDescContainerInitializeDelegate, UActorDescContainer*);
+	static ENGINE_API FActorDescContainerInitializeDelegate OnActorDescContainerInitialized;
+	
 	bool HasInvalidActors() const { return InvalidActors.Num() > 0; }
 	const TArray<FAssetData>& GetInvalidActors() const { return InvalidActors; }
 	void ClearInvalidActors() { InvalidActors.Empty(); }
 
-	ENGINE_API void RegisterActorDescriptor(FWorldPartitionActorDesc* ActorDesc, UWorld* InWorldContext);
+	ENGINE_API void RegisterActorDescriptor(FWorldPartitionActorDesc* ActorDesc);
 	ENGINE_API void UnregisterActorDescriptor(FWorldPartitionActorDesc* ActorDesc);
 
 	ENGINE_API void OnActorDescAdded(FWorldPartitionActorDesc* NewActorDesc);
@@ -111,6 +117,35 @@ public:
 
 	TArray<FAssetData> InvalidActors;
 
+	//~ Begin Deprecation
+	UE_DEPRECATED(5.4, "UActorDescContainer::Update is deprecated.")
+	ENGINE_API void Update() {}
+
+	UE_DEPRECATED(5.4, "Use UActorDescContainerInstance::GetInstancingContext instead")
+	ENGINE_API const FLinkerInstancingContext* GetInstancingContext() const { return nullptr; }
+	
+	UE_DEPRECATED(5.4, "Use UActorDescContainerInstance::GetInstanceTransform instead")
+	ENGINE_API FTransform GetInstanceTransform() const { return FTransform::Identity; }
+
+	UE_DEPRECATED(5.4, "Use UActorDescContainerInstance::LoadAllActors instead")
+	ENGINE_API void LoadAllActors(TArray<FWorldPartitionReference>& OutReferences) {}
+
+	UE_DEPRECATED(5.4, "Use UActorDescContainerInstance::OnObjectsReplaced instead")
+	ENGINE_API void OnObjectsReplaced(const TMap<UObject*, UObject*>& OldToNewObjectMap) {}
+
+	UE_DEPRECATED(5.4, "UActorDescContainer::IsTemplateContainer is deprecated")
+	ENGINE_API bool IsTemplateContainer() const { return true;}
+
+	UE_DEPRECATED(5.4, "UActorDescContainer::IsMainWorldPartitionContainer is deprecated")
+	ENGINE_API bool IsMainPartitionContainer() const { return false; }
+	
+	UE_DEPRECATED(5.4, "Use UActorDescContainerInstance::GetWorldPartition instead")
+	ENGINE_API UWorldPartition* GetWorldPartition() const { return nullptr; }
+
+	UE_DEPRECATED(5.4, "Use version without UWorld parameter")
+	ENGINE_API void RegisterActorDescriptor(FWorldPartitionActorDesc* ActorDesc, UWorld* InWorldContext) {}
+	//~ End Deprecation
+
 protected:
 	FNameActorDescMap ActorsByName;
 
@@ -120,7 +155,7 @@ protected:
 
 private:
 	// GetWorld() should never be called on an ActorDescContainer to avoid any confusion as it can be used as a template
-	virtual UWorld* GetWorld() const override { return nullptr; }
+	UWorld* GetWorld() const override { return nullptr; }
 
 	ENGINE_API bool ShouldRegisterDelegates();
 	ENGINE_API void RegisterEditorDelegates();
diff --git a/Engine/Source/Runtime/Engine/Public/WorldPartition/ActorDescList.h b/Engine/Source/Runtime/Engine/Public/WorldPartition/ActorDescList.h
--- a/Engine/Source/Runtime/Engine/Public/WorldPartition/ActorDescList.h
+++ b/Engine/Source/Runtime/Engine/Public/WorldPartition/ActorDescList.h
@@ -6,12 +6,22 @@
 #include "WorldPartition/WorldPartitionActorDesc.h"
 #include "WorldPartition/WorldPartitionActorDescType.h"
 
-class FActorDescList
+#if WITH_EDITOR
+// Default Iterator Value Type is the DescType itself
+template<typename DescType, class ActorType>
+struct TActorDescListIteratorValueType 
+{ 
+	using Type = DescType;
+};
+#endif
+
+template<class DescType>
+class TActorDescList
 {
 #if WITH_EDITOR
 	friend struct FWorldPartitionImplBase;
 
-	struct FActorGuidKeyFuncs : TDefaultMapKeyFuncs<FGuid, TUniquePtr<FWorldPartitionActorDesc>*, false>
+	struct FActorGuidKeyFuncs : TDefaultMapKeyFuncs<FGuid, TUniquePtr<DescType>*, false>
 	{
 		static FORCEINLINE uint32 GetKeyHash(const FGuid& Key)
 		{
@@ -19,29 +29,23 @@ class FActorDescList
 		}
 	};
 
-	using FGuidActorDescMap = TMap<FGuid, TUniquePtr<FWorldPartitionActorDesc>*, FDefaultSetAllocator, FActorGuidKeyFuncs>;
-	using FActorDescArray = TChunkedArray<TUniquePtr<FWorldPartitionActorDesc>>;
+	using FGuidActorDescMap = TMap<FGuid, TUniquePtr<DescType>*, FDefaultSetAllocator, FActorGuidKeyFuncs>;
+	using FActorDescArray = TChunkedArray<TUniquePtr<DescType>>;
 
 public:
-	FActorDescList() {}
-	virtual ~FActorDescList() {}
+	TActorDescList() {}
+	virtual ~TActorDescList() {}
 
 	// Non-copyable
-	FActorDescList(const FActorDescList&) = delete;
-	FActorDescList& operator=(const FActorDescList&) = delete;
-
-	ENGINE_API FWorldPartitionActorDesc* AddActor(const AActor* InActor);
-
-	ENGINE_API FWorldPartitionActorDesc* GetActorDesc(const FGuid& Guid);
-	ENGINE_API const FWorldPartitionActorDesc* GetActorDesc(const FGuid& Guid) const;
-
-	ENGINE_API FWorldPartitionActorDesc& GetActorDescChecked(const FGuid& Guid);
-	ENGINE_API const FWorldPartitionActorDesc& GetActorDescChecked(const FGuid& Guid) const;
-
-	int32 GetActorDescCount() const { return ActorsByGuid.Num(); }
-
-	bool IsEmpty() const { return GetActorDescCount() == 0; }
-	ENGINE_API void Empty();
+	TActorDescList(const TActorDescList&) = delete;
+	TActorDescList& operator=(const TActorDescList&) = delete;
+		
+	bool IsEmpty() const { return ActorsByGuid.Num() == 0; }
+	ENGINE_API void Empty()
+	{
+		ActorsByGuid.Empty();
+		ActorDescList.Empty();
+	}
 
 	template<bool bConst, class ActorType>
 	class TBaseIterator
@@ -50,9 +54,9 @@ public:
 
 	protected:
 		using MapType = FGuidActorDescMap;
-		using ValueType = typename FWorldPartitionActorDescType<ActorType>::Type;
-		using IteratorType = std::conditional_t<bConst, MapType::TConstIterator, MapType::TIterator>;
-		using ListType = std::conditional_t<bConst, const FActorDescList*, FActorDescList*>;
+		using ValueType = typename TActorDescListIteratorValueType<DescType, ActorType>::Type;
+		using IteratorType = std::conditional_t<bConst, typename MapType::TConstIterator, typename MapType::TIterator>;
+		using ListType = std::conditional_t<bConst, const TActorDescList<DescType>*, TActorDescList<DescType>*>;
 		using ReturnType = std::conditional_t<bConst, const ValueType*, ValueType*>;
 
 	public:
@@ -167,13 +171,83 @@ public:
 		{}
 	};
 
-	ENGINE_API void AddActorDescriptor(FWorldPartitionActorDesc* ActorDesc);
-	ENGINE_API void RemoveActorDescriptor(FWorldPartitionActorDesc* ActorDesc);
+	void AddActorDescriptor(DescType* ActorDesc);
+	void RemoveActorDescriptor(DescType* ActorDesc);
 
 protected:
-	ENGINE_API TUniquePtr<FWorldPartitionActorDesc>* GetActorDescriptor(const FGuid& ActorGuid);
+	TUniquePtr<DescType>* GetActorDescriptor(const FGuid& ActorGuid);
+	const TUniquePtr<DescType>* GetActorDescriptor(const FGuid& ActorGuid) const;
+
+	TUniquePtr<DescType>* GetActorDescriptorChecked(const FGuid& ActorGuid);
+	const TUniquePtr<DescType>* GetActorDescriptorChecked(const FGuid& ActorGuid) const;
 
 	FActorDescArray ActorDescList;
 	FGuidActorDescMap ActorsByGuid;
 #endif
 };
+
+#if WITH_EDITOR
+template<class DescType>
+void TActorDescList<DescType>::AddActorDescriptor(DescType* ActorDesc)
+{
+	check(ActorDesc);
+	checkf(!ActorsByGuid.Contains(ActorDesc->GetGuid()), TEXT("Duplicated actor descriptor guid '%s' detected: `%s`"), *ActorDesc->GetGuid().ToString(), *ActorDesc->GetActorName().ToString());
+
+	TUniquePtr<DescType>* NewActorDesc = new(ActorDescList) TUniquePtr<DescType>(ActorDesc);
+	ActorsByGuid.Add(ActorDesc->GetGuid(), NewActorDesc);
+}
+
+template<class DescType>
+void TActorDescList<DescType>::RemoveActorDescriptor(DescType* ActorDesc)
+{
+	check(ActorDesc);
+	verify(ActorsByGuid.Remove(ActorDesc->GetGuid()));
+}
+
+template<class DescType>
+TUniquePtr<DescType>* TActorDescList<DescType>::GetActorDescriptor(const FGuid& ActorGuid)
+{
+	return ActorsByGuid.FindRef(ActorGuid);
+}
+
+template<class DescType>
+const TUniquePtr<DescType>* TActorDescList<DescType>::GetActorDescriptor(const FGuid& ActorGuid) const
+{
+	return ActorsByGuid.FindRef(ActorGuid);
+}
+
+template<class DescType>
+TUniquePtr<DescType>* TActorDescList<DescType>::GetActorDescriptorChecked(const FGuid& ActorGuid)
+{
+	return ActorsByGuid.FindChecked(ActorGuid);
+}
+
+template<class DescType>
+const TUniquePtr<DescType>* TActorDescList<DescType>::GetActorDescriptorChecked(const FGuid& ActorGuid) const
+{
+	return ActorsByGuid.FindChecked(ActorGuid);
+}
+
+// FActorDescList supports subclassing through ActorType's and FWorldPartitionActorDescType will return  the proper iterator value type
+template<class ActorType>
+struct TActorDescListIteratorValueType<FWorldPartitionActorDesc, ActorType>
+{
+	using Type = typename FWorldPartitionActorDescType<ActorType>::Type;
+};
+#endif
+
+class FActorDescList : public TActorDescList<FWorldPartitionActorDesc>
+{
+#if WITH_EDITOR
+public:
+	ENGINE_API FWorldPartitionActorDesc* GetActorDesc(const FGuid& Guid);
+	ENGINE_API const FWorldPartitionActorDesc* GetActorDesc(const FGuid& Guid) const;
+
+	ENGINE_API FWorldPartitionActorDesc& GetActorDescChecked(const FGuid& Guid);
+	ENGINE_API const FWorldPartitionActorDesc& GetActorDescChecked(const FGuid& Guid) const;
+
+	int32 GetActorDescCount() const { return ActorsByGuid.Num(); }
+
+	ENGINE_API FWorldPartitionActorDesc* AddActor(const AActor* InActor);
+#endif
+};
\ No newline at end of file
diff --git a/Engine/Source/Runtime/Engine/Public/WorldPartition/ContentBundle/ContentBundleEditor.h b/Engine/Source/Runtime/Engine/Public/WorldPartition/ContentBundle/ContentBundleEditor.h
--- a/Engine/Source/Runtime/Engine/Public/WorldPartition/ContentBundle/ContentBundleEditor.h
+++ b/Engine/Source/Runtime/Engine/Public/WorldPartition/ContentBundle/ContentBundleEditor.h
@@ -20,6 +20,8 @@ class AActor;
 
 class URuntimeHashExternalStreamingObjectBase;
 class UWorldPartitionRuntimeCell;
+class UActorDescContainerInstance;
+class UWorldPartitionActorDescInstance;
 
 class FContentBundleEditor : public FContentBundleBase, IWorldPartitionCookPackageGenerator
 {
@@ -53,7 +55,12 @@ public:
 	ENGINE_API void ReferenceAllActors();
 	ENGINE_API void UnreferenceAllActors();
 
-	const TWeakObjectPtr<UActorDescContainer>& GetActorDescContainer() const { return ActorDescContainer; };
+	UE_DEPRECATED(5.4, "Use GetActorDescContainerInstance instead")
+	TWeakObjectPtr<UActorDescContainer> GetActorDescContainer() const { return nullptr; };
+
+	const TWeakObjectPtr<UActorDescContainerInstance> GetActorDescContainerInstance() const { return ActorDescContainerInstance; }
+
+
 	ENGINE_API void GenerateStreaming(TArray<FString>* OutPackageToGenerate, bool bIsPIE);
 
 	URuntimeHashExternalStreamingObjectBase* GetStreamingObject() const { return ExternalStreamingObject; }
@@ -87,12 +94,12 @@ private:
 	ENGINE_API void RegisterDelegates();
 	ENGINE_API void UnregisterDelegates();
 	
-	ENGINE_API void OnActorDescAdded(FWorldPartitionActorDesc* ActorDesc);
-	ENGINE_API void OnActorDescRemoved(FWorldPartitionActorDesc* ActorDesc);
+	ENGINE_API void OnActorDescInstanceAdded(FWorldPartitionActorDescInstance* ActorDesc);
+	ENGINE_API void OnActorDescInstanceRemoved(FWorldPartitionActorDescInstance* ActorDesc);
 
 	TObjectPtr<class UContentBundleUnsavedActorMonitor> UnsavedActorMonitor;
 
-	TWeakObjectPtr<UActorDescContainer> ActorDescContainer;
+	TWeakObjectPtr<UActorDescContainerInstance> ActorDescContainerInstance;
 	FWorldDataLayersReference WorldDataLayersActorReference;
 
 	TArray<FWorldPartitionReference> ForceLoadedActors;
diff --git a/Engine/Source/Runtime/Engine/Public/WorldPartition/ContentBundle/WorldDataLayerReference.h b/Engine/Source/Runtime/Engine/Public/WorldPartition/ContentBundle/WorldDataLayerReference.h
--- a/Engine/Source/Runtime/Engine/Public/WorldPartition/ContentBundle/WorldDataLayerReference.h
+++ b/Engine/Source/Runtime/Engine/Public/WorldPartition/ContentBundle/WorldDataLayerReference.h
@@ -13,7 +13,10 @@ class FWorldDataLayersReference
 {
 public:
 	ENGINE_API FWorldDataLayersReference();
-	ENGINE_API FWorldDataLayersReference(UActorDescContainer* Container, FName WorldDataLayerName); // Retrieve Actor
+
+	UE_DEPRECATED(5.4, "Use UActorDescContainerInstance version instead")
+	ENGINE_API FWorldDataLayersReference(UActorDescContainer* Container, FName WorldDataLayerName) {}
+	ENGINE_API FWorldDataLayersReference(UActorDescContainerInstance* ContainerInstance, FName WorldDataLayerName); // Retrieve Actor
 	ENGINE_API FWorldDataLayersReference(const FActorSpawnParameters& SpawnParameters); // Retrieve and Create Actor if Not Found, FActorSpawnParameters::Name must be set
 	ENGINE_API FWorldDataLayersReference(FWorldDataLayersReference&& Other);
 
@@ -32,7 +35,7 @@ public:
 	ENGINE_API FWorldDataLayersReference& operator=(FWorldDataLayersReference&& Other);
 
 private:
-	ENGINE_API bool TrySetReference(UActorDescContainer* Container, FName WorldDataLayerName);
+	ENGINE_API bool TrySetReference(UActorDescContainerInstance* ContainerInstance, FName WorldDataLayerName);
 
 	TVariant<AWorldDataLayers*, FWorldPartitionReference> WorldDataLayersVariant;
 };
diff --git a/Engine/Source/Runtime/Engine/Public/WorldPartition/DataLayer/DataLayerManager.h b/Engine/Source/Runtime/Engine/Public/WorldPartition/DataLayer/DataLayerManager.h
--- a/Engine/Source/Runtime/Engine/Public/WorldPartition/DataLayer/DataLayerManager.h
+++ b/Engine/Source/Runtime/Engine/Public/WorldPartition/DataLayer/DataLayerManager.h
@@ -17,6 +17,7 @@ class UDataLayerInstanceWithAsset;
 class UDataLayerInstance;
 class UDataLayerLoadingPolicy;
 class UActorDescContainer;
+class UActorDescContainerInstance;
 class UCanvas;
 class UWorld;
 
@@ -157,13 +158,15 @@ private:
 	ENGINE_API uint32 GetDataLayerEditorContextHash() const;
 	//~ End Editor Context
 
-	//~ Begin ActorDesc
-	ENGINE_API void OnActorDescContainerInitialized(UActorDescContainer* InActorDescContainer) const;
+	//~ Begin WorldPartitionActorDescInstance
 	ENGINE_API bool CanResolveDataLayers() const;
 	ENGINE_API void ResolveActorDescContainersDataLayers() const;
-	ENGINE_API void ResolveActorDescContainerDataLayers(UActorDescContainer* InActorDescContainer) const;
-	ENGINE_API void ResolveActorDescDataLayers(FWorldPartitionActorDesc* InActorDesc) const;
-	ENGINE_API void ResolveActorDescContainerDataLayersInternal(UActorDescContainer* InActorDescContainer, FWorldPartitionActorDesc* InActorDesc) const;
+	
+	void OnActorDescContainerInstanceInitialized(UActorDescContainerInstance* InActorDescContainerInstance);
+	void ResolveActorDescContainerInstanceDataLayers(UActorDescContainerInstance* InActorDescContainerInstance) const;
+	void ResolveActorDescInstanceDataLayers(FWorldPartitionActorDescInstance* InActorDescInstance) const;
+	void ResolveActorDescContainerInstanceDataLayersInternal(UActorDescContainerInstance* InActorDescContainerInstance, FWorldPartitionActorDescInstance* InActorDescInstance) const;
+
 	//~ End
 
 	//~ Begin Editor loading
@@ -201,7 +204,7 @@ private:
 	friend class FActorBrowsingMode;
 	friend class UDataLayerEditorSubsystem;
 	friend class UActorPartitionSubsystem;
-	friend class UActorDescContainer;
+	friend class UActorDescContainerInstance;
 	friend class AWorldSettings;
 	friend class AActor;
 #endif
diff --git a/Engine/Source/Runtime/Engine/Public/WorldPartition/DataLayer/DataLayerUtils.h b/Engine/Source/Runtime/Engine/Public/WorldPartition/DataLayer/DataLayerUtils.h
--- a/Engine/Source/Runtime/Engine/Public/WorldPartition/DataLayer/DataLayerUtils.h
+++ b/Engine/Source/Runtime/Engine/Public/WorldPartition/DataLayer/DataLayerUtils.h
@@ -7,7 +7,6 @@
 #include "UObject/UObjectGlobals.h"
 #include "WorldPartition/DataLayer/DataLayerType.h"
 #include "WorldPartition/DataLayer/DataLayerInstance.h"
-#include "WorldPartition/WorldPartitionActorDescView.h"
 #include "WorldPartition/WorldPartitionStreamingGeneration.h"
 
 class UDataLayerManager;
@@ -28,19 +27,26 @@ public:
 	}
 
 	static ENGINE_API TArray<FName> ResolvedDataLayerInstanceNames(const UDataLayerManager* InDataLayerManager, const FWorldPartitionActorDesc* InActorDesc, const TArray<const FWorldDataLayersActorDesc*>& InWorldDataLayersActorDescs = TArray<const FWorldDataLayersActorDesc*>());
-	
-	static ENGINE_API bool ResolveRuntimeDataLayerInstanceNames(const UDataLayerManager* InDataLayerManager, const FWorldPartitionActorDescView& InActorDescView, const TArray<const FWorldPartitionActorDescView*>& InWorldDataLayerViews, TArray<FName>& OutRuntimeDataLayerInstanceNames);
 
-	static ENGINE_API const FDataLayerInstanceDesc* GetDataLayerInstanceDescFromInstanceName(const TArray<const FWorldDataLayersActorDesc*>& InWorldDataLayersActorDescs, const FName& InDataLayerInstanceName);
+	static ENGINE_API bool ResolveRuntimeDataLayerInstanceNames(const UDataLayerManager* InDataLayerManager, const IWorldPartitionActorDescInstanceView& InActorDescView, const FStreamingGenerationActorDescViewMap& ActorDescViewMap, TArray<FName>& OutRuntimeDataLayerInstanceNames);
 
-	static ENGINE_API const FDataLayerInstanceDesc* GetDataLayerInstanceDescFromAssetPath(const TArray<const FWorldDataLayersActorDesc*>& InWorldDataLayersActorDescs, const FName& InDataLayerAssetPath);
+	PRAGMA_DISABLE_DEPRECATION_WARNINGS
 
-	static ENGINE_API TArray<const FWorldDataLayersActorDesc*> FindWorldDataLayerActorDescs(const TArray<const FWorldPartitionActorDescView*>& InWorldDataLayerViews);
+	UE_DEPRECATED(5.4, "Use FStreamingGenerationActorDescView version instead")
+	static ENGINE_API bool ResolveRuntimeDataLayerInstanceNames(const UDataLayerManager* InDataLayerManager, const class FWorldPartitionActorDescView& InActorDescView, const FStreamingGenerationActorDescViewMap& ActorDescViewMap, TArray<FName>& OutRuntimeDataLayerInstanceNames) { return false; }
+
+	PRAGMA_ENABLE_DEPRECATION_WARNINGS
+
+	static ENGINE_API const FDataLayerInstanceDesc* GetDataLayerInstanceDescFromInstanceName(const TArray<const FWorldDataLayersActorDesc*>& InWorldDataLayersActorDescs, const FName& DataLayerInstanceName);
+
+	static ENGINE_API const FDataLayerInstanceDesc* GetDataLayerInstanceDescFromAssetPath(const TArray<const FWorldDataLayersActorDesc*>& InWorldDataLayersActorDescs, const FName& DataLayerAssetPath);
+
+	static ENGINE_API TArray<const FWorldDataLayersActorDesc*> FindWorldDataLayerActorDescs(const FStreamingGenerationActorDescViewMap& ActorDescViewMap);
 
 	static ENGINE_API bool AreWorldDataLayersActorDescsSane(const TArray<const FWorldDataLayersActorDesc*>& InWorldDataLayersActorDescs);
 
 	static ENGINE_API FString GenerateUniqueDataLayerShortName(const UDataLayerManager* InDataLayerManager, const FString& InNewShortName);
-	
+
 	static ENGINE_API bool SetDataLayerShortName(UDataLayerInstance* InDataLayerInstance, const FString& InNewShortName);
 
 	static ENGINE_API bool FindDataLayerByShortName(const UDataLayerManager* InDataLayerManager, const FString& InShortName, TSet<UDataLayerInstance*>& OutDataLayerInstances);
@@ -50,4 +56,4 @@ public:
 	{
 		return InShortName.TrimStartAndEnd().Replace(TEXT("\""), TEXT(""));
 	}
-};
+};
\ No newline at end of file
diff --git a/Engine/Source/Runtime/Engine/Public/WorldPartition/ErrorHandling/WorldPartitionStreamingGenerationErrorHandler.h b/Engine/Source/Runtime/Engine/Public/WorldPartition/ErrorHandling/WorldPartitionStreamingGenerationErrorHandler.h
--- a/Engine/Source/Runtime/Engine/Public/WorldPartition/ErrorHandling/WorldPartitionStreamingGenerationErrorHandler.h
+++ b/Engine/Source/Runtime/Engine/Public/WorldPartition/ErrorHandling/WorldPartitionStreamingGenerationErrorHandler.h
@@ -3,53 +3,53 @@
 
 #if WITH_EDITOR
 #include "CoreMinimal.h"
-#include "WorldPartition/WorldPartitionActorDescView.h"
+#include "WorldPartition/WorldPartitionActorDescInstanceViewInterface.h"
 
 class UDataLayerInstance;
 class UDataLayerInstanceWithAsset;
 
+// deprecated
+class FWorldPartitionActorDescView;
+
 class IStreamingGenerationErrorHandler
 {
 public:
 	virtual ~IStreamingGenerationErrorHandler() {}
 
-	UE_DEPRECATED(5.2, "OnInvalidReference is deprecated, use the version which takes an optional actor descriptor view.")
-	virtual void OnInvalidReference(const FWorldPartitionActorDescView& ActorDescView, const FGuid& ReferenceGuid) {}
-
-	/** 
+	/**
 	 * Called when an actor has an invalid runtime grid.
 	 */
-	virtual void OnInvalidRuntimeGrid(const FWorldPartitionActorDescView& ActorDescView, FName GridName) = 0;
+	virtual void OnInvalidRuntimeGrid(const IWorldPartitionActorDescInstanceView& ActorDescView, FName GridName) = 0;
 
-	/** 
+	/**
 	 * Called when an actor references an invalid actor.
 	 */
-	virtual void OnInvalidReference(const FWorldPartitionActorDescView& ActorDescView, const FGuid& ReferenceGuid, FWorldPartitionActorDescView* ReferenceActorDescView) = 0;
+	virtual void OnInvalidReference(const IWorldPartitionActorDescInstanceView& ActorDescView, const FGuid& ReferenceGuid, IWorldPartitionActorDescInstanceView* ReferenceActorDescView) = 0;
 
-	/** 
+	/**
 	 * Called when an actor references an actor using a different grid placement.
 	 */
-	virtual void OnInvalidReferenceGridPlacement(const FWorldPartitionActorDescView& ActorDescView, const FWorldPartitionActorDescView& ReferenceActorDescView) = 0;
+	virtual void OnInvalidReferenceGridPlacement(const IWorldPartitionActorDescInstanceView& ActorDescView, const IWorldPartitionActorDescInstanceView& ReferenceActorDescView) = 0;
 
-	/** 
+	/**
 	 * Called when an actor references an actor using a different set of data layers.
 	 */
-	virtual void OnInvalidReferenceDataLayers(const FWorldPartitionActorDescView& ActorDescView, const FWorldPartitionActorDescView& ReferenceActorDescView) = 0;
+	virtual void OnInvalidReferenceDataLayers(const IWorldPartitionActorDescInstanceView& ActorDescView, const IWorldPartitionActorDescInstanceView& ReferenceActorDescView) = 0;
 
 	/**
 	 * Called when an actor references an actor using a different RuntimeGrid.
 	 */
-	virtual void OnInvalidReferenceRuntimeGrid(const FWorldPartitionActorDescView& ActorDescView, const FWorldPartitionActorDescView& ReferenceActorDescView) = 0;
+	virtual void OnInvalidReferenceRuntimeGrid(const IWorldPartitionActorDescInstanceView& ActorDescView, const IWorldPartitionActorDescInstanceView& ReferenceActorDescView) = 0;
 
-	/** 
+	/**
 	 * Called when the level script references a streamed actor.
 	 */
-	virtual void OnInvalidReferenceLevelScriptStreamed(const FWorldPartitionActorDescView& ActorDescView) = 0;
+	virtual void OnInvalidReferenceLevelScriptStreamed(const IWorldPartitionActorDescInstanceView& ActorDescView) = 0;
 
-	/** 
+	/**
 	 * Called when an actor descriptor references an actor using data layers.
 	 */
-	virtual void OnInvalidReferenceLevelScriptDataLayers(const FWorldPartitionActorDescView& ActorDescView) = 0;
+	virtual void OnInvalidReferenceLevelScriptDataLayers(const IWorldPartitionActorDescInstanceView& ActorDescView) = 0;
 
 	/**
 	 * Called when a data layer instance does not have a data layer asset
@@ -69,7 +69,8 @@ public:
 	/**
 	 * Called when an actor needs to be resaved.
 	 */
-	virtual void OnActorNeedsResave(const FWorldPartitionActorDescView& ActorDescView) = 0;
+	virtual void OnActorNeedsResave(const IWorldPartitionActorDescInstanceView& ActorDescView) = 0;
+
 
 	/**
 	 * Called when a level instance actor has errrors
@@ -83,20 +84,66 @@ public:
 		CirculalReference
 	};
 
-	virtual void OnLevelInstanceInvalidWorldAsset(const FWorldPartitionActorDescView& ActorDescView, FName WorldAsset, ELevelInstanceInvalidReason Reason) = 0;
+	virtual void OnLevelInstanceInvalidWorldAsset(const IWorldPartitionActorDescInstanceView& ActorDescView, FName WorldAsset, ELevelInstanceInvalidReason Reason) = 0;
 
 	/**
 	 * Called when an actor descriptor references another actor with a different set of actor filter.
 	 */
-	virtual void OnInvalidActorFilterReference(const FWorldPartitionActorDescView& ActorDescView, const FWorldPartitionActorDescView& ReferenceActorDescView) = 0;
+	virtual void OnInvalidActorFilterReference(const IWorldPartitionActorDescInstanceView& ActorDescView, const IWorldPartitionActorDescInstanceView& ReferenceActorDescView) = 0;
 
 	/**
 	 * Called when an actor descriptor has an invalid HLOD layer.
 	 */
-	virtual void OnInvalidHLODLayer(const FWorldPartitionActorDescView& ActorDescView) = 0;
+	virtual void OnInvalidHLODLayer(const IWorldPartitionActorDescInstanceView& ActorDescView) = 0;
 
 	// Helpers
-	static ENGINE_API FString GetActorName(const FWorldPartitionActorDescView& ActorDescView);
-	static ENGINE_API FString GetFullActorName(const FWorldPartitionActorDescView& ActorDescView);
+	static ENGINE_API FString GetActorName(const IWorldPartitionActorDescInstanceView& ActorDescView);
+	static ENGINE_API FString GetFullActorName(const IWorldPartitionActorDescInstanceView& ActorDescView);
+	
+	PRAGMA_DISABLE_DEPRECATION_WARNINGS
+
+	UE_DEPRECATED(5.2, "OnInvalidReference is deprecated, use the version which takes an optional actor descriptor view.")
+	virtual void OnInvalidReference(const FWorldPartitionActorDescView& ActorDescView, const FGuid& ReferenceGuid) {}
+
+	UE_DEPRECATED(5.4, "Use IWorldPartitionActorDescInstanceView version instead")
+	virtual void OnInvalidRuntimeGrid(const FWorldPartitionActorDescView& ActorDescView, FName GridName) {}
+
+	UE_DEPRECATED(5.4, "Use IWorldPartitionActorDescInstanceView version instead")
+	virtual void OnInvalidReference(const FWorldPartitionActorDescView& ActorDescView, const FGuid& ReferenceGuid, FWorldPartitionActorDescView* ReferenceActorDescView) {}
+
+	UE_DEPRECATED(5.4, "Use IWorldPartitionActorDescInstanceView version instead")
+	virtual void OnInvalidReferenceGridPlacement(const FWorldPartitionActorDescView& ActorDescView, const FWorldPartitionActorDescView& ReferenceActorDescView) {}
+
+	UE_DEPRECATED(5.4, "Use IWorldPartitionActorDescInstanceView version instead")
+	virtual void OnInvalidReferenceDataLayers(const FWorldPartitionActorDescView& ActorDescView, const FWorldPartitionActorDescView& ReferenceActorDescView) {}
+
+	UE_DEPRECATED(5.4, "Use IWorldPartitionActorDescInstanceView version instead")
+	virtual void OnInvalidReferenceRuntimeGrid(const FWorldPartitionActorDescView& ActorDescView, const FWorldPartitionActorDescView& ReferenceActorDescView) {}
+
+	UE_DEPRECATED(5.4, "Use IWorldPartitionActorDescInstanceView version instead")
+	virtual void OnInvalidReferenceLevelScriptStreamed(const FWorldPartitionActorDescView& ActorDescView) {}
+	
+	UE_DEPRECATED(5.4, "Use IWorldPartitionActorDescInstanceView version instead")
+	virtual void OnInvalidReferenceLevelScriptDataLayers(const FWorldPartitionActorDescView& ActorDescView) {}
+	
+	UE_DEPRECATED(5.4, "Use IWorldPartitionActorDescInstanceView version instead")
+	virtual void OnActorNeedsResave(const FWorldPartitionActorDescView& ActorDescView) {}
+
+	UE_DEPRECATED(5.4, "Use IWorldPartitionActorDescInstanceView version instead")
+	virtual void OnLevelInstanceInvalidWorldAsset(const FWorldPartitionActorDescView& ActorDescView, FName WorldAsset, ELevelInstanceInvalidReason Reason) {}
+
+	UE_DEPRECATED(5.4, "Use IWorldPartitionActorDescInstanceView version instead")
+	virtual void OnInvalidActorFilterReference(const FWorldPartitionActorDescView& ActorDescView, const FWorldPartitionActorDescView& ReferenceActorDescView) {}
+	
+	UE_DEPRECATED(5.4, "Use IWorldPartitionActorDescInstanceView version instead")
+	virtual void OnInvalidHLODLayer(const FWorldPartitionActorDescView& ActorDescView) {}
+
+	UE_DEPRECATED(5.4, "Use IWorldPartitionActorDescInstanceView version instead")
+	static ENGINE_API FString GetActorName(const FWorldPartitionActorDescView& ActorDescView) { return FString(); }
+
+	UE_DEPRECATED(5.4, "Use IWorldPartitionActorDescInstanceView version instead")
+	static ENGINE_API FString GetFullActorName(const FWorldPartitionActorDescView& ActorDescView) { return FString(); }
+
+	PRAGMA_ENABLE_DEPRECATION_WARNINGS
 };
 #endif
diff --git a/Engine/Source/Runtime/Engine/Public/WorldPartition/ErrorHandling/WorldPartitionStreamingGenerationLogErrorHandler.h b/Engine/Source/Runtime/Engine/Public/WorldPartition/ErrorHandling/WorldPartitionStreamingGenerationLogErrorHandler.h
--- a/Engine/Source/Runtime/Engine/Public/WorldPartition/ErrorHandling/WorldPartitionStreamingGenerationLogErrorHandler.h
+++ b/Engine/Source/Runtime/Engine/Public/WorldPartition/ErrorHandling/WorldPartitionStreamingGenerationLogErrorHandler.h
@@ -7,19 +7,19 @@
 class FStreamingGenerationLogErrorHandler : public IStreamingGenerationErrorHandler
 {
 public:
-	ENGINE_API virtual void OnInvalidRuntimeGrid(const FWorldPartitionActorDescView& ActorDescView, FName GridName) override;
-	ENGINE_API virtual void OnInvalidReference(const FWorldPartitionActorDescView& ActorDescView, const FGuid& ReferenceGuid, FWorldPartitionActorDescView* ReferenceActorDescView) override;
-	ENGINE_API virtual void OnInvalidReferenceGridPlacement(const FWorldPartitionActorDescView& ActorDescView, const FWorldPartitionActorDescView& ReferenceActorDescView) override;
-	ENGINE_API virtual void OnInvalidReferenceDataLayers(const FWorldPartitionActorDescView& ActorDescView, const FWorldPartitionActorDescView& ReferenceActorDescView) override;
-	ENGINE_API virtual void OnInvalidReferenceLevelScriptStreamed(const FWorldPartitionActorDescView& ActorDescView) override;
-	ENGINE_API virtual void OnInvalidReferenceLevelScriptDataLayers(const FWorldPartitionActorDescView& ActorDescView) override;
-	ENGINE_API virtual void OnInvalidReferenceRuntimeGrid(const FWorldPartitionActorDescView& ActorDescView, const FWorldPartitionActorDescView& ReferenceActorDescView) override;
+	ENGINE_API virtual void OnInvalidRuntimeGrid(const IWorldPartitionActorDescInstanceView& ActorDescView, FName GridName) override;
+	ENGINE_API virtual void OnInvalidReference(const IWorldPartitionActorDescInstanceView& ActorDescView, const FGuid& ReferenceGuid, IWorldPartitionActorDescInstanceView* ReferenceActorDescView) override;
+	ENGINE_API virtual void OnInvalidReferenceGridPlacement(const IWorldPartitionActorDescInstanceView& ActorDescView, const IWorldPartitionActorDescInstanceView& ReferenceActorDescView) override;
+	ENGINE_API virtual void OnInvalidReferenceDataLayers(const IWorldPartitionActorDescInstanceView& ActorDescView, const IWorldPartitionActorDescInstanceView& ReferenceActorDescView) override;
+	ENGINE_API virtual void OnInvalidReferenceLevelScriptStreamed(const IWorldPartitionActorDescInstanceView& ActorDescView) override;
+	ENGINE_API virtual void OnInvalidReferenceLevelScriptDataLayers(const IWorldPartitionActorDescInstanceView& ActorDescView) override;
+	ENGINE_API virtual void OnInvalidReferenceRuntimeGrid(const IWorldPartitionActorDescInstanceView& ActorDescView, const IWorldPartitionActorDescInstanceView& ReferenceActorDescView) override;
 	ENGINE_API virtual void OnInvalidReferenceDataLayerAsset(const UDataLayerInstanceWithAsset* DataLayerInstance) override;
 	ENGINE_API virtual void OnDataLayerHierarchyTypeMismatch(const UDataLayerInstance* DataLayerInstance, const UDataLayerInstance* Parent) override;
 	ENGINE_API virtual void OnDataLayerAssetConflict(const UDataLayerInstanceWithAsset* DataLayerInstance, const UDataLayerInstanceWithAsset* ConflictingDataLayerInstance) override;
-	ENGINE_API virtual void OnActorNeedsResave(const FWorldPartitionActorDescView& ActorDescView) override;
-	ENGINE_API virtual void OnLevelInstanceInvalidWorldAsset(const FWorldPartitionActorDescView& ActorDescView, FName WorldAsset, ELevelInstanceInvalidReason Reason) override;
-	ENGINE_API virtual void OnInvalidActorFilterReference(const FWorldPartitionActorDescView& ActorDescView, const FWorldPartitionActorDescView& ReferenceActorDescView) override;
-	ENGINE_API virtual void OnInvalidHLODLayer(const FWorldPartitionActorDescView& ActorDescView) override;
+	ENGINE_API virtual void OnActorNeedsResave(const IWorldPartitionActorDescInstanceView& ActorDescView) override;
+	ENGINE_API virtual void OnLevelInstanceInvalidWorldAsset(const IWorldPartitionActorDescInstanceView& ActorDescView, FName WorldAsset, ELevelInstanceInvalidReason Reason) override;
+	ENGINE_API virtual void OnInvalidActorFilterReference(const IWorldPartitionActorDescInstanceView& ActorDescView, const IWorldPartitionActorDescInstanceView& ReferenceActorDescView) override;
+	ENGINE_API virtual void OnInvalidHLODLayer(const IWorldPartitionActorDescInstanceView& ActorDescView) override;
 };
 #endif
diff --git a/Engine/Source/Runtime/Engine/Public/WorldPartition/ErrorHandling/WorldPartitionStreamingGenerationNullErrorHandler.h b/Engine/Source/Runtime/Engine/Public/WorldPartition/ErrorHandling/WorldPartitionStreamingGenerationNullErrorHandler.h
--- a/Engine/Source/Runtime/Engine/Public/WorldPartition/ErrorHandling/WorldPartitionStreamingGenerationNullErrorHandler.h
+++ b/Engine/Source/Runtime/Engine/Public/WorldPartition/ErrorHandling/WorldPartitionStreamingGenerationNullErrorHandler.h
@@ -7,19 +7,19 @@
 class FStreamingGenerationNullErrorHandler : public IStreamingGenerationErrorHandler
 {
 public:
-	virtual void OnInvalidRuntimeGrid(const FWorldPartitionActorDescView& ActorDescView, FName GridName) override {}
-	virtual void OnInvalidReference(const FWorldPartitionActorDescView& ActorDescView, const FGuid& ReferenceGuid, FWorldPartitionActorDescView* ReferenceActorDescView) override {}
-	virtual void OnInvalidReferenceGridPlacement(const FWorldPartitionActorDescView& ActorDescView, const FWorldPartitionActorDescView& ReferenceActorDescView) override {}
-	virtual void OnInvalidReferenceDataLayers(const FWorldPartitionActorDescView& ActorDescView, const FWorldPartitionActorDescView& ReferenceActorDescView) override {}
-	virtual void OnInvalidReferenceLevelScriptStreamed(const FWorldPartitionActorDescView& ActorDescView) override {}
-	virtual void OnInvalidReferenceLevelScriptDataLayers(const FWorldPartitionActorDescView& ActorDescView) override {}
-	virtual void OnInvalidReferenceRuntimeGrid(const FWorldPartitionActorDescView& ActorDescView, const FWorldPartitionActorDescView& ReferenceActorDescView) override {}
+	virtual void OnInvalidRuntimeGrid(const IWorldPartitionActorDescInstanceView& ActorDescView, FName GridName) override {}
+	virtual void OnInvalidReference(const IWorldPartitionActorDescInstanceView& ActorDescView, const FGuid& ReferenceGuid, IWorldPartitionActorDescInstanceView* ReferenceActorDescView) override {}
+	virtual void OnInvalidReferenceGridPlacement(const IWorldPartitionActorDescInstanceView& ActorDescView, const IWorldPartitionActorDescInstanceView& ReferenceActorDescView) override {}
+	virtual void OnInvalidReferenceDataLayers(const IWorldPartitionActorDescInstanceView& ActorDescView, const IWorldPartitionActorDescInstanceView& ReferenceActorDescView) override {}
+	virtual void OnInvalidReferenceLevelScriptStreamed(const IWorldPartitionActorDescInstanceView& ActorDescView) override {}
+	virtual void OnInvalidReferenceLevelScriptDataLayers(const IWorldPartitionActorDescInstanceView& ActorDescView) override {}
+	virtual void OnInvalidReferenceRuntimeGrid(const IWorldPartitionActorDescInstanceView& ActorDescView, const IWorldPartitionActorDescInstanceView& ReferenceActorDescView) override {}
 	virtual void OnInvalidReferenceDataLayerAsset(const UDataLayerInstanceWithAsset* DataLayerInstance) override {}
 	virtual void OnDataLayerHierarchyTypeMismatch(const UDataLayerInstance* DataLayerInstance, const UDataLayerInstance* Parent) override {};
 	virtual void OnDataLayerAssetConflict(const UDataLayerInstanceWithAsset* DataLayerInstance, const UDataLayerInstanceWithAsset* ConflictingDataLayerInstance) override {};
-	virtual void OnActorNeedsResave(const FWorldPartitionActorDescView& ActorDescView) override {};
-	virtual void OnLevelInstanceInvalidWorldAsset(const FWorldPartitionActorDescView& ActorDescView, FName WorldAsset, ELevelInstanceInvalidReason Reason) override {};
-	virtual void OnInvalidActorFilterReference(const FWorldPartitionActorDescView& ActorDescView, const FWorldPartitionActorDescView& ReferenceActorDescView) override {};
-	virtual void OnInvalidHLODLayer(const FWorldPartitionActorDescView& ActorDescView) override {}
+	virtual void OnActorNeedsResave(const IWorldPartitionActorDescInstanceView& ActorDescView) override {};
+	virtual void OnLevelInstanceInvalidWorldAsset(const IWorldPartitionActorDescInstanceView& ActorDescView, FName WorldAsset, ELevelInstanceInvalidReason Reason) override {};
+	virtual void OnInvalidActorFilterReference(const IWorldPartitionActorDescInstanceView& ActorDescView, const IWorldPartitionActorDescInstanceView& ReferenceActorDescView) override {};
+	virtual void OnInvalidHLODLayer(const IWorldPartitionActorDescInstanceView& ActorDescView) override {}
 };
 #endif
diff --git a/Engine/Source/Runtime/Engine/Public/WorldPartition/ErrorHandling/WorldPartitionStreamingGenerationProxyErrorHandler.h b/Engine/Source/Runtime/Engine/Public/WorldPartition/ErrorHandling/WorldPartitionStreamingGenerationProxyErrorHandler.h
--- a/Engine/Source/Runtime/Engine/Public/WorldPartition/ErrorHandling/WorldPartitionStreamingGenerationProxyErrorHandler.h
+++ b/Engine/Source/Runtime/Engine/Public/WorldPartition/ErrorHandling/WorldPartitionStreamingGenerationProxyErrorHandler.h
@@ -11,20 +11,20 @@ public:
 		: InnerErrorHandler(InInnerErrorHandler)
 	{}
 
-	ENGINE_API virtual void OnInvalidRuntimeGrid(const FWorldPartitionActorDescView& ActorDescView, FName GridName) override;
-	ENGINE_API virtual void OnInvalidReference(const FWorldPartitionActorDescView& ActorDescView, const FGuid& ReferenceGuid, FWorldPartitionActorDescView* ReferenceActorDescView) override;
-	ENGINE_API virtual void OnInvalidReferenceGridPlacement(const FWorldPartitionActorDescView& ActorDescView, const FWorldPartitionActorDescView& ReferenceActorDescView) override;
-	ENGINE_API virtual void OnInvalidReferenceDataLayers(const FWorldPartitionActorDescView& ActorDescView, const FWorldPartitionActorDescView& ReferenceActorDescView) override;
-	ENGINE_API virtual void OnInvalidReferenceLevelScriptStreamed(const FWorldPartitionActorDescView& ActorDescView) override;
-	ENGINE_API virtual void OnInvalidReferenceLevelScriptDataLayers(const FWorldPartitionActorDescView& ActorDescView) override;
-	ENGINE_API virtual void OnInvalidReferenceRuntimeGrid(const FWorldPartitionActorDescView& ActorDescView, const FWorldPartitionActorDescView& ReferenceActorDescView) override;
+	ENGINE_API virtual void OnInvalidRuntimeGrid(const IWorldPartitionActorDescInstanceView& ActorDescView, FName GridName) override;
+	ENGINE_API virtual void OnInvalidReference(const IWorldPartitionActorDescInstanceView& ActorDescView, const FGuid& ReferenceGuid, IWorldPartitionActorDescInstanceView* ReferenceActorDescView) override;
+	ENGINE_API virtual void OnInvalidReferenceGridPlacement(const IWorldPartitionActorDescInstanceView& ActorDescView, const IWorldPartitionActorDescInstanceView& ReferenceActorDescView) override;
+	ENGINE_API virtual void OnInvalidReferenceDataLayers(const IWorldPartitionActorDescInstanceView& ActorDescView, const IWorldPartitionActorDescInstanceView& ReferenceActorDescView) override;
+	ENGINE_API virtual void OnInvalidReferenceLevelScriptStreamed(const IWorldPartitionActorDescInstanceView& ActorDescView) override;
+	ENGINE_API virtual void OnInvalidReferenceLevelScriptDataLayers(const IWorldPartitionActorDescInstanceView& ActorDescView) override;
+	ENGINE_API virtual void OnInvalidReferenceRuntimeGrid(const IWorldPartitionActorDescInstanceView& ActorDescView, const IWorldPartitionActorDescInstanceView& ReferenceActorDescView) override;
 	ENGINE_API virtual void OnInvalidReferenceDataLayerAsset(const UDataLayerInstanceWithAsset* DataLayerInstance) override;
 	ENGINE_API virtual void OnDataLayerHierarchyTypeMismatch(const UDataLayerInstance* DataLayerInstance, const UDataLayerInstance* Parent) override;
 	ENGINE_API virtual void OnDataLayerAssetConflict(const UDataLayerInstanceWithAsset* DataLayerInstance, const UDataLayerInstanceWithAsset* ConflictingDataLayerInstance) override;
-	ENGINE_API virtual void OnActorNeedsResave(const FWorldPartitionActorDescView& ActorDescView) override;
-	ENGINE_API virtual void OnLevelInstanceInvalidWorldAsset(const FWorldPartitionActorDescView& ActorDescView, FName WorldAsset, ELevelInstanceInvalidReason Reason) override;
-	ENGINE_API virtual void OnInvalidActorFilterReference(const FWorldPartitionActorDescView& ActorDescView, const FWorldPartitionActorDescView& ReferenceActorDescView) override;
-	ENGINE_API virtual void OnInvalidHLODLayer(const FWorldPartitionActorDescView& ActorDescView) override;
+	ENGINE_API virtual void OnActorNeedsResave(const IWorldPartitionActorDescInstanceView& ActorDescView) override;
+	ENGINE_API virtual void OnLevelInstanceInvalidWorldAsset(const IWorldPartitionActorDescInstanceView& ActorDescView, FName WorldAsset, ELevelInstanceInvalidReason Reason) override;
+	ENGINE_API virtual void OnInvalidActorFilterReference(const IWorldPartitionActorDescInstanceView& ActorDescView, const IWorldPartitionActorDescInstanceView& ReferenceActorDescView) override;
+	ENGINE_API virtual void OnInvalidHLODLayer(const IWorldPartitionActorDescInstanceView& ActorDescView) override;
 
 protected:
 	IStreamingGenerationErrorHandler* InnerErrorHandler;
diff --git a/Engine/Source/Runtime/Engine/Public/WorldPartition/HLOD/HLODActorDesc.h b/Engine/Source/Runtime/Engine/Public/WorldPartition/HLOD/HLODActorDesc.h
--- a/Engine/Source/Runtime/Engine/Public/WorldPartition/HLOD/HLODActorDesc.h
+++ b/Engine/Source/Runtime/Engine/Public/WorldPartition/HLOD/HLODActorDesc.h
@@ -34,7 +34,6 @@ public:
 	static ENGINE_API int64 GetPackageSize(const AWorldPartitionHLOD* InHLODActor);
 
 	//~ Begin FWorldPartitionActorDesc Interface.
-	virtual bool IsRuntimeRelevant(const FActorContainerID& InContainerID) const override { return !bIsForcedNonSpatiallyLoaded; }
 	virtual bool ShouldValidateRuntimeGrid() const override { return false; }
 	virtual FBox GetEditorBounds() const override { return EditorBounds; }
 	//~ End FWorldPartitionActorDesc Interface.
@@ -47,6 +46,7 @@ protected:
 	ENGINE_API virtual bool Equals(const FWorldPartitionActorDesc* Other) const override;
 	virtual uint32 GetSizeOf() const override { return sizeof(FHLODActorDesc); }
 	ENGINE_API virtual void Serialize(FArchive& Ar) override;
+	virtual bool IsRuntimeRelevant(const FWorldPartitionActorDescInstance* InActorDescInstance) const override;
 	//~ End FWorldPartitionActorDesc Interface.
 
 	TArray<FGuid> ChildHLODActors;
diff --git a/Engine/Source/Runtime/Engine/Public/WorldPartition/HLOD/HLODLayer.h b/Engine/Source/Runtime/Engine/Public/WorldPartition/HLOD/HLODLayer.h
--- a/Engine/Source/Runtime/Engine/Public/WorldPartition/HLOD/HLODLayer.h
+++ b/Engine/Source/Runtime/Engine/Public/WorldPartition/HLOD/HLODLayer.h
@@ -23,7 +23,6 @@ class UMaterial;
 class UWorldPartition;
 class UWorldPartitionHLODModifier;
 class FWorldPartitionActorDesc;
-class FWorldPartitionActorDescView;
 
 UENUM()
 enum class EHLODLayerType : uint8
diff --git a/Engine/Source/Runtime/Engine/Public/WorldPartition/Landscape/LandscapeActorDesc.h b/Engine/Source/Runtime/Engine/Public/WorldPartition/Landscape/LandscapeActorDesc.h
--- a/Engine/Source/Runtime/Engine/Public/WorldPartition/Landscape/LandscapeActorDesc.h
+++ b/Engine/Source/Runtime/Engine/Public/WorldPartition/Landscape/LandscapeActorDesc.h
@@ -15,13 +15,12 @@ class FLandscapeActorDesc : public FPartitionActorDesc
 public:
 	ENGINE_API virtual void Init(const AActor* InActor) override;
 	ENGINE_API virtual bool Equals(const FWorldPartitionActorDesc* Other) const override;
-	ENGINE_API virtual void Unload() override;
 	ENGINE_API virtual const FGuid& GetSceneOutlinerParent() const override;
-
+	ENGINE_API virtual FBox GetEditorBounds() const override;
 protected:
 	virtual uint32 GetSizeOf() const override { return sizeof(FLandscapeActorDesc); }
 	ENGINE_API virtual void Serialize(FArchive& Ar) override;
-	ENGINE_API virtual FBox GetEditorBounds() const override;
+	ENGINE_API void OnUnloadingInstance(const FWorldPartitionActorDescInstance* InActorDescInstance) const override;
 
 private:
 	FGuid LandscapeActorGuid;
diff --git a/Engine/Source/Runtime/Engine/Public/WorldPartition/LevelInstance/LevelInstanceActorDesc.h b/Engine/Source/Runtime/Engine/Public/WorldPartition/LevelInstance/LevelInstanceActorDesc.h
--- a/Engine/Source/Runtime/Engine/Public/WorldPartition/LevelInstance/LevelInstanceActorDesc.h
+++ b/Engine/Source/Runtime/Engine/Public/WorldPartition/LevelInstance/LevelInstanceActorDesc.h
@@ -12,8 +12,10 @@
 #include "WorldPartition/Filter/WorldPartitionActorFilter.h"
 
 class ULevelInstanceSubsystem;
-class UActorDescContainer;
+class FWorldPartitionActorDescInstance;
 class IStreamingGenerationErrorHandler;
+class IWorldPartitionActorDescInstanceView;
+class UActorDescContainerInstance;
 enum class ELevelInstanceRuntimeBehavior : uint8;
 
 /**
@@ -24,14 +26,16 @@ class FLevelInstanceActorDesc : public FWorldPartitionActorDesc
 public:
 	ENGINE_API FLevelInstanceActorDesc();
 	ENGINE_API virtual ~FLevelInstanceActorDesc() override;
+		
+	ENGINE_API virtual bool IsChildContainerInstance() const override;
 
-	ENGINE_API virtual bool IsContainerInstance() const override;
-	virtual EWorldPartitionActorFilterType GetContainerFilterType() const { return IsContainerInstance() ? EWorldPartitionActorFilterType::Loading : EWorldPartitionActorFilterType::None; }
-	virtual FName GetContainerPackage() const override { return WorldAsset.GetLongPackageFName(); }
-	ENGINE_API virtual bool GetContainerInstance(FContainerInstance& OutContainerInstance) const override;
-	virtual const FWorldPartitionActorFilter* GetContainerFilter() const override { return &Filter; }
-	ENGINE_API virtual void CheckForErrors(IStreamingGenerationErrorHandler* ErrorHandler) const override;
+	virtual EWorldPartitionActorFilterType GetChildContainerFilterType() const { return IsChildContainerInstance() ? EWorldPartitionActorFilterType::Loading : EWorldPartitionActorFilterType::None; }
+	virtual FName GetChildContainerPackage() const override { return WorldAsset.GetLongPackageFName(); }
+	virtual const FWorldPartitionActorFilter* GetChildContainerFilter() const override { return &Filter; }
+		
+	ENGINE_API virtual void CheckForErrors(const IWorldPartitionActorDescInstanceView* InActorDescView, IStreamingGenerationErrorHandler* ErrorHandler) const override;
 
+	ENGINE_API virtual UActorDescContainer* GetChildContainer() const override { check(IsChildContainerInstance()); return ChildContainer.Get(); }
 protected:
 	ENGINE_API virtual void Init(const AActor* InActor) override;
 	ENGINE_API virtual void Init(const FWorldPartitionActorDescInitData& DescData) override;
@@ -40,21 +44,26 @@ protected:
 	ENGINE_API virtual void TransferWorldData(const FWorldPartitionActorDesc* From) override;
 	virtual uint32 GetSizeOf() const override { return sizeof(FLevelInstanceActorDesc); }
 	ENGINE_API virtual void Serialize(FArchive& Ar) override;
-	ENGINE_API virtual void SetContainer(UActorDescContainer* InContainer, UWorld* InWorldContext) override;
+	ENGINE_API virtual void SetContainer(UActorDescContainer* InContainer) override;
+	ENGINE_API virtual UActorDescContainerInstance* CreateChildContainerInstance(const FWorldPartitionActorDescInstance* InActorDescInstance) const override;
+	ENGINE_API virtual bool GetChildContainerInstance(const FWorldPartitionActorDescInstance* InActorDescInstance, FContainerInstance& OutContainerInstance) const override;
+	ENGINE_API virtual bool GetLoadedChildContainerInstance(const FWorldPartitionActorDescInstance* InActorDescInstance, FLoadedContainerInstance& OutContainerInstance) const override;
+	ENGINE_API FTransform GetChildContainerTransform() const;
 
 	FSoftObjectPath WorldAsset;
 	FTransform LevelInstanceTransform;
 	ELevelInstanceRuntimeBehavior DesiredRuntimeBehavior;
 
-	TWeakObjectPtr<UActorDescContainer> LevelInstanceContainer;
-	TWeakObjectPtr<UWorld> LevelInstanceContainerWorldContext;
+	TWeakObjectPtr<UActorDescContainer> ChildContainer;
 
 	FWorldPartitionActorFilter Filter;
-	bool bIsContainerInstance;
+	bool bIsChildContainerInstance;
+
 private:
-	ENGINE_API bool IsContainerInstanceInternal() const;
-	ENGINE_API void RegisterContainerInstance(UWorld* InWorldContext);
-	ENGINE_API void UnregisterContainerInstance();
+	ENGINE_API bool IsChildContainerInstanceInternal() const;
+	ENGINE_API void RegisterChildContainer();
+	ENGINE_API void UnregisterChildContainer();
 	ENGINE_API void UpdateBounds();
+	FString GetChildContainerName() const;
 };
 #endif
diff --git a/Engine/Source/Runtime/Engine/Public/WorldPartition/LoaderAdapter/LoaderAdapterPinnedActors.h b/Engine/Source/Runtime/Engine/Public/WorldPartition/LoaderAdapter/LoaderAdapterPinnedActors.h
--- a/Engine/Source/Runtime/Engine/Public/WorldPartition/LoaderAdapter/LoaderAdapterPinnedActors.h
+++ b/Engine/Source/Runtime/Engine/Public/WorldPartition/LoaderAdapter/LoaderAdapterPinnedActors.h
@@ -7,6 +7,7 @@
 
 class AActor;
 class FWorldPartitionActorDesc;
+class FWorldPartitionActorDescInstance;
 
 class FLoaderAdapterPinnedActors : public FLoaderAdapterActorList
 {
@@ -15,7 +16,10 @@ public:
 		: FLoaderAdapterActorList(InWorld)
 	{}
 
-	static ENGINE_API bool SupportsPinning(FWorldPartitionActorDesc* InActorDesc);
+	UE_DEPRECATED(5.4, "Use FWorldPartitionActorDescInstance version instead")
+	static ENGINE_API bool SupportsPinning(FWorldPartitionActorDesc* InActorDesc) { return false;}
+
+	static ENGINE_API bool SupportsPinning(FWorldPartitionActorDescInstance* InActorDescInstance);
 	static ENGINE_API bool SupportsPinning(AActor* InActor);
 
 protected:
diff --git a/Engine/Source/Runtime/Engine/Public/WorldPartition/LoaderAdapter/LoaderAdapterSpatial.h b/Engine/Source/Runtime/Engine/Public/WorldPartition/LoaderAdapter/LoaderAdapterSpatial.h
--- a/Engine/Source/Runtime/Engine/Public/WorldPartition/LoaderAdapter/LoaderAdapterSpatial.h
+++ b/Engine/Source/Runtime/Engine/Public/WorldPartition/LoaderAdapter/LoaderAdapterSpatial.h
@@ -23,7 +23,7 @@ protected:
 	virtual bool Intersect(const FBox& Box) const =0;
 
 private:
-	ENGINE_API void HandleIntersectingContainer(UWorldPartition* InWorldPartition, const FTransform InInstanceTransform, const FBox& InBoundingBox, TFunctionRef<void(const FWorldPartitionHandle&)> InOperation) const;
+	ENGINE_API void HandleIntersectingContainer(UWorldPartition* InWorldPartition, const FBox& InBoundingBox, TFunctionRef<void(const FWorldPartitionHandle&)> InOperation) const;
 
 protected:
 	bool bIncludeSpatiallyLoadedActors;
diff --git a/Engine/Source/Runtime/Engine/Public/WorldPartition/WorldPartition.h b/Engine/Source/Runtime/Engine/Public/WorldPartition/WorldPartition.h
--- a/Engine/Source/Runtime/Engine/Public/WorldPartition/WorldPartition.h
+++ b/Engine/Source/Runtime/Engine/Public/WorldPartition/WorldPartition.h
@@ -13,20 +13,22 @@
 #include "WorldPartition/WorldPartitionActorDesc.h"
 #include "WorldPartition/WorldPartitionStreamingSource.h"
 #include "WorldPartition/WorldPartitionHandle.h"
-#include "WorldPartition/ActorDescContainerCollection.h"
+#include "WorldPartition/ActorDescContainerInstanceCollection.h"
 #include "WorldPartition/Cook/WorldPartitionCookPackageGenerator.h"
 
 #if WITH_EDITOR
+#include "WorldPartition/ActorDescContainerInstance.h"
+#include "WorldPartition/WorldPartitionActorDescInstance.h"
 #include "WorldPartition/WorldPartitionStreamingGeneration.h"
 #include "WorldPartition/WorldPartitionActorLoaderInterface.h"
 #include "WorldPartition/WorldPartitionEditorLoaderAdapter.h"
 #include "PackageSourceControlHelper.h"
 #include "CookPackageSplitter.h"
+#include "Delegates/DelegateCombinations.h"
 #endif
 
 #include "WorldPartition.generated.h"
 
-class FWorldPartitionActorDesc;
 class UActorDescContainer;
 class UWorldPartitionEditorHash;
 class UWorldPartitionRuntimeCell;
@@ -108,7 +110,7 @@ public:
 #endif
 
 UCLASS(AutoExpandCategories=(WorldPartition), MinimalAPI)
-class UWorldPartition final : public UObject, public FActorDescContainerCollection, public IWorldPartitionCookPackageGenerator
+class UWorldPartition final : public UObject, public FActorDescContainerInstanceCollection, public IWorldPartitionCookPackageGenerator
 {
 	GENERATED_UCLASS_BODY()
 
@@ -119,6 +121,7 @@ class UWorldPartition final : public UObject, public FActorDescContainerCollecti
 	friend class FWorldPartitionDetails;
 	friend class FUnrealEdMisc;
 	friend class UActorDescContainer;
+	friend class UActorDescContainerInstance;
 
 public:
 #if WITH_EDITOR
@@ -157,15 +160,13 @@ private:
 	// WorldDeletegates Events
 	ENGINE_API void OnWorldRenamed(UWorld* RenamedWorld);
 
-	// ActorDescContainer Events
-	ENGINE_API void OnActorDescAdded(FWorldPartitionActorDesc* NewActorDesc);
-	ENGINE_API void OnActorDescRemoved(FWorldPartitionActorDesc* ActorDesc);
-	ENGINE_API void OnActorDescUpdating(FWorldPartitionActorDesc* ActorDesc);
-	ENGINE_API void OnActorDescUpdated(FWorldPartitionActorDesc* ActorDesc);
+	// ActorDescContainerInstance Events
+	void OnActorDescInstanceAdded(FWorldPartitionActorDescInstance* NewActorDescInstance);
+	void OnActorDescInstanceRemoved(FWorldPartitionActorDescInstance* ActorDescInstance);
+	void OnActorDescInstanceUpdating(FWorldPartitionActorDescInstance* ActorDescInstance);
+	void OnActorDescInstanceUpdated(FWorldPartitionActorDescInstance* ActorDescInstance);
 
-	ENGINE_API bool GetInstancingContext(const FLinkerInstancingContext*& OutInstancingContext) const;
-
-	ENGINE_API void InitializeActorDescContainerEditorStreaming(UActorDescContainer* InActorDescContainer, bool bInHashActorDescs);
+	ENGINE_API void InitializeActorDescContainerEditorStreaming(UActorDescContainerInstance* InActorDescContainer, bool bInHashActorDescs);
 #endif
 
 public:
@@ -196,25 +197,34 @@ public:
 	bool CanGenerateStreaming() const { return !StreamingPolicy; }
 
 	UE_DEPRECATED(5.3, "GenerateStreaming is deprecated, use GenerateStreaming with a param struct instead")
-	ENGINE_API bool GenerateStreaming(TArray<FString>* OutPackagesToGenerate = nullptr);
+	ENGINE_API bool GenerateStreaming(TArray<FString>* OutPackagesToGenerate = nullptr) { return false; }
 
 	UE_DEPRECATED(5.3, "GenerateContainerStreaming is deprecated, use GenerateContainerStreaming with a param struct instead")
-	ENGINE_API bool GenerateContainerStreaming(const UActorDescContainer* ActorDescContainer, TArray<FString>* OutPackagesToGenerate = nullptr);
+	ENGINE_API bool GenerateContainerStreaming(const UActorDescContainer* ActorDescContainer, TArray<FString>* OutPackagesToGenerate = nullptr) { return false; }
 
+	PRAGMA_DISABLE_DEPRECATION_WARNINGS
 	struct FGenerateStreamingParams
 	{
-		FGenerateStreamingParams()
+		FGenerateStreamingParams() 
 			: ErrorHandler(nullptr)
 		{}
-
+		
+		UE_DEPRECATED(5.4, "Use ContainerInstanceCollection instead")
 		FStreamingGenerationActorDescCollection ActorDescCollection;
+		
+		FStreamingGenerationContainerInstanceCollection ContainerInstanceCollection;
+
 		TOptional<const FString> OutputLogPath;
 		IStreamingGenerationErrorHandler* ErrorHandler;
 
-		FGenerateStreamingParams& SetActorDescContainer(const UActorDescContainer* InActorDescContainer) { ActorDescCollection.AddContainer(InActorDescContainer); return *this; }
-		FGenerateStreamingParams& SetOutputLogPath(const FString& InOutputLogPath) { OutputLogPath = InOutputLogPath; return *this; }
+		FGenerateStreamingParams& SetActorDescContainerInstance(const UActorDescContainerInstance* InContainerInstance) { ContainerInstanceCollection.AddContainer(InContainerInstance); return *this;	}
 		FGenerateStreamingParams& SetErrorHandler(IStreamingGenerationErrorHandler* InErrorHandler) { ErrorHandler = InErrorHandler; return *this; }
+		FGenerateStreamingParams& SetOutputLogPath(const FString& InOutputLogPath) { OutputLogPath = InOutputLogPath; return *this; }
+
+		UE_DEPRECATED(5.4, "Use SetActorDescContainerInstance instead")
+		FGenerateStreamingParams& SetActorDescContainer(const UActorDescContainer* InActorDescContainer) { return *this; }
 	};
+	PRAGMA_ENABLE_DEPRECATION_WARNINGS
 
 	struct FGenerateStreamingContext
 	{
@@ -292,23 +302,39 @@ public:
 	ENGINE_API void CheckForErrors(IStreamingGenerationErrorHandler* ErrorHandler) const;
 
 	/* Struct of optional parameters passed to check for errors function. */
+	PRAGMA_DISABLE_DEPRECATION_WARNINGS
 	struct FCheckForErrorsParams
 	{
 		ENGINE_API FCheckForErrorsParams();
 
 		IStreamingGenerationErrorHandler* ErrorHandler;
-		const FActorDescContainerCollection* ActorDescContainerCollection;
 		bool bEnableStreaming;
-		TMap<FGuid, const UActorDescContainer*> ActorGuidsToContainerMap;
 
+		const FActorDescContainerInstanceCollection* ActorDescContainerInstanceCollection;
+
+		TMap<FGuid, const UActorDescContainerInstance*> ActorGuidsToContainerInstanceMap;
+				
 		FCheckForErrorsParams& SetErrorHandler(IStreamingGenerationErrorHandler* InErrorHandler) { ErrorHandler = InErrorHandler; return *this; }
-		FCheckForErrorsParams& SetActorDescContainerCollection(const FActorDescContainerCollection* InActorDescContainerCollection) { ActorDescContainerCollection = InActorDescContainerCollection; return *this; }
+		FCheckForErrorsParams& SetActorDescContainerInstanceCollection(const FActorDescContainerInstanceCollection* InActorDescContainerInstanceCollection) { ActorDescContainerInstanceCollection = InActorDescContainerInstanceCollection; return *this; }
 		FCheckForErrorsParams& SetEnableStreaming(bool bInEnableStreaming) { bEnableStreaming = bInEnableStreaming; return *this; }
-		FCheckForErrorsParams& SetActorGuidsToContainerMap(const TMap<FGuid, const UActorDescContainer*>& InActorGuidsToContainerMap) { ActorGuidsToContainerMap = InActorGuidsToContainerMap; return *this; }
+		FCheckForErrorsParams& SetActorGuidsToContainerMap(const TMap<FGuid, const UActorDescContainerInstance*>& InActorGuidsToContainerInstanceMap) { ActorGuidsToContainerInstanceMap = InActorGuidsToContainerInstanceMap; return *this; }
+
+
+		UE_DEPRECATED(5.4, "Use ActorDescContainerInstanceCollection instead")
+		const FActorDescContainerCollection* ActorDescContainerCollection = nullptr;
+				
+		UE_DEPRECATED(5.4, "Use ActorGuidsToContainerInstanceMap instead")
+		TMap<FGuid, const UActorDescContainer*> ActorGuidsToContainerMap;
+
+		UE_DEPRECATED(5.4, "Use SetActorDescContainerInstanceCollection instead")
+		FCheckForErrorsParams& SetActorDescContainerCollection(const FActorDescContainerCollection* InActorDescContainerCollection) { return *this; }
+		UE_DEPRECATED(5.4, "Use SetActorGuidsToContainerInstanceMap instead")
+		FCheckForErrorsParams& SetActorGuidsToContainerMap(const TMap<FGuid, const UActorDescContainer*>& InActorGuidsToContainerMap) { return *this; }
 	};
+	PRAGMA_ENABLE_DEPRECATION_WARNINGS
 
 	UE_DEPRECATED(5.2, "CheckForErrors is deprecated, CheckForErrors with FCheckForErrorsParams should be used instead.")
-	static ENGINE_API void CheckForErrors(IStreamingGenerationErrorHandler* ErrorHandler, const UActorDescContainer* ActorDescContainer, bool bEnableStreaming, bool bIsChangelistValidation);
+	static ENGINE_API void CheckForErrors(IStreamingGenerationErrorHandler* ErrorHandler, const UActorDescContainer* ActorDescContainer, bool bEnableStreaming, bool bIsChangelistValidation) {}
 
 	static ENGINE_API void CheckForErrors(const FCheckForErrorsParams& Params);
 
@@ -331,16 +357,35 @@ public:
 		/* Custom filter function used to filter actors descriptors. */
 		TUniqueFunction<bool(const FWorldPartitionActorDesc*)> FilterActorDescFunc;
 	};
-	ENGINE_API UActorDescContainer* RegisterActorDescContainer(const FContainerRegistrationParams& InRegistrationParameters);
-	ENGINE_API bool UnregisterActorDescContainer(UActorDescContainer* Container);
+
+	// Event when world partition was enabled/disabled in the world
+	DECLARE_DELEGATE_TwoParams(FActorDescContainerInstancePreInitializeDelegate, UActorDescContainerInstance::FInitializeParams&, UActorDescContainerInstance*);
+	FActorDescContainerInstancePreInitializeDelegate OnActorDescContainerInstancePreInitialize;
+
+	UE_DEPRECATED(5.4, "Use RegisterActorDescContainerInstance")
+	ENGINE_API UActorDescContainer* RegisterActorDescContainer(const FContainerRegistrationParams& InRegistrationParameters) { return nullptr; }
+
+	UE_DEPRECATED(5.4, "Use UnregisterActorDescContainerInstance")
+	ENGINE_API bool UnregisterActorDescContainer(UActorDescContainer* Container) { return false; }
+
 	ENGINE_API void UninitializeActorDescContainers();
 
 	DECLARE_MULTICAST_DELEGATE_OneParam(FActorDescContainerRegistrationDelegate, UActorDescContainer*);
+
+	UE_DEPRECATED(5.4, "Use OnActorDescContainerInstanceRegistered")
 	FActorDescContainerRegistrationDelegate OnActorDescContainerRegistered;
+	UE_DEPRECATED(5.4, "Use OnActorDescContainerInstanceUnregistered")
 	FActorDescContainerRegistrationDelegate OnActorDescContainerUnregistered;
 
 	UE_DEPRECATED(5.3, "Use RegisterActorDescContainer with FContainerRegistrationParams instead.")
-	UActorDescContainer* RegisterActorDescContainer(const FName& ContainerPackage) { return RegisterActorDescContainer(FContainerRegistrationParams(ContainerPackage)); }
+	UActorDescContainer* RegisterActorDescContainer(const FName& ContainerPackage) { return nullptr; }
+		
+	ENGINE_API UActorDescContainerInstance* RegisterActorDescContainerInstance(const UActorDescContainerInstance::FInitializeParams& InInitializationParams);
+	ENGINE_API bool UnregisterActorDescContainerInstance(UActorDescContainerInstance* InContainerInstance);
+
+	DECLARE_MULTICAST_DELEGATE_OneParam(FActorDescContainerInstanceRegistrationDelegate, UActorDescContainerInstance*);
+	FActorDescContainerInstanceRegistrationDelegate OnActorDescContainerInstanceRegistered;
+	FActorDescContainerInstanceRegistrationDelegate OnActorDescContainerInstanceUnregistered;
 
 	// Actors pinning
 	ENGINE_API void PinActors(const TArray<FGuid>& ActorGuids);
@@ -444,10 +489,12 @@ private:
 #endif
 
 public:
-	UActorDescContainer* GetActorDescContainer() const { return ActorDescContainer; }
-
-	UPROPERTY(Transient)
-	TObjectPtr<UActorDescContainer> ActorDescContainer;
+	UE_DEPRECATED(5.4, "Use GetActorDescContainerInstance")
+	UActorDescContainer* GetActorDescContainer() const { return nullptr; }
+		
+#if WITH_EDITOR
+	UActorDescContainerInstance* GetActorDescContainerInstance() const { return ActorDescContainerInstance; }
+#endif			
 
 	UPROPERTY()
 	TObjectPtr<UWorldPartitionRuntimeHash> RuntimeHash;
@@ -493,6 +540,9 @@ private:
 
 	TSet<FString> GeneratedStreamingPackageNames;
 
+	UPROPERTY(Transient)
+	TObjectPtr<UActorDescContainerInstance> ActorDescContainerInstance;
+
 public:
 	TOptional<bool> bOverrideEnableStreamingInEditor;
 
@@ -516,10 +566,6 @@ private:
 	UPROPERTY(Transient)
 	mutable TObjectPtr<UWorldPartitionStreamingPolicy> StreamingPolicy;
 
-#if WITH_EDITORONLY_DATA
-	FLinkerInstancingContext InstancingContext;
-#endif
-
 #if WITH_EDITOR
 	static ENGINE_API int32 LoadingRangeBugItGo;
 	static ENGINE_API int32 EnableSimulationStreamingSource;
@@ -549,8 +595,8 @@ private:
 	ENGINE_API void UnregisterDelegates();	
 
 #if WITH_EDITOR
-	ENGINE_API void HashActorDesc(FWorldPartitionActorDesc* ActorDesc);
-	ENGINE_API void UnhashActorDesc(FWorldPartitionActorDesc* ActorDesc);
+	void HashActorDescInstance(FWorldPartitionActorDescInstance* ActorDescInstance);
+	void UnhashActorDescInstance(FWorldPartitionActorDescInstance* ActorDescInstance);
 	void OnContentBundleRemovedContent(const FContentBundleEditor* ContentBundle);
 	ENGINE_API bool IsStreamingEnabledInEditor() const;
 
