WorldPartition: Add UActorDescContainerInstance/FWorldPartitionActorDescInstance which take transient/context data from UActorDescContainer/FWorldPartitionActorDesc leaving them as pure on disk repres

diff --git a/Engine/Source/Runtime/Engine/Public/WorldPartition/WorldPartitionActorDesc.h b/Engine/Source/Runtime/Engine/Public/WorldPartition/WorldPartitionActorDesc.h
--- a/Engine/Source/Runtime/Engine/Public/WorldPartition/WorldPartitionActorDesc.h
+++ b/Engine/Source/Runtime/Engine/Public/WorldPartition/WorldPartitionActorDesc.h
@@ -12,6 +12,7 @@
 #include "UObject/WeakObjectPtr.h"
 #include "Templates/SubclassOf.h"
 #include "Misc/Guid.h"
+#include "WorldPartition/WorldPartitionActorDescType.h"
 #include "WorldPartition/WorldPartitionActorContainerID.h"
 #include "WorldPartition/Filter/WorldPartitionActorFilter.h"
 
@@ -37,7 +38,10 @@ struct FWorldPartitionAssetDataPatcher
 };
 
 class AActor;
+class UActorDescContainerInstance;
+class FWorldPartitionActorDescInstance;
 class UActorDescContainer;
+class IWorldPartitionActorDescInstanceView;
 class IStreamingGenerationErrorHandler;
 
 enum class EContainerClusterMode : uint8
@@ -71,11 +75,14 @@ inline bool CompareUnsortedArrays(const TArray<FName>& Array1, const TArray<FNam
 	return CompareUnsortedArrays(Array1, Array2, [](const FName& A, const FName& B) { return A.LexicalLess(B); });
 }
 
+#endif // WITH_EDITOR
+
 /**
- * Represents a potentially unloaded actor (editor-only)
+ * Represents a Actor Metadata (editor-only)
  */
 class FWorldPartitionActorDesc
 {
+#if WITH_EDITOR
 	friend class AActor;
 	friend class UWorldPartition;
 	friend struct FWorldPartitionHandleImpl;
@@ -85,17 +92,32 @@ class FWorldPartitionActorDesc
 	friend class FAssetRootPackagePatcher;
 	friend class FAssetHeaderPatcherInner;
 	friend class FActorDescArchive;
+	friend class FWorldPartitionActorDescInstance;
+	friend class FStreamingGenerationUnsavedDirtyActorDescInstance;
 	template<class U> friend class TActorDescContainerCollection;
 
 public:
-	struct FContainerInstance
+	struct FLoadedContainerInstance
 	{
-		const UActorDescContainer* Container = nullptr;
 		ULevel* LoadedLevel = nullptr;
 		bool bSupportsPartialEditorLoading = false;
+	};
+
+	struct FContainerInstance
+	{
+		UActorDescContainerInstance* ContainerInstance = nullptr;
 		FTransform Transform = FTransform::Identity;
-		EContainerClusterMode ClusterMode;		
+		EContainerClusterMode ClusterMode;
 		TMap<FActorContainerID, TSet<FGuid>> FilteredActors;
+
+		UE_DEPRECATED(5.4, "Use FLoadedContainerInstance instead")
+		ULevel* LoadedLevel = nullptr;
+
+		UE_DEPRECATED(5.4, "Use FLoadedContainerInstance instead")
+		bool bSupportsPartialEditorLoading = false;
+
+		UE_DEPRECATED(5.4, "Use ContainerInstance instead")
+		const UActorDescContainer* Container = nullptr;
 	};
 
 	virtual ~FWorldPartitionActorDesc() {}
@@ -106,38 +128,25 @@ public:
 	inline FTopLevelAssetPath GetNativeClass() const { return NativeClass; }
 	inline UClass* GetActorNativeClass() const { return ActorNativeClass; }
 
-PRAGMA_DISABLE_DEPRECATION_WARNINGS
-	UE_DEPRECATED(5.2, "GetOrigin is deprecated.")
-	inline FVector GetOrigin() const { return GetBounds().GetCenter(); }
-PRAGMA_ENABLE_DEPRECATION_WARNINGS
-
 	inline FName GetRuntimeGrid() const { return RuntimeGrid; }
-	inline bool GetIsSpatiallyLoaded() const { return (bIsForcedNonSpatiallyLoaded || !bIsBoundsValid) ? false : bIsSpatiallyLoaded; }
+	inline bool GetIsSpatiallyLoaded() const { return bIsBoundsValid ? bIsSpatiallyLoaded : false; }
 	inline bool GetIsSpatiallyLoadedRaw() const { return bIsSpatiallyLoaded; }
 	inline bool GetActorIsEditorOnly() const { return bActorIsEditorOnly; }
 	inline bool GetActorIsRuntimeOnly() const { return bActorIsRuntimeOnly; }
 
-	UE_DEPRECATED(5.1, "SetIsSpatiallyLoadedRaw is deprecated and should not be used.")
-	inline void SetIsSpatiallyLoadedRaw(bool bNewIsSpatiallyLoaded) { bIsSpatiallyLoaded = bNewIsSpatiallyLoaded; }
-
 	inline bool GetActorIsHLODRelevant() const { return bActorIsHLODRelevant; }
 	inline FSoftObjectPath GetHLODLayer() const { return HLODLayer; }
 	inline const TArray<FName>& GetDataLayers() const { return DataLayers; }
-	inline bool HasResolvedDataLayerInstanceNames() const { return ResolvedDataLayerInstanceNames.IsSet(); }
-	ENGINE_API const TArray<FName>& GetDataLayerInstanceNames() const;
+
 	inline const TArray<FName>& GetTags() const { return Tags; }
-	inline void SetDataLayerInstanceNames(const TArray<FName>& InDataLayerInstanceNames) { ResolvedDataLayerInstanceNames = InDataLayerInstanceNames; }
+
 	inline FName GetActorPackage() const { return ActorPackage; }
 	inline FSoftObjectPath GetActorSoftPath() const { return ActorPath; }
 	inline FName GetActorLabel() const { return ActorLabel; }
 	inline FName GetFolderPath() const { return FolderPath; }
 	inline const FGuid& GetFolderGuid() const { return FolderGuid; }
 
-PRAGMA_DISABLE_DEPRECATION_WARNINGS
-	UE_DEPRECATED(5.2, "GetBounds is deprecated, GetEditorBounds or GetRuntimeBounds should be used instead.")
-	FBox GetBounds() const { return GetEditorBounds(); }
-ENGINE_API PRAGMA_ENABLE_DEPRECATION_WARNINGS
-	virtual FBox GetEditorBounds() const;
+	ENGINE_API virtual FBox GetEditorBounds() const;
 	ENGINE_API FBox GetRuntimeBounds() const;
 
 	inline const FGuid& GetParentActor() const { return ParentActor; }
@@ -151,27 +160,19 @@ ENGINE_API PRAGMA_ENABLE_DEPRECATION_WARNINGS
 	ENGINE_API FName GetDisplayClassName() const;
 
 	inline bool IsDefaultActorDesc() const { return bIsDefaultActorDesc; }
-
-	virtual bool IsContainerInstance() const { return false; }
-	virtual EWorldPartitionActorFilterType GetContainerFilterType() const { return EWorldPartitionActorFilterType::None; }
-	virtual FName GetContainerPackage() const { return NAME_None; }
+		
+	virtual bool IsChildContainerInstance() const { return false; }
+	virtual FName GetChildContainerPackage() const { return NAME_None; }
+	virtual EWorldPartitionActorFilterType GetChildContainerFilterType() const { return EWorldPartitionActorFilterType::None; }
+	virtual const FWorldPartitionActorFilter* GetChildContainerFilter() const { return nullptr; }
+	virtual UActorDescContainer* GetChildContainer() const { return nullptr; }
 	
-	virtual const FWorldPartitionActorFilter* GetContainerFilter() const { return nullptr; }
-	virtual bool GetContainerInstance(FContainerInstance& OutContainerInstance) const { return false; }
-
 	ENGINE_API FGuid GetContentBundleGuid() const;
 
 	virtual const FGuid& GetSceneOutlinerParent() const { return GetParentActor(); }
 	virtual bool IsResaveNeeded() const { return false; }
-	ENGINE_API virtual bool IsRuntimeRelevant(const FActorContainerID& InContainerID) const;
-	ENGINE_API virtual bool IsEditorRelevant() const;
-	ENGINE_API virtual void CheckForErrors(IStreamingGenerationErrorHandler* ErrorHandler) const;
 
-	UE_DEPRECATED(5.2, "ShouldValidateRuntimeGrid is deprecated and should not be used.")
-	virtual bool ShouldValidateRuntimeGrid() const { return true; }
-
-	UE_DEPRECATED(5.3, "GetLevelPackage is deprecated use GetContainerPackage instead.")
-	virtual FName GetLevelPackage() const { return GetContainerPackage(); }
+	ENGINE_API virtual void CheckForErrors(const IWorldPartitionActorDescInstanceView* InActorDescView, IStreamingGenerationErrorHandler* ErrorHandler) const;
 
 	bool operator==(const FWorldPartitionActorDesc& Other) const
 	{
@@ -183,52 +184,114 @@ ENGINE_API PRAGMA_ENABLE_DEPRECATION_WARNINGS
 		return GetTypeHash(Key.Guid);
 	}
 
-	ENGINE_API const FText& GetUnloadedReason() const;
+	//~ Begin Deprecation
+	UE_DEPRECATED(5.1, "SetIsSpatiallyLoadedRaw is deprecated and should not be used.")
+	inline void SetIsSpatiallyLoadedRaw(bool bNewIsSpatiallyLoaded) { bIsSpatiallyLoaded = bNewIsSpatiallyLoaded; }
 
-	void SetUnloadedReason(FText* InUnloadedReason)
-	{
-		UnloadedReason = InUnloadedReason;
-	}
+PRAGMA_DISABLE_DEPRECATION_WARNINGS
+	UE_DEPRECATED(5.2, "GetOrigin is deprecated.")
+	inline FVector GetOrigin() const { return GetBounds().GetCenter(); }
 
-protected:
-	inline uint32 IncSoftRefCount() const
-	{
-		return ++SoftRefCount;
-	}
+	UE_DEPRECATED(5.2, "GetBounds is deprecated, GetEditorBounds or GetRuntimeBounds should be used instead.")
+	FBox GetBounds() const { return GetEditorBounds(); }
+PRAGMA_ENABLE_DEPRECATION_WARNINGS
 
-	inline uint32 DecSoftRefCount() const
-	{
-		check(SoftRefCount > 0);
-		return --SoftRefCount;
-	}
+	UE_DEPRECATED(5.2, "ShouldValidateRuntimeGrid is deprecated and should not be used")
+	virtual bool ShouldValidateRuntimeGrid() const { return true; }
 
-	inline uint32 IncHardRefCount() const
-	{
-		return ++HardRefCount;
-	}
+	UE_DEPRECATED(5.3, "GetLevelPackage is deprecated use GetChildContainerPackage instead")
+	virtual FName GetLevelPackage() const { return GetChildContainerPackage(); }
 
-	inline uint32 DecHardRefCount() const
-	{
-		check(HardRefCount > 0);
-		return --HardRefCount;
-	}
+	UE_DEPRECATED(5.4, "Use FWorldPartitionActorDescInstance::HasResolvedDataLayerInstanceNames")
+	inline bool HasResolvedDataLayerInstanceNames() const { return false; }
 
-	inline uint32 GetSoftRefCount() const
-	{
-		return SoftRefCount;
-	}
+	UE_DEPRECATED(5.4, "Use FWorldPartitionActorDescInstance::GetDataLayerInstanceNames")
+	ENGINE_API TArray<FName> GetDataLayerInstanceNames() const { return TArray<FName>(); }
 
-	inline uint32 GetHardRefCount() const
-	{
-		return HardRefCount;
-	}
+	UE_DEPRECATED(5.4, "Use FWorldPartitionActorDescInstance::SetDataLayerInstanceNames")
+	inline void SetDataLayerInstanceNames(const TArray<FName>& InDataLayerInstanceNames) { }
+
+	UE_DEPRECATED(5.4, "Use FWorldPartitionActorDescInstance::IsRuntimeRelevant instead")
+	ENGINE_API virtual bool IsRuntimeRelevant(const FActorContainerID& InContainerID) const { return false; }
+		
+	UE_DEPRECATED(5.4, "Use FWorldPartitionActorDescInstance::IsEditorRelevant instead")
+	ENGINE_API virtual bool IsEditorRelevant() const { return false; }
+		
+	UE_DEPRECATED(5.4, "Please use FWorldPartitionActorDescInstance::GetUnloadedReason")
+	ENGINE_API FText GetUnloadedReason() const { return FText(); }
+
+	UE_DEPRECATED(5.4, "Please use FWorldPartitionActorDescInstance::SetUnloadedReason")
+	void SetUnloadedReason(FText* InUnloadedReason) { }
+
+	UE_DEPRECATED(5.4, "Use IsChildContainerInstance instead")
+	virtual bool IsContainerInstance() const { return false; }
+
+	UE_DEPRECATED(5.4, "Use GetChildContainerPackage instead")
+	virtual FName GetContainerPackage() const { return NAME_None; }
+
+	UE_DEPRECATED(5.4, "Use GetChildContainerFilterType instead")
+	virtual EWorldPartitionActorFilterType GetContainerFilterType() const { return EWorldPartitionActorFilterType::None; }
+
+	UE_DEPRECATED(5.4, "Use GetChildContainerFilter instead")
+	virtual const FWorldPartitionActorFilter* GetContainerFilter() const { return nullptr; }
+
+	UE_DEPRECATED(5.4, "Use GetChildContainerInstance instead")
+	virtual bool GetContainerInstance(FContainerInstance& OutContainerInstance) const { return false; }
+
+	UE_DEPRECATED(5.4, "Use SetContainer without World param")
+	virtual void SetContainer(UActorDescContainer* InContainer, UWorld* InWorldContext)	{ }
 
+	UE_DEPRECATED(5.4, "Please use FWorldPartitionActorDescInstance::IsLoaded")
+	ENGINE_API bool IsLoaded(bool bEvenIfPendingKill=false) const { return false; }
+	
+	UE_DEPRECATED(5.4, "Please use FWorldPartitionActorDescInstance::GetActor")
+	ENGINE_API AActor* GetActor(bool bEvenIfPendingKill=true, bool bEvenIfUnreachable=false) const { return nullptr; }
+
+	UE_DEPRECATED(5.4, "Please use FWorldPartitionActorDescInstance::GetActor")
+	ENGINE_API TWeakObjectPtr<AActor>* GetActorPtr(bool bEvenIfPendingKill = true, bool bEvenIfUnreachable = false) const { return nullptr; }
+
+	UE_DEPRECATED(5.4, "Please use FWorldPartitionActorDescInstance::Load")
+	ENGINE_API AActor* Load() const { return nullptr;}
+
+	UE_DEPRECATED(5.4, "Please use FWorldPartitionActorDescInstance::Unload")
+	ENGINE_API virtual void Unload() {}
+
+	UE_DEPRECATED(5.4, "Please use FWorldPartitionActorDescInstance")
+	ENGINE_API void TransformInstance(const FString& From, const FString& To) {}
+
+	UE_DEPRECATED(5.4, "Use FWorldPartitionActorDescInstance version instead")
+	virtual void CheckForErrors(IStreamingGenerationErrorHandler* ErrorHandler) const {}
+protected:
+	UE_DEPRECATED(5.4, "Please use FWorldPartitionActorDescInstance::IncSoftRefCount")
+	inline uint32 IncSoftRefCount() const { return 0; }
+
+	UE_DEPRECATED(5.4, "Please use FWorldPartitionActorDescInstance::DecSoftRefCount")
+	inline uint32 DecSoftRefCount() const { return 0; }
+	
+	UE_DEPRECATED(5.4, "Please use FWorldPartitionActorDescInstance::IncHardRefCount")
+	inline uint32 IncHardRefCount() const { return 0;}
+	
+	UE_DEPRECATED(5.4, "Please use FWorldPartitionActorDescInstance::DecHardRefCount")
+	inline uint32 DecHardRefCount() const { return 0;}
+	
+	UE_DEPRECATED(5.4, "Please use FWorldPartitionActorDescInstance::GetSoftRefCount")
+	inline uint32 GetSoftRefCount() const { return 0; }
+	
+	UE_DEPRECATED(5.4, "Please use FWorldPartitionActorDescInstance::GetHardRefCount")
+	inline uint32 GetHardRefCount() const { return 0; }
+	//~ End Deprecation
+	
 public:
 	const TArray<FGuid>& GetReferences() const
 	{
 		return References;
 	}
 
+	const TArray<FGuid>& GetEditorOnlyReferences() const
+	{
+		return EditorOnlyReferences;
+	}
+
 	bool IsEditorOnlyReference(const FGuid& ReferenceGuid) const
 	{
 		return EditorOnlyReferences.Contains(ReferenceGuid);
@@ -239,7 +302,7 @@ public:
 		return Container;
 	}
 
-	virtual void SetContainer(UActorDescContainer* InContainer, UWorld* InWorldContext)
+	virtual void SetContainer(UActorDescContainer* InContainer)
 	{
 		check(!Container || !InContainer);
 		Container = InContainer;
@@ -254,12 +317,6 @@ public:
 
 	ENGINE_API FString ToString(EToStringMode Mode = EToStringMode::Compact) const;
 
-	ENGINE_API bool IsLoaded(bool bEvenIfPendingKill=false) const;
-	ENGINE_API AActor* GetActor(bool bEvenIfPendingKill=true, bool bEvenIfUnreachable=false) const;
-	ENGINE_API TWeakObjectPtr<AActor>* GetActorPtr(bool bEvenIfPendingKill=true, bool bEvenIfUnreachable=false) const;	
-	ENGINE_API AActor* Load() const;
-	ENGINE_API virtual void Unload();
-
 	ENGINE_API virtual void Init(const AActor* InActor);
 	ENGINE_API virtual void Init(const FWorldPartitionActorDescInitData& DescData);
 
@@ -273,9 +330,7 @@ public:
 	 */
 	ENGINE_API virtual bool ShouldResave(const FWorldPartitionActorDesc* Other) const;
 
-	ENGINE_API void SerializeTo(TArray<uint8>& OutData);
-
-	ENGINE_API void TransformInstance(const FString& From, const FString& To);
+	ENGINE_API void SerializeTo(TArray<uint8>& OutData) const;
 
 	using FActorDescDeprecator = TFunction<void(FArchive&, FWorldPartitionActorDesc*)>;
 	static ENGINE_API void RegisterActorDescDeprecator(TSubclassOf<AActor> ActorClass, const FActorDescDeprecator& Deprecator);
@@ -283,6 +338,14 @@ public:
 	ENGINE_API bool IsMainWorldOnly() const;
 	ENGINE_API bool IsListedInSceneOutliner() const;
 protected:
+	virtual void OnUnloadingInstance(const FWorldPartitionActorDescInstance* InActorDescInstance) const {}
+
+	virtual bool GetChildContainerInstance(const FWorldPartitionActorDescInstance* InActorDescInstance, FContainerInstance& OutContainerInstance) const { return false; }
+	virtual UActorDescContainerInstance* CreateChildContainerInstance(const FWorldPartitionActorDescInstance* InActorDescInstance) const { return nullptr; }
+	virtual bool GetLoadedChildContainerInstance(const FWorldPartitionActorDescInstance* InActorDescInstance, FLoadedContainerInstance& OutContainerInstance) const { return false; }
+	ENGINE_API virtual bool IsRuntimeRelevant(const FWorldPartitionActorDescInstance* InActorDescInstance) const;
+	ENGINE_API virtual bool IsEditorRelevant(const FWorldPartitionActorDescInstance* InActorDescInstance) const;
+
 	ENGINE_API FWorldPartitionActorDesc();
 
 	ENGINE_API virtual void TransferFrom(const FWorldPartitionActorDesc* From);
@@ -328,43 +391,12 @@ protected:
 	FGuid							ContentBundleGuid;
 	
 	// Transient
-	mutable uint32					SoftRefCount;
-	mutable uint32					HardRefCount;
 	UClass*							ActorNativeClass;
-	mutable TWeakObjectPtr<AActor>	ActorPtr;
 	UActorDescContainer*			Container;
-	TOptional<TArray<FName>>		ResolvedDataLayerInstanceNames; // Can only resolve in ActorDesc if Container is not used as a template
-	bool							bIsForcedNonSpatiallyLoaded;
 	bool							bIsDefaultActorDesc;
-	mutable FText*					UnloadedReason;
 
 	static ENGINE_API TMap<TSubclassOf<AActor>, FActorDescDeprecator> Deprecators;
 
-#if DO_CHECK
-public:
-	struct FRegisteringUnregisteringGuard
-	{
-		FRegisteringUnregisteringGuard(FWorldPartitionActorDesc* InActorDesc)
-			:ActorDesc(InActorDesc)
-		{
-			check(!ActorDesc->bIsRegisteringOrUnregistering);
-			ActorDesc->bIsRegisteringOrUnregistering = true;
-		}
-
-		~FRegisteringUnregisteringGuard()
-		{
-			check(ActorDesc->bIsRegisteringOrUnregistering);
-			ActorDesc->bIsRegisteringOrUnregistering = false;
-		}
-
-		FRegisteringUnregisteringGuard(const FRegisteringUnregisteringGuard&) = delete;
-		FRegisteringUnregisteringGuard& operator=(const FRegisteringUnregisteringGuard&) = delete;
-
-	private:
-		FWorldPartitionActorDesc* ActorDesc;
-	};
-private:
-	bool bIsRegisteringOrUnregistering = false;
-#endif
+#endif // WITH_EDITOR
 };
-#endif
+
diff --git a/Engine/Source/Runtime/Engine/Public/WorldPartition/WorldPartitionActorDescInstance.h b/Engine/Source/Runtime/Engine/Public/WorldPartition/WorldPartitionActorDescInstance.h
new file mode 100644
--- /dev/null
+++ b/Engine/Source/Runtime/Engine/Public/WorldPartition/WorldPartitionActorDescInstance.h
@@ -0,0 +1,205 @@
+// Copyright Epic Games, Inc. All Rights Reserved.
+
+#pragma once
+
+#include "CoreMinimal.h"
+#include "WorldPartition/WorldPartitionActorDescInstanceInterface.h"
+
+#if WITH_EDITOR
+
+#include "WorldPartition/WorldPartitionActorDesc.h"
+#include "UObject/WeakObjectPtr.h"
+
+class AActor;
+class UActorDescContainerInstance;
+
+#endif // WITH_EDITOR
+
+class FWorldPartitionActorDescInstance : public IWorldPartitionActorDescInstance
+{
+#if WITH_EDITOR
+	friend struct FWorldPartitionHandleImpl;
+	friend struct FWorldPartitionReferenceImpl;
+	friend struct FWorldPartitionActorDescUtils;
+	friend class FWorldPartitionActorDesc;
+	friend class UActorDescContainerInstance;
+	friend class FWorldPartitionLoadingContext;
+	friend struct FWorldPartitionActorDescUnitTestAcccessor;
+	friend class UDataLayerManager;
+	friend class UWorldPartition;
+
+protected:
+	ENGINE_API FWorldPartitionActorDescInstance();
+public:
+	ENGINE_API FWorldPartitionActorDescInstance(UActorDescContainerInstance* InContainerInstance, FWorldPartitionActorDesc* InActorDesc);
+	virtual ~FWorldPartitionActorDescInstance() {}
+
+	//~ Begin IWorldPartitionActorDescInstanceViewInterface
+	virtual const FGuid& GetGuid() const override { return GetActorDesc()->GetGuid(); }
+	virtual FTopLevelAssetPath GetBaseClass() const override { return GetActorDesc()->GetBaseClass(); }
+	virtual FTopLevelAssetPath GetNativeClass() const override { return GetActorDesc()->GetNativeClass(); }
+	virtual UClass* GetActorNativeClass() const override { return GetActorDesc()->GetActorNativeClass(); }
+	
+	virtual FName GetRuntimeGrid() const override { return GetActorDesc()->GetRuntimeGrid(); }
+	virtual bool GetIsSpatiallyLoaded() const override { return !GetForceNonSpatiallyLoaded() && GetActorDesc()->GetIsSpatiallyLoaded(); }
+	virtual bool GetActorIsEditorOnly() const override { return GetActorDesc()->GetActorIsEditorOnly(); }
+	virtual bool GetActorIsRuntimeOnly() const override { return GetActorDesc()->GetActorIsRuntimeOnly(); }
+	ENGINE_API virtual bool IsRuntimeRelevant() const override;
+	ENGINE_API virtual bool IsEditorRelevant() const override;
+
+	virtual bool IsUsingDataLayerAsset() const override { return GetActorDesc()->IsUsingDataLayerAsset(); }
+	virtual const TArray<FName>& GetDataLayers() const override { return GetActorDesc()->GetDataLayers(); }
+
+	virtual bool GetActorIsHLODRelevant() const override { return GetActorDesc()->GetActorIsHLODRelevant(); }
+	virtual FSoftObjectPath GetHLODLayer() const override { return GetActorDesc()->GetHLODLayer(); }
+
+	virtual const TArray<FName>& GetTags() const override { return GetActorDesc()->GetTags(); }
+	virtual FName GetActorPackage() const override { return GetActorDesc()->GetActorPackage(); }
+	ENGINE_API virtual FSoftObjectPath GetActorSoftPath() const override;
+	virtual FName GetActorLabel() const override { return GetActorDesc()->GetActorLabel(); }
+	ENGINE_API virtual FName GetActorName() const override;
+	virtual FName GetFolderPath() const override { return GetActorDesc()->GetFolderPath(); }
+	virtual const FGuid& GetFolderGuid() const override { return GetActorDesc()->GetFolderGuid(); }
+
+	virtual FBox GetEditorBounds() const override { return GetActorDesc()->GetEditorBounds(); }
+	virtual FBox GetRuntimeBounds() const override { return GetActorDesc()->GetRuntimeBounds(); }
+
+	virtual bool GetProperty(FName PropertyName, FName* PropertyValue) const override { return GetActorDesc()->GetProperty(PropertyName, PropertyValue); }
+	virtual bool HasProperty(FName PropertyName) const override { return GetActorDesc()->HasProperty(PropertyName); }
+
+	virtual const TArray<FGuid>& GetReferences() const override { return GetActorDesc()->GetReferences(); }
+	virtual const TArray<FGuid>& GetEditorOnlyReferences() const override { return GetActorDesc()->GetEditorOnlyReferences(); }
+	virtual bool IsEditorOnlyReference(const FGuid& ReferenceGuid) const override { return GetActorDesc()->IsEditorOnlyReference(ReferenceGuid); }
+	
+	virtual const FGuid& GetParentActor() const override { return GetActorDesc()->GetParentActor(); }
+		
+	virtual FGuid GetContentBundleGuid() const override { return GetActorDesc()->GetContentBundleGuid(); }
+
+	virtual bool IsChildContainerInstance() const override { return ChildContainerInstance || GetActorDesc()->IsChildContainerInstance(); }
+	virtual FName GetChildContainerPackage() const override { return GetActorDesc()->GetChildContainerPackage(); }
+	virtual EWorldPartitionActorFilterType GetChildContainerFilterType() const override { return GetActorDesc()->GetChildContainerFilterType(); }
+	virtual const FWorldPartitionActorFilter* GetChildContainerFilter() const override { return GetActorDesc()->GetChildContainerFilter(); }
+	virtual bool GetChildContainerInstance(FWorldPartitionActorDesc::FContainerInstance& OutContainerInstance) const override { return GetActorDesc()->GetChildContainerInstance(this, OutContainerInstance); }
+			
+	virtual bool IsMainWorldOnly() const override { return GetActorDesc()->IsMainWorldOnly(); }
+	virtual bool IsListedInSceneOutliner() const override { return GetActorDesc()->IsListedInSceneOutliner(); }
+	virtual const FGuid& GetSceneOutlinerParent() const override { return GetActorDesc()->GetSceneOutlinerParent(); }
+			
+	ENGINE_API virtual const FWorldPartitionActorDesc* GetActorDesc() const override { check(ActorDesc); return ActorDesc; }
+	
+	virtual bool HasResolvedDataLayerInstanceNames() const override { return ResolvedDataLayerInstanceNames.IsSet(); }
+	ENGINE_API virtual const TArray<FName>& GetDataLayerInstanceNames() const override;
+
+	ENGINE_API virtual bool IsLoaded(bool bEvenIfPendingKill = false) const override;
+	ENGINE_API virtual AActor* GetActor(bool bEvenIfPendingKill = true, bool bEvenIfUnreachable = false) const override;
+	
+	ENGINE_API virtual FString ToString(FWorldPartitionActorDesc::EToStringMode Mode = FWorldPartitionActorDesc::EToStringMode::Compact) const override;
+	
+	virtual UActorDescContainerInstance* GetContainerInstance() const override { return ContainerInstance; }
+	//~ End
+
+	inline UActorDescContainerInstance* GetChildContainerInstance() const { return ChildContainerInstance; }
+	bool GetLoadedChildContainerInstance(FWorldPartitionActorDesc::FLoadedContainerInstance& OutContainerInstance) const { return GetActorDesc()->GetLoadedChildContainerInstance(this, OutContainerInstance); }
+		
+	ENGINE_API virtual TWeakObjectPtr<AActor>* GetActorPtr(bool bEvenIfPendingKill = true, bool bEvenIfUnreachable = false) const;
+	ENGINE_API virtual bool IsValid() const;
+				
+	inline void SetForceNonSpatiallyLoadded(bool bForce) { bIsForcedNonSpatiallyLoaded = bForce; }
+	inline bool GetForceNonSpatiallyLoaded() const { return bIsForcedNonSpatiallyLoaded; }
+
+	inline void SetUnloadedReason(FText* InUnloadedReason) { UnloadedReason = InUnloadedReason; }
+
+	FName GetDisplayClassName() const { return GetActorDesc()->GetDisplayClassName(); }
+	
+	ENGINE_API const FText& GetUnloadedReason() const;
+		
+	ENGINE_API AActor* Load();
+	ENGINE_API void Unload();
+protected:
+	ENGINE_API void UpdateActorDesc(FWorldPartitionActorDesc* InActorDesc);
+	void Invalidate();
+		
+	inline uint32 IncSoftRefCount() const
+	{
+		return ++SoftRefCount;
+	}
+
+	inline uint32 DecSoftRefCount() const
+	{
+		check(SoftRefCount > 0);
+		return --SoftRefCount;
+	}
+
+	inline uint32 IncHardRefCount() const
+	{
+		return ++HardRefCount;
+	}
+
+	inline uint32 DecHardRefCount() const
+	{
+		check(HardRefCount > 0);
+		return --HardRefCount;
+	}
+
+	inline uint32 GetSoftRefCount() const
+	{
+		return SoftRefCount;
+	}
+
+	inline uint32 GetHardRefCount() const
+	{
+		return HardRefCount;
+	}
+
+	inline void SetDataLayerInstanceNames(const TArray<FName>& InDataLayerInstanceNames) { ResolvedDataLayerInstanceNames = InDataLayerInstanceNames; }
+
+	virtual ENGINE_API void RegisterChildContainerInstance();
+	virtual ENGINE_API void UnregisterChildContainerInstance();
+	virtual ENGINE_API void UpdateChildContainerInstance();
+
+protected:
+	TObjectPtr<UActorDescContainerInstance>		ContainerInstance;
+
+	mutable uint32								SoftRefCount;
+	mutable uint32								HardRefCount;
+	TOptional<TArray<FName>>					ResolvedDataLayerInstanceNames;
+	bool										bIsForcedNonSpatiallyLoaded;
+	mutable FText*								UnloadedReason;
+
+	// Instancing Path if set
+	TOptional<FSoftObjectPath>					ActorPath;
+
+	mutable TWeakObjectPtr<AActor>				ActorPtr;
+	FWorldPartitionActorDesc*					ActorDesc;
+
+	TObjectPtr<UActorDescContainerInstance>		ChildContainerInstance;
+
+#if DO_CHECK
+private:
+	struct FRegisteringUnregisteringGuard
+	{
+		FRegisteringUnregisteringGuard(FWorldPartitionActorDescInstance* InActorDescInstance)
+			:ActorDescInstance(InActorDescInstance)
+		{
+			check(!ActorDescInstance->bIsRegisteringOrUnregistering);
+			ActorDescInstance->bIsRegisteringOrUnregistering = true;
+		}
+
+		~FRegisteringUnregisteringGuard()
+		{
+			check(ActorDescInstance->bIsRegisteringOrUnregistering);
+			ActorDescInstance->bIsRegisteringOrUnregistering = false;
+		}
+
+		FRegisteringUnregisteringGuard(const FRegisteringUnregisteringGuard&) = delete;
+		FRegisteringUnregisteringGuard& operator=(const FRegisteringUnregisteringGuard&) = delete;
+
+	private:
+		FWorldPartitionActorDescInstance* ActorDescInstance;
+	};
+	bool bIsRegisteringOrUnregistering = false;
+#endif // DO_CHECK
+
+#endif // WITH_EDITOR
+};
+
diff --git a/Engine/Source/Runtime/Engine/Public/WorldPartition/WorldPartitionActorDescInstanceInterface.h b/Engine/Source/Runtime/Engine/Public/WorldPartition/WorldPartitionActorDescInstanceInterface.h
new file mode 100644
--- /dev/null
+++ b/Engine/Source/Runtime/Engine/Public/WorldPartition/WorldPartitionActorDescInstanceInterface.h
@@ -0,0 +1,92 @@
+// Copyright Epic Games, Inc. All Rights Reserved.
+#pragma once
+
+#if WITH_EDITOR
+
+#include "CoreMinimal.h"
+#include "UObject/SoftObjectPath.h"
+#include "UObject/TopLevelAssetPath.h"
+#include "Misc/Optional.h"
+#include "WorldPartition/WorldPartitionActorDesc.h"
+
+class AActor;
+class IStreamingGenerationErrorHandler;
+class UWorldPartition;
+struct FWorldPartitionActorFilter;
+enum class EWorldPartitionActorFilterType : uint8;
+
+#endif // WITH_EDITOR
+
+/**
+ * Interface on top of an actor desc, representing an instance of a FWorldPartitionActorDesc
+ */
+class IWorldPartitionActorDescInstance
+{
+#if WITH_EDITOR
+public:
+	virtual const FGuid& GetGuid() const = 0;
+
+	virtual FTopLevelAssetPath GetBaseClass() const = 0;
+	virtual FTopLevelAssetPath GetNativeClass() const = 0;
+	virtual UClass* GetActorNativeClass() const = 0;
+	
+	virtual FName GetRuntimeGrid() const = 0;
+	virtual bool GetIsSpatiallyLoaded() const = 0;
+	virtual bool GetActorIsEditorOnly() const = 0;
+	virtual bool GetActorIsRuntimeOnly() const = 0;
+	virtual bool IsRuntimeRelevant() const = 0;
+	virtual bool IsEditorRelevant() const = 0;
+
+	virtual bool IsUsingDataLayerAsset() const = 0;
+	virtual const TArray<FName>& GetDataLayers() const = 0;
+
+	virtual bool GetActorIsHLODRelevant() const = 0;
+	virtual FSoftObjectPath GetHLODLayer() const = 0;
+	
+	virtual const TArray<FName>& GetTags() const = 0;
+	virtual FName GetActorPackage() const = 0;
+	virtual FSoftObjectPath GetActorSoftPath() const = 0;
+	virtual FName GetActorLabel() const = 0;
+	virtual FName GetActorName() const = 0;
+	virtual FName GetFolderPath() const = 0;
+	virtual const FGuid& GetFolderGuid() const = 0;
+		
+	virtual FBox GetEditorBounds() const = 0;
+	virtual FBox GetRuntimeBounds() const = 0;
+		
+	virtual bool GetProperty(FName PropertyName, FName* PropertyValue) const = 0;
+	virtual bool HasProperty(FName PropertyName) const = 0;
+
+	virtual const TArray<FGuid>& GetReferences() const = 0;
+	virtual const TArray<FGuid>& GetEditorOnlyReferences() const = 0;
+	virtual bool IsEditorOnlyReference(const FGuid& ReferenceGuid) const = 0;
+		
+	virtual const FGuid& GetParentActor() const = 0;
+	
+	virtual FGuid GetContentBundleGuid() const = 0;
+	
+	virtual bool IsChildContainerInstance() const = 0;
+	virtual FName GetChildContainerPackage() const = 0;
+	virtual EWorldPartitionActorFilterType GetChildContainerFilterType() const = 0;
+	virtual const FWorldPartitionActorFilter* GetChildContainerFilter() const = 0;
+	virtual bool GetChildContainerInstance(FWorldPartitionActorDesc::FContainerInstance& OutContainerInstance) const = 0;
+			
+	virtual bool IsMainWorldOnly() const = 0;
+	virtual bool IsListedInSceneOutliner() const = 0;
+	virtual const FGuid& GetSceneOutlinerParent() const = 0;
+		
+	virtual FString ToString(FWorldPartitionActorDesc::EToStringMode Mode) const = 0;
+	virtual const FWorldPartitionActorDesc* GetActorDesc() const = 0;
+
+	virtual FName GetActorLabelOrName() const { return GetActorLabel().IsNone() ? GetActorName() : GetActorLabel(); }
+
+	virtual bool HasResolvedDataLayerInstanceNames() const = 0;
+	virtual const TArray<FName>& GetDataLayerInstanceNames() const = 0;
+
+	virtual AActor* GetActor(bool bEvenIfPendingKill = true, bool bEvenIfUnreachable = false) const = 0;
+	virtual bool IsLoaded(bool bEvenIfPendingKill = false) const = 0;
+
+	virtual UActorDescContainerInstance* GetContainerInstance() const = 0;
+#endif // WITH_EDITOR
+};
+
diff --git a/Engine/Source/Runtime/Engine/Public/WorldPartition/WorldPartitionActorDescInstanceViewInterface.h b/Engine/Source/Runtime/Engine/Public/WorldPartition/WorldPartitionActorDescInstanceViewInterface.h
new file mode 100644
--- /dev/null
+++ b/Engine/Source/Runtime/Engine/Public/WorldPartition/WorldPartitionActorDescInstanceViewInterface.h
@@ -0,0 +1,26 @@
+// Copyright Epic Games, Inc. All Rights Reserved.
+#pragma once
+
+#if WITH_EDITOR
+
+#include "CoreMinimal.h"
+#include "WorldPartition/WorldPartitionActorDescInstanceInterface.h"
+
+class AActor;
+class IStreamingGenerationErrorHandler;
+class UWorldPartition;
+struct FWorldPartitionActorFilter;
+enum class EWorldPartitionActorFilterType : uint8;
+
+#endif // WITH_EDITOR
+
+/**
+ * Interface for a view on top of an actor desc, used to cache information that can be (potentially) different than the actor desc
+ * itself due to streaming generation logic, etc.
+ */
+class IWorldPartitionActorDescInstanceView : public IWorldPartitionActorDescInstance
+{
+public:
+	virtual void CheckForErrors(IStreamingGenerationErrorHandler* ErrorHandler) const = 0;
+};
+
diff --git a/Engine/Source/Runtime/Engine/Public/WorldPartition/WorldPartitionActorDescUtils.h b/Engine/Source/Runtime/Engine/Public/WorldPartition/WorldPartitionActorDescUtils.h
--- a/Engine/Source/Runtime/Engine/Public/WorldPartition/WorldPartitionActorDescUtils.h
+++ b/Engine/Source/Runtime/Engine/Public/WorldPartition/WorldPartitionActorDescUtils.h
@@ -8,6 +8,7 @@
 
 class AActor;
 struct FAssetData;
+class FWorldPartitionActorDescInstance;
 
 struct FWorldPartitionActorDescUtils
 {
@@ -80,9 +81,12 @@ struct FWorldPartitionActorDescUtils
 	 * Replaces the actor descriptor's actor pointer with the provided new actor pointer.
 	 * @param InOldActor	The old actor that the provided actor descriptor was representing.
 	 * @param InNewActor	The new actor that the provided actor descriptor should be representing.
-	 * @param InActorDesc	Actor descriptor that will get its actor pointer updated.
+	 * @param InActorDescInstance	Actor descriptor instance that will get its actor pointer updated.
 	 */
-	static ENGINE_API void ReplaceActorDescriptorPointerFromActor(const AActor* InOldActor, AActor* InNewActor, FWorldPartitionActorDesc* InActorDesc);
+	static ENGINE_API void ReplaceActorDescriptorPointerFromActor(const AActor* InOldActor, AActor* InNewActor, FWorldPartitionActorDescInstance* InActorDescInstance);
+
+	UE_DEPRECATED(5.4, "Use FWorldPartitionActorDescInstance version instead")
+	static ENGINE_API void ReplaceActorDescriptorPointerFromActor(const AActor* InOldActor, AActor* InNewActor, FWorldPartitionActorDesc* InActorDesc) {}
 
 	UE_DEPRECATED(5.3, "Use FWorldPartitionHelpers::FixupRedirectedAssetPath instead")
 	static ENGINE_API bool FixupRedirectedAssetPath(FName& InOutAssetPath);
diff --git a/Engine/Source/Runtime/Engine/Public/WorldPartition/WorldPartitionActorDescView.h b/Engine/Source/Runtime/Engine/Public/WorldPartition/WorldPartitionActorDescView.h
--- a/Engine/Source/Runtime/Engine/Public/WorldPartition/WorldPartitionActorDescView.h
+++ b/Engine/Source/Runtime/Engine/Public/WorldPartition/WorldPartitionActorDescView.h
@@ -16,60 +16,110 @@ struct FWorldPartitionActorFilter;
 enum class EWorldPartitionActorFilterType : uint8;
 
 /**
- * A view on top of an actor descriptor, used to store information that can be potentially different than the actor
- * descriptor itself due to streaming generation logic, etc.
+ * A view on top of an actor desc, used to cache information that can be (potentially) different than the actor desc
+ * itself due to streaming generation logic, etc.
  */
-class FWorldPartitionActorDescView
+class UE_DEPRECATED(5.4, "Class is deprecated in favor of IWorldPartitionActorDescInstanceView interface") FWorldPartitionActorDescView
 {
+	friend class UWorldPartitionRuntimeHash;
+
 public:
 	ENGINE_API FWorldPartitionActorDescView();
-	ENGINE_API FWorldPartitionActorDescView(const FWorldPartitionActorDesc* InActorDesc);
-	
-	virtual ~FWorldPartitionActorDescView() {}
-
-	/** Interface on top of FWorldPartitionActorDesc */
-	ENGINE_API virtual const FGuid& GetGuid() const;
-	ENGINE_API virtual FTopLevelAssetPath GetBaseClass() const;
-	ENGINE_API virtual FTopLevelAssetPath GetNativeClass() const;
-	ENGINE_API virtual UClass* GetActorNativeClass() const;
-	ENGINE_API virtual FName GetRuntimeGrid() const;
-	ENGINE_API virtual bool GetIsSpatiallyLoaded() const;
-	ENGINE_API virtual bool GetActorIsEditorOnly() const;
-	ENGINE_API virtual bool GetActorIsRuntimeOnly() const;
-	ENGINE_API virtual bool GetActorIsHLODRelevant() const;
-	ENGINE_API virtual FSoftObjectPath GetHLODLayer() const;
-	ENGINE_API virtual bool HasResolvedDataLayerInstanceNames() const;
-	ENGINE_API virtual const TArray<FName>& GetDataLayerInstanceNames() const;
-	ENGINE_API virtual const TArray<FName>& GetTags() const;
-	ENGINE_API virtual FName GetActorPackage() const;	
-	ENGINE_API virtual FSoftObjectPath GetActorSoftPath() const;
-	ENGINE_API virtual FName GetActorLabel() const;
-	ENGINE_API virtual FBox GetEditorBounds() const;
-	ENGINE_API virtual FBox GetRuntimeBounds() const;
-	ENGINE_API virtual const TArray<FGuid>& GetReferences() const;
-	ENGINE_API virtual const TArray<FGuid>& GetEditorReferences() const;
-	ENGINE_API virtual FString ToString() const;
-	ENGINE_API virtual const FGuid& GetParentActor() const;
-	ENGINE_API virtual FName GetActorName() const;
-	ENGINE_API virtual const FGuid& GetFolderGuid() const;
-	ENGINE_API virtual FGuid GetContentBundleGuid() const;
-	ENGINE_API virtual FName GetContainerPackage() const;
-	ENGINE_API virtual bool IsContainerInstance() const;
-	ENGINE_API virtual bool GetContainerInstance(FWorldPartitionActorDesc::FContainerInstance& OutContainerInstance) const;
-	ENGINE_API virtual EWorldPartitionActorFilterType GetContainerFilterType() const;
-	ENGINE_API virtual const FWorldPartitionActorFilter* GetContainerFilter() const;	
-	ENGINE_API virtual void CheckForErrors(IStreamingGenerationErrorHandler* ErrorHandler) const;
-	ENGINE_API virtual FName GetActorLabelOrName() const;
-	ENGINE_API virtual AActor* GetActor() const;
-	ENGINE_API virtual bool IsEditorOnlyReference(const FGuid& ReferenceGuid) const;
-	ENGINE_API virtual bool GetProperty(FName PropertyName, FName* PropertyValue) const;
-	ENGINE_API virtual bool HasProperty(FName PropertyName) const;
-
-	/** Helper functions */
+	ENGINE_API FWorldPartitionActorDescView(const FWorldPartitionActorDesc * InActorDesc);
+
+	ENGINE_API const FGuid& GetGuid() const;
+	ENGINE_API FTopLevelAssetPath GetBaseClass() const;
+	ENGINE_API FTopLevelAssetPath GetNativeClass() const;
+	ENGINE_API UClass* GetActorNativeClass() const;
+	UE_DEPRECATED(5.2, "GetOrigin is deprecated.")
+		ENGINE_API FVector GetOrigin() const;
+	ENGINE_API FName GetRuntimeGrid() const;
+	ENGINE_API bool GetIsSpatiallyLoaded() const;
+	ENGINE_API bool GetActorIsEditorOnly() const;
+	ENGINE_API bool GetActorIsRuntimeOnly() const;
+	ENGINE_API bool GetActorIsHLODRelevant() const;
+	ENGINE_API FSoftObjectPath GetHLODLayer() const;
+	ENGINE_API const TArray<FName>& GetDataLayerInstanceNames() const;
+	ENGINE_API const TArray<FName>& GetRuntimeDataLayerInstanceNames() const;
+	ENGINE_API const TArray<FName>& GetTags() const;
+	ENGINE_API FName GetActorPackage() const;
+
+	ENGINE_API FSoftObjectPath GetActorSoftPath() const;
+	ENGINE_API FName GetActorLabel() const;
+
+	UE_DEPRECATED(5.2, "GetBounds is deprecated, GetEditorBounds or GetRuntimeBounds should be used instead.")
+		ENGINE_API FBox GetBounds() const;
+
+	ENGINE_API FBox GetEditorBounds() const;
+	ENGINE_API FBox GetRuntimeBounds() const;
+
+	ENGINE_API const TArray<FGuid>& GetReferences() const;
+	ENGINE_API const TArray<FGuid>& GetEditorReferences() const;
+	ENGINE_API FString ToString() const;
+	ENGINE_API const FGuid& GetParentActor() const;
+	ENGINE_API FName GetActorName() const;
+	ENGINE_API const FGuid& GetFolderGuid() const;
+
+	ENGINE_API FGuid GetContentBundleGuid() const;
+	ENGINE_API FName GetContainerPackage() const;
+	ENGINE_API bool IsContainerInstance() const;
+	ENGINE_API bool GetContainerInstance(FWorldPartitionActorDesc::FContainerInstance & OutContainerInstance) const;
+	ENGINE_API EWorldPartitionActorFilterType GetContainerFilterType() const;
+	ENGINE_API const FWorldPartitionActorFilter* GetContainerFilter() const;
+
+	UE_DEPRECATED(5.3, "GetLevelPackage is deprecated use GetContainerPackage instead.")
+		FName GetLevelPackage() const { return GetContainerPackage(); }
+
+	ENGINE_API void CheckForErrors(IStreamingGenerationErrorHandler * ErrorHandler) const;
+
+	ENGINE_API FName GetActorLabelOrName() const;
+
+	UE_DEPRECATED(5.2, "ShouldValidateRuntimeGrid is deprecated and should not be used.")
+		ENGINE_API bool ShouldValidateRuntimeGrid() const;
+
+	ENGINE_API void SetForcedNonSpatiallyLoaded();
+	ENGINE_API void SetForcedNoRuntimeGrid();
+	ENGINE_API void SetForcedNoDataLayers();
+	ENGINE_API void SetRuntimeDataLayerInstanceNames(const TArray<FName>&InRuntimeDataLayerInstanceNames);
+	ENGINE_API void SetRuntimeReferences(const TArray<FGuid>&InRuntimeReferences);
+	ENGINE_API void SetEditorReferences(const TArray<FGuid>&InEditorReferences);
+	ENGINE_API void SetDataLayerInstanceNames(const TArray<FName>&InDataLayerInstanceNames);
+
+	ENGINE_API void SetForcedNoHLODLayer();
+	ENGINE_API void SetRuntimeHLODLayer(const FSoftObjectPath & InHLODLayer);
+
+	ENGINE_API AActor* GetActor() const;
+
+	bool operator==(const FWorldPartitionActorDescView & Other) const
+	{
+		return GetGuid() == Other.GetGuid();
+	}
+
+	friend uint32 GetTypeHash(const FWorldPartitionActorDescView & Key)
+	{
+		return GetTypeHash(Key.GetGuid());
+	}
+
 	const FWorldPartitionActorDesc* GetActorDesc() const { return ActorDesc; }
-	ENGINE_API virtual const TArray<FName>& GetRuntimeDataLayerInstanceNames() const;
+
+	ENGINE_API bool IsEditorOnlyReference(const FGuid & ReferenceGuid) const;
+
+	ENGINE_API bool GetProperty(FName PropertyName, FName * PropertyValue) const;
+	ENGINE_API bool HasProperty(FName PropertyName) const;
+
+	ENGINE_API void SetParentView(const FWorldPartitionActorDescView * InParentView);
 
 protected:
 	const FWorldPartitionActorDesc* ActorDesc;
+	const FWorldPartitionActorDescView* ParentView;
+	bool bIsForcedNonSpatiallyLoaded;
+	bool bIsForcedNoRuntimeGrid;
+	bool bIsForcedNoDataLayers;
+	bool bIsForceNoHLODLayer;
+	TOptional<TArray<FName>> ResolvedDataLayerInstanceNames;
+	TOptional<TArray<FName>> RuntimeDataLayerInstanceNames;
+	TOptional<TArray<FGuid>> RuntimeReferences;
+	TOptional<FSoftObjectPath> RuntimedHLODLayer;
+	TArray<FGuid> EditorReferences;
 };
-#endif
+#endif
\ No newline at end of file
diff --git a/Engine/Source/Runtime/Engine/Public/WorldPartition/WorldPartitionActorDescViewProxy.h b/Engine/Source/Runtime/Engine/Public/WorldPartition/WorldPartitionActorDescViewProxy.h
--- a/Engine/Source/Runtime/Engine/Public/WorldPartition/WorldPartitionActorDescViewProxy.h
+++ b/Engine/Source/Runtime/Engine/Public/WorldPartition/WorldPartitionActorDescViewProxy.h
@@ -1,13 +1,15 @@
 // Copyright Epic Games, Inc. All Rights Reserved.
 #pragma once
 
+#if WITH_EDITOR
+
 #include "WorldPartitionActorDescView.h"
 
-#if WITH_EDITOR
+PRAGMA_DISABLE_DEPRECATION_WARNINGS
 /**
  * A view on top of an actor desc, which reverts to the actor if its loaded and dirty.
  */
-class ENGINE_API FWorldPartitionActorViewProxy : public FWorldPartitionActorDescView
+class ENGINE_API UE_DEPRECATED(5.4, "Class is deprecated") FWorldPartitionActorViewProxy : public FWorldPartitionActorDescView
 {
 public:
 	FWorldPartitionActorViewProxy(const FWorldPartitionActorDesc* InActorDesc);
@@ -15,4 +17,6 @@ public:
 private:
 	TUniquePtr<FWorldPartitionActorDesc> CachedActorDesc;
 };
+PRAGMA_ENABLE_DEPRECATION_WARNINGS
+
 #endif
diff --git a/Engine/Source/Runtime/Engine/Public/WorldPartition/WorldPartitionActorLoaderInterface.h b/Engine/Source/Runtime/Engine/Public/WorldPartition/WorldPartitionActorLoaderInterface.h
--- a/Engine/Source/Runtime/Engine/Public/WorldPartition/WorldPartitionActorLoaderInterface.h
+++ b/Engine/Source/Runtime/Engine/Public/WorldPartition/WorldPartitionActorLoaderInterface.h
@@ -6,6 +6,7 @@
 #include "UObject/Interface.h"
 #include "Misc/Optional.h"
 #include "WorldPartition/WorldPartitionHandle.h"
+#include "WorldPartition/WorldPartitionActorDesc.h"
 #include "WorldPartitionActorLoaderInterface.generated.h"
 
 UINTERFACE(MinimalAPI)
@@ -21,7 +22,7 @@ class IWorldPartitionActorLoaderInterface
 #if WITH_EDITOR
 	using FReferenceMap = TMap<FGuid, FWorldPartitionReference>;
 	using FActorReferenceMap = TMap<FGuid, FReferenceMap>;
-	using FContainerReferenceMap = TMap<TWeakObjectPtr<UActorDescContainer>, FActorReferenceMap>;
+	using FContainerReferenceMap = TMap<TWeakObjectPtr<UActorDescContainerInstance>, FActorReferenceMap>;
 
 public:
 	/** Base class for actor loaders */
@@ -45,10 +46,15 @@ public:
 		virtual TOptional<FString> GetLabel() const { return TOptional<FString>(); }
 		virtual TOptional<FColor> GetColor() const { return TOptional<FColor>(); }
 
-		ENGINE_API void OnActorDescContainerInitialize(UActorDescContainer* Container);
-		ENGINE_API void OnActorDescContainerUninitialize(UActorDescContainer* Container);
+		UE_DEPRECATED(5.4, "Use OnActorDescContainerInstanceInitialize")
+		ENGINE_API void OnActorDescContainerInitialize(UActorDescContainer* Container) {}
+		UE_DEPRECATED(5.4, "Use OnActorDescContainerInstanceUninitialize")
+		ENGINE_API void OnActorDescContainerUninitialize(UActorDescContainer* Container) {}
 
 	protected:
+		void OnActorDescContainerInstanceInitialize(UActorDescContainerInstance* ContainerInstance);
+		void OnActorDescContainerInstanceUninitialize(UActorDescContainerInstance* ContainerInstance);
+
 		// Private interface
 		virtual void ForEachActor(TFunctionRef<void(const FWorldPartitionHandle&)> InOperation) const =0;
 
@@ -67,8 +73,8 @@ public:
 		ENGINE_API void OnRefreshLoadedState(bool bFromUserOperation);
 
 		// Helpers
-		ENGINE_API FActorReferenceMap& GetContainerReferences(UActorDescContainer* InContainer);
-		ENGINE_API const FActorReferenceMap* GetContainerReferencesConst(UActorDescContainer* InContainer) const;
+		ENGINE_API FActorReferenceMap& GetContainerReferences(UActorDescContainerInstance* InContainerInstance);
+		ENGINE_API const FActorReferenceMap* GetContainerReferencesConst(UActorDescContainerInstance* InContainerInstance) const;
 
 	private:
 		UWorld* World;
@@ -96,6 +102,8 @@ public:
 
 	static ENGINE_API void RefreshLoadedState(bool bIsFromUserChange);
 
+	static ENGINE_API bool GetLoadedChildContainerInstance(const FWorldPartitionHandle& Handle, FWorldPartitionActorDesc::FLoadedContainerInstance& OutContainerInstance);
+
 private:
 	DECLARE_EVENT_OneParam(IWorldPartitionActorLoaderInterface, FOnActorLoaderInterfaceRefreshState, bool /*bIsFromUserChange*/);
 	static ENGINE_API FOnActorLoaderInterfaceRefreshState ActorLoaderInterfaceRefreshState;
diff --git a/Engine/Source/Runtime/Engine/Public/WorldPartition/WorldPartitionBlueprintLibrary.h b/Engine/Source/Runtime/Engine/Public/WorldPartition/WorldPartitionBlueprintLibrary.h
--- a/Engine/Source/Runtime/Engine/Public/WorldPartition/WorldPartitionBlueprintLibrary.h
+++ b/Engine/Source/Runtime/Engine/Public/WorldPartition/WorldPartitionBlueprintLibrary.h
@@ -9,6 +9,7 @@
 #include "WorldPartitionBlueprintLibrary.generated.h"
 
 class FWorldPartitionActorDesc;
+class FWorldPartitionActorDescInstance;
 class UWorldPartition;
 class UDataLayerManager;
 
@@ -90,10 +91,10 @@ class UWorldPartitionBlueprintLibrary : public UBlueprintFunctionLibrary
 	static FDelegateHandle OnWorldPartitionUninitializedHandle;
 
 	static bool GetActorDescs(const UWorldPartition* WorldPartition, TArray<FActorDesc>& OutActorDescs);
-	static bool GetActorDescs(const UActorDescContainer* InContainer, const FTransform& InTransform, TArray<FActorDesc>& OutActorDescs);
+	static bool GetActorDescs(const UActorDescContainerInstance* InContainer, const FTransform& InTransform, TArray<FActorDesc>& OutActorDescs);
 	static bool GetIntersectingActorDescs(UWorldPartition* WorldPartition, const FBox& InBox, TArray<FActorDesc>& OutActorDescs);
-	static bool GetIntersectingActorDescs(const UActorDescContainer* InContainer, const FBox& InBox, const FTransform& InTransform, TArray<FActorDesc>& OutActorDescs);
-	static bool HandleIntersectingActorDesc(const FWorldPartitionActorDesc* ActorDesc, const FBox& InBox, const FTransform& InTransform, TArray<FActorDesc>& OutActorDescs);
+	static bool GetIntersectingActorDescs(const UActorDescContainerInstance* InContainerInstance, const FBox& InBox, const FTransform& InTransform, TArray<FActorDesc>& OutActorDescs);
+	static bool HandleIntersectingActorDesc(const FWorldPartitionActorDescInstance* ActorDescInstance, const FBox& InBox, const FTransform& InTransform, TArray<FActorDesc>& OutActorDescs);
 #endif
 
 public:
diff --git a/Engine/Source/Runtime/Engine/Public/WorldPartition/WorldPartitionEditorHash.h b/Engine/Source/Runtime/Engine/Public/WorldPartition/WorldPartitionEditorHash.h
--- a/Engine/Source/Runtime/Engine/Public/WorldPartition/WorldPartitionEditorHash.h
+++ b/Engine/Source/Runtime/Engine/Public/WorldPartition/WorldPartitionEditorHash.h
@@ -45,6 +45,9 @@ public:
 		FForEachIntersectingActorParams& SetMinimumBox(const FBox& InMinimumBox) { MinimumBox = InMinimumBox; return *this; };
 	};
 
-	ENGINE_API virtual int32 ForEachIntersectingActor(const FBox& Box, TFunctionRef<void(FWorldPartitionActorDesc*)> InOperation, const FForEachIntersectingActorParams& Params = FForEachIntersectingActorParams()) PURE_VIRTUAL(UWorldPartitionEditorHash::ForEachIntersectingActor, return 0;);
+	UE_DEPRECATED(5.4, "Use ForEachIntersectingActor with FWorldPartitionActorDescInstance")
+	ENGINE_API virtual int32 ForEachIntersectingActor(const FBox& Box, TFunctionRef<void(FWorldPartitionActorDesc*)> InOperation, const FForEachIntersectingActorParams& Params = FForEachIntersectingActorParams()) { return 0; }
+
+	ENGINE_API virtual int32 ForEachIntersectingActor(const FBox& Box, TFunctionRef<void(FWorldPartitionActorDescInstance*)> InOperation, const FForEachIntersectingActorParams& Params = FForEachIntersectingActorParams()) PURE_VIRTUAL(UWorldPartitionEditorHash::ForEachIntersectingActor, return 0;);
 #endif
 };
diff --git a/Engine/Source/Runtime/Engine/Public/WorldPartition/WorldPartitionEditorSpatialHash.h b/Engine/Source/Runtime/Engine/Public/WorldPartition/WorldPartitionEditorSpatialHash.h
--- a/Engine/Source/Runtime/Engine/Public/WorldPartition/WorldPartitionEditorSpatialHash.h
+++ b/Engine/Source/Runtime/Engine/Public/WorldPartition/WorldPartitionEditorSpatialHash.h
@@ -198,7 +198,10 @@ public:
 	ENGINE_API virtual void HashActor(FWorldPartitionHandle& InActorHandle) override;
 	ENGINE_API virtual void UnhashActor(FWorldPartitionHandle& InActorHandle) override;
 
-	ENGINE_API virtual int32 ForEachIntersectingActor(const FBox& Box, TFunctionRef<void(FWorldPartitionActorDesc*)> InOperation, const FForEachIntersectingActorParams& Params = FForEachIntersectingActorParams()) override;
+	UE_DEPRECATED(5.4, "Use ForEachIntersectingActor with FWorldPartitionActorDescInstance")
+	ENGINE_API virtual int32 ForEachIntersectingActor(const FBox& Box, TFunctionRef<void(FWorldPartitionActorDesc*)> InOperation, const FForEachIntersectingActorParams& Params = FForEachIntersectingActorParams()) override { return 0;  }
+
+	ENGINE_API virtual int32 ForEachIntersectingActor(const FBox& Box, TFunctionRef<void(FWorldPartitionActorDescInstance*)> InOperation, const FForEachIntersectingActorParams& Params = FForEachIntersectingActorParams()) override;
 	// UWorldPartitionEditorHash interface end
 #endif
 
diff --git a/Engine/Source/Runtime/Engine/Public/WorldPartition/WorldPartitionHandle.h b/Engine/Source/Runtime/Engine/Public/WorldPartition/WorldPartitionHandle.h
--- a/Engine/Source/Runtime/Engine/Public/WorldPartition/WorldPartitionHandle.h
+++ b/Engine/Source/Runtime/Engine/Public/WorldPartition/WorldPartitionHandle.h
@@ -12,9 +12,16 @@ class FWorldPartitionActorDesc;
 struct FWorldPartitionHandleImpl;
 struct FWorldPartitionReferenceImpl;
 
+class AActor;
+class UActorDescContainerInstance;
+class FWorldPartitionActorDescInstance;
+
 template<class U>
 class TActorDescContainerCollection;
 
+template<class U>
+class TActorDescContainerInstanceCollection;
+
 class FWorldPartitionLoadingContext
 {
 public:
@@ -30,8 +37,8 @@ public:
 		ENGINE_API virtual ~IContext();
 
 	private:
-		virtual void RegisterActor(FWorldPartitionActorDesc* ActorDesc) = 0;
-		virtual void UnregisterActor(FWorldPartitionActorDesc* ActorDesc) = 0;
+		virtual void RegisterActor(FWorldPartitionActorDescInstance* InActorDescInstance) = 0;
+		virtual void UnregisterActor(FWorldPartitionActorDescInstance* InActorDescInstance) = 0;
 	};
 
 	/**
@@ -40,8 +47,8 @@ public:
 	class FImmediate : public IContext
 	{
 	private:
-		ENGINE_API virtual void RegisterActor(FWorldPartitionActorDesc* ActorDesc) override;
-		ENGINE_API virtual void UnregisterActor(FWorldPartitionActorDesc* ActorDesc) override;
+		ENGINE_API virtual void RegisterActor(FWorldPartitionActorDescInstance* InActorDescInstance) override;
+		ENGINE_API virtual void UnregisterActor(FWorldPartitionActorDescInstance* InActorDescInstance) override;
 	};
 
 	/**
@@ -59,16 +66,16 @@ public:
 		bool GetNeedsClearTransactions() const { return bNeedsClearTransactions; }
 
 	private:
-		ENGINE_API virtual void RegisterActor(FWorldPartitionActorDesc* ActorDesc) override;
-		ENGINE_API virtual void UnregisterActor(FWorldPartitionActorDesc* ActorDesc) override;
+		ENGINE_API virtual void RegisterActor(FWorldPartitionActorDescInstance* InActorDescInstance) override;
+		ENGINE_API virtual void UnregisterActor(FWorldPartitionActorDescInstance* InActorDescInstance) override;
 
-		struct FContainerOps
+		struct FContainerInstanceOps
 		{
-			TSet<FWorldPartitionActorDesc*> Registrations;
-			TSet<FWorldPartitionActorDesc*> Unregistrations;
+			TSet<FWorldPartitionActorDescInstance*> Registrations;
+			TSet<FWorldPartitionActorDescInstance*> Unregistrations;
 		};
 
-		TMap<UActorDescContainer*, FContainerOps> ContainerOps;
+		TMap<UActorDescContainerInstance*, FContainerInstanceOps> ContainerInstanceOps;
 
 		int32 NumRegistrations = 0;
 		int32 NumUnregistrations = 0;
@@ -81,12 +88,17 @@ public:
 	class FNull : public IContext
 	{
 	private:
-		virtual void RegisterActor(FWorldPartitionActorDesc* ActorDesc) override {}
-		virtual void UnregisterActor(FWorldPartitionActorDesc* ActorDesc) override {}
+		virtual void RegisterActor(FWorldPartitionActorDescInstance* InActorDescInstance) override {}
+		virtual void UnregisterActor(FWorldPartitionActorDescInstance* InActorDescInstance) override {}
 	};
 
-	static ENGINE_API void LoadAndRegisterActor(FWorldPartitionActorDesc* ActorDesc);
-	static ENGINE_API void UnloadAndUnregisterActor(FWorldPartitionActorDesc* ActorDesc);
+	UE_DEPRECATED(5.4, "Use FWorldPartitionActorDescInstance version instead")
+	static ENGINE_API void LoadAndRegisterActor(FWorldPartitionActorDesc* ActorDesc) {}
+	UE_DEPRECATED(5.4, "Use FWorldPartitionActorDescInstance version instead")
+	static ENGINE_API void UnloadAndUnregisterActor(FWorldPartitionActorDesc* ActorDesc) {}
+
+	static ENGINE_API void LoadAndRegisterActor(FWorldPartitionActorDescInstance* ActorDesc);
+	static ENGINE_API void UnloadAndUnregisterActor(FWorldPartitionActorDescInstance* ActorDesc);
 
 private:
 	static FImmediate DefaultContext;
@@ -98,12 +110,13 @@ class TWorldPartitionHandle
 {
 public:
 	FORCEINLINE TWorldPartitionHandle()
-		: ActorDesc(nullptr)
+		: ContainerInstance(nullptr)
+		, ActorDescInstance(nullptr)
 	{}
 
-	FORCEINLINE TWorldPartitionHandle(TUniquePtr<FWorldPartitionActorDesc>* InActorDesc)
-		: Container(Impl::GetActorDescContainer(InActorDesc))
-		, ActorDesc(InActorDesc)
+	FORCEINLINE TWorldPartitionHandle(TUniquePtr<FWorldPartitionActorDescInstance>* InActorDescInstance)
+		: ContainerInstance(Impl::GetActorDescContainerInstance(InActorDescInstance))
+		, ActorDescInstance(InActorDescInstance)
 	{
 		if (IsValid())
 		{
@@ -111,9 +124,24 @@ public:
 		}
 	}
 
-	FORCEINLINE TWorldPartitionHandle(UActorDescContainer* InContainer, const FGuid& ActorGuid)
-		: Container(InContainer)
-		, ActorDesc(Impl::GetActorDesc(InContainer, ActorGuid))
+	FORCEINLINE TWorldPartitionHandle(FWorldPartitionActorDescInstance* InActorDescInstance)
+		: ContainerInstance(Impl::GetActorDescContainerInstance(InActorDescInstance))
+		, ActorDescInstance(nullptr)
+	{
+		if (ContainerInstance != nullptr)
+		{
+			ActorDescInstance = Impl::GetActorDescInstance(ContainerInstance.Get(), Impl::GetActorDescInstanceGuid(InActorDescInstance));
+		}
+
+		if (IsValid())
+		{
+			IncRefCount();
+		}
+	}
+
+	FORCEINLINE TWorldPartitionHandle(UActorDescContainerInstance* InContainerInstance, const FGuid& InActorGuid)
+		: ContainerInstance(InContainerInstance)
+		, ActorDescInstance(Impl::GetActorDescInstance(InContainerInstance, InActorGuid))
 	{
 		if (IsValid())
 		{
@@ -122,13 +150,13 @@ public:
 	}
 
 	template<class U>
-	FORCEINLINE TWorldPartitionHandle(TActorDescContainerCollection<U>* ContainerCollection, const FGuid& ActorGuid)
-		: Container(Impl::GetActorDescContainer(ContainerCollection, ActorGuid))
-		, ActorDesc(nullptr)
+	FORCEINLINE TWorldPartitionHandle(TActorDescContainerInstanceCollection<U>* ContainerCollection, const FGuid& ActorGuid)
+		: ContainerInstance(Impl::GetActorDescContainerInstance(ContainerCollection, ActorGuid))
+		, ActorDescInstance(nullptr)
 	{
-		if (Container != nullptr)
+		if (ContainerInstance != nullptr)
 		{
-			ActorDesc = Impl::GetActorDesc(Container.Get(), ActorGuid);
+			ActorDescInstance = Impl::GetActorDescInstance(ContainerInstance.Get(), ActorGuid);
 		}
 
 		if (IsValid())
@@ -137,16 +165,45 @@ public:
 		}
 	}
 
+	UE_DEPRECATED(5.4, "Use FWorldPartitionActorDescInstance version")
+	FORCEINLINE TWorldPartitionHandle(TUniquePtr<FWorldPartitionActorDesc>* InActorDesc)
+		: ContainerInstance(nullptr)
+		, ActorDescInstance(nullptr)
+	{
+	}
+
+	UE_DEPRECATED(5.4, "Use FActorDescContainerInstance version")
+	FORCEINLINE TWorldPartitionHandle(UActorDescContainer* InContainer, const FGuid& ActorGuid)
+		: ContainerInstance(nullptr)
+		, ActorDescInstance(nullptr)
+	{
+	}
+		
+	template<class U>
+	struct TActorDescContainerCollectionDeprecated
+	{
+		UE_DEPRECATED(5.4, "Use TActorDescContainerInstanceCollection version")
+		static void Deprecated() {}
+	};
+
+	template<class U>
+	FORCEINLINE TWorldPartitionHandle(TActorDescContainerCollection<U>* ContainerCollection, const FGuid& ActorGuid)
+		: ContainerInstance(nullptr)
+		, ActorDescInstance(nullptr)
+	{
+		TActorDescContainerCollectionDeprecated<U>::Deprecated();
+	}
+	
 	FORCEINLINE TWorldPartitionHandle(const TWorldPartitionHandle& Other)
-		: Container(nullptr)
-		, ActorDesc(nullptr)
+		: ContainerInstance(nullptr)
+		, ActorDescInstance(nullptr)
 	{
 		*this = Other;
 	}
 
 	FORCEINLINE TWorldPartitionHandle(TWorldPartitionHandle&& Other)
-		: Container(nullptr)
-		, ActorDesc(nullptr)
+		: ContainerInstance(nullptr)
+		, ActorDescInstance(nullptr)
 	{
 		*this = MoveTemp(Other);
 	}
@@ -163,7 +220,7 @@ public:
 
 	template <typename T>
 	FORCEINLINE TWorldPartitionHandle<Impl>(TWorldPartitionHandle<T>&& Other)
-		: ActorDesc(nullptr)
+		: ActorDescInstance(nullptr)
 	{
 		*this = MoveTemp(Other);
 	}
@@ -185,8 +242,8 @@ public:
 				DecRefCount();
 			}
 
-			Container = Other.Container;
-			ActorDesc = Other.ActorDesc;
+			ContainerInstance = Other.ContainerInstance;
+			ActorDescInstance = Other.ActorDescInstance;
 
 			if (IsValid())
 			{
@@ -206,11 +263,11 @@ public:
 				DecRefCount();
 			}
 
-			Container = Other.Container;
-			ActorDesc = Other.ActorDesc;
+			ContainerInstance = Other.ContainerInstance;
+			ActorDescInstance = Other.ActorDescInstance;
 			
-			Other.Container = nullptr;
-			Other.ActorDesc = nullptr;
+			Other.ContainerInstance = nullptr;
+			Other.ActorDescInstance = nullptr;
 		}
 
 		return *this;
@@ -225,8 +282,8 @@ public:
 			DecRefCount();
 		}
 
-		Container = Other.Container;
-		ActorDesc = Other.ActorDesc;
+		ContainerInstance = Other.ContainerInstance;
+		ActorDescInstance = Other.ActorDescInstance;
 
 		if (IsValid())
 		{
@@ -244,44 +301,71 @@ public:
 			DecRefCount();
 		}
 
-		Container = Other.Container;
-		ActorDesc = Other.ActorDesc;
+		ContainerInstance = Other.ContainerInstance;
+		ActorDescInstance = Other.ActorDescInstance;
 
 		if (IsValid())
 		{
 			IncRefCount();
 			Other.DecRefCount();
 
-			Other.Container = nullptr;
-			Other.ActorDesc = nullptr;
+			Other.ContainerInstance = nullptr;
+			Other.ActorDescInstance = nullptr;
 		}
 
 		return *this;
 	}
 
+	UE_DEPRECATED(5.4, "Use GetInstance instead")
 	FORCEINLINE FWorldPartitionActorDesc* operator->() const
 	{
+		PRAGMA_DISABLE_DEPRECATION_WARNINGS
 		return Get();
+		PRAGMA_ENABLE_DEPRECATION_WARNINGS
 	}
 
+	UE_DEPRECATED(5.4, "Use GetInstance instead")
 	FORCEINLINE FWorldPartitionActorDesc* operator*() const
 	{
+		PRAGMA_DISABLE_DEPRECATION_WARNINGS
 		return Get();
+		PRAGMA_ENABLE_DEPRECATION_WARNINGS
+	}
+
+	FORCEINLINE FWorldPartitionActorDescInstance* GetInstance() const
+	{
+		return IsValid() ? ActorDescInstance->Get() : nullptr;
+	}
+
+	FORCEINLINE UActorDescContainerInstance* GetContainerInstance() const
+	{
+		return IsValid() ? Impl::GetActorDescContainerInstance(ActorDescInstance) : nullptr;
 	}
 
 	FORCEINLINE bool IsValid() const
 	{
-		return Container.IsValid() && ActorDesc && ActorDesc->IsValid();
+		return ContainerInstance.IsValid() && ActorDescInstance && ActorDescInstance->IsValid();
 	}
 
 	FORCEINLINE bool IsLoaded() const
 	{
-		return IsValid() && Impl::IsActorDescLoaded(ActorDesc->Get());
+		return IsValid() && Impl::IsLoaded(ActorDescInstance->Get());
 	}
 
+	FORCEINLINE AActor* GetActor() const
+	{
+		return IsValid() ? Impl::GetActor(ActorDescInstance->Get()) : nullptr;
+	}
+
+	UE_DEPRECATED(5.4, "Use GetInstance instead")
 	FORCEINLINE FWorldPartitionActorDesc* Get() const
 	{
-		return IsValid() ? ActorDesc->Get() : nullptr;
+		if (ContainerInstance.IsValid() && ActorDescInstance && ActorDescInstance->IsValid())
+		{
+			return Impl::GetActorDesc(ActorDescInstance->Get());
+		}
+
+		return nullptr;
 	}
 
 	FORCEINLINE void Reset()
@@ -291,18 +375,18 @@ public:
 			DecRefCount();
 		}
 
-		Container = nullptr;
-		ActorDesc = nullptr;
+		ContainerInstance = nullptr;
+		ActorDescInstance = nullptr;
 	}
 	
 	friend FORCEINLINE uint32 GetTypeHash(const TWorldPartitionHandle<Impl>& HandleBase)
 	{
-		return ::PointerHash(HandleBase.ActorDesc);
+		return ::PointerHash(HandleBase.ActorDescInstance);
 	}
 
 	FORCEINLINE bool operator==(const TWorldPartitionHandle& Other) const
 	{
-		return ActorDesc == Other.ActorDesc;
+		return ActorDescInstance == Other.ActorDescInstance;
 	}
 
 	FORCEINLINE bool operator!=(const TWorldPartitionHandle& Other) const
@@ -314,7 +398,7 @@ public:
 	template <typename T>
 	FORCEINLINE bool operator==(const TWorldPartitionHandle<T>& Other) const
 	{
-		return Get() == *Other;
+		return GetInstance() == Other.GetInstance();
 	}
 
 	template <typename T>
@@ -326,42 +410,74 @@ public:
 public:
 	FORCEINLINE void IncRefCount()
 	{
-		Impl::IncRefCount(ActorDesc->Get());
+		Impl::IncRefCount(ActorDescInstance->Get());
 	}
 
 	FORCEINLINE void DecRefCount()
 	{
-		Impl::DecRefCount(ActorDesc->Get());
+		Impl::DecRefCount(ActorDescInstance->Get());
 	}
 
-	TWeakObjectPtr<UActorDescContainer> Container;
-	TUniquePtr<FWorldPartitionActorDesc>* ActorDesc;
+	TWeakObjectPtr<UActorDescContainerInstance> ContainerInstance;
+	TUniquePtr<FWorldPartitionActorDescInstance>* ActorDescInstance;
 };
 
 struct FWorldPartitionImplBase
 {
-	static ENGINE_API TUniquePtr<FWorldPartitionActorDesc>* GetActorDesc(UActorDescContainer* Container, const FGuid& ActorGuid);
-	static ENGINE_API UActorDescContainer* GetActorDescContainer(TUniquePtr<FWorldPartitionActorDesc>* ActorDesc);
-	static ENGINE_API bool IsActorDescLoaded(FWorldPartitionActorDesc* ActorDesc);
-
+	UE_DEPRECATED(5.4, "Use FWorldPartitionActorDescInstance")
+	static ENGINE_API TUniquePtr<FWorldPartitionActorDesc>* GetActorDesc(UActorDescContainer* Container, const FGuid& ActorGuid) { return nullptr;}
+	UE_DEPRECATED(5.4, "Use FWorldPartitionActorDescInstance")
+	static ENGINE_API UActorDescContainer* GetActorDescContainer(TUniquePtr<FWorldPartitionActorDesc>* ActorDesc) { return nullptr; }
+	UE_DEPRECATED(5.4, "Use FWorldPartitionActorDescInstance")
+	static ENGINE_API bool IsActorDescLoaded(FWorldPartitionActorDesc* ActorDesc) { return false; }
+	UE_DEPRECATED(5.4, "Use FWorldPartitionActorDescInstance")
+	static ENGINE_API AActor* GetActorDescActor(FWorldPartitionActorDesc* ActorDesc) { return nullptr; }
+	
 	template<class U>
 	static UActorDescContainer* GetActorDescContainer(TActorDescContainerCollection<U>* ContainerCollection, const FGuid& ActorGuid)
 	{
-		return ContainerCollection ? ContainerCollection->GetActorDescContainer(ActorGuid) : nullptr;
+		return nullptr;
+	}
+
+	static ENGINE_API bool IsLoaded(FWorldPartitionActorDescInstance* InActorDescInstance);
+	static ENGINE_API UActorDescContainerInstance* GetActorDescContainerInstance(TUniquePtr<FWorldPartitionActorDescInstance>* InActorDescIntance);
+	static ENGINE_API UActorDescContainerInstance* GetActorDescContainerInstance(FWorldPartitionActorDescInstance* InActorDescIntance);
+	static ENGINE_API FGuid GetActorDescInstanceGuid(const FWorldPartitionActorDescInstance* InActorDescIntance);
+	static ENGINE_API TUniquePtr<FWorldPartitionActorDescInstance>* GetActorDescInstance(UActorDescContainerInstance* InContainerInstance, const FGuid& InActorGuid);
+	static ENGINE_API FWorldPartitionActorDesc* GetActorDesc(FWorldPartitionActorDescInstance* InActorDescInstance);
+
+	static ENGINE_API AActor* GetActor(FWorldPartitionActorDescInstance* InActorDescInstance);
+
+	template<class U>
+	static UActorDescContainerInstance* GetActorDescContainerInstance(TActorDescContainerInstanceCollection<U>* ContainerCollection, const FGuid& ActorGuid)
+	{
+		return ContainerCollection ? ContainerCollection->GetActorDescContainerInstance(ActorGuid) : nullptr;
 	}
 };
 
 struct FWorldPartitionHandleImpl : FWorldPartitionImplBase
 {
-	static ENGINE_API void IncRefCount(FWorldPartitionActorDesc* ActorDesc);
-	static ENGINE_API void DecRefCount(FWorldPartitionActorDesc* ActorDesc);
+	UE_DEPRECATED(5.4, "Use FWorldPartitionActorDescInstance")
+	static ENGINE_API void IncRefCount(FWorldPartitionActorDesc* ActorDesc) {}
+	UE_DEPRECATED(5.4, "Use FWorldPartitionActorDescInstance")
+	static ENGINE_API void DecRefCount(FWorldPartitionActorDesc* ActorDesc) {}
+
+	static ENGINE_API void IncRefCount(FWorldPartitionActorDescInstance* InActorDescInstance);
+	static ENGINE_API void DecRefCount(FWorldPartitionActorDescInstance* InActorDescInstance);
+
 	static ENGINE_API TWorldPartitionHandle<FWorldPartitionReferenceImpl> ToReference(const TWorldPartitionHandle<FWorldPartitionHandleImpl>& Source);
 };
 
 struct FWorldPartitionReferenceImpl : FWorldPartitionImplBase
 {
-	static ENGINE_API void IncRefCount(FWorldPartitionActorDesc* ActorDesc);
-	static ENGINE_API void DecRefCount(FWorldPartitionActorDesc* ActorDesc);
+	UE_DEPRECATED(5.4, "Use FWorldPartitionActorDescInstance")
+	static ENGINE_API void IncRefCount(FWorldPartitionActorDesc* ActorDesc) {}
+	UE_DEPRECATED(5.4, "Use FWorldPartitionActorDescInstance")
+	static ENGINE_API void DecRefCount(FWorldPartitionActorDesc* ActorDesc) {}
+
+	static ENGINE_API void IncRefCount(FWorldPartitionActorDescInstance* InActorDescInstance);
+	static ENGINE_API void DecRefCount(FWorldPartitionActorDescInstance* InActorDescInstance);
+
 	static ENGINE_API TWorldPartitionHandle<FWorldPartitionHandleImpl> ToHandle(const TWorldPartitionHandle<FWorldPartitionReferenceImpl>& Source);
 };
 
diff --git a/Engine/Source/Runtime/Engine/Public/WorldPartition/WorldPartitionHelpers.h b/Engine/Source/Runtime/Engine/Public/WorldPartition/WorldPartitionHelpers.h
--- a/Engine/Source/Runtime/Engine/Public/WorldPartition/WorldPartitionHelpers.h
+++ b/Engine/Source/Runtime/Engine/Public/WorldPartition/WorldPartitionHelpers.h
@@ -16,6 +16,7 @@ class UWorldPartition;
 
 #if WITH_EDITOR
 class FWorldPartitionActorDesc;
+class FWorldPartitionActorDescInstance;
 #endif
 
 namespace FWorldPartitionHelpersPrivate
@@ -76,21 +77,49 @@ public:
 private:
 	static ENGINE_API bool IsActorDescClassCompatibleWith(const FWorldPartitionActorDesc* ActorDesc, const UClass* Class);
 
+	template<class ActorClass>
+	class TDeprecated
+	{
+		UE_DEPRECATED(5.4, "Use ForEachIntersectingActorDescInstance instead")
+		static void DeprecatedForEachIntersectingActorDesc() {}
+
+		UE_DEPRECATED(5.4, "Use ForEachActorDescInstance instead")
+		static void DeprecatedForEachActorDesc() {}
+	};
+
 public:
 	template <class ActorClass = AActor>
 	static void ForEachIntersectingActorDesc(UWorldPartition* WorldPartition, const FBox& Box, TFunctionRef<bool(const FWorldPartitionActorDesc*)> Func)
 	{
-		ForEachIntersectingActorDesc(WorldPartition, Box, ActorClass::StaticClass(), Func);
+		TDeprecated<ActorClass>::DeprecatedForEachIntersectingActorDesc();
 	}
 
 	template<class ActorClass = AActor>
 	static void ForEachActorDesc(UWorldPartition* WorldPartition, TFunctionRef<bool(const FWorldPartitionActorDesc*)> Func)
 	{
-		ForEachActorDesc(WorldPartition, ActorClass::StaticClass(), Func);
+		TDeprecated<ActorClass>::DeprecatedForEachActorDesc();
 	}
 
-	static ENGINE_API void ForEachIntersectingActorDesc(UWorldPartition* WorldPartition, const FBox& Box, TSubclassOf<AActor> ActorClass, TFunctionRef<bool(const FWorldPartitionActorDesc*)> Func);
-	static ENGINE_API void ForEachActorDesc(UWorldPartition* WorldPartition, TSubclassOf<AActor> ActorClass, TFunctionRef<bool(const FWorldPartitionActorDesc*)> Func);
+	UE_DEPRECATED(5.4, "Use ForEachIntersectingActorDescInstance")
+	static ENGINE_API void ForEachIntersectingActorDesc(UWorldPartition* WorldPartition, const FBox& Box, TSubclassOf<AActor> ActorClass, TFunctionRef<bool(const FWorldPartitionActorDesc*)> Func) {}
+	UE_DEPRECATED(5.4, "Use ForEachActorDescInstance")
+	static ENGINE_API void ForEachActorDesc(UWorldPartition* WorldPartition, TSubclassOf<AActor> ActorClass, TFunctionRef<bool(const FWorldPartitionActorDesc*)> Func) {}
+
+
+	template <class ActorClass = AActor>
+	static void ForEachIntersectingActorDescInstance(UWorldPartition* WorldPartition, const FBox& Box, TFunctionRef<bool(const FWorldPartitionActorDescInstance*)> Func)
+	{
+		ForEachIntersectingActorDescInstance(WorldPartition, Box, ActorClass::StaticClass(), Func);
+	}
+
+	template<class ActorClass = AActor>
+	static void ForEachActorDescInstance(UWorldPartition* WorldPartition, TFunctionRef<bool(const FWorldPartitionActorDescInstance*)> Func)
+	{
+		ForEachActorDescInstance(WorldPartition, ActorClass::StaticClass(), Func);
+	}
+
+	static ENGINE_API void ForEachIntersectingActorDescInstance(UWorldPartition* WorldPartition, const FBox& Box, TSubclassOf<AActor> ActorClass, TFunctionRef<bool(const FWorldPartitionActorDescInstance*)> Func);
+	static ENGINE_API void ForEachActorDescInstance(UWorldPartition* WorldPartition, TSubclassOf<AActor> ActorClass, TFunctionRef<bool(const FWorldPartitionActorDescInstance*)> Func);
 
 	/* Struct of optional parameters passed to foreach actordesc functions. */
 	struct FForEachActorWithLoadingParams
@@ -123,9 +152,16 @@ public:
 		TMap<FGuid, FWorldPartitionReference> ActorReferences;
 	};
 
-	static ENGINE_API void ForEachActorWithLoading(UWorldPartition* WorldPartition, TFunctionRef<bool(const FWorldPartitionActorDesc*)> Func, const FForEachActorWithLoadingParams& Params = FForEachActorWithLoadingParams());
-	static ENGINE_API void ForEachActorWithLoading(UWorldPartition* WorldPartition, TFunctionRef<bool(const FWorldPartitionActorDesc*)> Func, const FForEachActorWithLoadingParams& Params, FForEachActorWithLoadingResult& Result);
+	UE_DEPRECATED(5.4, "Use ForEachActorWithLoading with FWorldPartitionActorDescInstance")
+	static ENGINE_API void ForEachActorWithLoading(UWorldPartition* WorldPartition, TFunctionRef<bool(const FWorldPartitionActorDesc*)> Func, const FForEachActorWithLoadingParams& Params = FForEachActorWithLoadingParams()) {}
+
+	UE_DEPRECATED(5.4, "Use ForEachActorWithLoading with FWorldPartitionActorDescInstance")
+	static ENGINE_API void ForEachActorWithLoading(UWorldPartition* WorldPartition, TFunctionRef<bool(const FWorldPartitionActorDesc*)> Func, const FForEachActorWithLoadingParams& Params, FForEachActorWithLoadingResult& Result) {}
 	
+	static ENGINE_API void ForEachActorWithLoading(UWorldPartition* WorldPartition, TFunctionRef<bool(const FWorldPartitionActorDescInstance*)> Func, const FForEachActorWithLoadingParams& Params = FForEachActorWithLoadingParams());
+	static ENGINE_API void ForEachActorWithLoading(UWorldPartition* WorldPartition, TFunctionRef<bool(const FWorldPartitionActorDescInstance*)> Func, const FForEachActorWithLoadingParams& Params, FForEachActorWithLoadingResult& Result);
+
+
 	static ENGINE_API bool HasExceededMaxMemory();
 	static ENGINE_API bool ShouldCollectGarbage();
 	static ENGINE_API void DoCollectGarbage();
diff --git a/Engine/Source/Runtime/Engine/Public/WorldPartition/WorldPartitionRuntimeCell.h b/Engine/Source/Runtime/Engine/Public/WorldPartition/WorldPartitionRuntimeCell.h
--- a/Engine/Source/Runtime/Engine/Public/WorldPartition/WorldPartitionRuntimeCell.h
+++ b/Engine/Source/Runtime/Engine/Public/WorldPartition/WorldPartitionRuntimeCell.h
@@ -4,7 +4,6 @@
 #include "CoreMinimal.h"
 #include "Templates/SubclassOf.h"
 #include "WorldPartition/DataLayer/DataLayerInstance.h"
-#include "WorldPartition/WorldPartitionActorDescView.h"
 #include "WorldPartition/WorldPartitionStreamingSource.h"
 #include "WorldPartition/WorldPartitionActorContainerID.h"
 #include "WorldPartition/WorldPartitionRuntimeCellData.h"
@@ -20,6 +19,7 @@ class UDataLayerAsset;
 class UDataLayerInstance;
 class UWorldPartition;
 class UDataLayerManager;
+class FStreamingGenerationActorDescView;
 struct FHierarchicalLogArchive;
 
 enum class EWorldPartitionDataLayersLogicOperator : uint8;
@@ -241,7 +241,15 @@ class UWorldPartitionRuntimeCell : public UObject, public IWorldPartitionCell
 	virtual TSet<FName> GetActorPackageNames() const override { return TSet<FName>(); }
 	//~End IWorldPartitionCell Interface
 
-	ENGINE_API virtual void AddActorToCell(const FWorldPartitionActorDescView& ActorDescView, const FActorContainerID& InContainerID, const FTransform& InContainerTransform, const UActorDescContainer* InContainer) PURE_VIRTUAL(UWorldPartitionRuntimeCell::AddActorToCell,);
+	ENGINE_API virtual void AddActorToCell(const FStreamingGenerationActorDescView& ActorDescView) PURE_VIRTUAL(UWorldPartitionRuntimeCell::AddActorToCell, );
+	
+	PRAGMA_DISABLE_DEPRECATION_WARNINGS
+
+	UE_DEPRECATED(5.4, "Implement FStreamingGenerationActorDescView version instead")
+	virtual void AddActorToCell(const class FWorldPartitionActorDescView& ActorDescView, const FActorContainerID& InContainerID, const FTransform& InContainerTransform, const UActorDescContainer* InContainer) {}
+
+	PRAGMA_ENABLE_DEPRECATION_WARNINGS
+
 	ENGINE_API virtual void Fixup() PURE_VIRTUAL(UWorldPartitionRuntimeCell::Fixup, );
 	ENGINE_API virtual int32 GetActorCount() const PURE_VIRTUAL(UWorldPartitionRuntimeCell::GetActorCount, return 0;);
 
diff --git a/Engine/Source/Runtime/Engine/Public/WorldPartition/WorldPartitionRuntimeContainerResolving.h b/Engine/Source/Runtime/Engine/Public/WorldPartition/WorldPartitionRuntimeContainerResolving.h
--- a/Engine/Source/Runtime/Engine/Public/WorldPartition/WorldPartitionRuntimeContainerResolving.h
+++ b/Engine/Source/Runtime/Engine/Public/WorldPartition/WorldPartitionRuntimeContainerResolving.h
@@ -82,6 +82,11 @@ public:
 		return Containers.Add(InContainerPackage);	
 	}
 
+	bool ContainsContainer(FName InContainerPackage)
+	{
+		return Containers.Contains(InContainerPackage);
+	}
+
 	void BuildContainerIDToEditorPathMap();
 	const FString* FindContainerEditorPath(const FActorContainerID& InContainerID) const { return ContainerIDToEditorPath.Find(InContainerID); }
 #endif
diff --git a/Engine/Source/Runtime/Engine/Public/WorldPartition/WorldPartitionRuntimeHash.h b/Engine/Source/Runtime/Engine/Public/WorldPartition/WorldPartitionRuntimeHash.h
--- a/Engine/Source/Runtime/Engine/Public/WorldPartition/WorldPartitionRuntimeHash.h
+++ b/Engine/Source/Runtime/Engine/Public/WorldPartition/WorldPartitionRuntimeHash.h
@@ -188,9 +188,6 @@ protected:
 	TArray<FAlwaysLoadedActorForPIE> AlwaysLoadedActorsForPIE;
 
 	TMap<FString, UWorldPartitionRuntimeCell*> PackagesToGenerateForCook;
-
-public:
-	mutable FActorDescList ModifiedActorDescListForPIE;
 #endif
 
 protected:
diff --git a/Engine/Source/Runtime/Engine/Public/WorldPartition/WorldPartitionRuntimeLevelStreamingCell.h b/Engine/Source/Runtime/Engine/Public/WorldPartition/WorldPartitionRuntimeLevelStreamingCell.h
--- a/Engine/Source/Runtime/Engine/Public/WorldPartition/WorldPartitionRuntimeLevelStreamingCell.h
+++ b/Engine/Source/Runtime/Engine/Public/WorldPartition/WorldPartitionRuntimeLevelStreamingCell.h
@@ -10,6 +10,7 @@
 #include "WorldPartitionRuntimeLevelStreamingCell.generated.h"
 
 class UWorld;
+class FStreamingGenerationActorDescView;
 
 UCLASS(MinimalAPI)
 class UWorldPartitionRuntimeLevelStreamingCell : public UWorldPartitionRuntimeCell
@@ -46,7 +47,7 @@ class UWorldPartitionRuntimeLevelStreamingCell : public UWorldPartitionRuntimeCe
 
 #if WITH_EDITOR
 	//~Begin UWorldPartitionRuntimeCell Interface
-	ENGINE_API virtual void AddActorToCell(const FWorldPartitionActorDescView& ActorDescView, const FActorContainerID& InContainerID, const FTransform& InContainerTransform, const UActorDescContainer* InContainer) override;
+	ENGINE_API virtual void AddActorToCell(const FStreamingGenerationActorDescView& ActorDescView) override;
 	ENGINE_API virtual void Fixup() override;
 	ENGINE_API virtual int32 GetActorCount() const override;
 	ENGINE_API virtual void DumpStateLog(FHierarchicalLogArchive& Ar) const override;
diff --git a/Engine/Source/Runtime/Engine/Public/WorldPartition/WorldPartitionStreamingGeneration.h b/Engine/Source/Runtime/Engine/Public/WorldPartition/WorldPartitionStreamingGeneration.h
--- a/Engine/Source/Runtime/Engine/Public/WorldPartition/WorldPartitionStreamingGeneration.h
+++ b/Engine/Source/Runtime/Engine/Public/WorldPartition/WorldPartitionStreamingGeneration.h
@@ -3,80 +3,152 @@
 
 #if WITH_EDITOR
 #include "CoreMinimal.h"
+#include "OverrideVoidReturnInvoker.h"
 #include "WorldPartition/ActorDescContainerCollection.h"
-#include "WorldPartition/WorldPartitionActorDescViewMap.h"
+#include "WorldPartition/ActorDescContainerInstanceCollection.h"
+#include "WorldPartition/WorldPartitionActorDescInstance.h"
+#include "WorldPartition/WorldPartitionActorDescInstanceInterface.h"
+#include "WorldPartition/WorldPartitionActorDescInstanceViewInterface.h"
 
-class FStreamingGenerationActorDescCollection : public TActorDescContainerCollection<TObjectPtr<const UActorDescContainer>>
-{
-public:
-	FStreamingGenerationActorDescCollection() = default;
-	ENGINE_API FStreamingGenerationActorDescCollection(std::initializer_list<TObjectPtr<const UActorDescContainer>> ActorDescContainerArray);
-	ENGINE_API FStreamingGenerationActorDescCollection(const TArray<const UActorDescContainer*>& ActorDescContainers);
-
-	ENGINE_API UWorld* GetWorld() const;
+// deprecated
+class FWorldPartitionActorDescView;
 
-	// @todo_ow : Remove once conversion to ExternalDataLayer is complete. 
-	// It is present to handle content bundles streaming generation via the same code path. 
-	ENGINE_API FGuid GetContentBundleGuid() const;
+class FStreamingGenerationActorDescViewMap;
 
-	ENGINE_API const UActorDescContainer* GetMainActorDescContainer() const;
-	ENGINE_API FName GetMainContainerPackageName() const;
-	ENGINE_API TArrayView<const UActorDescContainer* const> GetExternalDataLayerContainers();
+class FStreamingGenerationActorDescView : public IWorldPartitionActorDescInstanceView
+{
+	friend class FStreamingGenerationActorDescViewMap;
+	friend class FWorldPartitionStreamingGenerator;
 
-	ENGINE_API virtual void OnCollectionChanged() override;
+protected:
+	// Used for invalid reference error reporting
+	FStreamingGenerationActorDescView(const IWorldPartitionActorDescInstance* InActorDescInstance)
+		: ActorDescViewMap(nullptr)
+		, ActorDescInstance(InActorDescInstance)
+		, ParentView(nullptr)
+		, bIsForcedNonSpatiallyLoaded(false)
+		, bIsForcedNoRuntimeGrid(false)
+		, bIsForcedNoDataLayers(false)
+		, bIsForceNoHLODLayer(false)
+		, bIsUnsaved(false)
+	{
+	}
 
-private:
-	ENGINE_API void SortCollection();
+public:
+	FStreamingGenerationActorDescView(const FStreamingGenerationActorDescViewMap& InActorDescViewMap, const IWorldPartitionActorDescInstance* InActorDescInstance, bool bInUnsaved = false)
+		: FStreamingGenerationActorDescView(InActorDescInstance)
+	{
+		ActorDescViewMap = &InActorDescViewMap;
+		bIsUnsaved = bInUnsaved;
+	}
+		
+	virtual ~FStreamingGenerationActorDescView() {}
 
-	static constexpr int MainContainerIdx = 0;
-	static constexpr int ExternalDataLayerContainerStartIdx = MainContainerIdx + 1;
-};
+	//~ Begin FWorldPartitionActorDescViewProxy	
+	virtual const FGuid& GetGuid() const override { return ActorDescInstance->GetGuid(); }
 
-class FStreamingGenerationActorDescView : public FWorldPartitionActorDescView
-{
-public:
-	ENGINE_API FStreamingGenerationActorDescView();
-	ENGINE_API FStreamingGenerationActorDescView(const FWorldPartitionActorDesc* InActorDesc);
+	virtual FTopLevelAssetPath GetBaseClass() const override { return ActorDescInstance->GetBaseClass(); }
+	virtual FTopLevelAssetPath GetNativeClass() const override { return ActorDescInstance->GetNativeClass(); }
+	virtual UClass* GetActorNativeClass() const override { return ActorDescInstance->GetActorNativeClass(); }
 
-	//~ Begin FWorldPartitionActorDescView interface
 	ENGINE_API virtual FName GetRuntimeGrid() const override;
 	ENGINE_API virtual bool GetIsSpatiallyLoaded() const override;
+	virtual bool GetActorIsEditorOnly() const override { return ActorDescInstance->GetActorIsEditorOnly(); }
+	virtual bool GetActorIsRuntimeOnly() const override { return ActorDescInstance->GetActorIsRuntimeOnly(); }
+	virtual bool IsRuntimeRelevant() const override { return ActorDescInstance->IsRuntimeRelevant(); }
+	virtual bool IsEditorRelevant() const override { return ActorDescInstance->IsEditorRelevant(); }
+
+	virtual bool IsUsingDataLayerAsset() const override { return ActorDescInstance->IsUsingDataLayerAsset(); }
+	virtual const TArray<FName>& GetDataLayers() const override { return ActorDescInstance->GetDataLayers(); }
+
+	virtual bool GetActorIsHLODRelevant() const override { return ActorDescInstance->GetActorIsHLODRelevant(); }
 	ENGINE_API virtual FSoftObjectPath GetHLODLayer() const override;
-	ENGINE_API virtual const TArray<FName>& GetDataLayerInstanceNames() const override;
+
+	virtual const TArray<FName>& GetTags() const override { return ActorDescInstance->GetTags(); }
+	virtual FName GetActorPackage() const override { return ActorDescInstance->GetActorPackage(); }
+	virtual FSoftObjectPath GetActorSoftPath() const override { return ActorDescInstance->GetActorSoftPath(); }
+	virtual FName GetActorLabel() const override { return ActorDescInstance->GetActorLabel(); }
+	virtual FName GetActorName() const override { return ActorDescInstance->GetActorName(); }
+	virtual FName GetFolderPath() const override { return ActorDescInstance->GetFolderPath(); }
+	virtual const FGuid& GetFolderGuid() const override { return ActorDescInstance->GetFolderGuid(); }
+
+	virtual FBox GetEditorBounds() const override { return ActorDescInstance->GetEditorBounds(); }
+	virtual FBox GetRuntimeBounds() const override { return ActorDescInstance->GetRuntimeBounds(); }
+
+	virtual bool GetProperty(FName PropertyName, FName* PropertyValue) const override { return ActorDescInstance->GetProperty(PropertyName, PropertyValue); }
+	virtual bool HasProperty(FName PropertyName) const override { return ActorDescInstance->HasProperty(PropertyName); }
+
 	ENGINE_API virtual const TArray<FGuid>& GetReferences() const override;
-	ENGINE_API virtual const TArray<FGuid>& GetEditorReferences() const override;
-	//~ End FWorldPartitionActorDescView interface
+	virtual const TArray<FGuid>& GetEditorOnlyReferences() const override { return ActorDescInstance->GetEditorOnlyReferences(); }
+	virtual bool IsEditorOnlyReference(const FGuid& ReferenceGuid) const override { return ActorDescInstance->IsEditorOnlyReference(ReferenceGuid); }
+
+	virtual const FGuid& GetParentActor() const override { return ActorDescInstance->GetParentActor(); }
+
+	virtual FGuid GetContentBundleGuid() const override { return ActorDescInstance->GetContentBundleGuid(); }
+
+	virtual bool IsChildContainerInstance() const override { return ActorDescInstance->IsChildContainerInstance(); }
+	virtual FName GetChildContainerPackage() const override { return ActorDescInstance->GetChildContainerPackage(); }
+	virtual EWorldPartitionActorFilterType GetChildContainerFilterType() const override { return ActorDescInstance->GetChildContainerFilterType(); }
+	virtual const FWorldPartitionActorFilter* GetChildContainerFilter() const override { return ActorDescInstance->GetChildContainerFilter(); }
+	virtual bool GetChildContainerInstance(FWorldPartitionActorDesc::FContainerInstance& OutContainerInstance) const override { return ActorDescInstance->GetChildContainerInstance(OutContainerInstance); }
+
+	virtual bool IsMainWorldOnly() const override { return ActorDescInstance->IsMainWorldOnly(); }
+	virtual bool IsListedInSceneOutliner() const override { return ActorDescInstance->IsListedInSceneOutliner(); }
+	virtual const FGuid& GetSceneOutlinerParent() const override { return ActorDescInstance->GetSceneOutlinerParent(); }
+
+	virtual void CheckForErrors(IStreamingGenerationErrorHandler* ErrorHandler) const override { GetActorDesc()->CheckForErrors(this, ErrorHandler); }
+
+	virtual FString ToString(FWorldPartitionActorDesc::EToStringMode Mode) const override { return ActorDescInstance->ToString(Mode); }
+	virtual const FWorldPartitionActorDesc* GetActorDesc() const override { return ActorDescInstance->GetActorDesc(); }
+
+	virtual bool HasResolvedDataLayerInstanceNames() const override { return ActorDescInstance->HasResolvedDataLayerInstanceNames(); }
+	ENGINE_API const TArray<FName>& GetDataLayerInstanceNames() const override;
+
+	virtual AActor* GetActor(bool bEvenIfPendingKill = true, bool bEvenIfUnreachable = false) const override { return ActorDescInstance->GetActor(bEvenIfPendingKill, bEvenIfUnreachable); }
+	virtual bool IsLoaded(bool bEvenIfPendingKill = false) const override { return ActorDescInstance->IsLoaded(bEvenIfPendingKill); }
+
+	virtual UActorDescContainerInstance* GetContainerInstance() const override { return ActorDescInstance->GetContainerInstance(); }
+	//~ End
+
+	bool IsUnsaved() const { return bIsUnsaved; }
+	const FStreamingGenerationActorDescViewMap& GetActorDescViewMap() const { return *ActorDescViewMap; }
+
+	ENGINE_API const TArray<FGuid>& GetEditorReferences() const;
+	ENGINE_API const TArray<FName>& GetRuntimeDataLayerInstanceNames() const;
 
-	ENGINE_API bool ShouldValidateRuntimeGrid() const;
-	ENGINE_API void SetParentView(const FWorldPartitionActorDescView* InParentView);
-	ENGINE_API void SetDataLayerInstanceNames(const TArray<FName>& InDataLayerInstanceNames);
 	ENGINE_API void SetForcedNonSpatiallyLoaded();
 	ENGINE_API void SetForcedNoRuntimeGrid();
 	ENGINE_API void SetForcedNoDataLayers();
 	ENGINE_API void SetRuntimeDataLayerInstanceNames(const TArray<FName>& InRuntimeDataLayerInstanceNames);
 	ENGINE_API void SetRuntimeReferences(const TArray<FGuid>& InRuntimeReferences);
-	ENGINE_API void SetEditorReferences(const TArray<FGuid>& InEditorReferences);		
+	ENGINE_API void SetEditorReferences(const TArray<FGuid>& InEditorReferences);
+	ENGINE_API void SetDataLayerInstanceNames(const TArray<FName>& InDataLayerInstanceNames);
+	ENGINE_API void SetParentView(const FStreamingGenerationActorDescView* InParentView);
+
 	ENGINE_API void SetForcedNoHLODLayer();
 	ENGINE_API void SetRuntimeHLODLayer(const FSoftObjectPath& InHLODLayer);
 
-	ENGINE_API const TArray<FName>& GetRuntimeDataLayerInstanceNames() const;
-
-	bool operator==(const FWorldPartitionActorDescView& Other) const
+	bool operator==(const FStreamingGenerationActorDescView& Other) const
 	{
 		return GetGuid() == Other.GetGuid();
 	}
 
-	friend uint32 GetTypeHash(const FWorldPartitionActorDescView& Key)
+	friend uint32 GetTypeHash(const FStreamingGenerationActorDescView& Key)
 	{
 		return GetTypeHash(Key.GetGuid());
 	}
 
-protected:
-	const FWorldPartitionActorDescView* ParentView;
+private:
+	const FStreamingGenerationActorDescViewMap* ActorDescViewMap;
+	const IWorldPartitionActorDescInstance* ActorDescInstance;
+		
+	const FStreamingGenerationActorDescView* ParentView;
+
 	bool bIsForcedNonSpatiallyLoaded;
 	bool bIsForcedNoRuntimeGrid;
 	bool bIsForcedNoDataLayers;
 	bool bIsForceNoHLODLayer;
+	bool bIsUnsaved;
 	TOptional<TArray<FName>> ResolvedDataLayerInstanceNames;
 	TOptional<TArray<FName>> RuntimeDataLayerInstanceNames;
 	TOptional<TArray<FGuid>> RuntimeReferences;
@@ -84,6 +156,153 @@ protected:
 	TArray<FGuid> EditorReferences;
 };
 
-class FStreamingGenerationActorDescViewMap : public TActorDescViewMap<FStreamingGenerationActorDescView> {};
+class FStreamingGenerationActorDescViewMap
+{
+	friend class FWorldPartitionStreamingGenerator;
+
+private:
+	template <class Func>
+	void ForEachActorDescView(Func InFunc)
+	{
+		TOverrideVoidReturnInvoker Invoker(true, InFunc);
+
+		for (TUniquePtr<FStreamingGenerationActorDescView>& ActorDescView : ActorDescViewList)
+		{
+			if (!Invoker(*ActorDescView))
+			{
+				return;
+			}
+		}
+	}
+
+	FStreamingGenerationActorDescView* FindByGuid(const FGuid& InGuid)
+	{
+		if (FStreamingGenerationActorDescView** ActorDescViewPtr = ActorDescViewsByGuid.Find(InGuid))
+		{
+			return *ActorDescViewPtr;
+		}
+		return nullptr;
+	}
+
+	FStreamingGenerationActorDescView& FindByGuidChecked(const FGuid& InGuid)
+	{
+		return *ActorDescViewsByGuid.FindChecked(InGuid);
+	}
+
+public:
+	ENGINE_API FStreamingGenerationActorDescViewMap();
+
+	// Non-copyable but movable
+	FStreamingGenerationActorDescViewMap(const FStreamingGenerationActorDescViewMap&) = delete;
+	FStreamingGenerationActorDescViewMap(FStreamingGenerationActorDescViewMap&&) = default;
+
+	FStreamingGenerationActorDescViewMap& operator=(const FStreamingGenerationActorDescViewMap&) = delete;
+	FStreamingGenerationActorDescViewMap& operator=(FStreamingGenerationActorDescViewMap&&) = default;
+
+	ENGINE_API FStreamingGenerationActorDescView* Emplace(const FGuid& InActorGuid, const FStreamingGenerationActorDescView& InActorDescView);
+	ENGINE_API FStreamingGenerationActorDescView* Emplace(const FWorldPartitionActorDescInstance* InActorDescInstance);
+
+	PRAGMA_DISABLE_DEPRECATION_WARNINGS
+
+	UE_DEPRECATED(5.4, "Use FStreamingGenerationActorDescView version instead")
+	ENGINE_API FWorldPartitionActorDescView* Emplace(const FGuid& InActorGuid, const FWorldPartitionActorDescView& InActorDescView) { return nullptr; }
+
+	UE_DEPRECATED(5.4, "Use FWorldPartitionActorDescInstance version instead")
+	ENGINE_API class FWorldPartitionActorDescView* Emplace(const FWorldPartitionActorDesc* InActorDesc) { return nullptr; }
+
+	PRAGMA_ENABLE_DEPRECATION_WARNINGS
+			
+	FORCEINLINE int32 Num() const
+	{
+		return ActorDescViewList.Num();
+	}
+
+	template <class Func>
+	void ForEachActorDescView(Func InFunc) const
+	{
+		TOverrideVoidReturnInvoker Invoker(true, InFunc);
+
+		for (const TUniquePtr<FStreamingGenerationActorDescView>& ActorDescView : ActorDescViewList)
+		{
+			if (!Invoker(*ActorDescView))
+			{
+				return;
+			}
+		}
+	}
+
+	const FStreamingGenerationActorDescView* FindByGuid(const FGuid& InGuid) const
+	{
+		if (const FStreamingGenerationActorDescView* const* ActorDescViewPtr = ActorDescViewsByGuid.Find(InGuid))
+		{
+			return *ActorDescViewPtr;
+		}
+		return nullptr;
+	}
+
+	const FStreamingGenerationActorDescView& FindByGuidChecked(const FGuid& InGuid) const
+	{
+		return *ActorDescViewsByGuid.FindChecked(InGuid);
+	}
+
+	template <class ClassType>
+	TArray<const FStreamingGenerationActorDescView*> FindByExactNativeClass() const
+	{
+		return FindByExactNativeClass(ClassType::StaticClass());
+	}
+
+	ENGINE_API TArray<const FStreamingGenerationActorDescView*> FindByExactNativeClass(UClass* InExactNativeClass) const;
+
+	const TMap<FGuid, FStreamingGenerationActorDescView*>& GetActorDescViewsByGuid() const { return ActorDescViewsByGuid; }
+
+protected:
+	TArray<TUniquePtr<FStreamingGenerationActorDescView>> ActorDescViewList;
+
+	TMap<FGuid, FStreamingGenerationActorDescView*> ActorDescViewsByGuid;
+	TMultiMap<FName, const FStreamingGenerationActorDescView*> ActorDescViewsByClass;
+};
+
+class FStreamingGenerationContainerInstanceCollection : public TActorDescContainerInstanceCollection<TObjectPtr<const UActorDescContainerInstance>>
+{
+public:
+	FStreamingGenerationContainerInstanceCollection() = default;
+	ENGINE_API FStreamingGenerationContainerInstanceCollection(std::initializer_list<TObjectPtr<const UActorDescContainerInstance>> ActorDescContainerInstanceArray);
+	ENGINE_API FStreamingGenerationContainerInstanceCollection(const TArray<const UActorDescContainerInstance*>& ActorDescContainerInstances);
+
+	ENGINE_API UWorld* GetWorld() const;
+
+	// @todo_ow : Remove once conversion to ExternalDataLayer is complete. 
+	// It is present to handle content bundles streaming generation via the same code path. 
+	ENGINE_API FGuid GetContentBundleGuid() const;
+
+	ENGINE_API const UActorDescContainerInstance* GetMainContainer() const;
+	ENGINE_API FName GetMainContainerPackageName() const;
+	ENGINE_API TArrayView<const UActorDescContainerInstance* const> GetExternalDataLayerContainers();
+
+	ENGINE_API virtual void OnCollectionChanged() override;
+
+private:
+	ENGINE_API void SortCollection();
+
+	static constexpr int MainContainerIdx = 0;
+	static constexpr int ExternalDataLayerContainerStartIdx = MainContainerIdx + 1;
+};
+
+class UE_DEPRECATED(5.4, "Use FStreamingGenerationContainerInstanceCollection instead") FStreamingGenerationActorDescCollection : public TActorDescContainerCollection<TObjectPtr<const UActorDescContainer>>
+{
+public:
+	FStreamingGenerationActorDescCollection() = default;
+	FStreamingGenerationActorDescCollection(std::initializer_list<TObjectPtr<const UActorDescContainer>> ActorDescContainerArray) {}
+	FStreamingGenerationActorDescCollection(const TArray<const UActorDescContainer*>& ActorDescContainers) {}
+
+	UWorld* GetWorld() const { return nullptr; }
+	FGuid GetContentBundleGuid() const { return FGuid(); }
+	const UActorDescContainer* GetMainActorDescContainer() const { return nullptr; }
+	FName GetMainContainerPackageName() const { return NAME_None; }
+	TArrayView<const UActorDescContainer* const> GetExternalDataLayerContainers() { return TArrayView<const UActorDescContainer*>(); }
+	virtual void OnCollectionChanged() override {}
+};
+
+using FActorDescViewMap UE_DEPRECATED(5.4, "Use FStreamingGenerationActorDescViewMap instead") = FStreamingGenerationActorDescViewMap;
 
 #endif // WITH_EDITOR
diff --git a/Engine/Source/Runtime/Engine/Public/WorldPartition/WorldPartitionStreamingGenerationContext.h b/Engine/Source/Runtime/Engine/Public/WorldPartition/WorldPartitionStreamingGenerationContext.h
--- a/Engine/Source/Runtime/Engine/Public/WorldPartition/WorldPartitionStreamingGenerationContext.h
+++ b/Engine/Source/Runtime/Engine/Public/WorldPartition/WorldPartitionStreamingGenerationContext.h
@@ -6,9 +6,9 @@
 #include "WorldPartition/WorldPartitionActorContainerID.h"
 
 class FStreamingGenerationActorDescViewMap;
+class FStreamingGenerationActorDescView;
 class FWorldPartitionStreamingGenerator;
-class FStreamingGenerationActorDescCollection;
-class FWorldPartitionActorDescView;
+class FStreamingGenerationContainerInstanceCollection;
 class UActorDescContainer;
 class UDataLayerInstance;
 struct FWorldPartitionRuntimeContainerResolver;
@@ -28,10 +28,7 @@ public:
 		TArray<FGuid> Actors;
 	};
 
-	/**
-	 * An actor set container represents the list of actor sets in an actor container, e.g. a level instance.
-	 */
-	struct FActorSetContainer
+	struct UE_DEPRECATED(5.4, "Use FActorSetContainerInstance") FActorSetContainer
 	{
 		FActorSetContainer()
 			: ActorDescViewMap(nullptr)
@@ -43,7 +40,26 @@ public:
 		FActorSetContainer& operator=(const FActorSetContainer&) = delete;
 
 		const FStreamingGenerationActorDescViewMap* ActorDescViewMap;
-		const FStreamingGenerationActorDescCollection* ActorDescCollection; // Only used by UWorldPartitionRuntimeSpatialHash::SetupHLODActors
+		const class FStreamingGenerationActorDescCollection* ActorDescCollection; // Only used by UWorldPartitionRuntimeSpatialHash::SetupHLODActors
+		TArray<TUniquePtr<FActorSet>> ActorSets;
+	};
+
+	/**
+	 * An actor set container represents the list of actor sets in an actor container, e.g. a level instance.
+	 */
+	struct FActorSetContainerInstance
+	{
+		FActorSetContainerInstance()
+			: ActorDescViewMap(nullptr)
+			, ContainerInstanceCollection(nullptr)
+		{}
+
+		// Non-copyable
+		FActorSetContainerInstance(const FActorSetContainerInstance&) = delete;
+		FActorSetContainerInstance& operator=(const FActorSetContainerInstance&) = delete;
+
+		const FStreamingGenerationActorDescViewMap* ActorDescViewMap;
+		const FStreamingGenerationContainerInstanceCollection* ContainerInstanceCollection; // Only used by UWorldPartitionRuntimeSpatialHash::SetupHLODActors
 		TArray<TUniquePtr<FActorSet>> ActorSets;
 	};
 
@@ -57,7 +73,16 @@ public:
 		bool bIsSpatiallyLoaded;
 		TArray<const UDataLayerInstance*> DataLayers;
 		FGuid ContentBundleID;
-		const FActorSetContainer* ContainerInstance;
+
+		PRAGMA_DISABLE_DEPRECATION_WARNINGS
+
+		UE_DEPRECATED(5.4, "Use ActorSetContainerInstance instead")
+		const FActorSetContainer* ContainerInstance = nullptr;
+
+		PRAGMA_ENABLE_DEPRECATION_WARNINGS
+
+		const FActorSetContainerInstance* ActorSetContainerInstance = nullptr;
+
 		FActorContainerID ContainerID;
 		FTransform Transform;
 		const FActorSet* ActorSet;
@@ -94,16 +119,27 @@ public:
 		FGuid ActorGuid;
 		const FActorSetInstance* ActorSetInstance;
 
-		ENGINE_API const FWorldPartitionActorDescView& GetActorDescView() const;
+		ENGINE_API const FStreamingGenerationActorDescView& GetActorDescView() const;
 		ENGINE_API const FActorContainerID& GetContainerID() const;
 		ENGINE_API const FTransform& GetTransform() const;
 		ENGINE_API const FBox GetBounds() const;
 	};
 
 	virtual FBox GetWorldBounds() const = 0;
-	virtual const FActorSetContainer* GetMainWorldContainer() const = 0;
+
+	virtual const FActorSetContainerInstance* GetMainWorldContainerInstance() const = 0;
 	virtual void ForEachActorSetInstance(TFunctionRef<void(const FActorSetInstance&)> Func) const = 0;
-	virtual void ForEachActorSetContainer(TFunctionRef<void(const FActorSetContainer&)> Func) const = 0;
+	virtual void ForEachActorSetContainerInstance(TFunctionRef<void(const FActorSetContainerInstance&)> Func) const = 0;
+
+	PRAGMA_DISABLE_DEPRECATION_WARNINGS
+
+	UE_DEPRECATED(5.4, "Implement GetMainWorldContainerInstance instead")
+	virtual const FActorSetContainer* GetMainWorldContainer() const { return nullptr; };
+
+	UE_DEPRECATED(5.4, "Implement ForEachActorSetContainerInstance instead")
+	virtual void ForEachActorSetContainer(TFunctionRef<void(const FActorSetContainer&)> Func) const {};
+
+	PRAGMA_ENABLE_DEPRECATION_WARNINGS
 };
 
 class FStreamingGenerationContextProxy : public IStreamingGenerationContext
@@ -118,9 +154,9 @@ public:
 		return SourceContext->GetWorldBounds();
 	}
 
-	virtual const FActorSetContainer* GetMainWorldContainer() const override
+	virtual const FActorSetContainerInstance* GetMainWorldContainerInstance() const override
 	{
-		return SourceContext->GetMainWorldContainer();
+		return SourceContext->GetMainWorldContainerInstance();
 	}
 
 	virtual void ForEachActorSetInstance(TFunctionRef<void(const FActorSetInstance&)> Func) const override
@@ -128,9 +164,9 @@ public:
 		SourceContext->ForEachActorSetInstance(Func);
 	}
 
-	virtual void ForEachActorSetContainer(TFunctionRef<void(const FActorSetContainer&)> Func) const override
+	virtual void ForEachActorSetContainerInstance(TFunctionRef<void(const FActorSetContainerInstance&)> Func) const override
 	{
-		SourceContext->ForEachActorSetContainer(Func);
+		SourceContext->ForEachActorSetContainerInstance(Func);
 	}
 
 protected:
