Iris
- Fixed ensure when client destroys actor marked as bNetTemporary.

diff --git a/Engine/Source/Runtime/Engine/Private/Net/Iris/ReplicationSystem/ActorReplicationBridge.cpp b/Engine/Source/Runtime/Engine/Private/Net/Iris/ReplicationSystem/ActorReplicationBridge.cpp
--- a/Engine/Source/Runtime/Engine/Private/Net/Iris/ReplicationSystem/ActorReplicationBridge.cpp
+++ b/Engine/Source/Runtime/Engine/Private/Net/Iris/ReplicationSystem/ActorReplicationBridge.cpp
@@ -434,8 +434,9 @@ void UActorReplicationBridge::EndReplication(AActor* Actor, EEndPlayReason::Type
 			Flags |= EEndReplicationFlags::Destroy | EEndReplicationFlags::DestroyNetHandle | EEndReplicationFlags::ClearNetPushId;
 		}
 
-		// If we are shutting down we do not need to validate that we are not detaching remote instances by accident.
-		if ((EndPlayReason == EEndPlayReason::EndPlayInEditor) || (EndPlayReason == EEndPlayReason::Quit))
+		// If we are shutting down or the actor is a nettemporary we do not need to validate that we are not detaching remote instances by accident.
+		const bool bIsShuttingDown = (EndPlayReason == EEndPlayReason::EndPlayInEditor) || (EndPlayReason == EEndPlayReason::Quit);
+		if (bIsShuttingDown || Actor->bNetTemporary)
 		{
 			Flags |= EEndReplicationFlags::SkipPendingEndReplicationValidation;
 		}
