WorldPartition: change FWorldPartitionHelpers::TDeprecated to struct so that api is public (usage will generate deprecation warnings instead of errors)

#rnx

diff --git a/Engine/Source/Runtime/Engine/Public/WorldPartition/WorldPartitionHelpers.h b/Engine/Source/Runtime/Engine/Public/WorldPartition/WorldPartitionHelpers.h
--- a/Engine/Source/Runtime/Engine/Public/WorldPartition/WorldPartitionHelpers.h
+++ b/Engine/Source/Runtime/Engine/Public/WorldPartition/WorldPartitionHelpers.h
@@ -78,7 +78,7 @@ private:
 	static ENGINE_API bool IsActorDescClassCompatibleWith(const FWorldPartitionActorDesc* ActorDesc, const UClass* Class);
 
 	template<class ActorClass>
-	class TDeprecated
+	struct TDeprecated
 	{
 		UE_DEPRECATED(5.4, "Use ForEachIntersectingActorDescInstance instead")
 		static void DeprecatedForEachIntersectingActorDesc() {}
