Texture add APIs SetModernSettingsForNewOrChangedTexture and BlockOnAnyAsyncBuild
SetModernSettingsForNewOrChangedTexture should be used on newly made or changed textures
BlockOnAnyAsyncBuild instead of directly calling to TextureCompilingManager or Modify
TextureRenderTarget::GetRenderTargetTextureClass
make TextureSource::Init do UseHashAsGuid so it behaves the same as Lock/Unlock

diff --git a/Engine/Source/Runtime/Engine/Classes/Engine/Texture.h b/Engine/Source/Runtime/Engine/Classes/Engine/Texture.h
--- a/Engine/Source/Runtime/Engine/Classes/Engine/Texture.h
+++ b/Engine/Source/Runtime/Engine/Classes/Engine/Texture.h
@@ -247,7 +247,8 @@ struct FTextureSource
 	/** PNG Compresses the source art if possible or tells the bulk data to zlib compress when it saves out to disk. */
 	ENGINE_API void Compress();
 
-	/** Force the GUID to change even if mip data has not been modified. */
+	/** Force the GUID to change even if mip data has not been modified.
+	This should be avoided; typically let the data hash be used as guid to uniquely identify the data. */
 	ENGINE_API void ForceGenerateGuid();
 
 	/** Use FMipLock to encapsulate a locked mip */
@@ -255,7 +256,8 @@ struct FTextureSource
 	/** Lock a mip for reading. */
 	ENGINE_API const uint8* LockMipReadOnly(int32 BlockIndex, int32 LayerIndex, int32 MipIndex);
 
-	/** Lock a mip for editing. */
+	/** Lock a mip for editing.
+	Note that Lock/Unlock for edit automatically does UseHashAsGuid. */
 	ENGINE_API uint8* LockMip(int32 BlockIndex, int32 LayerIndex, int32 MipIndex);
 
 	/** Unlock a mip. */
@@ -333,6 +335,8 @@ struct FTextureSource
 
 	/** Trivial accessors. These will only give values for Block0 so may not be correct for UDIM/multi-block textures, use GetBlock() for this case. */
 	FGuid GetPersistentId() const { return BulkData.GetIdentifier(); }
+	/** GetId() returns a hash of the Id member (data hash) and also the attributes of the Source.
+	( GetId does not just return Id ) **/
 	ENGINE_API FGuid GetId() const;
 	FORCEINLINE int64 GetSizeX() const { return SizeX; }
 	FORCEINLINE int64 GetSizeY() const { return SizeY; }
@@ -601,7 +605,8 @@ private:
 	bool EnsureBlocksAreSorted();
 
 public:
-	/** Uses a hash as the GUID, useful to prevent creating new GUIDs on load for legacy assets. */
+	/** Uses a hash as the GUID, useful to prevent creating new GUIDs on load for legacy assets.
+	This is automatically done by Init() and Mip Lock/Unlock.  New textures should always have the data hash as Id. */
 	ENGINE_API void UseHashAsGuid();
 
 	void ReleaseSourceMemory(); // release the memory from the mips (does almost the same as remove source data except doesn't rebuild the guid)
@@ -612,7 +617,9 @@ private:
 #endif
 
 #if WITH_EDITORONLY_DATA
-	/** GUID used to track changes to the source data. */
+	/** GUID used to track changes to the source data.
+	Typically with UseHashAsGuid , this "Id" is the hash of the BulkData.
+	Note that GetId() is not == Id. */
 	UPROPERTY(VisibleAnywhere, Category=TextureSource)
 	FGuid Id;
 
@@ -1708,6 +1715,13 @@ public:
 	*/
 	ENGINE_API virtual void UpdateOodleTextureSdkVersionToLatest(bool bDoPrePostEditChangeIfChanging = false);
 	
+	/** Set new default settings that are desired for textures.
+	These cannot be set in the constructor automatically to avoid changing old content.
+	When new textures are made, or the texture content changes, call this to update settings.
+	ApplyDefaultsForNewlyImportedTextures calls this, you don't need to call both.
+	*/
+	ENGINE_API virtual void SetModernSettingsForNewOrChangedTexture();
+
 	/** Get TextureFormatName with platform remaps and conditional prefix
 	 *   this is the entry point API for getting the final texture format name
 	 *  OutFormats will be resized to the number of sub-flavors of the platform (typically just 1)
@@ -1761,6 +1775,10 @@ public:
 	virtual bool IsCompiling() const override { return IsDefaultTexture(); }
 	//~ End AsyncCompilation Interface
 
+	// if texture is currently in async build action or queue, block until it is done
+	//	note: Modify() and PreEditChange() do this.  You should usually be using PreEdit/PostEdit and not calling this directly.
+	ENGINE_API virtual void BlockOnAnyAsyncBuild();
+
 	//~ Begin UObject Interface.
 	ENGINE_API virtual bool Modify(bool bAlwaysMarkDirty = true) override;
 	ENGINE_API virtual void PostEditChangeProperty(FPropertyChangedEvent& PropertyChangedEvent) override;
diff --git a/Engine/Source/Runtime/Engine/Classes/Engine/TextureRenderTarget.h b/Engine/Source/Runtime/Engine/Classes/Engine/TextureRenderTarget.h
--- a/Engine/Source/Runtime/Engine/Classes/Engine/TextureRenderTarget.h
+++ b/Engine/Source/Runtime/Engine/Classes/Engine/TextureRenderTarget.h
@@ -71,12 +71,12 @@ class UTextureRenderTarget : public UTexture
 	/**
 	 * Returns the UTexture class that corresponds to this render target (e.g. UTexture2D for UTextureRenderTarget2D)
 	 */
-	ENGINE_API virtual TSubclassOf<UTexture> GetTextureUClass() const PURE_VIRTUAL(UTextureRenderTarget, return nullptr;)
+	ENGINE_API virtual TSubclassOf<UTexture> GetTextureUClass() const PURE_VIRTUAL(GetTextureUClass, return nullptr;)
+	// GetTextureClass() will just return "RenderTarget" ; to get the sub-type (2d/cube), use GetRenderTargetTextureClass
+	ENGINE_API virtual ETextureClass GetRenderTargetTextureClass() const PURE_VIRTUAL(GetRenderTargetTextureClass, return ETextureClass::Invalid;)
 
 	ENGINE_API virtual EPixelFormat GetFormat() const PURE_VIRTUAL(GetFormat,return PF_Unknown;)
-
 	ENGINE_API virtual bool IsSRGB() const PURE_VIRTUAL(IsSRGB,return false;)
-
 	ENGINE_API virtual float GetDisplayGamma() const PURE_VIRTUAL(GetDisplayGamma,return 0.f;)
 
 	// UTextureRenderTarget default display gamma if none is set
diff --git a/Engine/Source/Runtime/Engine/Classes/Engine/TextureRenderTarget2D.h b/Engine/Source/Runtime/Engine/Classes/Engine/TextureRenderTarget2D.h
--- a/Engine/Source/Runtime/Engine/Classes/Engine/TextureRenderTarget2D.h
+++ b/Engine/Source/Runtime/Engine/Classes/Engine/TextureRenderTarget2D.h
@@ -175,6 +175,7 @@ class UTextureRenderTarget2D : public UTextureRenderTarget
 	virtual EPixelFormat GetFormat() const override;
 	virtual bool IsSRGB() const override;
 	virtual float GetDisplayGamma() const override;
+	virtual ETextureClass GetRenderTargetTextureClass() const override { return ETextureClass::TwoD; }
 	//~ End UTextureRenderTarget Interface
 
 	/**
diff --git a/Engine/Source/Runtime/Engine/Classes/Engine/TextureRenderTarget2DArray.h b/Engine/Source/Runtime/Engine/Classes/Engine/TextureRenderTarget2DArray.h
--- a/Engine/Source/Runtime/Engine/Classes/Engine/TextureRenderTarget2DArray.h
+++ b/Engine/Source/Runtime/Engine/Classes/Engine/TextureRenderTarget2DArray.h
@@ -106,6 +106,7 @@ class UTextureRenderTarget2DArray : public UTextureRenderTarget
 	virtual EPixelFormat GetFormat() const override;
 	virtual bool IsSRGB() const override;
 	virtual float GetDisplayGamma() const override;
+	virtual ETextureClass GetRenderTargetTextureClass() const override { return ETextureClass::Array; }
 	//~ End UTextureRenderTarget Interface
 };
 
diff --git a/Engine/Source/Runtime/Engine/Classes/Engine/TextureRenderTargetCube.h b/Engine/Source/Runtime/Engine/Classes/Engine/TextureRenderTargetCube.h
--- a/Engine/Source/Runtime/Engine/Classes/Engine/TextureRenderTargetCube.h
+++ b/Engine/Source/Runtime/Engine/Classes/Engine/TextureRenderTargetCube.h
@@ -98,6 +98,7 @@ class UTextureRenderTargetCube : public UTextureRenderTarget
 	virtual EPixelFormat GetFormat() const override;
 	virtual bool IsSRGB() const override;
 	virtual float GetDisplayGamma() const override;
+	virtual ETextureClass GetRenderTargetTextureClass() const override { return ETextureClass::Cube; }
 	//~ End UTextureRenderTarget Interface
 };
 
diff --git a/Engine/Source/Runtime/Engine/Classes/Engine/TextureRenderTargetVolume.h b/Engine/Source/Runtime/Engine/Classes/Engine/TextureRenderTargetVolume.h
--- a/Engine/Source/Runtime/Engine/Classes/Engine/TextureRenderTargetVolume.h
+++ b/Engine/Source/Runtime/Engine/Classes/Engine/TextureRenderTargetVolume.h
@@ -104,6 +104,7 @@ class UTextureRenderTargetVolume : public UTextureRenderTarget
 	virtual EPixelFormat GetFormat() const override;
 	virtual bool IsSRGB() const override;
 	virtual float GetDisplayGamma() const override;
+	virtual ETextureClass GetRenderTargetTextureClass() const override { return ETextureClass::Volume; }
 	//~ End UTextureRenderTarget Interface
 };
 
diff --git a/Engine/Source/Runtime/Engine/Private/ImageUtils.cpp b/Engine/Source/Runtime/Engine/Private/ImageUtils.cpp
--- a/Engine/Source/Runtime/Engine/Private/ImageUtils.cpp
+++ b/Engine/Source/Runtime/Engine/Private/ImageUtils.cpp
@@ -325,9 +325,9 @@ bool FImageUtils::ExportRenderTargetToDDS(TArray64<uint8> & OutData, UTextureRen
 	}
 
 	// some code dupe with GetRenderTargetImage to identify 2d/cube/vol :
-	UTextureRenderTargetCube * TexRTCube = Cast<UTextureRenderTargetCube>(TexRT);
-	bool bIsCube = ( TexRTCube != nullptr );
-	// UTextureRenderTargetCubeArray does not exist at the moment
+	ETextureClass TexRTClass = TexRT->GetRenderTargetTextureClass();
+	check( TexRTClass != ETextureClass::Invalid && TexRTClass != ETextureClass::RenderTarget );
+	bool bIsCube = TexRTClass == ETextureClass::Cube || TexRTClass == ETextureClass::CubeArray;
 	
 	int64 SizeX = FMath::RoundToInt64( TexRT->GetSurfaceWidth() ); // GetSurfaceWidth returns SizeX but as float
 	int64 SizeY = FMath::RoundToInt64( TexRT->GetSurfaceHeight() );
@@ -452,18 +452,11 @@ bool FImageUtils::GetRenderTargetImage(UTextureRenderTarget* TexRT, FImage & Out
 
 	ERawImageFormat::Type ReadFormat = UTextureRenderTarget::GetReadPixelsFormat(RTFormat,false);
 	
-	//UTextureRenderTarget2D * TexRT2D = Cast<UTextureRenderTarget2D>(TexRT);
-	//UTextureRenderTarget2DArray * TexRT2DA = Cast<UTextureRenderTarget2DArray>(TexRT);
-	//UTextureRenderTargetVolume * TexRTVolume = Cast<UTextureRenderTargetVolume>(TexRT);
-
 	// we have to identify cubes because they are treated differently
-	// can't use TexRT->GetTextureClass() because it's always RenderTarget
-	UTextureRenderTargetCube * TexRTCube = Cast<UTextureRenderTargetCube>(TexRT);
-	bool bIsCube = ( TexRTCube != nullptr );
-	// UTextureRenderTargetCubeArray does not exist at the moment
+	ETextureClass TexRTClass = TexRT->GetRenderTargetTextureClass();
+	check( TexRTClass != ETextureClass::Invalid && TexRTClass != ETextureClass::RenderTarget );
+	bool bIsCube = TexRTClass == ETextureClass::Cube || TexRTClass == ETextureClass::CubeArray;
 	
-	// UTextureRenderTarget2DArray does not derive from UTextureRenderTarget2D, so TexRT2D will be null
-
 	int64 TexRT_SizeX = FMath::RoundToInt64( TexRT->GetSurfaceWidth() ); // GetSurfaceWidth returns SizeX but as float
 	int64 TexRT_SizeY = FMath::RoundToInt64( TexRT->GetSurfaceHeight() );
 	int64 TexRT_SizeZ = FMath::RoundToInt64( TexRT->GetSurfaceDepth() );
@@ -886,6 +879,8 @@ UTexture2D* FImageUtils::CreateTexture2D(int32 SrcWidth, int32 SrcHeight, const
 	}
 
 	Tex2D->VirtualTextureStreaming = InParams.bVirtualTexture;
+	
+	Tex2D->SetModernSettingsForNewOrChangedTexture();
 
 	Tex2D->PostEditChange();
 	return Tex2D;
diff --git a/Engine/Source/Runtime/Engine/Private/Texture.cpp b/Engine/Source/Runtime/Engine/Private/Texture.cpp
--- a/Engine/Source/Runtime/Engine/Private/Texture.cpp
+++ b/Engine/Source/Runtime/Engine/Private/Texture.cpp
@@ -180,6 +180,7 @@ UTexture::UTexture(const FObjectInitializer& ObjectInitializer)
 	bPadWithBorderColor = false;
 	ResizeDuringBuildX = 0;
 	ResizeDuringBuildY = 0;
+	bChromaKeyTexture = false;
 	ChromaKeyColor = FColorList::Magenta;
 	ChromaKeyThreshold = 1.0f / 255.0f;
 	VirtualTextureStreaming = 0;
@@ -196,6 +197,32 @@ UTexture::UTexture(const FObjectInitializer& ObjectInitializer)
 	}
 }
 
+#if WITH_EDITOR
+void UTexture::SetModernSettingsForNewOrChangedTexture()
+{
+	UpdateOodleTextureSdkVersionToLatest();
+
+	// here we can change values that must have different defaults for backwards compatibility
+	// we set them to the new desired value here, the Texture constructor sets the legacy value
+	
+	if ( GetCompositeTexture() == nullptr )
+	{
+		CompositeTextureMode = CTM_Disabled;
+	}
+
+	// set AlphaCoverageThresholds to a better default than zero :
+	if ( bDoScaleMipsForAlphaCoverage == false && AlphaCoverageThresholds == FVector4(0,0,0,0) )
+	{
+		AlphaCoverageThresholds = FVector4(0, 0, 0, 0.75f);
+	}
+
+	// bNormalizeNormals is ignored if we're not a normal map
+	bNormalizeNormals = true;
+	
+	bUseNewMipFilter = true;
+}
+#endif
+
 const FTextureResource* UTexture::GetResource() const
 {
 	if (IsInParallelGameThread() || IsInGameThread() || IsInSlateThread() || IsInAsyncLoadingThread())
@@ -436,13 +463,13 @@ bool UTexture::IsPostLoadThreadSafe() const
 
 bool UTexture::IsDefaultTexture() const
 {
+	// "IsDefaultTexture" actually means that a temporary default stand-in is being used
+	//	 because the texture is being async built
 	return false;
 }
 
-bool UTexture::Modify(bool bAlwaysMarkDirty)
+void UTexture::BlockOnAnyAsyncBuild()
 {
-	// Before applying any modification to the texture
-	// make sure no compilation is still ongoing.
 	if (!IsAsyncCacheComplete())
 	{
 		FinishCachePlatformData();
@@ -453,6 +480,15 @@ bool UTexture::Modify(bool bAlwaysMarkDirty)
 		FTextureCompilingManager::Get().FinishCompilation({this});
 	}
 
+	check( ! IsDefaultTexture() ); // this is always true even in failure/error cases
+}
+
+bool UTexture::Modify(bool bAlwaysMarkDirty)
+{
+	// Before applying any modification to the texture
+	// make sure no compilation is still ongoing.
+	BlockOnAnyAsyncBuild();
+
 	return Super::Modify(bAlwaysMarkDirty);
 }
 
@@ -1968,12 +2004,17 @@ void FTextureSource::InitBlocked(const ETextureSourceFormat* InLayerFormats,
 			{
 				FMemory::Memcpy(DataPtr, InDataPerBlock[i], BlockSize);
 			}
+			else
+			{
+				memset(DataPtr,0,BlockSize);
+			}
 			DataPtr += BlockSize;
 		}
 	}
 
 	BulkData.UpdatePayload(Buffer.MoveToShared(), Owner);
 	BulkData.SetCompressionOptions(UE::Serialization::ECompressionOptions::Default);
+	UseHashAsGuid();
 }
 
 void FTextureSource::InitBlocked(const ETextureSourceFormat* InLayerFormats,
@@ -1986,6 +2027,7 @@ void FTextureSource::InitBlocked(const ETextureSourceFormat* InLayerFormats,
 
 	BulkData.UpdatePayload(MoveTemp(NewData), Owner);
 	BulkData.SetCompressionOptions(UE::Serialization::ECompressionOptions::Default);
+	UseHashAsGuid();
 }
 
 void FTextureSource::InitLayered(
@@ -2020,9 +2062,11 @@ void FTextureSource::InitLayered(
 	else
 	{
 		BulkData.UpdatePayload(FUniqueBuffer::Alloc(TotalBytes).MoveToShared(), Owner);
+		// ?? unitialized ??
 	}
 
 	BulkData.SetCompressionOptions(UE::Serialization::ECompressionOptions::Default);
+	UseHashAsGuid();
 }
 
 void FTextureSource::InitLayered(
@@ -2045,6 +2089,7 @@ void FTextureSource::InitLayered(
 
 	BulkData.UpdatePayload(MoveTemp(NewData), Owner);
 	BulkData.SetCompressionOptions(UE::Serialization::ECompressionOptions::Default);
+	UseHashAsGuid();
 }
 
 void FTextureSource::Init(
@@ -2143,6 +2188,8 @@ void FTextureSource::InitWithCompressedSourceData(
 	{
 		BulkData.SetCompressionOptions(UE::Serialization::ECompressionOptions::Disabled);
 	}
+	
+	UseHashAsGuid();
 }
 
 void FTextureSource::InitWithCompressedSourceData(
@@ -2177,6 +2224,8 @@ void FTextureSource::InitWithCompressedSourceData(
 	{
 		BulkData.SetCompressionOptions(UE::Serialization::ECompressionOptions::Disabled);
 	}
+	
+	UseHashAsGuid();
 }
 
 FTextureSource FTextureSource::CopyTornOff() const
@@ -2270,6 +2319,13 @@ void FTextureSource::Compress()
 	{
 		BulkData.SetCompressionOptions(UE::Serialization::ECompressionOptions::Disabled);
 	}
+
+	// note: we changed BulkData payload here to put PNG data in it
+	//	but we do NOT call UseHashAsGuid
+	//	we try to keep "Id" == to the hash of the BulkData when it was the raw data
+	//	the invariant
+	//	( Id == UE::Serialization::IoHashToGuid(BulkData.GetPayloadId()) )
+	//	is no longer true after this
 }
 
 FSharedBuffer FTextureSource::Decompress(IImageWrapperModule* ) const
@@ -2939,12 +2995,16 @@ void FTextureSource::ForceGenerateGuid()
 
 void FTextureSource::ReleaseSourceMemory()
 {
+	check( LockState == ELockState::None && NumLockedMips == 0 );
+
 	bHasHadBulkDataCleared = true;
 	BulkData.UnloadData();
 }
 
 void FTextureSource::RemoveSourceData()
 {
+	check( LockState == ELockState::None && NumLockedMips == 0 );
+
 	SizeX = 0;
 	SizeY = 0;
 	NumSlices = 0;
diff --git a/Engine/Source/Runtime/Engine/Private/TextureRenderTarget.cpp b/Engine/Source/Runtime/Engine/Private/TextureRenderTarget.cpp
--- a/Engine/Source/Runtime/Engine/Private/TextureRenderTarget.cpp
+++ b/Engine/Source/Runtime/Engine/Private/TextureRenderTarget.cpp
@@ -421,6 +421,8 @@ bool UTextureRenderTarget::UpdateTexture(UTexture* InTexture, EConstructTextureF
 
 UTexture* UTextureRenderTarget::ConstructTexture(UObject* InOuter, const FString& InNewTextureName, EObjectFlags InObjectFlags, EConstructTextureFlags InFlags, const TArray<uint8>* InAlphaOverride, FText* OutErrorMessage)
 {
+	// InFlags = CTF_Default = CTF_Compress | CTF_SRGB,
+
 	FTextureRenderTargetResource* RenderTarget = GameThread_GetRenderTargetResource();
 	if (RenderTarget == nullptr)
 	{
@@ -450,6 +452,8 @@ UTexture* UTextureRenderTarget::ConstructTexture(UObject* InOuter, const FString
 	//		I see SRGB = 1 but Gamma = 1.0
 	//	 see also IsSRGB() which is yet another query that has different ideas
 	//	furthermore, float formats do not support anything but Linear gamma
+	//
+	// InFlags typically starts as Default with CTF_SRGB on, and then it's turned off here :
 	float Gamma = RenderTarget->GetDisplayGamma();
 	if (FMath::Abs(Gamma - 1.0f) < UE_KINDA_SMALL_NUMBER)
 	{
@@ -477,6 +481,8 @@ UTexture* UTextureRenderTarget::ConstructTexture(UObject* InOuter, const FString
 		Result->DeferCompression = false;
 	}
 
+	Result->SetModernSettingsForNewOrChangedTexture();
+
 	if ((InFlags & CTF_SkipPostEdit) == 0)
 	{
 		Result->PostEditChange();
diff --git a/Engine/Source/Runtime/Engine/Private/VolumeTexture.cpp b/Engine/Source/Runtime/Engine/Private/VolumeTexture.cpp
--- a/Engine/Source/Runtime/Engine/Private/VolumeTexture.cpp
+++ b/Engine/Source/Runtime/Engine/Private/VolumeTexture.cpp
@@ -252,17 +252,18 @@ ENGINE_API bool UVolumeTexture::UpdateSourceFromFunction(TFunction<void(int32, i
 		return false;
 	}
 	
-	Modify(true);
+	PreEditChange(nullptr);
 
 	Source2DTexture = nullptr;
 	const int64 FormatDataSize = ERawImageFormat::GetBytesPerPixel(FImageCoreUtils::ConvertToRawImageFormat(Format));
 
 	// Allocate temp buffer used to fill texture
 	uint8* const NewData = (uint8*)FMemory::Malloc(SizeX * SizeY * SizeZ * FormatDataSize);
-	uint8* CurPos = NewData;
+	uint8* DataPtr = NewData;
 
-	// Temporary storage for a single voxel value extracted from the lambda, type depends on Format
-	void* const NewValue = FMemory::Malloc(FormatDataSize);
+	// use PixelDataScratch just in case the function does something like write RGBA32F
+	//	 when our format is actually BGRA8 , to try to avoid crashing
+	uint64 PixelDataScratch[4];
 
 	// Loop over all voxels and fill from our TFunction
 	for (int32 PosZ = 0; PosZ < SizeZ; ++PosZ)
@@ -271,11 +272,12 @@ ENGINE_API bool UVolumeTexture::UpdateSourceFromFunction(TFunction<void(int32, i
 		{
 			for (int32 PosX = 0; PosX < SizeX; ++PosX)
 			{
-				Func(PosX, PosY, PosZ, NewValue);
-
-				FMemory::Memcpy(CurPos, NewValue, FormatDataSize);
+				// ?? Func knows our pixel format somehow?
+				//	should have defined this so that Func always writes FLinearColor and we convert
+				Func(PosX, PosY, PosZ, PixelDataScratch);
 
-				CurPos += FormatDataSize;
+				memcpy(DataPtr,PixelDataScratch,FormatDataSize);
+				DataPtr += FormatDataSize;
 			}
 		}
 	}
@@ -285,14 +287,10 @@ ENGINE_API bool UVolumeTexture::UpdateSourceFromFunction(TFunction<void(int32, i
 	
 	// Free temp buffers
 	FMemory::Free(NewData);
-	FMemory::Free(NewValue);
 
 	SetLightingGuid(); // Because the content has changed, use a new GUID.
 
-	ValidateSettingsAfterImportOrEdit();
-
-	// Make sure to update the texture resource so the results of filling the texture 
-	UpdateResource();
+	PostEditChange();
 
 	bSourceValid = true;
 #endif
diff --git a/Engine/Source/Runtime/TextureUtilitiesCommon/Private/TextureImportSettings.cpp b/Engine/Source/Runtime/TextureUtilitiesCommon/Private/TextureImportSettings.cpp
--- a/Engine/Source/Runtime/TextureUtilitiesCommon/Private/TextureImportSettings.cpp
+++ b/Engine/Source/Runtime/TextureUtilitiesCommon/Private/TextureImportSettings.cpp
@@ -76,37 +76,31 @@ namespace UE::TextureUtilitiesCommon
 	
 		// things that are done for both fresh import and reimport :
 
-		Texture->UpdateOodleTextureSdkVersionToLatest();
-
 		if ( bIsReimport )
 		{
+			Texture->UpdateOodleTextureSdkVersionToLatest();
+
 			return;
 		}
 
 		// things that are for fresh import only :
+		
+		// SetModern does UpdateOodleTextureSdkVersionToLatest
+		Texture->SetModernSettingsForNewOrChangedTexture();
 
-		// here we can change values that must have different defaults for backwards compatibility
-		// we set them to the new desired value here, the Texture constructor sets the legacy value
-	
 		const UTextureImportSettings* Settings = GetDefault<UTextureImportSettings>();
 
-		if ( Settings->bEnableNormalizeNormals )
+		// cannot check for TC_Normalmap here
+		//	because of the way NormalmapIdentification is delayed in Interchange
+		//	it's harmless to just always turn it on
+		//	it will be ignored if we are not TC_Normalmap
+		//	OutBuildSettings.bNormalizeNormals = Texture.bNormalizeNormals && Texture.IsNormalMap();
+		//if ( Texture->CompressionSettings == TC_Normalmap )
 		{
-			// cannot check for TC_Normalmap here
-			//	because of the way NormalmapIdentification is delayed in Interchange
-			//	it's harmless to just always turn it on
-			//	it will be ignored if we are not TC_Normalmap
-			//	OutBuildSettings.bNormalizeNormals = Texture.bNormalizeNormals && Texture.IsNormalMap();
-			//if ( Texture->CompressionSettings == TC_Normalmap )
-			{
-				Texture->bNormalizeNormals = true;
-			}
+			Texture->bNormalizeNormals = Settings->bEnableNormalizeNormals;
 		}
 
-		if ( Settings->bEnableFastMipFilter )
-		{
-			Texture->bUseNewMipFilter = true;
-		}
+		Texture->bUseNewMipFilter = Settings->bEnableFastMipFilter;
 
 		// the pipeline before here will have set floating point textures to TC_HDR
 		//	could alternatively check Texture->HasHDRSource
