[Material Editor] Reintroduced translation DDC results serialization time output to cook-stats.

diff --git a/Engine/Source/Runtime/Engine/Private/Materials/HLSLMaterialTranslator.cpp b/Engine/Source/Runtime/Engine/Private/Materials/HLSLMaterialTranslator.cpp
--- a/Engine/Source/Runtime/Engine/Private/Materials/HLSLMaterialTranslator.cpp
+++ b/Engine/Source/Runtime/Engine/Private/Materials/HLSLMaterialTranslator.cpp
@@ -1200,13 +1200,16 @@ bool FHLSLMaterialTranslator::Translate()
 			UE_LOG(LogMaterial, Display, TEXT("Material '%s' translation results were in the DDC"), *Material->GetMaterialInterface()->GetFullName());
 		}
 
+		STAT(double SerializeTime = FPlatformTime::Seconds());
+
 		// Serialize the environment defines from the buffer retrieved from the DDC.
 		FMemoryReaderView EnvironmentDefinesBufferReader{ TArrayView<uint8>{ (uint8*)EnvironmentDefinesBuffer.GetData(), (int)EnvironmentDefinesBuffer.GetSize() } };
 		FObjectAndNameAsStringProxyArchive EnvironmentDefinesBufferReaderProxy{ EnvironmentDefinesBufferReader, true };
 		EnvironmentDefines->Serialize(EnvironmentDefinesBufferReaderProxy);
 
 		MaterialCompilationOutput = DDCMaterialCompilationOutput;
-
+		
+		STAT(DDCRequestSerializeTime += SerializeTime - FPlatformTime::Seconds());
 		STAT(GShaderCompilerStats->IncrementMaterialCacheHit());
 		bSuccess = true;
 	}
@@ -1232,15 +1235,15 @@ bool FHLSLMaterialTranslator::Translate()
 	// Report timings to Material Cook Stats
 #if STATS
 	TotalTime = FPlatformTime::Seconds() - TotalTime;
-	GShaderCompilerStats->IncrementMaterialTranslated(TotalTime, TranslationOnlyTime, 0 /* todo */);
-#endif
-
+	GShaderCompilerStats->IncrementMaterialTranslated(TotalTime, TranslationOnlyTime, DDCRequestSerializeTime);
 	INC_FLOAT_STAT_BY(STAT_ShaderCompiling_HLSLTranslation, (float)TotalTime);
 
-#if ENABLE_COOK_STATS && STATS
+#if ENABLE_COOK_STATS
 	// Write out a CSV file MaterialTranslationLog.txt containing info about all material translations.
 	FCsvLogFile::Get().AddEntry(Material->GetMaterialInterface()->GetFullName(), TranslationDateTime, TotalTime);
-#endif
+#endif // ENABLE_COOK_STATS
+
+#endif // STATS
 
 	return bSuccess;
 }
@@ -15434,6 +15437,8 @@ void FHLSLMaterialTranslator::AsyncQueryDDC(
 				// Load the results if we hit the cache.
 				if (!MaterialCompilationOutputBuffer.IsNull() && !TranslationResultsBuffer.IsNull() && !EnvironmentDefinesBuffer.IsNull())
 				{
+					STAT(double SerializeTime = FPlatformTime::Seconds());
+
 					// Read the material compilation output
 					FShaderMapPointerTable PointerTable;
 					FPlatformTypeLayoutParameters LayoutParams;
@@ -15446,6 +15451,7 @@ void FHLSLMaterialTranslator::AsyncQueryDDC(
 					FMemoryReaderView ResultsMemoryReader{ TArrayView<uint8>{ (uint8*)TranslationResultsBuffer.GetData(), (int)TranslationResultsBuffer.GetSize() } };
 					ResultsMemoryReader << MaterialSourceTemplateParams;
 
+					STAT(DDCRequestSerializeTime = FPlatformTime::Seconds() - SerializeTime);
 					DDCQueryHit = true;
 				}
 			}
diff --git a/Engine/Source/Runtime/Engine/Private/Materials/HLSLMaterialTranslator.h b/Engine/Source/Runtime/Engine/Private/Materials/HLSLMaterialTranslator.h
--- a/Engine/Source/Runtime/Engine/Private/Materials/HLSLMaterialTranslator.h
+++ b/Engine/Source/Runtime/Engine/Private/Materials/HLSLMaterialTranslator.h
@@ -1419,6 +1419,11 @@ protected:
 	/** Signals when the async DDC query task has completed AND there was a hit */
 	TAtomic<bool> DDCQueryHit;
 
+#ifdef STATS
+	/** When a translation DDC request hits, this is set to the time spent serializing. */
+	double DDCRequestSerializeTime = 0;
+#endif
+
 	friend class FAddUniformExpressionScope;
 };
 
