Fixed performance regression with TActorIterator - rollback CL 29270624

diff --git a/Engine/Source/Runtime/Engine/Public/EngineUtils.h b/Engine/Source/Runtime/Engine/Public/EngineUtils.h
--- a/Engine/Source/Runtime/Engine/Public/EngineUtils.h
+++ b/Engine/Source/Runtime/Engine/Public/EngineUtils.h
@@ -189,35 +189,39 @@ public:
 #if WITH_EDITOR
 		// In the editor, you are more likely to have many worlds in memory at once.
 		// As an optimization to avoid iterating over many actors that are not in the world we are asking for,
+		// if the filter class is AActor, just use the actors that are in the world you asked for.
 		// This could be useful in runtime code as well if there are many worlds in memory, but for now we will leave
 		// it in editor code.
-		
-		// First determine the number of actors in the world to reduce reallocations when we append them to the array below.
-		const int32 MaxActors = Algo::Accumulate(InWorld->GetLevels(), 0, [](int32 InTotalNum, const ULevel* InLevel) { return InTotalNum + (InLevel ? InLevel->Actors.Num() : 0); });
-
-		// Presize the array
-		ObjectArray.Reserve(MaxActors);
-
-		for (ULevel* Level : InWorld->GetLevels())
+		if (InClass == AActor::StaticClass())
 		{
-			if (Level)
+			// First determine the number of actors in the world to reduce reallocations when we append them to the array below.
+			int32 NumActors = 0;
+			for (ULevel* Level : InWorld->GetLevels())
 			{
-				if(InClass == AActor::StaticClass())
+				if (Level)
 				{
-					ObjectArray.Append(Level->Actors);
+					NumActors += Level->Actors.Num();
 				}
-				else
+			}
+
+			// Presize the array
+			ObjectArray.Reserve(NumActors);
+
+			// Fill the array
+			for (ULevel* Level : InWorld->GetLevels())
+			{
+				if (Level)
 				{
-					Algo::CopyIf(Level->Actors, ObjectArray, [InClass](AActor* InActor) { return InActor && InActor->IsA(InClass); });
+					ObjectArray.Append(Level->Actors);
 				}
 			}
 		}
-#else // WITH_EDITOR
+		else
+#endif // WITH_EDITOR
 		{
 			constexpr EObjectFlags ExcludeFlags = RF_ClassDefaultObject;
 			GetObjectsOfClass(InClass, ObjectArray, true, ExcludeFlags, EInternalObjectFlags::Garbage);
 		}
-#endif
 
 		const auto ActorSpawnedDelegate = FOnActorSpawned::FDelegate::CreateRaw(this, &FActorIteratorState::OnActorSpawned);
 		ActorSpawnedDelegateHandle = CurrentWorld->AddOnActorSpawnedHandler(ActorSpawnedDelegate);
