multi-trace against cluster unions that hit nothing in the subtrace should retry

diff --git a/Engine/Source/Runtime/Engine/Private/Collision/SceneQuery.cpp b/Engine/Source/Runtime/Engine/Private/Collision/SceneQuery.cpp
--- a/Engine/Source/Runtime/Engine/Private/Collision/SceneQuery.cpp
+++ b/Engine/Source/Runtime/Engine/Private/Collision/SceneQuery.cpp
@@ -531,8 +531,11 @@ bool TSceneCastCommonImpWithRetryRequest(const UWorld* World, typename Traits::T
 				{
 					if constexpr (Traits::IsMulti())
 					{
+						const bool bHadBlockingHit = bBlockingHit;
 						TArray<FHitResult> AllNewHits;
 						TArray<int32> ClusterUnionIndices;
+						TArray<AActor*, TInlineAllocator<1>> ClusterUnionActorsToIgnoreIfRetry;
+						bBlockingHit = false;
 
 						for (int32 Index = 0; Index < OutHits.Num(); ++Index)
 						{
@@ -547,23 +550,48 @@ bool TSceneCastCommonImpWithRetryRequest(const UWorld* World, typename Traits::T
 
 								if (ClusterUnionHit.bHit)
 								{
+									bBlockingHit = true;
 									AllNewHits.Append(NewHit);
 								}
+								else if (OutHits[Index].bBlockingHit)
+								{
+									// Subtrace has no blocking hit but the cluster union trace was a blocking hit.
+									// We need to make sure this cluster union gets ignored if we retry.
+									ClusterUnionActorsToIgnoreIfRetry.Add(OutHits[Index].GetActor());
+								}
+							}
+							else
+							{
+								bBlockingHit |= OutHits[Index].bBlockingHit;
 							}
 						}
 
-						for (int32 Index = ClusterUnionIndices.Num() - 1; Index >= 0; --Index)
+						if (bHadBlockingHit && !bBlockingHit)
 						{
-							// No shrinking since we're going to be adding more elements shortly.
-							OutHits.RemoveAtSwap(ClusterUnionIndices[Index], 1, false);
+							// We had a blocking hit, but after subtracing against a cluster union we no longer have a blocking hit because its subcomponent(s) were ignored (e.g. if ignored actors/ignored components is used).
+							// In this case we want to retry and ignore the cluster unions that were hit to continue the trace until it reaches the end/finds another blocking hit.
+							Traits::ResetOutHits(OutHits, Start, End);
+							OutRetryParams = Params;
+							for (AActor* IgnoreActor : ClusterUnionActorsToIgnoreIfRetry)
+							{
+								OutRetryParams.AddIgnoredActor(IgnoreActor);
+							}
+							bOutRequestRetry = true;
+							return false;
 						}
-
-						if (Params.bReplaceHitWithSubComponents)
+						else
 						{
-							OutHits.Append(AllNewHits);
+							for (int32 Index = ClusterUnionIndices.Num() - 1; Index >= 0; --Index)
+							{
+								// No shrinking since we're going to be adding more elements shortly.
+								OutHits.RemoveAtSwap(ClusterUnionIndices[Index], 1, false);
+							}
+	
+							if (Params.bReplaceHitWithSubComponents)
+							{
+								OutHits.Append(AllNewHits);
+							}
 						}
-						
-						bBlockingHit &= !OutHits.IsEmpty();
 					}
 					else
 					{
