FImageUtils::CreateTexture for a modern way to make textures from FImage
supercedes CreateTexture2D

#rnx

diff --git a/Engine/Source/Runtime/Engine/Private/ImageUtils.cpp b/Engine/Source/Runtime/Engine/Private/ImageUtils.cpp
--- a/Engine/Source/Runtime/Engine/Private/ImageUtils.cpp
+++ b/Engine/Source/Runtime/Engine/Private/ImageUtils.cpp
@@ -885,8 +885,51 @@ UTexture2D* FImageUtils::CreateTexture2D(int32 SrcWidth, int32 SrcHeight, const
 	Tex2D->PostEditChange();
 	return Tex2D;
 #else
-	UE_LOG(LogImageUtils, Fatal,TEXT("ConstructTexture2D not supported on console."));
-	return NULL;
+	UE_LOG(LogImageUtils, Fatal,TEXT("CreateTexture2D not supported without WITH_EDITOR."));
+	return nullptr;
+#endif
+}
+	
+UTexture * FImageUtils::CreateTexture(ETextureClass TextureClass, const FImageView & Image, UObject* Outer, const FString& Name, EObjectFlags Flags, bool DoPostEditChange )
+{
+#if WITH_EDITOR
+	UTexture * Tex;
+
+	switch(TextureClass)
+	{
+	case ETextureClass::TwoD:
+		Tex = NewObject<UTexture2D>(Outer, FName(*Name), Flags);
+		break;
+	case ETextureClass::Cube:
+		Tex = NewObject<UTextureCube>(Outer, FName(*Name), Flags);
+		break;
+	case ETextureClass::Array:
+		Tex = NewObject<UTexture2DArray>(Outer, FName(*Name), Flags);
+		break;
+	case ETextureClass::CubeArray:
+		Tex = NewObject<UTextureCubeArray>(Outer, FName(*Name), Flags);
+		break;
+	case ETextureClass::Volume:
+		Tex = NewObject<UVolumeTexture>(Outer, FName(*Name), Flags);
+		break;
+	default:
+		UE_LOG(LogImageUtils, Fatal,TEXT("CreateTexture invalid TextureClass."));
+		return nullptr;
+	}
+	
+	Tex->Source.Init(Image);
+	
+	Tex->SetModernSettingsForNewOrChangedTexture();
+
+	if ( DoPostEditChange )
+	{
+		Tex->PostEditChange();
+	}
+
+	return Tex;
+#else
+	UE_LOG(LogImageUtils, Fatal,TEXT("CreateTexture not supported without WITH_EDITOR."));
+	return nullptr;
 #endif
 }
 
diff --git a/Engine/Source/Runtime/Engine/Public/ImageUtils.h b/Engine/Source/Runtime/Engine/Public/ImageUtils.h
--- a/Engine/Source/Runtime/Engine/Public/ImageUtils.h
+++ b/Engine/Source/Runtime/Engine/Public/ImageUtils.h
@@ -246,9 +246,24 @@ public:
 	 * this fills a TextureSource , and will then Build a Platform texture from that
 	 * can be used WITH_EDITOR only
 	 * contrast to CreateTexture2DFromImage
+	 *
+	 * Prefer the more modern CreateTexture from FImage.
 	 */
 	ENGINE_API static UTexture2D* CreateTexture2D(int32 SrcWidth, int32 SrcHeight, const TArray<FColor> &SrcData, UObject* Outer, const FString& Name, const EObjectFlags &Flags, const FCreateTexture2DParameters& InParams);
 	
+	/**
+	 * Creates a texture of any type from an Image
+	 * This is the modern preferred way to create a texture.
+	 *
+	 * If you need to change the default settings, then use DoPostEditChange = false, and call PostEditChange() yourself after setting all properties.
+	 * Typically you may want to set LODGroup and CompressionSettings.
+	 *
+	 * this fills the TextureSource , and will then Build a Platform texture from that
+	 * can be used WITH_EDITOR only
+	 *
+	 */
+	ENGINE_API static UTexture * CreateTexture(ETextureClass TextureClass, const FImageView & Image, UObject* Outer, const FString& Name, EObjectFlags Flags = RF_NoFlags, bool DoPostEditChange = true );
+	
 	/**
 	 * Creates a 2D texture from an FImage
 	 * 
