[NetPhysics] Input decay during resim when there are no more valid inputs

diff --git a/Engine/Plugins/Experimental/ChaosVehiclesPlugin/Source/ChaosVehicles/Private/ChaosVehicleManagerAsyncCallback.cpp b/Engine/Plugins/Experimental/ChaosVehiclesPlugin/Source/ChaosVehicles/Private/ChaosVehicleManagerAsyncCallback.cpp
--- a/Engine/Plugins/Experimental/ChaosVehiclesPlugin/Source/ChaosVehicles/Private/ChaosVehicleManagerAsyncCallback.cpp
+++ b/Engine/Plugins/Experimental/ChaosVehiclesPlugin/Source/ChaosVehicles/Private/ChaosVehicleManagerAsyncCallback.cpp
@@ -281,6 +281,18 @@ void FNetworkVehicleInputs::InterpolateDatas(const FNetworkVehicleInputs& MinDat
 	VehicleInputs.TransmissionType = LerpFactor < 0.5 ? MinDatas.VehicleInputs.TransmissionType : MaxDatas.VehicleInputs.TransmissionType;
 }
 
+void FNetworkVehicleInputs::DecayDatas(float DecayAmount)
+{
+	// Local adjustment to DecayAmount for vehicle implementation
+	DecayAmount = FMath::Min(DecayAmount * 2, 1.0f);
+
+	// Apply decay on steering inputs
+	VehicleInputs.PitchInput = FMath::Lerp(VehicleInputs.PitchInput, 0.0f, DecayAmount);
+	VehicleInputs.RollInput = FMath::Lerp(VehicleInputs.RollInput, 0.0f, DecayAmount);
+	VehicleInputs.SteeringInput = FMath::Lerp(VehicleInputs.SteeringInput, 0.0f, DecayAmount);
+	VehicleInputs.YawInput = FMath::Lerp(VehicleInputs.YawInput, 0.0f, DecayAmount);
+}
+
 bool FNetworkVehicleStates::NetSerialize(FArchive& Ar, class UPackageMap* Map, bool& bOutSuccess)
 {
 	FNetworkPhysicsDatas::SerializeFrames(Ar);
diff --git a/Engine/Plugins/Experimental/ChaosVehiclesPlugin/Source/ChaosVehicles/Public/ChaosVehicleManagerAsyncCallback.h b/Engine/Plugins/Experimental/ChaosVehiclesPlugin/Source/ChaosVehicles/Public/ChaosVehicleManagerAsyncCallback.h
--- a/Engine/Plugins/Experimental/ChaosVehiclesPlugin/Source/ChaosVehicles/Public/ChaosVehicleManagerAsyncCallback.h
+++ b/Engine/Plugins/Experimental/ChaosVehiclesPlugin/Source/ChaosVehicles/Public/ChaosVehicleManagerAsyncCallback.h
@@ -128,6 +128,9 @@ struct CHAOSVEHICLES_API FNetworkVehicleInputs : public FNetworkPhysicsDatas
 	/**  Build the datas from the network physics component */
 	virtual void BuildDatas(const UActorComponent* NetworkComponent) override;
 
+	/** Decay input during resimulation */
+	virtual void DecayDatas(float DecayAmount) override;
+
 	/**  Serialize datas function that will be used to transfer the struct across the network */
 	bool NetSerialize(FArchive& Ar, class UPackageMap* Map, bool& bOutSuccess);
 
diff --git a/Engine/Source/Runtime/Engine/Private/PhysicsEngine/Experimental/PhysScene_Chaos.cpp b/Engine/Source/Runtime/Engine/Private/PhysicsEngine/Experimental/PhysScene_Chaos.cpp
--- a/Engine/Source/Runtime/Engine/Private/PhysicsEngine/Experimental/PhysScene_Chaos.cpp
+++ b/Engine/Source/Runtime/Engine/Private/PhysicsEngine/Experimental/PhysScene_Chaos.cpp
@@ -138,17 +138,6 @@ public:
 	{
 		using namespace Chaos;
 
-		if (!PendingCommands.IsEmpty())
-		{
-			if (!PendingCommands[0].OwningObject.IsStale())
-			{
-				if (Chaos::FRewindData* RewindData = GetSolver()->GetRewindData())
-				{
-					RewindData->SetResimFrame(INDEX_NONE);
-				}
-			}
-		}
-
 		const UPhysicsSettings* PhysicsSettings = UPhysicsSettings::Get();
 		const bool bAllowResim = PhysicsSettings->PhysicsPrediction.bEnablePhysicsPrediction;
 		const int32 NumFrames = PhysicsSettings->GetPhysicsHistoryCount();
diff --git a/Engine/Source/Runtime/Engine/Private/PhysicsEngine/NetworkPhysicsComponent.cpp b/Engine/Source/Runtime/Engine/Private/PhysicsEngine/NetworkPhysicsComponent.cpp
--- a/Engine/Source/Runtime/Engine/Private/PhysicsEngine/NetworkPhysicsComponent.cpp
+++ b/Engine/Source/Runtime/Engine/Private/PhysicsEngine/NetworkPhysicsComponent.cpp
@@ -179,7 +179,6 @@ int32 FNetworkPhysicsCallback::TriggerRewindIfNeeded_Internal(int32 LatestStepCo
 			UE_LOG(LogChaos, Log, TEXT("CLIENT | PT | TriggerRewindIfNeeded_Internal | Replication Frame = %d"), ReplicationFrame);
 #endif
 			ResimFrame = (ResimFrame == INDEX_NONE) ? ReplicationFrame : (ReplicationFrame == INDEX_NONE) ? ResimFrame : FMath::Min(ReplicationFrame, ResimFrame);
-			RewindData->SetResimFrame(INDEX_NONE);
 		}
 
 		if (ResimFrame != INDEX_NONE)
@@ -728,7 +727,7 @@ void UNetworkPhysicsComponent::OnPreProcessInputsInternal(const int32 PhysicsSte
 	}
 #endif
 
-	if(InputsHistory && StatesHistory && ActorComponent)
+	if (InputsHistory && StatesHistory && ActorComponent)
 	{ 
 		bool bIsSolverReset = false;
 		bool bIsSolverResim = false;
@@ -752,6 +751,13 @@ void UNetworkPhysicsComponent::OnPreProcessInputsInternal(const int32 PhysicsSte
 	#endif
 			if (InputsHistory->ExtractDatas(PhysicsStep, bIsSolverReset, PhysicsDatas))
 			{ 
+				// Calculate input decay if we are resimulating and we don't have up to date inputs
+				if (bIsSolverResim && PhysicsDatas->LocalFrame < PhysicsStep)
+				{
+					const float InputDecay = GetCurrentInputDecay(PhysicsDatas);
+					PhysicsDatas->DecayDatas(InputDecay);
+				}
+
 				PhysicsDatas->ApplyDatas(ActorComponent);
 			}
 		}
@@ -786,7 +792,7 @@ void UNetworkPhysicsComponent::OnPostProcessInputsInternal(const int32 PhysicsSt
 	}
 #endif
 
-	if(InputsHistory && StatesHistory && ActorComponent)
+	if (InputsHistory && StatesHistory && ActorComponent)
 	{
 		if (FPhysScene* PhysScene = GetWorld()->GetPhysicsScene())
 		{
@@ -827,7 +833,7 @@ void UNetworkPhysicsComponent::OnPostProcessInputsInternal(const int32 PhysicsSt
 			int32 InputFrame = INDEX_NONE;
 			{
 				FNetworkPhysicsDatas* PhysicsDatas = InputsDatas.Get();
-				if(InputsHistory->ExtractDatas(PhysicsStep, false, PhysicsDatas, true))
+				if (InputsHistory->ExtractDatas(PhysicsStep, false, PhysicsDatas, true))
 				{
 					InputFrame = PhysicsDatas->InputFrame;
 				}
@@ -850,6 +856,40 @@ void UNetworkPhysicsComponent::OnPostProcessInputsInternal(const int32 PhysicsSt
 	}
 }
 
+const float UNetworkPhysicsComponent::GetCurrentInputDecay(FNetworkPhysicsDatas* PhysicsDatas)
+{
+	if (!PhysicsDatas)
+	{
+		return 0.0f;
+	}
+
+	FPhysScene* PhysScene = GetWorld()->GetPhysicsScene();
+	if (!PhysScene)
+	{
+		return 0.0f;
+	}
+
+	Chaos::FPhysicsSolver* Solver = PhysScene->GetSolver();
+	if (!Solver)
+	{
+		return 0.0f;
+	}
+
+	Chaos::FRewindData* RewindData = Solver->GetRewindData();
+	if (!RewindData)
+	{
+		return 0.0f;
+	}
+	
+	const float NumPredictedInputs = RewindData->CurrentFrame() - PhysicsDatas->LocalFrame; // Number of frames we have used the same PhysicsDatas for during resim
+	const float MaxPredictedInputs = RewindData->GetLatestFrame() - 1 - PhysicsDatas->LocalFrame; // Max number of frames PhysicsDatas registered frame until end of resim
+
+	// Linear decay
+	const float PredictionAlpha = MaxPredictedInputs > 0 ? (NumPredictedInputs / MaxPredictedInputs) : 0.0f;
+
+	return PredictionAlpha;
+}
+
 bool UNetworkPhysicsComponent::HasServerWorld() const
 {
 	return GetWorld()->IsNetMode(NM_DedicatedServer) || GetWorld()->IsNetMode(NM_ListenServer);
@@ -866,7 +906,7 @@ bool UNetworkPhysicsComponent::HasLocalController() const
 
 bool UNetworkPhysicsComponent::IsLocallyControlled() const
 {
-	if (bIsLocallyPossessed && !GetWorld()->IsNetMode(NM_DedicatedServer))
+	if (bIsRelayingLocalInputs && !GetWorld()->IsNetMode(NM_DedicatedServer))
 	{
 		return true;
 	}
diff --git a/Engine/Source/Runtime/Engine/Public/Physics/NetworkPhysicsComponent.h b/Engine/Source/Runtime/Engine/Public/Physics/NetworkPhysicsComponent.h
--- a/Engine/Source/Runtime/Engine/Public/Physics/NetworkPhysicsComponent.h
+++ b/Engine/Source/Runtime/Engine/Public/Physics/NetworkPhysicsComponent.h
@@ -359,6 +359,13 @@ struct FNetworkPhysicsDatas
 
 	// Build the datas from the network physics component
 	virtual void BuildDatas(const UActorComponent* NetworkComponent) {}
+	
+	/** Use to decay desired data during resimulation if data is forward predicted.
+	* @param DecayAmount = Total amount of decay as a multiplier. 10% decay = 0.1.
+	* NOTE: Decay is not accumulated, the data will be in its original state each time DecayDatas is called. DecayAmount will increase each time the input is predicted (reused).
+	* EXAMPLE: Use to decay steering inputs to make resimulation not predict too much with a high steering value. Use DecayAmount of 0.1 to turn a steering value of 0.5 into 0.45 for example.
+	*/ 
+	virtual void DecayDatas(float DecayAmount) {}
 
 	friend UNetworkPhysicsComponent;
 };
@@ -437,21 +444,24 @@ public:
 	ENGINE_API bool HasServerWorld() const;
 
 	// Check if the player controller exists and is local
-	UE_DEPRECATED(5.4, "Deprecated, use IsLocallyControlled() which takes both local player controlled and local possession into account.")
+	UE_DEPRECATED(5.4, "Deprecated, use IsLocallyControlled() which takes both local player controlled and local relayed inputs into account.")
 	ENGINE_API bool HasLocalController() const;
 	
-	// Check if this is controlled locally through possession or an existing local player controller
+	// Check if this is controlled locally through relayed inputs or an existing local player controller
 	ENGINE_API bool IsLocallyControlled() const;
 
-	/** Mark this as controlled through local possession rather than controlled as a pawn
-	* Set if NetworkPhysicsComponent is implemented on an AActor instead of APawn and it's currently being fed inputs from the local player */
-	ENGINE_API void SetIsLocallyPossessed(bool bPossess)
+	/** Mark this as controlled through locally relayed inputs rather than controlled as a pawn through a player controller.
+	* Set if NetworkPhysicsComponent is implemented on an AActor instead of APawn and it's currently being fed inputs from the local player / autonomous proxy */
+	ENGINE_API void SetIsRelayingLocalInputs(bool bInRelayingLocalInputs)
 	{
-		bIsLocallyPossessed = bPossess;
+		bIsRelayingLocalInputs = bInRelayingLocalInputs;
 	}
 
-	/** Check if this is controlled locally through possession. It's recommended to use HasLocalController() when checking if this is locally controlled. */
-	ENGINE_API const bool GetIsLocallyPossessed() const { return bIsLocallyPossessed; }
+	/** Check if this is controlled locally through relayed inputs from autonomous proxy. It's recommended to use IsLocallyControlled() when checking if this is locally controlled. */
+	ENGINE_API const bool GetIsRelayingLocalInputs() const { return bIsRelayingLocalInputs; }
+
+	/** Returns the current amount of input decay during resimulation as a magnitude from 0.0 to 1.0. Returns 0 if not currently resimulating. */
+	ENGINE_API const float GetCurrentInputDecay(FNetworkPhysicsDatas* PhysicsDatas);
 
 protected : 
 
@@ -513,8 +523,8 @@ private:
 	// Actor component that will be used to fill the histories
 	TObjectPtr<UActorComponent> ActorComponent;
 
-	// Locally possessed makes this component act as if it's a locally controlled pawn.
-	bool bIsLocallyPossessed = false;
+	// Locally relayed inputs makes this component act as if it's a locally controlled pawn.
+	bool bIsRelayingLocalInputs = false;
 };
 
 template<typename PhysicsTraits>
diff --git a/Engine/Source/Runtime/Experimental/Chaos/Private/PBDRigidsSolver.cpp b/Engine/Source/Runtime/Experimental/Chaos/Private/PBDRigidsSolver.cpp
--- a/Engine/Source/Runtime/Experimental/Chaos/Private/PBDRigidsSolver.cpp
+++ b/Engine/Source/Runtime/Experimental/Chaos/Private/PBDRigidsSolver.cpp
@@ -1849,7 +1849,7 @@ namespace Chaos
 			{
 				FResimDebugInfo DebugInfo;
 				QUICK_SCOPE_CYCLE_COUNTER(ChaosRewindAndResim);
-				if(MRewindData->RewindToFrame(ResimStep))
+				if (MRewindData->RewindToFrame(ResimStep))
 				{
 #if DEBUG_REWIND_DATA
 					UE_LOG(LogTemp, Warning, TEXT("COMMON | PT | ConditionalApplyRewind_Internal | PERFORMING RESIMULATION | Resim From Frame = %d | Num Steps = %d | To Current Frame: %d"), ResimStep, NumResimSteps, CurrentFrame);
@@ -1857,7 +1857,7 @@ namespace Chaos
 
 					GetEvolution()->SetResim(true);
 					CurrentFrame = ResimStep;
-					
+
 					TArray<FPushPhysicsData*> RecordedPushData = MarshallingManager.StealHistory_Internal(NumResimSteps);
 					bool bFirst = true;
 
@@ -1935,6 +1935,8 @@ namespace Chaos
 				UE_LOG(LogTemp, Log, TEXT("COMMON | PT | ConditionalApplyRewind_Internal | Resimulation failed, invalid rewind frame data | Current Frame = %d | Num Steps = %d | Resim Frame = %d | Last Frame = %d | Rewind History Size = %d"), CurrentFrame, NumResimSteps, ResimStep, LastStep, MarshallingManager.GetNumHistory_Internal());
 			}
 #endif
+			// Clear the ResimFrame no matter if resimulation succeeded or failed (if it failed it's not going to succeed next frame either based on the same ResimFrame)
+			MRewindData->SetResimFrame(INDEX_NONE);
 		}
 	}
 
diff --git a/Engine/Source/Runtime/Experimental/Chaos/Private/RewindData.cpp b/Engine/Source/Runtime/Experimental/Chaos/Private/RewindData.cpp
--- a/Engine/Source/Runtime/Experimental/Chaos/Private/RewindData.cpp
+++ b/Engine/Source/Runtime/Experimental/Chaos/Private/RewindData.cpp
@@ -378,9 +378,10 @@ bool FRewindData::RewindToFrame(int32 Frame)
 	}
 #endif
 
+	ResimFrame = Frame;
 	CurFrame = Frame;
 	bNeedsSave = false;
-	
+
 	return true;
 }
 
diff --git a/Engine/Source/Runtime/Experimental/Chaos/Public/RewindData.h b/Engine/Source/Runtime/Experimental/Chaos/Public/RewindData.h
--- a/Engine/Source/Runtime/Experimental/Chaos/Public/RewindData.h
+++ b/Engine/Source/Runtime/Experimental/Chaos/Public/RewindData.h
@@ -1298,6 +1298,7 @@ public:
 
 	int32 Capacity() const { return Managers.Capacity(); }
 	int32 CurrentFrame() const { return CurFrame; }
+	int32 GetLatestFrame() const { return LatestFrame; }
 	int32 GetFramesSaved() const { return FramesSaved; }
 
 	FReal GetDeltaTimeForFrame(int32 Frame) const
