Fixes up an unnecessary implicit TObjectPtr to raw pointer array conversion, addressing a 4.27 -> 5.0 performance regression of GetChildComponent after the introduction of TObjectPtr.

diff --git a/Engine/Source/Runtime/Engine/Private/Components/SceneComponent.cpp b/Engine/Source/Runtime/Engine/Private/Components/SceneComponent.cpp
--- a/Engine/Source/Runtime/Engine/Private/Components/SceneComponent.cpp
+++ b/Engine/Source/Runtime/Engine/Private/Components/SceneComponent.cpp
@@ -1885,7 +1885,7 @@ USceneComponent* USceneComponent::GetChildComponent(int32 ChildIndex) const
 		return nullptr;
 	}
 
-	const TArray<USceneComponent*>& AttachedChildren = GetAttachChildren();
+	const TArray<TObjectPtr<USceneComponent>>& AttachedChildren = GetAttachChildren();
 	if (ChildIndex >= AttachedChildren.Num())
 	{
 		UE_LOG(LogBlueprint, Log, TEXT("SceneComponent::GetChild called with an out of range ChildIndex: %d; Number of children is %d."), ChildIndex, AttachedChildren.Num());
