Fix: make sure LevelInstance exists (it is created after when creating a new Level Instance). This needs to be fixed properly but for now since the created LI is then unloaded it doesn't matter that the parent isn't set.

diff --git a/Engine/Source/Runtime/Engine/Private/LevelInstance/LevelInstanceEditorLevelStreaming.cpp b/Engine/Source/Runtime/Engine/Private/LevelInstance/LevelInstanceEditorLevelStreaming.cpp
--- a/Engine/Source/Runtime/Engine/Private/LevelInstance/LevelInstanceEditorLevelStreaming.cpp
+++ b/Engine/Source/Runtime/Engine/Private/LevelInstance/LevelInstanceEditorLevelStreaming.cpp
@@ -154,14 +154,16 @@ void ULevelStreamingLevelInstanceEditor::OnLevelActorAdded(AActor* InActor)
 
 void ULevelStreamingLevelInstanceEditor::OnPreInitializeContainerInstance(UActorDescContainerInstance::FInitializeParams& InInitParams, UActorDescContainerInstance* InContainerInstance)
 {
-	// Apply override container
-	AActor* LevelInstanceActor = Cast<AActor>(GetLevelInstance());
-	if (UWorldPartition* OwningWorldPartition = LevelInstanceActor->GetLevel()->GetWorldPartition())
+	// @todo_ow: Make sure LevelInstance exists (it is created after when creating a new Level Instance). This needs to be fixed properly but for now since the created LI is then unloaded it doesn't matter that the parent isn't set.
+	if (AActor* LevelInstanceActor = Cast<AActor>(GetLevelInstance()))
 	{
-		if (FWorldPartitionActorDescInstance* ActorDescInstance = OwningWorldPartition->GetActorDescInstance(LevelInstanceActor->GetActorGuid()); ActorDescInstance && ActorDescInstance->IsChildContainerInstance())
+		if (UWorldPartition* OwningWorldPartition = LevelInstanceActor->GetLevel()->GetWorldPartition())
 		{
-			// Add parenting info to init param
-			InInitParams.SetParent(ActorDescInstance->GetContainerInstance(), ActorDescInstance->GetGuid());
+			if (FWorldPartitionActorDescInstance* ActorDescInstance = OwningWorldPartition->GetActorDescInstance(LevelInstanceActor->GetActorGuid()); ActorDescInstance && ActorDescInstance->IsChildContainerInstance())
+			{
+				// Add parenting info to init param
+				InInitParams.SetParent(ActorDescInstance->GetContainerInstance(), ActorDescInstance->GetGuid());
+			}
 		}
 	}
 }
