Cleanup SubObject replication logs

#rnx

diff --git a/Engine/Source/Runtime/Engine/Private/ActorReplication.cpp b/Engine/Source/Runtime/Engine/Private/ActorReplication.cpp
--- a/Engine/Source/Runtime/Engine/Private/ActorReplication.cpp
+++ b/Engine/Source/Runtime/Engine/Private/ActorReplication.cpp
@@ -752,7 +752,7 @@ void AActor::DestroyReplicatedSubObjectOnRemotePeers(UObject* SubObject)
 		return;
 	}
 
-	UE_LOG(LogNetSubObject, Verbose, TEXT("%s (0x%p) requested to Delete replicated subobject on clients %s (0x%p)"), *GetName(), this, *SubObject->GetName(), SubObject);
+	UE_LOG(LogNetSubObject, Verbose, TEXT("%s (0x%p) requested to Delete replicated subobject: %s (0x%p)"), *GetName(), this, *SubObject->GetName(), SubObject);
 
 	if (FWorldContext* const Context = GEngine->GetWorldContextFromWorld(GetWorld()))
 	{
@@ -781,7 +781,7 @@ void AActor::DestroyReplicatedSubObjectOnRemotePeers(UActorComponent* OwnerCompo
 		return;
 	}
 
-	UE_LOG(LogNetSubObject, Verbose, TEXT("%s::%s (0x%p) requested to Delete replicated subobject on clients %s (0x%p)"), *GetName(), *OwnerComponent->GetName(), OwnerComponent, *SubObject->GetName(), SubObject);
+	UE_LOG(LogNetSubObject, Verbose, TEXT("%s::%s (0x%p) requested to Delete replicated subobject: %s (0x%p)"), *GetName(), *OwnerComponent->GetName(), OwnerComponent, *SubObject->GetName(), SubObject);
 
 	if (FWorldContext* const Context = GEngine->GetWorldContextFromWorld(GetWorld()))
 	{
@@ -810,7 +810,7 @@ void AActor::TearOffReplicatedSubObjectOnRemotePeers(UActorComponent* OwnerCompo
 		return;
 	}
 
-	UE_LOG(LogNetSubObject, Verbose, TEXT("%s::%s (0x%p) requested to TearOff replicated subobject on clients %s (0x%p)"), *GetName(), *OwnerComponent->GetName(), this, *SubObject->GetName(), SubObject);
+	UE_LOG(LogNetSubObject, Verbose, TEXT("%s::%s (0x%p) requested to TearOff replicated subobject: %s (0x%p)"), *GetName(), *OwnerComponent->GetName(), this, *SubObject->GetName(), SubObject);
 
 	if (FWorldContext* const Context = GEngine->GetWorldContextFromWorld(GetWorld()))
 	{
@@ -839,7 +839,7 @@ void AActor::TearOffReplicatedSubObjectOnRemotePeers(UObject* SubObject)
 		return;
 	}
 
-	UE_LOG(LogNetSubObject, Verbose, TEXT("%s (0x%p) requested to TearOff replicated subobject on clients %s (0x%p)"), *GetName(), this, *SubObject->GetName(), SubObject);
+	UE_LOG(LogNetSubObject, Verbose, TEXT("%s (0x%p) requested to TearOff replicated subobject: %s (0x%p)"), *GetName(), this, *SubObject->GetName(), SubObject);
 
 	if (FWorldContext* const Context = GEngine->GetWorldContextFromWorld(GetWorld()))
 	{
