Add a network metric that tracks the maximum size of the incoming and outgoing reliable message queue size over the lifetime of a net driver. By default this metric reports to CSV and Stats.

diff --git a/Engine/Config/BaseEngine.ini b/Engine/Config/BaseEngine.ini
--- a/Engine/Config/BaseEngine.ini
+++ b/Engine/Config/BaseEngine.ini
@@ -3492,6 +3492,8 @@ IntervalSeconds=1
 +Listeners=(MetricName=SharedSerializationPropertyMiss, Class=/Script/Engine.NetworkMetricsCSV_Replication)
 +Listeners=(MetricName=SharedSerializationRPCHit, Class=/Script/Engine.NetworkMetricsCSV_Replication)
 +Listeners=(MetricName=SharedSerializationRPCMiss, Class=/Script/Engine.NetworkMetricsCSV_Replication)
++Listeners=(MetricName=OutgoingReliableMessageQueueMaxSize, Class=/Script/Engine.NetworkMetricsCSV_Replication)
++Listeners=(MetricName=IncomingReliableMessageQueueMaxSize, Class=/Script/Engine.NetworkMetricsCSV_Replication)
 +Listeners=(MetricName=AvgPing, Class=/Script/Engine.NetworkMetricsPerfCounters)
 +Listeners=(MetricName=InBunches, Class=/Script/Engine.NetworkMetricsPerfCounters)
 +Listeners=(MetricName=InPackets, Class=/Script/Engine.NetworkMetricsPerfCounters)
diff --git a/Engine/Source/Runtime/Engine/Classes/Engine/NetDriver.h b/Engine/Source/Runtime/Engine/Classes/Engine/NetDriver.h
--- a/Engine/Source/Runtime/Engine/Classes/Engine/NetDriver.h
+++ b/Engine/Source/Runtime/Engine/Classes/Engine/NetDriver.h
@@ -1005,6 +1005,9 @@ private:
 	/** Create an instance of UNetworkMetricsStats that is associated with a given Stat and cached with other listeners in NetworkMetricsListeners. */
 	void RegisterStatsListener(const FName MetricName, const FName StatName);
 
+	/** Reset any network metrics database values at the beginning of a frame. */
+	void ResetNetworkMetrics();
+
 public:
 	/** Get the value of MaxChannelsOverride cached from the net driver definition */
 	int32 GetMaxChannelsOverride() const { return MaxChannelsOverride; }
diff --git a/Engine/Source/Runtime/Engine/Private/DataChannel.cpp b/Engine/Source/Runtime/Engine/Private/DataChannel.cpp
--- a/Engine/Source/Runtime/Engine/Private/DataChannel.cpp
+++ b/Engine/Source/Runtime/Engine/Private/DataChannel.cpp
@@ -662,6 +662,7 @@ void UChannel::ReceivedRawBunch( FInBunch & Bunch, bool & bOutSkipAck )
 		New->Next     = *InPtr;
 		*InPtr        = New;
 		NumInRec++;
+		Connection->GetDriver()->GetMetrics()->SetMaxInt(UE::Net::Metric::IncomingReliableMessageQueueMaxSize, NumInRec);
 
 		if ( NumInRec >= RELIABLE_BUFFER )
 		{
@@ -1461,6 +1462,7 @@ FOutBunch* UChannel::PrepBunch(FOutBunch* Bunch, FOutBunch* OutBunch, bool Merge
 			Bunch->Next	= NULL;
 			Bunch->ChSequence = ++Connection->OutReliable[ChIndex];
 			NumOutRec++;
+			Connection->GetDriver()->GetMetrics()->SetMaxInt(UE::Net::Metric::OutgoingReliableMessageQueueMaxSize, NumOutRec);
 			OutBunch = new FOutBunch(*Bunch);
 			FOutBunch** OutLink = &OutRec;
 			while(*OutLink) // This was rewritten from a single-line for loop due to compiler complaining about empty body for loops (-Wempty-body)
diff --git a/Engine/Source/Runtime/Engine/Private/Net/NetworkMetricsDatabase.cpp b/Engine/Source/Runtime/Engine/Private/Net/NetworkMetricsDatabase.cpp
--- a/Engine/Source/Runtime/Engine/Private/Net/NetworkMetricsDatabase.cpp
+++ b/Engine/Source/Runtime/Engine/Private/Net/NetworkMetricsDatabase.cpp
@@ -43,6 +43,24 @@ void UNetworkMetricsDatabase::SetFloat(const FName MetricName, float Value)
 	}
 }
 
+void UNetworkMetricsDatabase::SetMinFloat(const FName MetricName, float Value)
+{
+	UE::Net::FNetworkMetric<float>* Metric = MetricFloats.Find(MetricName);
+	if (ensureMsgf(Metric, TEXT("Cannot find float metric %s."), *MetricName.ToString()))
+	{
+		Metric->Value = FMath::Min<float>(Value, Metric->Value);
+	}
+}
+
+void UNetworkMetricsDatabase::SetMaxFloat(const FName MetricName, float Value)
+{
+	UE::Net::FNetworkMetric<float>* Metric = MetricFloats.Find(MetricName);
+	if (ensureMsgf(Metric, TEXT("Cannot find float metric %s."), *MetricName.ToString()))
+	{
+		Metric->Value = FMath::Max<float>(Value, Metric->Value);
+	}
+}
+
 void UNetworkMetricsDatabase::SetInt(const FName MetricName, int64 Value)
 {
 	UE::Net::FNetworkMetric<int64>* Metric = MetricInts.Find(MetricName);
@@ -52,6 +70,24 @@ void UNetworkMetricsDatabase::SetInt(const FName MetricName, int64 Value)
 	}
 }
 
+void UNetworkMetricsDatabase::SetMinInt(const FName MetricName, int64 Value)
+{
+	UE::Net::FNetworkMetric<int64>* Metric = MetricInts.Find(MetricName);
+	if (ensureMsgf(Metric, TEXT("Cannot find integer metric %s."), *MetricName.ToString()))
+	{
+		Metric->Value = FMath::Min<int64>(Value, Metric->Value);
+	}
+}
+
+void UNetworkMetricsDatabase::SetMaxInt(const FName MetricName, int64 Value)
+{
+	UE::Net::FNetworkMetric<int64>* Metric = MetricInts.Find(MetricName);
+	if (ensureMsgf(Metric, TEXT("Cannot find integer metric %s."), *MetricName.ToString()))
+	{
+		Metric->Value = FMath::Max<int64>(Value, Metric->Value);
+	}
+}
+
 void UNetworkMetricsDatabase::IncrementInt(const FName MetricName, int64 Value)
 {
 	UE::Net::FNetworkMetric<int64>* Metric = MetricInts.Find(MetricName);
diff --git a/Engine/Source/Runtime/Engine/Private/Net/NetworkMetricsDefs.cpp b/Engine/Source/Runtime/Engine/Private/Net/NetworkMetricsDefs.cpp
--- a/Engine/Source/Runtime/Engine/Private/Net/NetworkMetricsDefs.cpp
+++ b/Engine/Source/Runtime/Engine/Private/Net/NetworkMetricsDefs.cpp
@@ -92,6 +92,8 @@ const FName PingBucketInt4("PingBucketInt4");
 const FName PingBucketInt5("PingBucketInt5");
 const FName PingBucketInt6("PingBucketInt6");
 const FName PingBucketInt7("PingBucketInt7");
+const FName OutgoingReliableMessageQueueMaxSize("OutgoingReliableMessageQueueMaxSize");
+const FName IncomingReliableMessageQueueMaxSize("IncomingReliableMessageQueueMaxSize");
 
 }
 
diff --git a/Engine/Source/Runtime/Engine/Private/NetDriver.cpp b/Engine/Source/Runtime/Engine/Private/NetDriver.cpp
--- a/Engine/Source/Runtime/Engine/Private/NetDriver.cpp
+++ b/Engine/Source/Runtime/Engine/Private/NetDriver.cpp
@@ -123,6 +123,8 @@ DEFINE_STAT(STAT_NetSaturated);
 DEFINE_STAT(STAT_ImportedNetGuids);
 DEFINE_STAT(STAT_PendingOuterNetGuids);
 DEFINE_STAT(STAT_UnmappedReplicators);
+DEFINE_STAT(STAT_OutgoingReliableMessageQueueMaxSize);
+DEFINE_STAT(STAT_IncomingReliableMessageQueueMaxSize);
 
 // Voice specific stats
 DEFINE_STAT(STAT_VoiceBytesSent);
@@ -1735,6 +1737,10 @@ void UNetDriver::SetupNetworkMetrics()
 	GetMetrics()->CreateInt(UE::Net::Metric::PingBucketInt6, 0);
 	GetMetrics()->CreateInt(UE::Net::Metric::PingBucketInt7, 0);
 
+	// The maximum size of the incoming/outgoing reliable message queues across all channels and connections in the last frame.
+	GetMetrics()->CreateInt(UE::Net::Metric::OutgoingReliableMessageQueueMaxSize, 0);
+	GetMetrics()->CreateInt(UE::Net::Metric::IncomingReliableMessageQueueMaxSize, 0);
+
 	GetMetrics()->CreateInt(UE::Net::Metric::ImportedNetGuids, 0);
 	GetMetrics()->CreateInt(UE::Net::Metric::PendingOuterNetGuids, 0);
 	GetMetrics()->CreateInt(UE::Net::Metric::NetInBunchTimeOvershootPercent, 0);
@@ -1833,6 +1839,8 @@ void UNetDriver::SetupNetworkMetricsListeners()
 	RegisterStatsListener(UE::Net::Metric::NetSaturated, GET_STATFNAME(STAT_NetSaturated));
 	RegisterStatsListener(UE::Net::Metric::NetInBunchTimeOvershootPercent, GET_STATFNAME(STAT_NetInBunchTimeOvershootPercent));
 	RegisterStatsListener(UE::Net::Metric::GatherPrioritizeTimeMS, GET_STATFNAME(STAT_NetServerGatherPrioritizeRepActorsTime));
+	RegisterStatsListener(UE::Net::Metric::OutgoingReliableMessageQueueMaxSize, GET_STATFNAME(STAT_OutgoingReliableMessageQueueMaxSize));
+	RegisterStatsListener(UE::Net::Metric::IncomingReliableMessageQueueMaxSize, GET_STATFNAME(STAT_IncomingReliableMessageQueueMaxSize));
 #endif
 }
 
@@ -1962,9 +1970,6 @@ static FAutoConsoleVariableRef CVarLogPendingGuidsOnShutdown(
 /** Shutdown all connections managed by this net driver */
 void UNetDriver::Shutdown()
 {
-	NetworkMetricsDatabase->Reset();
-	NetworkMetricsListeners.Reset();
-
 	// Client closing connection to server
 	if (ServerConnection)
 	{
@@ -2068,6 +2073,9 @@ void UNetDriver::Shutdown()
 	}
 
 	UpdateCrashContext();
+
+	NetworkMetricsDatabase->Reset();
+	NetworkMetricsListeners.Reset();
 }
 
 bool UNetDriver::IsServer() const
@@ -2185,6 +2193,8 @@ void UNetDriver::TickDispatch( float DeltaTime )
 	// Checks for standby cheats if enabled	
 	UpdateStandbyCheatStatus();
 
+	ResetNetworkMetrics();
+
 	if (ServerConnection == nullptr)
 	{
 		// Delete any straggler connections
@@ -5843,6 +5853,12 @@ void UNetDriver::RegisterStatsListener(const FName MetricName, const FName StatN
 	}
 }
 
+void UNetDriver::ResetNetworkMetrics()
+{
+	GetMetrics()->SetInt(UE::Net::Metric::OutgoingReliableMessageQueueMaxSize, 0);
+	GetMetrics()->SetInt(UE::Net::Metric::IncomingReliableMessageQueueMaxSize, 0);
+}
+
 void UNetDriver::ReleaseToChannelPool(UChannel* Channel)
 {
 	LLM_SCOPE_BYTAG(NetChannel);
diff --git a/Engine/Source/Runtime/Engine/Public/EngineStats.h b/Engine/Source/Runtime/Engine/Public/EngineStats.h
--- a/Engine/Source/Runtime/Engine/Public/EngineStats.h
+++ b/Engine/Source/Runtime/Engine/Public/EngineStats.h
@@ -255,6 +255,8 @@ DECLARE_DWORD_ACCUMULATOR_STAT_EXTERN(TEXT("Pending Outer NetGuids"),STAT_Pendin
 DECLARE_DWORD_ACCUMULATOR_STAT_EXTERN(TEXT("Unmapped Replicators"),STAT_UnmappedReplicators,STATGROUP_Net, );
 DECLARE_DWORD_ACCUMULATOR_STAT_EXTERN(TEXT("Num Replicated Actors Sent"),STAT_NumReplicatedActors,STATGROUP_Net, );
 DECLARE_DWORD_ACCUMULATOR_STAT_EXTERN(TEXT("Num Replicated Actor Bytes Sent"),STAT_NumReplicatedActorBytes,STATGROUP_Net, );
+DECLARE_DWORD_ACCUMULATOR_STAT_EXTERN(TEXT("Largest Reliable Outgoing Queue Size"),STAT_OutgoingReliableMessageQueueMaxSize,STATGROUP_Net, );
+DECLARE_DWORD_ACCUMULATOR_STAT_EXTERN(TEXT("Largest Reliable Incoming Queue Size"),STAT_IncomingReliableMessageQueueMaxSize,STATGROUP_Net, );
 
 #if !UE_BUILD_SHIPPING
 /**
diff --git a/Engine/Source/Runtime/Engine/Public/Net/NetworkMetricsDatabase.h b/Engine/Source/Runtime/Engine/Public/Net/NetworkMetricsDatabase.h
--- a/Engine/Source/Runtime/Engine/Public/Net/NetworkMetricsDatabase.h
+++ b/Engine/Source/Runtime/Engine/Public/Net/NetworkMetricsDatabase.h
@@ -41,8 +41,16 @@ public:
 	void CreateInt(const FName MetricName, int64 DefaultValue);
 	/* Set the value of an existing floating point metric. */
 	void SetFloat(const FName MetricName, float Value);
+	/* Set the value of a floating point metric if it's smaller than the existing value. */
+	void SetMinFloat(const FName MetricName, float Value);
+	/* Set the value of a floating point metric if it's bigger than the existing value. */
+	void SetMaxFloat(const FName MetricName, float Value);
 	/* Set the value of an existing integer metric. */
 	void SetInt(const FName MetricName, int64 Value);
+	/* Set the value of an integer metric if it's smaller than the existing value. */
+	void SetMinInt(const FName MetricName, int64 Value);
+	/* Set the value of an integer metric if it's bigger than the existing value. */
+	void SetMaxInt(const FName MetricName, int64 Value);
 	/* Increment the value of an existing integer metric. */
 	void IncrementInt(const FName MetricName, int64 Value);
 	/* Returns true if a metric has been created in the database. */
diff --git a/Engine/Source/Runtime/Engine/Public/Net/NetworkMetricsDefs.h b/Engine/Source/Runtime/Engine/Public/Net/NetworkMetricsDefs.h
--- a/Engine/Source/Runtime/Engine/Public/Net/NetworkMetricsDefs.h
+++ b/Engine/Source/Runtime/Engine/Public/Net/NetworkMetricsDefs.h
@@ -94,5 +94,7 @@ extern ENGINE_API const FName PingBucketInt4;
 extern ENGINE_API const FName PingBucketInt5;
 extern ENGINE_API const FName PingBucketInt6;
 extern ENGINE_API const FName PingBucketInt7;
+extern ENGINE_API const FName OutgoingReliableMessageQueueMaxSize;
+extern ENGINE_API const FName IncomingReliableMessageQueueMaxSize;
 
 }
