Fix: crash when creating sub-level instance from existing actors in the parent level instance (missing event hook allowing unhash of desc instance in instanced LIs)

#rnx

diff --git a/Engine/Source/Runtime/Engine/Private/WorldPartition/ActorDescContainerInstance.cpp b/Engine/Source/Runtime/Engine/Private/WorldPartition/ActorDescContainerInstance.cpp
--- a/Engine/Source/Runtime/Engine/Private/WorldPartition/ActorDescContainerInstance.cpp
+++ b/Engine/Source/Runtime/Engine/Private/WorldPartition/ActorDescContainerInstance.cpp
@@ -208,14 +208,6 @@ void UActorDescContainerInstance::RegisterDelegates()
 	{
 		check(Container);
 
-		// Containers with an instancing context:
-		// PIE / NewMap : Do not need to listen to Add/Removed/Updated events as they are either not changing (PIE) or don't have valid actor descs yet (NewMap)
-		if (!GetInstancingContext())
-		{
-			Container->OnActorDescAddedEvent.AddUObject(this, &UActorDescContainerInstance::OnActorDescAdded);
-			Container->OnActorDescRemovedEvent.AddUObject(this, &UActorDescContainerInstance::OnActorDescRemoved);
-		}
-
 		// Only listen to Object replaced events on ContainerInstance that have a direct World Partition outer (Loaded Container Instances: Main World or Loaded Level Instances)
 		if (HasWorldPartition())
 		{
@@ -228,7 +220,16 @@ void UActorDescContainerInstance::RegisterDelegates()
 			UActorDescContainerSubsystem::GetChecked().ContainerUpdated().AddUObject(this, &UActorDescContainerInstance::OnContainerUpdated);
 		}
 
-		// Always listen to Updated event to invalidate FWorldPartitionActorDescInstance Cached ActorDesc & Hash/Unhash if this container is a loaded world
+		// No need to register Added descs events for instanced worlds as they don't support it for now (Level Instances get reloaded after an edit)
+		if (!GetInstancingContext())
+		{
+			Container->OnActorDescAddedEvent.AddUObject(this, &UActorDescContainerInstance::OnActorDescAdded);
+		}
+
+		// Important to hook the other events that will invalidate existing FWorldPartitionActorDescInstance's because those can be hashed and loaded
+		// even in Instanced worlds (Level Instances)
+		Container->OnActorDescRemovedEvent.AddUObject(this, &UActorDescContainerInstance::OnActorDescRemoved);
+		
 		Container->OnActorDescUpdatingEvent.AddUObject(this, &UActorDescContainerInstance::OnActorDescUpdating);
 		Container->OnActorDescUpdatedEvent.AddUObject(this, &UActorDescContainerInstance::OnActorDescUpdated);
 	}
@@ -384,8 +385,6 @@ FWorldPartitionActorDescInstance* UActorDescContainerInstance::AddActorDescInsta
 
 void UActorDescContainerInstance::RemoveActor(const FGuid& InActorGuid)
 {
-	check(!InstancingContext.IsSet());
-
 	if (TUniquePtr<FWorldPartitionActorDescInstance>* ActorDescInstance = GetActorDescriptor(InActorGuid))
 	{
 		if (UWorldPartition* WorldPartition = GetWorldPartition())
