Latch an error state in the case where the decoder returns negative samples

diff --git a/Engine/Source/Runtime/BinkAudioDecoder/Module/Private/BinkAudioInfo.cpp b/Engine/Source/Runtime/BinkAudioDecoder/Module/Private/BinkAudioInfo.cpp
--- a/Engine/Source/Runtime/BinkAudioDecoder/Module/Private/BinkAudioInfo.cpp
+++ b/Engine/Source/Runtime/BinkAudioDecoder/Module/Private/BinkAudioInfo.cpp
@@ -684,7 +684,7 @@ FDecodeResult FBinkAudioInfo::Decode(const uint8* CompressedData, const int32 Co
 
 bool FBinkAudioInfo::HasError() const
 {
-	return bErrorStateLatch;
+	return Super::HasError() || bErrorStateLatch;
 }
 
 class BINKAUDIODECODER_API FBinkAudioDecoderModule : public IModuleInterface
diff --git a/Engine/Source/Runtime/Engine/Private/AudioDecompress.cpp b/Engine/Source/Runtime/Engine/Private/AudioDecompress.cpp
--- a/Engine/Source/Runtime/Engine/Private/AudioDecompress.cpp
+++ b/Engine/Source/Runtime/Engine/Private/AudioDecompress.cpp
@@ -12,9 +12,41 @@
 #include "Stats/StatsTrace.h"
 #include "Sound/StreamedAudioChunkSeekTable.h"
 
+namespace AudioDecompressPrivate
+{
+#ifndef WITH_AUDIO_DECODER_DIAGNOSTICS
+	// Optionally enable this manually with a flag. They use strings, so shouldn't be enabled by default.
+	#define WITH_AUDIO_DECODER_DIAGNOSTICS (0)
+#endif //WITH_AUDIO_DECODER_DIAGNOSTICS
+
+#if WITH_AUDIO_DECODER_DIAGNOSTICS
+	
+	static FString ForceDecoderErrorOnWaveCVar;
+	FAutoConsoleVariableRef CVarForceDecoderErrorOnWave(
+		TEXT("au.debug.force_decoder_error_on_wave"),
+		ForceDecoderErrorOnWaveCVar,
+		TEXT("Force Decoder Error On Any decoding wave matching this string."),
+		ECVF_Default);
+
+	static FString ForceDecoderNegativeSamplesOnWaveCVar;
+	FAutoConsoleVariableRef CVarForceDecoderNegativeSamplesOnWave(
+		TEXT("au.debug.force_decoder_negative_samples_on_wave"),
+		ForceDecoderNegativeSamplesOnWaveCVar,
+		TEXT("Force Negative Samples on Decode call to simulate error" ),
+		ECVF_Default);
+
+	// Macro to string match against the debugging wave of choice.
+	#define DECODER_MATCHES_WAVE(STR)\
+		(!STR.IsEmpty() && StreamingSoundWave.IsValid() && StreamingSoundWave->GetFName().ToString().Contains(STR))
+	
+#else  //WITH_AUDIO_DECODER_DIAGNOSTICS
+	#define DECODER_MATCHES_WAVE(STR) (false)
+#endif //WITH_AUDIO_DECODER_DIAGNOSTICS
+}
+
 IStreamedCompressedInfo::IStreamedCompressedInfo()
 	: bIsStreaming(false)
-	,SrcBufferData(nullptr)
+	, SrcBufferData(nullptr)
 	, SrcBufferDataSize(0)
 	, SrcBufferOffset(0)
 	, AudioDataOffset(0)
@@ -182,7 +214,7 @@ void IStreamedCompressedInfo::ExpandFile(uint8* DstBuffer, struct FSoundQualityI
 
 	while (RawPCMOffset < QualityInfo->SampleDataSize)
 	{
-		int32 DecodedFrames = DecompressToPCMBuffer( /*Unused*/ 0);
+		const int32 DecodedFrames = DecompressToPCMBuffer( /*Unused*/ 0);
 
 		if (DecodedFrames < 0)
 		{
@@ -243,6 +275,11 @@ bool IStreamedCompressedInfo::StreamCompressedInfoInternal(const FSoundWaveProxy
 	return false;
 }
 
+bool ICompressedAudioInfo::HasError() const
+{
+	return bHasError || DECODER_MATCHES_WAVE(AudioDecompressPrivate::ForceDecoderErrorOnWaveCVar);
+}
+
 bool IStreamedCompressedInfo::StreamCompressedData(uint8* Destination, bool bLooping, uint32 BufferSize, int32& OutNumBytesStreamed)
 {
 	check(Destination);
@@ -422,13 +459,16 @@ bool IStreamedCompressedInfo::StreamCompressedData(uint8* Destination, bool bLoo
 	while (RawPCMOffset < BufferSize)
 	{
 		// Decompress the next compression frame of audio (many samples) into the PCM buffer
-		int32 DecodedFrames = DecompressToPCMBuffer(/*Unused*/ 0);
+		const int32 DecodedFrames = DecompressToPCMBuffer(/*Unused*/ 0);
 
 		if (DecodedFrames < 0)
 		{
 			UE_LOG(LogAudioStreamCaching, Warning, TEXT("Zero pad buffer Chunk=%d, Wave=%s, Reason=Decoder returned negative samples."),
 				CurrentChunkIndex, *StreamingSoundWave->GetFName().ToString());
 
+			// Flag that the decoder has an unrecoverable error, so we don't try again.
+			bHasError = true;
+
 			LastPCMByteSize = 0;
 			ZeroBuffer(Destination + RawPCMOffset, BufferSize - RawPCMOffset);
 			return false;
@@ -543,12 +583,18 @@ int32 IStreamedCompressedInfo::DecompressToPCMBuffer(uint16 /*Unused*/ )
 	LastPCMOffset = 0;
 	
 	const FDecodeResult DecodeResult = Decode(SrcPtr, FrameSize, LastDecodedPCM.GetData(), LastDecodedPCM.Num());
-	if (DecodeResult.NumCompressedBytesConsumed != INDEX_NONE )
+	
+	if (DecodeResult.NumCompressedBytesConsumed == INDEX_NONE || DECODER_MATCHES_WAVE(AudioDecompressPrivate::ForceDecoderNegativeSamplesOnWaveCVar))
 	{
-		SrcBufferOffset -= FrameSize;
-		SrcBufferOffset += DecodeResult.NumCompressedBytesConsumed;
+		UE_LOG(LogAudio, Warning, TEXT("Decoder error: Decode call returned INDEX_NONE which indicates an error. : Wave='%s', FrameSize=%d, SrcBufferOffset=%u, SrcBufferDataSize=%u"), 
+			*GetStreamingSoundWave()->GetSoundWaveData()->GetFName().ToString(), FrameSize, SrcBufferOffset, SrcBufferDataSize);
+
+		// Error.
+		return -1;
 	}
 
+	SrcBufferOffset -= FrameSize;
+	SrcBufferOffset += DecodeResult.NumCompressedBytesConsumed;
 	return DecodeResult.NumAudioFramesProduced;
 }
 
diff --git a/Engine/Source/Runtime/Engine/Private/VorbisAudioInfo.cpp b/Engine/Source/Runtime/Engine/Private/VorbisAudioInfo.cpp
--- a/Engine/Source/Runtime/Engine/Private/VorbisAudioInfo.cpp
+++ b/Engine/Source/Runtime/Engine/Private/VorbisAudioInfo.cpp
@@ -781,6 +781,11 @@ bool FVorbisAudioInfo::StreamCompressedData(uint8* InDestination, bool bLooping,
 	return( bLooped );
 }
 
+bool FVorbisAudioInfo::HasError() const
+{
+	return Super::HasError() | bHasError;
+}
+
 void LoadVorbisLibraries()
 {
 	static bool bIsInitialized = false;
diff --git a/Engine/Source/Runtime/Engine/Public/AudioDecompress.h b/Engine/Source/Runtime/Engine/Public/AudioDecompress.h
--- a/Engine/Source/Runtime/Engine/Public/AudioDecompress.h
+++ b/Engine/Source/Runtime/Engine/Public/AudioDecompress.h
@@ -137,7 +137,7 @@ public:
 	/**
 	*  Returns true if a non-recoverable error has occurred.
 	*/
-	virtual bool HasError() const { return false; }
+	ENGINE_API virtual bool HasError() const;
 
 protected:
 	/** Internal override implemented by subclasses. */
@@ -170,6 +170,7 @@ public:
 	virtual const FSoundWaveProxyPtr& GetStreamingSoundWave() { return StreamingSoundWave; }
 
 protected:
+	mutable bool bHasError = false;
 	FSoundWaveProxyPtr StreamingSoundWave;
 };
 
diff --git a/Engine/Source/Runtime/Engine/Public/VorbisAudioInfo.h b/Engine/Source/Runtime/Engine/Public/VorbisAudioInfo.h
--- a/Engine/Source/Runtime/Engine/Public/VorbisAudioInfo.h
+++ b/Engine/Source/Runtime/Engine/Public/VorbisAudioInfo.h
@@ -104,7 +104,7 @@ public:
 	virtual bool StreamCompressedData(uint8* InDestination, bool bLooping, uint32 BufferSize, int32& OutNumBytesStreamed) override;
 	virtual int32 GetCurrentChunkIndex() const override {return CurrentStreamingChunkIndex;}
 	virtual int32 GetCurrentChunkOffset() const override {return BufferOffset % CurrentStreamingChunksSize;}
-	virtual bool HasError() const override { return bHasError; }
+	virtual bool HasError() const override;
 	// End of ICompressedAudioInfo Interface
 
 protected:
@@ -114,6 +114,7 @@ protected:
 	ENGINE_API int32 GetAudioDataStartOffset() const;
 
 private:
+	using Super = ICompressedAudioInfo; 
 	const uint8* GetLoadedChunk(FSoundWaveProxyPtr InSoundWave, uint32 ChunkIndex, uint32& OutChunkSize);
 
 	struct FVorbisFileWrapper* VFWrapper;
