Exposed DataTable column utils to scripting

* GetDataTableColumnNames to get the raw property name of each column.
* GetDataTableColumnNameFromExportName to convert a friendly export name to its raw column name.

diff --git a/Engine/Source/Runtime/Engine/Classes/Kismet/DataTableFunctionLibrary.h b/Engine/Source/Runtime/Engine/Classes/Kismet/DataTableFunctionLibrary.h
--- a/Engine/Source/Runtime/Engine/Classes/Kismet/DataTableFunctionLibrary.h
+++ b/Engine/Source/Runtime/Engine/Classes/Kismet/DataTableFunctionLibrary.h
@@ -37,12 +37,30 @@ class UDataTableFunctionLibrary : public UBlueprintFunctionLibrary
     
 	// Returns whether or not Table contains a row named RowName
   	UFUNCTION(BlueprintCallable, Category = "DataTable")
- 	static ENGINE_API bool DoesDataTableRowExist(UDataTable* Table, FName RowName);
+ 	static ENGINE_API bool DoesDataTableRowExist(const UDataTable* Table, FName RowName);
     
 	UFUNCTION(BlueprintCallable, Category = "DataTable")
-	static ENGINE_API void GetDataTableRowNames(UDataTable* Table, TArray<FName>& OutRowNames);
+	static ENGINE_API void GetDataTableRowNames(const UDataTable* Table, TArray<FName>& OutRowNames);
 
-	/** Export from the DataTable all the row for one column. Export it as string. The row name is not included. */
+	/**
+	 * Get the name of each column in this DataTable.
+	 * @note These are always the raw property names (@see GetDataTableColumnAsString) rather than the friendly export name that would be used in a CSV/JSON export (@see GetDataTableColumnNameFromExportName).
+	 */
+	UFUNCTION(BlueprintCallable, Category = "DataTable")
+	static ENGINE_API void GetDataTableColumnNames(const UDataTable* Table, TArray<FName>& OutColumnNames);
+
+	/**
+	 * Get the raw property name of a data table column from its friendly export name.
+	 * @return True if a column was found for the friendly name, false otherwise.
+	 */
+	UFUNCTION(BlueprintCallable, Category = "DataTable")
+	static ENGINE_API bool GetDataTableColumnNameFromExportName(const UDataTable* Table, const FString& ColumnExportName, FName& OutColumnName);
+
+	/**
+	 * Export from the DataTable all the row for one column. Export it as string. The row name is not included.
+	 * @see GetDataTableColumnNames.
+	 * @see GetDataTableColumnNameFromExportName.
+	 */
 	UFUNCTION(BlueprintCallable, Category = "DataTable")
 	static ENGINE_API TArray<FString> GetDataTableColumnAsString(const UDataTable* DataTable, FName PropertyName);
 
diff --git a/Engine/Source/Runtime/Engine/Private/DataTableFunctionLibrary.cpp b/Engine/Source/Runtime/Engine/Private/DataTableFunctionLibrary.cpp
--- a/Engine/Source/Runtime/Engine/Private/DataTableFunctionLibrary.cpp
+++ b/Engine/Source/Runtime/Engine/Private/DataTableFunctionLibrary.cpp
@@ -34,7 +34,7 @@ void UDataTableFunctionLibrary::EvaluateCurveTableRow(UCurveTable* CurveTable, F
 	}
 }
 
-bool UDataTableFunctionLibrary::DoesDataTableRowExist(UDataTable* Table, FName RowName)
+bool UDataTableFunctionLibrary::DoesDataTableRowExist(const UDataTable* Table, FName RowName)
 {
 	if (!Table)
 	{
@@ -87,7 +87,7 @@ bool UDataTableFunctionLibrary::GetDataTableRowFromName(UDataTable* Table, FName
 	return false;
 }
 
-void UDataTableFunctionLibrary::GetDataTableRowNames(UDataTable* Table, TArray<FName>& OutRowNames)
+void UDataTableFunctionLibrary::GetDataTableRowNames(const UDataTable* Table, TArray<FName>& OutRowNames)
 {
 	if (Table)
 	{
@@ -99,6 +99,36 @@ void UDataTableFunctionLibrary::GetDataTableRowNames(UDataTable* Table, TArray<F
 	}
 }
 
+void UDataTableFunctionLibrary::GetDataTableColumnNames(const UDataTable* Table, TArray<FName>& OutColumnNames)
+{
+	if (Table && Table->GetRowStruct())
+	{
+		OutColumnNames = DataTableUtils::GetStructPropertyNames(Table->GetRowStruct());
+	}
+	else
+	{
+		OutColumnNames.Empty();
+	}
+}
+
+bool UDataTableFunctionLibrary::GetDataTableColumnNameFromExportName(const UDataTable* Table, const FString& ColumnExportName, FName& OutColumnName)
+{
+	if (Table && Table->GetRowStruct())
+	{
+		for (TFieldIterator<const FProperty> It(Table->GetRowStruct()); It; ++It)
+		{
+			const FString PropertyExportName = DataTableUtils::GetPropertyExportName(*It);
+			if (PropertyExportName == ColumnExportName)
+			{
+				OutColumnName = It->GetFName();
+				return true;
+			}
+		}
+	}
+
+	return false;
+}
+
 
 #if WITH_EDITOR
 bool UDataTableFunctionLibrary::FillDataTableFromCSVString(UDataTable* DataTable, const FString& InString)
diff --git a/Engine/Source/Runtime/Engine/Private/DataTableUtils.cpp b/Engine/Source/Runtime/Engine/Private/DataTableUtils.cpp
--- a/Engine/Source/Runtime/Engine/Private/DataTableUtils.cpp
+++ b/Engine/Source/Runtime/Engine/Private/DataTableUtils.cpp
@@ -370,10 +370,10 @@ FText DataTableUtils::GetPropertyValueAsText(const FProperty* InProp, const uint
 	return Result;
 }
 
-TArray<FName> DataTableUtils::GetStructPropertyNames(UStruct* InStruct)
+TArray<FName> DataTableUtils::GetStructPropertyNames(const UStruct* InStruct)
 {
 	TArray<FName> PropNames;
-	for (TFieldIterator<FProperty> It(InStruct); It; ++It)
+	for (TFieldIterator<const FProperty> It(InStruct); It; ++It)
 	{
 		PropNames.Add(It->GetFName());
 	}
diff --git a/Engine/Source/Runtime/Engine/Public/DataTableUtils.h b/Engine/Source/Runtime/Engine/Public/DataTableUtils.h
--- a/Engine/Source/Runtime/Engine/Public/DataTableUtils.h
+++ b/Engine/Source/Runtime/Engine/Public/DataTableUtils.h
@@ -65,7 +65,7 @@ namespace DataTableUtils
 	/**
 	 * Util to get all property names from a struct.
 	 */
-	ENGINE_API TArray<FName> GetStructPropertyNames(UStruct* InStruct);
+	ENGINE_API TArray<FName> GetStructPropertyNames(const UStruct* InStruct);
 
 	/**
 	 * Util that removes invalid chars and then make an FName.
