CIS fix, missing include
#rnx

diff --git a/Engine/Source/Runtime/Engine/Private/WorldPartition/ContentBundle/ContentBundle.cpp b/Engine/Source/Runtime/Engine/Private/WorldPartition/ContentBundle/ContentBundle.cpp
--- a/Engine/Source/Runtime/Engine/Private/WorldPartition/ContentBundle/ContentBundle.cpp
+++ b/Engine/Source/Runtime/Engine/Private/WorldPartition/ContentBundle/ContentBundle.cpp
@@ -7,6 +7,7 @@
 #include "WorldPartition/WorldPartitionRuntimeHash.h"
 #include "WorldPartition/ContentBundle/ContentBundleLog.h"
 #include "Engine/World.h"
+#include "UObject/Package.h"
 
 #if WITH_EDITOR
 #include "WorldPartition/ContentBundle/ContentBundleWorldSubsystem.h"
