CIS fix, missing include
#rnx

diff --git a/Engine/Source/Runtime/Engine/Private/WorldPartition/WorldPartitionLevelStreamingPolicy.cpp b/Engine/Source/Runtime/Engine/Private/WorldPartition/WorldPartitionLevelStreamingPolicy.cpp
--- a/Engine/Source/Runtime/Engine/Private/WorldPartition/WorldPartitionLevelStreamingPolicy.cpp
+++ b/Engine/Source/Runtime/Engine/Private/WorldPartition/WorldPartitionLevelStreamingPolicy.cpp
@@ -13,6 +13,7 @@
 #include "Engine/Level.h"
 #include "Engine/Canvas.h"
 #include "Engine/LevelStreamingGCHelper.h"
+#include "UObject/Package.h"
 #if WITH_EDITOR
 #include "Misc/PackageName.h"
 #endif
