Exposed Data Table export to scripting

* ExportDataTableToCSVString
* ExportDataTableToCSVFile
* ExportDataTableToJSONString
* ExportDataTableToJSONFile

diff --git a/Engine/Source/Runtime/Engine/Classes/Kismet/DataTableFunctionLibrary.h b/Engine/Source/Runtime/Engine/Classes/Kismet/DataTableFunctionLibrary.h
--- a/Engine/Source/Runtime/Engine/Classes/Kismet/DataTableFunctionLibrary.h
+++ b/Engine/Source/Runtime/Engine/Classes/Kismet/DataTableFunctionLibrary.h
@@ -149,7 +149,6 @@ class UDataTableFunctionLibrary : public UBlueprintFunctionLibrary
 	UFUNCTION(BlueprintCallable, Category = "Editor Scripting | DataTable", DisplayName = "Fill Data Table from JSON String")
 	static ENGINE_API bool FillDataTableFromJSONString(UDataTable* DataTable, const FString& JSONString);
 
-
 	/** 
 	 * Empty and fill a Data Table from JSON file.
 	 * @param	JSONFilePath	The file path of the JSON file.
@@ -158,6 +157,38 @@ class UDataTableFunctionLibrary : public UBlueprintFunctionLibrary
 	UFUNCTION(BlueprintCallable, Category = "Editor Scripting | DataTable", DisplayName = "Fill Data Table from JSON File")
 	static ENGINE_API bool FillDataTableFromJSONFile(UDataTable* DataTable, const FString& JSONFilePath, UScriptStruct* ImportRowStruct = nullptr);
 
+	/** 
+	 * Export a Data Table to CSV string.
+	 * @param	OutCSVString Output representing the contents of a CSV file.
+	 * @return	True if the operation succeeds, check the log for errors if it didn't succeed.
+	 */
+	UFUNCTION(BlueprintCallable, Category = "Editor Scripting | DataTable", DisplayName="Export Data Table to CSV String")
+	static ENGINE_API bool ExportDataTableToCSVString(const UDataTable* DataTable, FString& OutCSVString);
+
+	/** 
+	 * Export a Data Table to CSV file.
+	 * @param	CSVFilePath	The file path of the CSV file to write (output file is UTF-8).
+	 * @return	True if the operation succeeds, check the log for errors if it didn't succeed.
+	 */
+	UFUNCTION(BlueprintCallable, Category = "Editor Scripting | DataTable", DisplayName = "Export Data Table to CSV File")
+	static ENGINE_API bool ExportDataTableToCSVFile(const UDataTable* DataTable, const FString& CSVFilePath);
+
+	/** 
+	 * Export a Data Table to JSON string.
+	 * @param	OutJSONString Output representing the contents of a JSON file.
+	 * @return	True if the operation succeeds, check the log for errors if it didn't succeed.
+	 */
+	UFUNCTION(BlueprintCallable, Category = "Editor Scripting | DataTable", DisplayName = "Export Data Table to JSON String")
+	static ENGINE_API bool ExportDataTableToJSONString(const UDataTable* DataTable, FString& OutJSONString);
+
+	/** 
+	 * Export a Data Table to JSON file.
+	 * @param	JSONFilePath The file path of the JSON file to write (output file is UTF-8).
+	 * @return	True if the operation succeeds, check the log for errors if it didn't succeed.
+	 */
+	UFUNCTION(BlueprintCallable, Category = "Editor Scripting | DataTable", DisplayName = "Export Data Table to JSON File")
+	static ENGINE_API bool ExportDataTableToJSONFile(const UDataTable* DataTable, const FString& JSONFilePath);
+
 	/** Add a row to a Data Table with the provided name and data. */
     UFUNCTION(BlueprintCallable, CustomThunk, Category = "Editor Scripting | DataTable", meta=(AutoCreateRefTerm="RowName", CustomStructureParam="RowData"))
 	static ENGINE_API void AddDataTableRow(UDataTable* const DataTable, const FName& RowName, const FTableRowBase& RowData);
diff --git a/Engine/Source/Runtime/Engine/Private/DataTableFunctionLibrary.cpp b/Engine/Source/Runtime/Engine/Private/DataTableFunctionLibrary.cpp
--- a/Engine/Source/Runtime/Engine/Private/DataTableFunctionLibrary.cpp
+++ b/Engine/Source/Runtime/Engine/Private/DataTableFunctionLibrary.cpp
@@ -9,6 +9,7 @@
 #include "EditorFramework/AssetImportData.h"
 #include "Factories/CSVImportFactory.h"
 #include "HAL/FileManager.h"
+#include "Misc/FileHelper.h"
 #endif //WITH_EDITOR
 
 UDataTableFunctionLibrary::UDataTableFunctionLibrary(const FObjectInitializer& ObjectInitializer)
@@ -218,7 +219,7 @@ bool UDataTableFunctionLibrary::FillDataTableFromJSONFile(UDataTable* DataTable,
 
 	if (!InFilePath.EndsWith(TEXT(".json")))
 	{
-		UE_LOG(LogDataTable, Error, TEXT("FillDataTableFromJSONFile - The file is not a json."));
+		UE_LOG(LogDataTable, Error, TEXT("FillDataTableFromJSONFile - The file is not a JSON file."));
 		return false;
 	}
 
@@ -228,6 +229,64 @@ bool UDataTableFunctionLibrary::FillDataTableFromJSONFile(UDataTable* DataTable,
 	return ImportFactory->ReimportCSV(DataTable) == EReimportResult::Succeeded;
 }
 
+bool UDataTableFunctionLibrary::ExportDataTableToCSVString(const UDataTable* DataTable, FString& OutCSVString)
+{
+	if (!DataTable || !DataTable->GetRowStruct())
+	{
+		UE_LOG(LogDataTable, Error, TEXT("ExportDataTableToCSVString - The DataTable is invalid."));
+		return false;
+	}
+
+	OutCSVString = DataTable->GetTableAsCSV();
+	return true;
+}
+
+bool UDataTableFunctionLibrary::ExportDataTableToCSVFile(const UDataTable* DataTable, const FString& CSVFilePath)
+{
+	if (!DataTable || !DataTable->GetRowStruct())
+	{
+		UE_LOG(LogDataTable, Error, TEXT("ExportDataTableToCSVFile - The DataTable is invalid."));
+		return false;
+	}
+
+	if (!FFileHelper::SaveStringToFile(DataTable->GetTableAsCSV(), *CSVFilePath, FFileHelper::EEncodingOptions::ForceUTF8WithoutBOM))
+	{
+		UE_LOG(LogDataTable, Error, TEXT("ExportDataTableToCSVFile - Failed to write file '%s'."), *CSVFilePath);
+		return false;
+	}
+
+	return true;
+}
+
+bool UDataTableFunctionLibrary::ExportDataTableToJSONString(const UDataTable* DataTable, FString& OutJSONString)
+{
+	if (!DataTable || !DataTable->GetRowStruct())
+	{
+		UE_LOG(LogDataTable, Error, TEXT("ExportDataTableToJSONString - The DataTable is invalid."));
+		return false;
+	}
+
+	OutJSONString = DataTable->GetTableAsJSON();
+	return true;
+}
+
+bool UDataTableFunctionLibrary::ExportDataTableToJSONFile(const UDataTable* DataTable, const FString& JSONFilePath)
+{
+	if (!DataTable || !DataTable->GetRowStruct())
+	{
+		UE_LOG(LogDataTable, Error, TEXT("ExportDataTableToJSONFile - The DataTable is invalid."));
+		return false;
+	}
+
+	if (!FFileHelper::SaveStringToFile(DataTable->GetTableAsJSON(), *JSONFilePath, FFileHelper::EEncodingOptions::ForceUTF8WithoutBOM))
+	{
+		UE_LOG(LogDataTable, Error, TEXT("ExportDataTableToJSONFile - Failed to write file '%s'."), *JSONFilePath);
+		return false;
+	}
+
+	return true;
+}
+
 void UDataTableFunctionLibrary::AddDataTableRow(UDataTable* const DataTable, const FName& RowName, const FTableRowBase& RowData)
 {
 	if (!DataTable)
