Fix actor component created post-beginplay that were getting added to the replicated subobject list despite IsReplicated being false

diff --git a/Engine/Source/Runtime/Engine/Private/Actor.cpp b/Engine/Source/Runtime/Engine/Private/Actor.cpp
--- a/Engine/Source/Runtime/Engine/Private/Actor.cpp
+++ b/Engine/Source/Runtime/Engine/Private/Actor.cpp
@@ -5596,7 +5596,7 @@ void AActor::HandleRegisterComponentWithWorld(UActorComponent* Component)
 
 		// The component was finally initialized, it can now be replicated
 		// Note that if this component does not ask to be initialized, it would have started to be replicated inside AddOwnedComponent.
-		if (bOwnerBeginPlayStarted)
+		if (bOwnerBeginPlayStarted && Component->GetIsReplicated())
 		{
 			AddComponentForReplication(Component);
 		}
diff --git a/Engine/Source/Runtime/Engine/Private/ActorReplication.cpp b/Engine/Source/Runtime/Engine/Private/ActorReplication.cpp
--- a/Engine/Source/Runtime/Engine/Private/ActorReplication.cpp
+++ b/Engine/Source/Runtime/Engine/Private/ActorReplication.cpp
@@ -642,13 +642,16 @@ void AActor::AddComponentForReplication(UActorComponent* Component)
 		return;
 	}
 
+	const ELifetimeCondition NetCondition = AllowActorComponentToReplicate(Component);
+
+	// Check if the UCS component was built from a template that's not replicated to the client.
 	if (Component->CreationMethod == EComponentCreationMethod::UserConstructionScript && !Component->IsNameStableForNetworking())
 	{
 		ensureMsgf(Component->GetArchetype() == Component->GetClass()->GetDefaultObject(), TEXT("Replicated component %s::%s was added dynamically outside the construction script. This is not well supported and the component on the client will be initialized using the wrong archetype."),
 			*GetName(), *Component->GetName());
 	}
 
-	const ELifetimeCondition NetCondition = AllowActorComponentToReplicate(Component);
+	ensureMsgf(ReplicatedComponents.Find(Component)!=INDEX_NONE, TEXT("AActor::AddComponentForReplication %s::%s but the componenbt was not found in the ReplicatedComponents list."), *GetName(), *Component->GetName());
 
 	FReplicatedComponentInfo* ComponentInfo = ReplicatedComponentsInfo.FindByKey(Component);
 	if (!ComponentInfo)
