Allow all FillDataTableFromX functions to provide an optional row struct override for automated imports

diff --git a/Engine/Source/Runtime/Engine/Classes/Kismet/DataTableFunctionLibrary.h b/Engine/Source/Runtime/Engine/Classes/Kismet/DataTableFunctionLibrary.h
--- a/Engine/Source/Runtime/Engine/Classes/Kismet/DataTableFunctionLibrary.h
+++ b/Engine/Source/Runtime/Engine/Classes/Kismet/DataTableFunctionLibrary.h
@@ -131,7 +131,7 @@ class UDataTableFunctionLibrary : public UBlueprintFunctionLibrary
 	 * @return	True if the operation succeeds, check the log for errors if it didn't succeed.
 	 */
 	UFUNCTION(BlueprintCallable, Category = "Editor Scripting | DataTable", DisplayName="Fill Data Table from CSV String")
-	static ENGINE_API bool FillDataTableFromCSVString(UDataTable* DataTable, const FString& CSVString);
+	static ENGINE_API bool FillDataTableFromCSVString(UDataTable* DataTable, const FString& CSVString, UScriptStruct* ImportRowStruct = nullptr);
 
 	/** 
 	 * Empty and fill a Data Table from CSV file.
@@ -139,7 +139,7 @@ class UDataTableFunctionLibrary : public UBlueprintFunctionLibrary
 	 * @return	True if the operation succeeds, check the log for errors if it didn't succeed.
 	 */
 	UFUNCTION(BlueprintCallable, Category = "Editor Scripting | DataTable", DisplayName = "Fill Data Table from CSV File")
-	static ENGINE_API bool FillDataTableFromCSVFile(UDataTable* DataTable, const FString& CSVFilePath);
+	static ENGINE_API bool FillDataTableFromCSVFile(UDataTable* DataTable, const FString& CSVFilePath, UScriptStruct* ImportRowStruct = nullptr);
 
 	/** 
 	 * Empty and fill a Data Table from JSON string.
@@ -147,7 +147,7 @@ class UDataTableFunctionLibrary : public UBlueprintFunctionLibrary
 	 * @return	True if the operation succeeds, check the log for errors if it didn't succeed.
 	 */
 	UFUNCTION(BlueprintCallable, Category = "Editor Scripting | DataTable", DisplayName = "Fill Data Table from JSON String")
-	static ENGINE_API bool FillDataTableFromJSONString(UDataTable* DataTable, const FString& JSONString);
+	static ENGINE_API bool FillDataTableFromJSONString(UDataTable* DataTable, const FString& JSONString, UScriptStruct* ImportRowStruct = nullptr);
 
 	/** 
 	 * Empty and fill a Data Table from JSON file.
diff --git a/Engine/Source/Runtime/Engine/Private/DataTableFunctionLibrary.cpp b/Engine/Source/Runtime/Engine/Private/DataTableFunctionLibrary.cpp
--- a/Engine/Source/Runtime/Engine/Private/DataTableFunctionLibrary.cpp
+++ b/Engine/Source/Runtime/Engine/Private/DataTableFunctionLibrary.cpp
@@ -132,7 +132,7 @@ bool UDataTableFunctionLibrary::GetDataTableColumnNameFromExportName(const UData
 
 
 #if WITH_EDITOR
-bool UDataTableFunctionLibrary::FillDataTableFromCSVString(UDataTable* DataTable, const FString& InString)
+bool UDataTableFunctionLibrary::FillDataTableFromCSVString(UDataTable* DataTable, const FString& InString, UScriptStruct* ImportRowStruct)
 {
 	if (!DataTable)
 	{
@@ -141,6 +141,7 @@ bool UDataTableFunctionLibrary::FillDataTableFromCSVString(UDataTable* DataTable
 	}
 
 	UCSVImportFactory* ImportFactory = NewObject<UCSVImportFactory>();
+	ImportFactory->AutomatedImportSettings.ImportRowStruct = ImportRowStruct;
 
 	bool bWasCancelled = false;
 	const TCHAR* Buffer = *InString;
@@ -158,7 +159,7 @@ bool UDataTableFunctionLibrary::FillDataTableFromCSVString(UDataTable* DataTable
 	return Result != nullptr && !bWasCancelled;
 }
 
-bool UDataTableFunctionLibrary::FillDataTableFromCSVFile(UDataTable* DataTable, const FString& InFilePath)
+bool UDataTableFunctionLibrary::FillDataTableFromCSVFile(UDataTable* DataTable, const FString& InFilePath, UScriptStruct* ImportRowStruct)
 {
 	if (!DataTable)
 	{
@@ -174,10 +175,11 @@ bool UDataTableFunctionLibrary::FillDataTableFromCSVFile(UDataTable* DataTable,
 
 	DataTable->AssetImportData->Update(InFilePath);
 	UCSVImportFactory* ImportFactory = NewObject<UCSVImportFactory>();
+	ImportFactory->AutomatedImportSettings.ImportRowStruct = ImportRowStruct;
 	return ImportFactory->ReimportCSV(DataTable) == EReimportResult::Succeeded;
 }
 
-bool UDataTableFunctionLibrary::FillDataTableFromJSONString(UDataTable* DataTable, const FString& InString)
+bool UDataTableFunctionLibrary::FillDataTableFromJSONString(UDataTable* DataTable, const FString& InString, UScriptStruct* ImportRowStruct)
 {
 	if (!DataTable)
 	{
@@ -186,6 +188,7 @@ bool UDataTableFunctionLibrary::FillDataTableFromJSONString(UDataTable* DataTabl
 	}
 
 	UCSVImportFactory* ImportFactory = NewObject<UCSVImportFactory>();
+	ImportFactory->AutomatedImportSettings.ImportRowStruct = ImportRowStruct;
 
 	bool bWasCancelled = false;
 	const TCHAR* Buffer = *InString;
