Data Table automation improvements

* Always consider the FillDataTableFromX functions to be running an automated import when the optional ImportRowStruct is provided.
* Added GetDataTableRowStruct to query the row struct of an existing data table.

diff --git a/Engine/Source/Editor/UnrealEd/Classes/Factories/CSVImportFactory.h b/Engine/Source/Editor/UnrealEd/Classes/Factories/CSVImportFactory.h
--- a/Engine/Source/Editor/UnrealEd/Classes/Factories/CSVImportFactory.h
+++ b/Engine/Source/Editor/UnrealEd/Classes/Factories/CSVImportFactory.h
@@ -48,6 +48,9 @@ struct FCSVImportSettings
 	UPROPERTY(BlueprintReadWrite, Category="Misc")
 	TEnumAsByte<ERichCurveInterpMode> ImportCurveInterpMode;
 
+	/** True to force IsAutomatedImport to return true during the import */
+	bool bForceAutomatedImport = false;
+
 	/** Despite its name, DataToImport can be JSON instead of CSV if this bool is set */
 	bool bDataIsJson = false;
 	FString DataToImport;
@@ -60,6 +63,7 @@ class UCSVImportFactory : public UFactory, public IImportSettingsParser
 
 public:
 	//~ Begin UFactory Interface
+	UNREALED_API virtual bool IsAutomatedImport() const override;
 	UNREALED_API virtual FText GetDisplayName() const override;
 	UNREALED_API virtual UObject* FactoryCreateFile(UClass* InClass, UObject* InParent, FName InName, EObjectFlags Flags,
 		const FString& Filename, const TCHAR* Parms, FFeedbackContext* Warn, bool& bOutOperationCanceled) override;
diff --git a/Engine/Source/Editor/UnrealEd/Private/Factories/CSVImportFactory.cpp b/Engine/Source/Editor/UnrealEd/Private/Factories/CSVImportFactory.cpp
--- a/Engine/Source/Editor/UnrealEd/Private/Factories/CSVImportFactory.cpp
+++ b/Engine/Source/Editor/UnrealEd/Private/Factories/CSVImportFactory.cpp
@@ -77,6 +77,12 @@ UCSVImportFactory::UCSVImportFactory(const FObjectInitializer& ObjectInitializer
 	Formats.Add(TEXT("csv;Comma-separated values"));
 }
 
+bool UCSVImportFactory::IsAutomatedImport() const
+{
+	return Super::IsAutomatedImport()
+		|| AutomatedImportSettings.bForceAutomatedImport;
+}
+
 FText UCSVImportFactory::GetDisplayName() const
 {
 	return LOCTEXT("CSVImportFactoryDescription", "Comma Separated Values");
diff --git a/Engine/Source/Runtime/Engine/Classes/Kismet/DataTableFunctionLibrary.h b/Engine/Source/Runtime/Engine/Classes/Kismet/DataTableFunctionLibrary.h
--- a/Engine/Source/Runtime/Engine/Classes/Kismet/DataTableFunctionLibrary.h
+++ b/Engine/Source/Runtime/Engine/Classes/Kismet/DataTableFunctionLibrary.h
@@ -35,6 +35,10 @@ class UDataTableFunctionLibrary : public UBlueprintFunctionLibrary
 	UFUNCTION(BlueprintCallable, Category = "DataTable", meta = (ExpandEnumAsExecs="OutResult", DataTablePin="CurveTable"))
 	static ENGINE_API void EvaluateCurveTableRow(UCurveTable* CurveTable, FName RowName, float InXY, TEnumAsByte<EEvaluateCurveTableResult::Type>& OutResult, float& OutXY,const FString& ContextString);
     
+	/** Get the row struct used by the given Data Table, if any */
+	UFUNCTION(BlueprintPure, Category = "DataTable")
+ 	static ENGINE_API const UScriptStruct* GetDataTableRowStruct(const UDataTable* Table);
+
 	// Returns whether or not Table contains a row named RowName
   	UFUNCTION(BlueprintCallable, Category = "DataTable")
  	static ENGINE_API bool DoesDataTableRowExist(const UDataTable* Table, FName RowName);
@@ -43,7 +47,7 @@ class UDataTableFunctionLibrary : public UBlueprintFunctionLibrary
 	static ENGINE_API void GetDataTableRowNames(const UDataTable* Table, TArray<FName>& OutRowNames);
 
 	/**
-	 * Get the name of each column in this DataTable.
+	 * Get the name of each column in this Data Table.
 	 * @note These are always the raw property names (@see GetDataTableColumnAsString) rather than the friendly export name that would be used in a CSV/JSON export (@see GetDataTableColumnNameFromExportName).
 	 */
 	UFUNCTION(BlueprintCallable, Category = "DataTable")
@@ -127,34 +131,38 @@ class UDataTableFunctionLibrary : public UBlueprintFunctionLibrary
 #if WITH_EDITOR
 	/** 
 	 * Empty and fill a Data Table from CSV string.
-	 * @param	CSVString	The Data that representing the contents of a CSV file.
+	 * @param	CSVString			The Data that representing the contents of a CSV file.
+	 * @param	ImportRowStruct		Optional row struct to apply on import. If set will also force the import to run automated (no questions or dialogs).
 	 * @return	True if the operation succeeds, check the log for errors if it didn't succeed.
 	 */
-	UFUNCTION(BlueprintCallable, Category = "Editor Scripting | DataTable", DisplayName="Fill Data Table from CSV String")
+	UFUNCTION(BlueprintCallable, Category = "Editor Scripting | DataTable", DisplayName="Fill Data Table from CSV String", meta=(AdvancedDisplay="ImportRowStruct"))
 	static ENGINE_API bool FillDataTableFromCSVString(UDataTable* DataTable, const FString& CSVString, UScriptStruct* ImportRowStruct = nullptr);
 
 	/** 
 	 * Empty and fill a Data Table from CSV file.
-	 * @param	CSVFilePath	The file path of the CSV file.
+	 * @param	CSVFilePath			The file path of the CSV file.
+	 * @param	ImportRowStruct		Optional row struct to apply on import. If set will also force the import to run automated (no questions or dialogs).
 	 * @return	True if the operation succeeds, check the log for errors if it didn't succeed.
 	 */
-	UFUNCTION(BlueprintCallable, Category = "Editor Scripting | DataTable", DisplayName = "Fill Data Table from CSV File")
+	UFUNCTION(BlueprintCallable, Category = "Editor Scripting | DataTable", DisplayName = "Fill Data Table from CSV File", meta=(AdvancedDisplay="ImportRowStruct"))
 	static ENGINE_API bool FillDataTableFromCSVFile(UDataTable* DataTable, const FString& CSVFilePath, UScriptStruct* ImportRowStruct = nullptr);
 
 	/** 
 	 * Empty and fill a Data Table from JSON string.
-	 * @param	JSONString	The Data that representing the contents of a JSON file.
+	 * @param	JSONString			The Data that representing the contents of a JSON file.
+	 * @param	ImportRowStruct		Optional row struct to apply on import. If set will also force the import to run automated (no questions or dialogs).
 	 * @return	True if the operation succeeds, check the log for errors if it didn't succeed.
 	 */
-	UFUNCTION(BlueprintCallable, Category = "Editor Scripting | DataTable", DisplayName = "Fill Data Table from JSON String")
+	UFUNCTION(BlueprintCallable, Category = "Editor Scripting | DataTable", DisplayName = "Fill Data Table from JSON String", meta=(AdvancedDisplay="ImportRowStruct"))
 	static ENGINE_API bool FillDataTableFromJSONString(UDataTable* DataTable, const FString& JSONString, UScriptStruct* ImportRowStruct = nullptr);
 
 	/** 
 	 * Empty and fill a Data Table from JSON file.
-	 * @param	JSONFilePath	The file path of the JSON file.
+	 * @param	JSONFilePath		The file path of the JSON file.
+	 * @param	ImportRowStruct		Optional row struct to apply on import. If set will also force the import to run automated (no questions or dialogs).
 	 * @return	True if the operation succeeds, check the log for errors if it didn't succeed.
 	 */
-	UFUNCTION(BlueprintCallable, Category = "Editor Scripting | DataTable", DisplayName = "Fill Data Table from JSON File")
+	UFUNCTION(BlueprintCallable, Category = "Editor Scripting | DataTable", DisplayName = "Fill Data Table from JSON File", meta=(AdvancedDisplay="ImportRowStruct"))
 	static ENGINE_API bool FillDataTableFromJSONFile(UDataTable* DataTable, const FString& JSONFilePath, UScriptStruct* ImportRowStruct = nullptr);
 
 	/** 
diff --git a/Engine/Source/Runtime/Engine/Private/DataTableFunctionLibrary.cpp b/Engine/Source/Runtime/Engine/Private/DataTableFunctionLibrary.cpp
--- a/Engine/Source/Runtime/Engine/Private/DataTableFunctionLibrary.cpp
+++ b/Engine/Source/Runtime/Engine/Private/DataTableFunctionLibrary.cpp
@@ -35,6 +35,13 @@ void UDataTableFunctionLibrary::EvaluateCurveTableRow(UCurveTable* CurveTable, F
 	}
 }
 
+const UScriptStruct* UDataTableFunctionLibrary::GetDataTableRowStruct(const UDataTable* Table)
+{
+	return Table
+		? Table->GetRowStruct()
+		: nullptr;
+}
+
 bool UDataTableFunctionLibrary::DoesDataTableRowExist(const UDataTable* Table, FName RowName)
 {
 	if (!Table)
@@ -141,6 +148,7 @@ bool UDataTableFunctionLibrary::FillDataTableFromCSVString(UDataTable* DataTable
 	}
 
 	UCSVImportFactory* ImportFactory = NewObject<UCSVImportFactory>();
+	ImportFactory->AutomatedImportSettings.bForceAutomatedImport = ImportRowStruct != nullptr;
 	ImportFactory->AutomatedImportSettings.ImportRowStruct = ImportRowStruct;
 
 	bool bWasCancelled = false;
@@ -174,8 +182,11 @@ bool UDataTableFunctionLibrary::FillDataTableFromCSVFile(UDataTable* DataTable,
 	}
 
 	DataTable->AssetImportData->Update(InFilePath);
+
 	UCSVImportFactory* ImportFactory = NewObject<UCSVImportFactory>();
+	ImportFactory->AutomatedImportSettings.bForceAutomatedImport = ImportRowStruct != nullptr;
 	ImportFactory->AutomatedImportSettings.ImportRowStruct = ImportRowStruct;
+	
 	return ImportFactory->ReimportCSV(DataTable) == EReimportResult::Succeeded;
 }
 
@@ -188,6 +199,7 @@ bool UDataTableFunctionLibrary::FillDataTableFromJSONString(UDataTable* DataTabl
 	}
 
 	UCSVImportFactory* ImportFactory = NewObject<UCSVImportFactory>();
+	ImportFactory->AutomatedImportSettings.bForceAutomatedImport = ImportRowStruct != nullptr;
 	ImportFactory->AutomatedImportSettings.ImportRowStruct = ImportRowStruct;
 
 	bool bWasCancelled = false;
@@ -227,8 +239,11 @@ bool UDataTableFunctionLibrary::FillDataTableFromJSONFile(UDataTable* DataTable,
 	}
 
 	DataTable->AssetImportData->Update(InFilePath);
+
 	UCSVImportFactory* ImportFactory = NewObject<UCSVImportFactory>();
+	ImportFactory->AutomatedImportSettings.bForceAutomatedImport = ImportRowStruct != nullptr;
 	ImportFactory->AutomatedImportSettings.ImportRowStruct = ImportRowStruct;
+	
 	return ImportFactory->ReimportCSV(DataTable) == EReimportResult::Succeeded;
 }
 
