Update tooltip on USkeleton CompatibleSkeletons

diff --git a/Engine/Source/Runtime/Engine/Classes/Animation/Skeleton.h b/Engine/Source/Runtime/Engine/Classes/Animation/Skeleton.h
--- a/Engine/Source/Runtime/Engine/Classes/Animation/Skeleton.h
+++ b/Engine/Source/Runtime/Engine/Classes/Animation/Skeleton.h
@@ -324,7 +324,9 @@ protected:
 	TArray<FVirtualBone> VirtualBones;
 
 	/**
-	 * The list of compatible skeletons.
+	 * The list of compatible skeletons. This skeleton will be able to use animation data originating from skeletons within this array, such as animation sequences.
+	 * This property is not bi-directional.
+	 * 
 	 * This is an array of TSoftObjectPtr in order to prevent all skeletons to be loaded, as we only want to load things on demand.
 	 * As this is EditAnywhere and an array of TSoftObjectPtr, checking validity of pointers is needed.
 	 **/
