Exposed DataTable column export names to scripting

* GetDataTableColumnExportNames

diff --git a/Engine/Source/Runtime/Engine/Classes/Kismet/DataTableFunctionLibrary.h b/Engine/Source/Runtime/Engine/Classes/Kismet/DataTableFunctionLibrary.h
--- a/Engine/Source/Runtime/Engine/Classes/Kismet/DataTableFunctionLibrary.h
+++ b/Engine/Source/Runtime/Engine/Classes/Kismet/DataTableFunctionLibrary.h
@@ -53,6 +53,13 @@ class UDataTableFunctionLibrary : public UBlueprintFunctionLibrary
 	UFUNCTION(BlueprintCallable, Category = "DataTable")
 	static ENGINE_API void GetDataTableColumnNames(const UDataTable* Table, TArray<FName>& OutColumnNames);
 
+	/**
+	 * Get the friendly export name of each column in this Data Table.
+	 * @see GetDataTableColumnNameFromExportName.
+	 */
+	UFUNCTION(BlueprintCallable, Category = "DataTable")
+	static ENGINE_API void GetDataTableColumnExportNames(const UDataTable* Table, TArray<FString>& OutExportColumnNames);
+
 	/**
 	 * Get the raw property name of a data table column from its friendly export name.
 	 * @return True if a column was found for the friendly name, false otherwise.
diff --git a/Engine/Source/Runtime/Engine/Private/DataTableFunctionLibrary.cpp b/Engine/Source/Runtime/Engine/Private/DataTableFunctionLibrary.cpp
--- a/Engine/Source/Runtime/Engine/Private/DataTableFunctionLibrary.cpp
+++ b/Engine/Source/Runtime/Engine/Private/DataTableFunctionLibrary.cpp
@@ -119,6 +119,19 @@ void UDataTableFunctionLibrary::GetDataTableColumnNames(const UDataTable* Table,
 	}
 }
 
+void UDataTableFunctionLibrary::GetDataTableColumnExportNames(const UDataTable* Table, TArray<FString>& OutExportColumnNames)
+{
+	OutExportColumnNames.Empty();
+
+	if (Table && Table->GetRowStruct())
+	{
+		for (TFieldIterator<const FProperty> It(Table->GetRowStruct()); It; ++It)
+		{
+			OutExportColumnNames.Add(DataTableUtils::GetPropertyExportName(*It));
+		}
+	}
+}
+
 bool UDataTableFunctionLibrary::GetDataTableColumnNameFromExportName(const UDataTable* Table, const FString& ColumnExportName, FName& OutColumnName)
 {
 	if (Table && Table->GetRowStruct())
