Use different name for distributed ShaderCompilingThread.

#rnx

diff --git a/Engine/Source/Runtime/Engine/Private/ShaderCompiler/ShaderCompiler.cpp b/Engine/Source/Runtime/Engine/Private/ShaderCompiler/ShaderCompiler.cpp
--- a/Engine/Source/Runtime/Engine/Private/ShaderCompiler/ShaderCompiler.cpp
+++ b/Engine/Source/Runtime/Engine/Private/ShaderCompiler/ShaderCompiler.cpp
@@ -3440,7 +3440,7 @@ void FShaderCompileThreadRunnableBase::StartThread()
 {
 	if (Manager->bAllowAsynchronousShaderCompiling && !FPlatformProperties::RequiresCookedData())
 	{
-		Thread = FRunnableThread::Create(this, TEXT("ShaderCompilingThread"), 0, TPri_Normal, FPlatformAffinity::GetPoolThreadMask());
+		Thread = FRunnableThread::Create(this, GetThreadName(), 0, TPri_Normal, FPlatformAffinity::GetPoolThreadMask());
 	}
 }
 
diff --git a/Engine/Source/Runtime/Engine/Private/ShaderCompiler/ShaderCompilerDistributed.cpp b/Engine/Source/Runtime/Engine/Private/ShaderCompiler/ShaderCompilerDistributed.cpp
--- a/Engine/Source/Runtime/Engine/Private/ShaderCompiler/ShaderCompilerDistributed.cpp
+++ b/Engine/Source/Runtime/Engine/Private/ShaderCompiler/ShaderCompilerDistributed.cpp
@@ -398,3 +398,8 @@ int32 FShaderCompileDistributedThreadRunnable_Interface::CompilingLoop()
 	// Return true if there is more work to be done.
 	return Manager->AllJobs.GetNumOutstandingJobs() > 0;
 }
+
+const TCHAR* FShaderCompileDistributedThreadRunnable_Interface::GetThreadName() const
+{
+	return TEXT("ShaderCompilingThread-Distributed");
+}
diff --git a/Engine/Source/Runtime/Engine/Public/ShaderCompiler.h b/Engine/Source/Runtime/Engine/Public/ShaderCompiler.h
--- a/Engine/Source/Runtime/Engine/Public/ShaderCompiler.h
+++ b/Engine/Source/Runtime/Engine/Public/ShaderCompiler.h
@@ -130,6 +130,12 @@ protected:
 
 	virtual void PrintWorkerMemoryUsage() {}
 
+	// Returns a name for this thread instance. Defaults to "ShaderCompilingThread".
+	virtual const TCHAR* GetThreadName() const
+	{
+		return TEXT("ShaderCompilingThread");
+	}
+
 public:
 	FShaderCompileThreadRunnableBase(class FShaderCompilingManager* InManager);
 	virtual ~FShaderCompileThreadRunnableBase()
@@ -280,6 +286,8 @@ private:
 
 	TArray<FString> GetDependencyFilesForJobs(TArray<FShaderCommonCompileJobPtr>& Jobs);
 	void DispatchShaderCompileJobsBatch(TArray<FShaderCommonCompileJobPtr>& JobsToSerialize);
+
+	virtual const TCHAR* GetThreadName() const override;
 };
 
 /** Results for a single compiled and finalized shader map. */
