Added ScriptName aliases for DataTableFunctionLibrary functions

This lets them be used by Python directly on the DataTable object instance

diff --git a/Engine/Source/Runtime/Engine/Classes/Kismet/DataTableFunctionLibrary.h b/Engine/Source/Runtime/Engine/Classes/Kismet/DataTableFunctionLibrary.h
--- a/Engine/Source/Runtime/Engine/Classes/Kismet/DataTableFunctionLibrary.h
+++ b/Engine/Source/Runtime/Engine/Classes/Kismet/DataTableFunctionLibrary.h
@@ -36,35 +36,35 @@ class UDataTableFunctionLibrary : public UBlueprintFunctionLibrary
 	static ENGINE_API void EvaluateCurveTableRow(UCurveTable* CurveTable, FName RowName, float InXY, TEnumAsByte<EEvaluateCurveTableResult::Type>& OutResult, float& OutXY,const FString& ContextString);
     
 	/** Get the row struct used by the given Data Table, if any */
-	UFUNCTION(BlueprintPure, Category = "DataTable")
+	UFUNCTION(BlueprintPure, Category = "DataTable", meta=(ScriptMethod="GetRowStruct"))
  	static ENGINE_API const UScriptStruct* GetDataTableRowStruct(const UDataTable* Table);
 
 	// Returns whether or not Table contains a row named RowName
-  	UFUNCTION(BlueprintCallable, Category = "DataTable")
+  	UFUNCTION(BlueprintCallable, Category = "DataTable", meta=(ScriptMethod="DoesRowExist"))
  	static ENGINE_API bool DoesDataTableRowExist(const UDataTable* Table, FName RowName);
     
-	UFUNCTION(BlueprintCallable, Category = "DataTable")
+	UFUNCTION(BlueprintCallable, Category = "DataTable", meta=(ScriptMethod="GetRowNames"))
 	static ENGINE_API void GetDataTableRowNames(const UDataTable* Table, TArray<FName>& OutRowNames);
 
 	/**
 	 * Get the name of each column in this Data Table.
 	 * @note These are always the raw property names (@see GetDataTableColumnAsString) rather than the friendly export name that would be used in a CSV/JSON export (@see GetDataTableColumnNameFromExportName).
 	 */
-	UFUNCTION(BlueprintCallable, Category = "DataTable")
+	UFUNCTION(BlueprintCallable, Category = "DataTable", meta=(ScriptMethod="GetColumnNames"))
 	static ENGINE_API void GetDataTableColumnNames(const UDataTable* Table, TArray<FName>& OutColumnNames);
 
 	/**
 	 * Get the friendly export name of each column in this Data Table.
 	 * @see GetDataTableColumnNameFromExportName.
 	 */
-	UFUNCTION(BlueprintCallable, Category = "DataTable")
+	UFUNCTION(BlueprintCallable, Category = "DataTable", meta=(ScriptMethod="GetColumnExportNames"))
 	static ENGINE_API void GetDataTableColumnExportNames(const UDataTable* Table, TArray<FString>& OutExportColumnNames);
 
 	/**
 	 * Get the raw property name of a data table column from its friendly export name.
 	 * @return True if a column was found for the friendly name, false otherwise.
 	 */
-	UFUNCTION(BlueprintCallable, Category = "DataTable")
+	UFUNCTION(BlueprintCallable, Category = "DataTable", meta=(ScriptMethod="GetColumnNameFromExportName"))
 	static ENGINE_API bool GetDataTableColumnNameFromExportName(const UDataTable* Table, const FString& ColumnExportName, FName& OutColumnName);
 
 	/**
@@ -72,7 +72,7 @@ class UDataTableFunctionLibrary : public UBlueprintFunctionLibrary
 	 * @see GetDataTableColumnNames.
 	 * @see GetDataTableColumnNameFromExportName.
 	 */
-	UFUNCTION(BlueprintCallable, Category = "DataTable")
+	UFUNCTION(BlueprintCallable, Category = "DataTable", meta=(ScriptMethod="GetColumnAsString"))
 	static ENGINE_API TArray<FString> GetDataTableColumnAsString(const UDataTable* DataTable, FName PropertyName);
 
     /** Get a Row from a DataTable given a RowName */
@@ -142,7 +142,7 @@ class UDataTableFunctionLibrary : public UBlueprintFunctionLibrary
 	 * @param	ImportRowStruct		Optional row struct to apply on import. If set will also force the import to run automated (no questions or dialogs).
 	 * @return	True if the operation succeeds, check the log for errors if it didn't succeed.
 	 */
-	UFUNCTION(BlueprintCallable, Category = "Editor Scripting | DataTable", DisplayName="Fill Data Table from CSV String", meta=(AdvancedDisplay="ImportRowStruct"))
+	UFUNCTION(BlueprintCallable, Category = "Editor Scripting | DataTable", DisplayName="Fill Data Table from CSV String", meta=(AdvancedDisplay="ImportRowStruct", ScriptMethod="FillFromCSVString"))
 	static ENGINE_API bool FillDataTableFromCSVString(UDataTable* DataTable, const FString& CSVString, UScriptStruct* ImportRowStruct = nullptr);
 
 	/** 
@@ -151,7 +151,7 @@ class UDataTableFunctionLibrary : public UBlueprintFunctionLibrary
 	 * @param	ImportRowStruct		Optional row struct to apply on import. If set will also force the import to run automated (no questions or dialogs).
 	 * @return	True if the operation succeeds, check the log for errors if it didn't succeed.
 	 */
-	UFUNCTION(BlueprintCallable, Category = "Editor Scripting | DataTable", DisplayName = "Fill Data Table from CSV File", meta=(AdvancedDisplay="ImportRowStruct"))
+	UFUNCTION(BlueprintCallable, Category = "Editor Scripting | DataTable", DisplayName = "Fill Data Table from CSV File", meta=(AdvancedDisplay="ImportRowStruct", ScriptMethod="FillFromCSVFile"))
 	static ENGINE_API bool FillDataTableFromCSVFile(UDataTable* DataTable, const FString& CSVFilePath, UScriptStruct* ImportRowStruct = nullptr);
 
 	/** 
@@ -160,7 +160,7 @@ class UDataTableFunctionLibrary : public UBlueprintFunctionLibrary
 	 * @param	ImportRowStruct		Optional row struct to apply on import. If set will also force the import to run automated (no questions or dialogs).
 	 * @return	True if the operation succeeds, check the log for errors if it didn't succeed.
 	 */
-	UFUNCTION(BlueprintCallable, Category = "Editor Scripting | DataTable", DisplayName = "Fill Data Table from JSON String", meta=(AdvancedDisplay="ImportRowStruct"))
+	UFUNCTION(BlueprintCallable, Category = "Editor Scripting | DataTable", DisplayName = "Fill Data Table from JSON String", meta=(AdvancedDisplay="ImportRowStruct", ScriptMethod="FillFromJSONString"))
 	static ENGINE_API bool FillDataTableFromJSONString(UDataTable* DataTable, const FString& JSONString, UScriptStruct* ImportRowStruct = nullptr);
 
 	/** 
@@ -169,7 +169,7 @@ class UDataTableFunctionLibrary : public UBlueprintFunctionLibrary
 	 * @param	ImportRowStruct		Optional row struct to apply on import. If set will also force the import to run automated (no questions or dialogs).
 	 * @return	True if the operation succeeds, check the log for errors if it didn't succeed.
 	 */
-	UFUNCTION(BlueprintCallable, Category = "Editor Scripting | DataTable", DisplayName = "Fill Data Table from JSON File", meta=(AdvancedDisplay="ImportRowStruct"))
+	UFUNCTION(BlueprintCallable, Category = "Editor Scripting | DataTable", DisplayName = "Fill Data Table from JSON File", meta=(AdvancedDisplay="ImportRowStruct", ScriptMethod="FillFromJSONFile"))
 	static ENGINE_API bool FillDataTableFromJSONFile(UDataTable* DataTable, const FString& JSONFilePath, UScriptStruct* ImportRowStruct = nullptr);
 
 	/** 
@@ -177,7 +177,7 @@ class UDataTableFunctionLibrary : public UBlueprintFunctionLibrary
 	 * @param	OutCSVString Output representing the contents of a CSV file.
 	 * @return	True if the operation succeeds, check the log for errors if it didn't succeed.
 	 */
-	UFUNCTION(BlueprintCallable, Category = "Editor Scripting | DataTable", DisplayName="Export Data Table to CSV String")
+	UFUNCTION(BlueprintCallable, Category = "Editor Scripting | DataTable", DisplayName="Export Data Table to CSV String", meta=(ScriptMethod="ExportToCSVString"))
 	static ENGINE_API bool ExportDataTableToCSVString(const UDataTable* DataTable, FString& OutCSVString);
 
 	/** 
@@ -185,7 +185,7 @@ class UDataTableFunctionLibrary : public UBlueprintFunctionLibrary
 	 * @param	CSVFilePath	The file path of the CSV file to write (output file is UTF-8).
 	 * @return	True if the operation succeeds, check the log for errors if it didn't succeed.
 	 */
-	UFUNCTION(BlueprintCallable, Category = "Editor Scripting | DataTable", DisplayName = "Export Data Table to CSV File")
+	UFUNCTION(BlueprintCallable, Category = "Editor Scripting | DataTable", DisplayName = "Export Data Table to CSV File", meta=(ScriptMethod="ExportToCSVFile"))
 	static ENGINE_API bool ExportDataTableToCSVFile(const UDataTable* DataTable, const FString& CSVFilePath);
 
 	/** 
@@ -193,7 +193,7 @@ class UDataTableFunctionLibrary : public UBlueprintFunctionLibrary
 	 * @param	OutJSONString Output representing the contents of a JSON file.
 	 * @return	True if the operation succeeds, check the log for errors if it didn't succeed.
 	 */
-	UFUNCTION(BlueprintCallable, Category = "Editor Scripting | DataTable", DisplayName = "Export Data Table to JSON String")
+	UFUNCTION(BlueprintCallable, Category = "Editor Scripting | DataTable", DisplayName = "Export Data Table to JSON String", meta=(ScriptMethod="ExportToJSONString"))
 	static ENGINE_API bool ExportDataTableToJSONString(const UDataTable* DataTable, FString& OutJSONString);
 
 	/** 
@@ -201,7 +201,7 @@ class UDataTableFunctionLibrary : public UBlueprintFunctionLibrary
 	 * @param	JSONFilePath The file path of the JSON file to write (output file is UTF-8).
 	 * @return	True if the operation succeeds, check the log for errors if it didn't succeed.
 	 */
-	UFUNCTION(BlueprintCallable, Category = "Editor Scripting | DataTable", DisplayName = "Export Data Table to JSON File")
+	UFUNCTION(BlueprintCallable, Category = "Editor Scripting | DataTable", DisplayName = "Export Data Table to JSON File", meta=(ScriptMethod="ExportToJSONFile"))
 	static ENGINE_API bool ExportDataTableToJSONFile(const UDataTable* DataTable, const FString& JSONFilePath);
 
 	/** Add a row to a Data Table with the provided name and data. */
