Add s.OptimizeActorRegistration console variable that prevents redundant calls to functions like PostRegisterAllComponents for actors that are spawned during initialization
This defaults to 1 to use the fast path but can be changed to use the old path or enable extra ensures for validation

diff --git a/Engine/Source/Runtime/Engine/Private/Actor.cpp b/Engine/Source/Runtime/Engine/Private/Actor.cpp
--- a/Engine/Source/Runtime/Engine/Private/Actor.cpp
+++ b/Engine/Source/Runtime/Engine/Private/Actor.cpp
@@ -89,6 +89,9 @@ FMakeNoiseDelegate AActor::MakeNoiseDelegate = FMakeNoiseDelegate::CreateStatic(
 /** Selects if actor and actorcomponents will replicate subobjects using the registration list by default. */
 extern bool GDefaultUseSubObjectReplicationList;
 
+/** Enables optimizations to actor and component registration */
+extern int32 GOptimizeActorRegistration;
+
 #if !UE_BUILD_SHIPPING
 FOnProcessEvent AActor::ProcessEventDelegate;
 #endif
@@ -5344,9 +5347,14 @@ void AActor::UnregisterAllComponents(const bool bForReregister)
 		}
 	}
 
-	bHasRegisteredAllComponents = false;
+	if (bHasRegisteredAllComponents || GOptimizeActorRegistration == 0)
+	{
+		// With registration optimizations enabled, we need to make sure it unregisters components that were partially registered during construction,
+		// but we do not want to call PostUnregisterAllComponents if it was only partially registered
+		bHasRegisteredAllComponents = false;
 
-	PostUnregisterAllComponents();
+		PostUnregisterAllComponents();
+	}
 }
 
 void AActor::PostUnregisterAllComponents()
@@ -5358,6 +5366,13 @@ void AActor::PostUnregisterAllComponents()
 
 void AActor::RegisterAllComponents()
 {
+	if (bHasRegisteredAllComponents && GOptimizeActorRegistration > 0)
+	{
+		// Stops it from calling redundant Pre/Post registration functions, should never get here if it was deferred
+		ensureMsgf(!bHasDeferredComponentRegistration, TEXT("Actor %s has both bHasDeferredComponentRegistration and bHasRegisteredAllComponents"), *GetPathName());
+		return;
+	}
+
 	QUICK_SCOPE_CYCLE_COUNTER(STAT_AActor_RegisterAllComponents);
 	
 	PreRegisterAllComponents();
@@ -5365,9 +5380,6 @@ void AActor::RegisterAllComponents()
 	// 0 - means register all components
 	bool bAllRegistered = IncrementalRegisterComponents(0);
 	check(bAllRegistered);
-
-	// Clear this flag as it's no longer deferred
-	bHasDeferredComponentRegistration = false;
 }
 
 // Walks through components hierarchy and returns closest to root parent component that is unregistered
@@ -5477,6 +5489,17 @@ bool AActor::IncrementalRegisterComponents(int32 NumComponentsToRegister, FRegis
 #if PERF_TRACK_DETAILED_ASYNC_STATS
 		QUICK_SCOPE_CYCLE_COUNTER(STAT_AActor_IncrementalRegisterComponents_PostRegisterAllComponents);
 #endif
+
+#if !UE_BUILD_SHIPPING
+		if (GOptimizeActorRegistration == 2)
+		{
+			ensureMsgf(!bHasRegisteredAllComponents, TEXT("PostRegisterAllComponents called twice for %s!"), *GetPathName());
+		}
+#endif
+
+		// Clear this flag as it's no longer deferred
+		bHasDeferredComponentRegistration = false;
+
 		bHasRegisteredAllComponents = true;
 		// Finally, call PostRegisterAllComponents
 		PostRegisterAllComponents();
diff --git a/Engine/Source/Runtime/Engine/Private/Level.cpp b/Engine/Source/Runtime/Engine/Private/Level.cpp
--- a/Engine/Source/Runtime/Engine/Private/Level.cpp
+++ b/Engine/Source/Runtime/Engine/Private/Level.cpp
@@ -92,6 +92,17 @@ static FAutoConsoleVariableRef CVarActorClusteringEnabled(
 	ECVF_Default
 );
 
+int32 GOptimizeActorRegistration = 1;
+static FAutoConsoleVariableRef CVarOptimizeActorRegistration(
+	TEXT("s.OptimizeActorRegistration"),
+	GOptimizeActorRegistration,
+	TEXT("Enables optimizations to actor component registration functions like PostRegisterAllComponents\n")
+	TEXT(" 0: Disable optimizations for legacy code that depends on redundant calls to registration functions\n")
+	TEXT(" 1: Enables optimizations and assumes the code is working properly\n")
+	TEXT(" 2: Enables optimization logic, but ensures it is working properly in non shipping builds"),
+	ECVF_Default
+);
+
 // RouteActorInit for a single actor generally takes about half the time to complete compared to component initialization
 float GRouteActorInitializationWorkUnitWeighting = 0.5f;
 static FAutoConsoleVariableRef CVarRouteActorInitializationWorkUnitWeighting(
@@ -1711,20 +1722,33 @@ bool ULevel::IncrementalRegisterComponents(bool bPreRegisterComponents, int32 Nu
 		bool bAllComponentsRegistered = true;
 		if (IsValid(Actor))
 		{
+			if (!Actor->HasActorRegisteredAllComponents() || GOptimizeActorRegistration == 0)
+			{
 #if PERF_TRACK_DETAILED_ASYNC_STATS
-			FScopeCycleCounterUObject ContextScope(Actor);
+				FScopeCycleCounterUObject ContextScope(Actor);
 #endif
-			if (bPreRegisterComponents && !bHasCurrentActorCalledPreRegister)
+				if (bPreRegisterComponents && !bHasCurrentActorCalledPreRegister)
+				{
+					Actor->PreRegisterAllComponents();
+					bHasCurrentActorCalledPreRegister = true;
+				}
+				bAllComponentsRegistered = Actor->IncrementalRegisterComponents(NumComponentsToUpdate, Context);
+			}
+#if !UE_BUILD_SHIPPING
+			else if (GOptimizeActorRegistration == 2)
 			{
-				Actor->PreRegisterAllComponents();
-				bHasCurrentActorCalledPreRegister = true;
+				// Verify that there aren't any leftover components
+				Actor->ForEachComponent(false, [](UActorComponent* Component)
+				{
+					ensureMsgf(Component->IsRegistered() || !Component->bAutoRegister, TEXT("Component %s should be registered!"), *Component->GetPathName());
+				});
 			}
-			bAllComponentsRegistered = Actor->IncrementalRegisterComponents(NumComponentsToUpdate, Context);
+#endif
 		}
 
 		if (bAllComponentsRegistered)
 		{
-			// All components have been registered fro this actor, move to a next one
+			// All components have been registered for this actor, move to a next one
 			CurrentActorIndexForIncrementalUpdate++;
 			bHasCurrentActorCalledPreRegister = false;
 		}
