World Partition: Fixed a bug in TActorDescContainerInstanceCollection::TBaseIterator where iteration would stop on the first container if it contained no relevant actor.

diff --git a/Engine/Source/Runtime/Engine/Public/WorldPartition/ActorDescContainerInstanceCollection.h b/Engine/Source/Runtime/Engine/Public/WorldPartition/ActorDescContainerInstanceCollection.h
--- a/Engine/Source/Runtime/Engine/Public/WorldPartition/ActorDescContainerInstanceCollection.h
+++ b/Engine/Source/Runtime/Engine/Public/WorldPartition/ActorDescContainerInstanceCollection.h
@@ -104,6 +104,11 @@ public:
 			if (ContainerIterator)
 			{
 				ActorsIterator = MakeUnique<ActDescIteratorType>(*ContainerIterator);
+
+				if (!*ActorsIterator)
+				{
+					AdvanceToRelevantActorInNextContainer();
+				}
 			}
 		}
 
