CurveLinearColor delete WritePixel
unused and not great
just use FColor if you want that

#rnx

diff --git a/Engine/Source/Runtime/Engine/Classes/Curves/CurveLinearColor.h b/Engine/Source/Runtime/Engine/Classes/Curves/CurveLinearColor.h
--- a/Engine/Source/Runtime/Engine/Classes/Curves/CurveLinearColor.h
+++ b/Engine/Source/Runtime/Engine/Classes/Curves/CurveLinearColor.h
@@ -101,7 +101,5 @@ public:
 	UE_DEPRECATED(4.26, "OnUpdateGradient is deprecated. Please use UCurveBase::OnUpdateCurve instead")
 	FOnUpdateGradient OnUpdateGradient;
 #endif
-protected:
-	ENGINE_API void WritePixel(uint8* Pixel, const FLinearColor& Color);
 };
 
diff --git a/Engine/Source/Runtime/Engine/Private/Curves/CurveLinearColor.cpp b/Engine/Source/Runtime/Engine/Private/Curves/CurveLinearColor.cpp
--- a/Engine/Source/Runtime/Engine/Private/Curves/CurveLinearColor.cpp
+++ b/Engine/Source/Runtime/Engine/Private/Curves/CurveLinearColor.cpp
@@ -349,12 +349,3 @@ void UCurveLinearColor::Serialize(FArchive& Ar)
 	Ar.UsingCustomVersion(FReleaseObjectVersion::GUID);
 	Super::Serialize(Ar);
 }
-
-void UCurveLinearColor::WritePixel(uint8* Pixel, const FLinearColor& Color)
-{
-	Pixel[0] = FMath::RoundToInt(Color.B * 255.f);
-	Pixel[1] = FMath::RoundToInt(Color.G * 255.f);
-	Pixel[2] = FMath::RoundToInt(Color.R * 255.f);
-	Pixel[3] = FMath::RoundToInt(Color.A * 255.f);
-}
-
