Adding possibility to replace classes upon reading back in the CPFUOReader
Enabled this class replacement for object that have overridable serialization enabled
Fixed the PreCreateInstances to continue looping through its archetype object to see if there are other subobject attached to them that are not archetypes that would need to be created.
Also start using a topologie sort to sort the CDO to reinstantiate. Each class depend on thier super class and any embeded class instances.

diff --git a/Engine/Source/Editor/Kismet/Private/BlueprintCompilationManager.cpp b/Engine/Source/Editor/Kismet/Private/BlueprintCompilationManager.cpp
--- a/Engine/Source/Editor/Kismet/Private/BlueprintCompilationManager.cpp
+++ b/Engine/Source/Editor/Kismet/Private/BlueprintCompilationManager.cpp
@@ -43,6 +43,7 @@
 #include "UObject/UObjectHash.h"
 #include "Kismet2/KismetDebugUtilities.h"
 #include "BlueprintEditorModule.h"
+#include "Algo/TopologicalSort.h"
 #include "Animation/AnimBlueprint.h"
 #include "Stats/StatsHierarchical.h"
 #include "UObject/PropertyBagRepository.h"
@@ -513,11 +514,6 @@ struct FReinstancingJob
 
 	// always set:
 	TPair<UClass*, UClass*> OldToNew;
-
-	// Sorting dependencies
-	TSet<const UClass*> BPClassDependencies;
-	bool DependsOn(const FReinstancingJob& OtherJob) const;
-	void CalculateBPClassDependencies();
 };
 
 FReinstancingJob::FReinstancingJob(TSharedPtr<FBlueprintCompileReinstancer> InReinstancer)
@@ -547,41 +543,6 @@ FReinstancingJob::FReinstancingJob(TPair<UClass*, UClass*> InOldToNew)
 {
 }
 
-bool FReinstancingJob::DependsOn(const FReinstancingJob& OtherJob) const
-{
-	return OtherJob.Reinstancer.IsValid() && OtherJob.OldToNew.Key && BPClassDependencies.Contains(OtherJob.OldToNew.Key);
-}
-
-void FReinstancingJob::CalculateBPClassDependencies()
-{
-	if (Reinstancer.IsValid() && OldToNew.Key)
-	{
-		const UObject* OldCDO = OldToNew.Key->ClassDefaultObject;
-
-		// Gather subobjects on old CDO and remember depends BP classes
-		auto GatherDependentBPClasses = [this](const UObject* CDO, const auto& Recurse) -> void
-		{
-			TArray<UObject*> ContainedOldObjects;
-			GetObjectsWithOuter(CDO, ContainedOldObjects);
-			for (const UObject* OldObject : ContainedOldObjects)
-			{
-				const UClass* DependentClass = OldObject->GetClass();
-				while (DependentClass && UBlueprint::GetBlueprintFromClass(DependentClass))
-				{
-					bool bIsAlreadyInSet = false;
-					BPClassDependencies.Add(DependentClass, &bIsAlreadyInSet);
-					if(!bIsAlreadyInSet)
-					{
-						Recurse(DependentClass->ClassDefaultObject, Recurse);
-					}
-					DependentClass = DependentClass->GetSuperClass();
-				}
-			}
-		};
-		GatherDependentBPClasses(OldCDO, GatherDependentBPClasses);
-	}
-}
-
 namespace UE::Kismet::BlueprintCompilationManager::Private
 {
 	namespace ConsoleVariables
@@ -2393,7 +2354,6 @@ void FBlueprintCompilationManagerImpl::ReinstanceBatch(TArray<FReinstancingJob>&
 			{
 				UBlueprint* CompiledBlueprint = UBlueprint::GetBlueprintFromClass(ReinstancingJob.OldToNew.Value);
 				ReinstancingJob.Reinstancer->UpdateBytecodeReferences(DependentBPs, FieldMappings);
-				ReinstancingJob.CalculateBPClassDependencies();
 			}
 		}
 
@@ -2403,31 +2363,47 @@ void FBlueprintCompilationManagerImpl::ReinstanceBatch(TArray<FReinstancingJob>&
 	// Now we can update templates and archetypes - note that we don't look for direct references to archetypes - doing
 	// so is very expensive and it will be much faster to directly update anything that cares to cache direct references
 	// to an archetype here (e.g. a UClass::ClassDefaultObject member):
-	
-	// 1. Sort classes so that most derived types are updated last - right now the only caller of this function
-	// also sorts, but we don't want to make too many assumptions about caller. We could refine this API so that
-	// we're not taking a raw list of reinstancers:
-	Reinstancers.Sort(
-		[](const FReinstancingJob& ReinstancingDataA, const FReinstancingJob& ReinstancingDataB)
+	TArray<FReinstancingJob*> ReinstancersPtr;
+	ReinstancersPtr.Reset(Reinstancers.Num());
+	for (FReinstancingJob& ReinstancingJob : Reinstancers)
+	{
+		ReinstancersPtr.Add(&ReinstancingJob);
+	}
+
+	Algo::TopologicalSort(ReinstancersPtr, [&Reinstancers](FReinstancingJob* ReinstancingJob)
+	{
+		TArray<FReinstancingJob*> Dependencies;
+		auto AddDependentClass = [&Dependencies,&Reinstancers](UClass* DependentClass)
 		{
-			if(ReinstancingDataA.DependsOn(ReinstancingDataB))
+			if (FReinstancingJob* ReinstancingJob = Reinstancers.FindByPredicate([DependentClass](FReinstancingJob& ReinstancingJob) { return ReinstancingJob.OldToNew.Key == DependentClass;}))
 			{
-				return false;
+				Dependencies.Add(ReinstancingJob);
 			}
-			if(ReinstancingDataB.DependsOn(ReinstancingDataA))
+		};
+
+		if (UClass* OldClass = ReinstancingJob->OldToNew.Key)
+		{
+			AddDependentClass(OldClass->GetSuperClass());
+
+			if (const UObject* CDO = OldClass->ClassDefaultObject)
 			{
-				return true;
+				TArray<UObject*> ContainedOldObjects;
+				GetObjectsWithOuter(CDO, ContainedOldObjects);
+				for (const UObject* OldObject : ContainedOldObjects)
+				{
+					AddDependentClass(OldObject->GetClass());
+				}
 			}
-			return FBlueprintCompileReinstancer::ReinstancerOrderingFunction(
-					ReinstancingDataA.OldToNew.Value, 
-					ReinstancingDataB.OldToNew.Value);
 		}
-	);
+
+		return Dependencies;
+	});
 
 	// 2. Copy defaults from old CDO - CDO may be missing if this class was reinstanced and relinked here,
 	// so use GetDefaultObject(true):
-	for (const FReinstancingJob& ReinstancingJob : Reinstancers)
+	for (const FReinstancingJob* ReinstancingJobPtr : ReinstancersPtr)
 	{
+		const FReinstancingJob& ReinstancingJob = *ReinstancingJobPtr;
 		UObject* OldCDO = nullptr;
 		UClass* OldClass = ReinstancingJob.OldToNew.Key;
 
@@ -2476,7 +2452,7 @@ void FBlueprintCompilationManagerImpl::ReinstanceBatch(TArray<FReinstancingJob>&
 				}
 				for(const auto& Pair : CreatedInstanceMap)
 				{
-					FBlueprintCompileReinstancer::CopyPropertiesForUnrelatedObjects(Pair.Key, Pair.Value, /*bClearExternalReferences*/true, bUseDeltaSerialization, /*bOnlyHandleDirectSubObjects*/true, &OldToNewInstanceMap);
+					FBlueprintCompileReinstancer::CopyPropertiesForUnrelatedObjects(Pair.Key, Pair.Value, /*bClearExternalReferences*/true, bUseDeltaSerialization, /*bOnlyHandleDirectSubObjects*/true, &OldToNewInstanceMap, &InOutOldToNewClassMap);
 				}
 
 				if (OldToNewTemplates && !OldToNewInstanceMap.IsEmpty())
@@ -2544,8 +2520,9 @@ void FBlueprintCompilationManagerImpl::ReinstanceBatch(TArray<FReinstancingJob>&
 		ArchetypeReferencers.Add(GUnrealEd->Trans);
 	}
 
-	for (const FReinstancingJob& ReinstancingJob : Reinstancers)
+	for (const FReinstancingJob* ReinstancingJobPtr : ReinstancersPtr)
 	{
+		const FReinstancingJob& ReinstancingJob = *ReinstancingJobPtr;
 		UClass* OldClass = ReinstancingJob.OldToNew.Key;
 		if(OldClass)
 		{
@@ -2655,8 +2632,9 @@ void FBlueprintCompilationManagerImpl::ReinstanceBatch(TArray<FReinstancingJob>&
 	}
 
 	// This loop finishes the reinstancing of archetypes after the entire Outer hierarchy has been updated with new instances:
-	for (const FReinstancingJob& ReinstancingJob : Reinstancers)
+	for (const FReinstancingJob* ReinstancingJobPtr : ReinstancersPtr)
 	{
+		const FReinstancingJob& ReinstancingJob = *ReinstancingJobPtr;
 		UClass* OldClass = ReinstancingJob.OldToNew.Key;
 		if(OldClass)
 		{
@@ -2687,7 +2665,7 @@ void FBlueprintCompilationManagerImpl::ReinstanceBatch(TArray<FReinstancingJob>&
 					}
 					for (const auto& Pair : CreatedInstanceMap)
 					{
-						FBlueprintCompileReinstancer::CopyPropertiesForUnrelatedObjects(Pair.Key, Pair.Value, /*bClearExternalReferences*/true, bUseDeltaSerialization, /*bOnlyHandleDirectSubObjects*/true, &OldToNewInstanceMap);
+						FBlueprintCompileReinstancer::CopyPropertiesForUnrelatedObjects(Pair.Key, Pair.Value, /*bClearExternalReferences*/true, bUseDeltaSerialization, /*bOnlyHandleDirectSubObjects*/true, &OldToNewInstanceMap, &InOutOldToNewClassMap);
 					}
 
 					if (OldToNewTemplates && !OldToNewInstanceMap.IsEmpty())
@@ -2708,8 +2686,9 @@ void FBlueprintCompilationManagerImpl::ReinstanceBatch(TArray<FReinstancingJob>&
 	// 4. update known references to archetypes (e.g. component templates, WidgetTree). We don't want to run the normal 
 	// reference finder to update these because searching the entire object graph is time consuming. Instead we just replace
 	// all references in our UBlueprint and its generated class:
-	for (const FReinstancingJob& ReinstancingJob : Reinstancers)
+	for (const FReinstancingJob* ReinstancingJobPtr : ReinstancersPtr)
 	{
+		const FReinstancingJob& ReinstancingJob = *ReinstancingJobPtr;
 		ArchetypeReferencers.Add(ReinstancingJob.OldToNew.Value);
 		ArchetypeReferencers.Add(ReinstancingJob.OldToNew.Value->ClassGeneratedBy);
 
diff --git a/Engine/Source/Editor/UnrealEd/Private/Kismet2/KismetReinstanceUtilities.cpp b/Engine/Source/Editor/UnrealEd/Private/Kismet2/KismetReinstanceUtilities.cpp
--- a/Engine/Source/Editor/UnrealEd/Private/Kismet2/KismetReinstanceUtilities.cpp
+++ b/Engine/Source/Editor/UnrealEd/Private/Kismet2/KismetReinstanceUtilities.cpp
@@ -2174,8 +2174,12 @@ static void ReplaceObjectHelper(UObject*& OldObject, UClass* OldClass, UObject*&
 	Options.bNotifyObjectReplacement = true;
 	Options.bSkipCompilerGeneratedDefaults = true;
 	Options.bOnlyHandleDirectSubObjects = true;
-	Options.bReplaceInternalReferenceUponRead = FOverridableManager::Get().IsEnabled(*OldObject);
 	Options.OptionalReplacementMappings = &OldToNewInstanceMap;
+	if (FOverridableManager::Get().IsEnabled(*OldObject))
+	{
+		Options.bReplaceInternalReferenceUponRead = true;
+		Options.OptionalOldToNewClassMappings = &OldToNewClassMap;
+	}
 	// this currently happens because of some misguided logic in UBlueprintGeneratedClass::FindArchetype that
 	// points us to a mismatched archetype, in which case delta serialization becomes unsafe.. without
 	// that logic we could lose data, so for now i'm disabling delta serialization when we detect that situation
@@ -3056,7 +3060,7 @@ void FBlueprintCompileReinstancer::ReparentChild(UClass* ChildClass)
 	ChildClass->StaticLink(true);
 }
 
-void FBlueprintCompileReinstancer::CopyPropertiesForUnrelatedObjects(UObject* OldObject, UObject* NewObject, bool bClearExternalReferences, bool bForceDeltaSerialization /* = false */, bool bOnlyHandleDirectSubObjects/* = false */, TMap<UObject*, UObject*>* OldToNewInstanceMap /*=nullptr*/)
+void FBlueprintCompileReinstancer::CopyPropertiesForUnrelatedObjects(UObject* OldObject, UObject* NewObject, bool bClearExternalReferences, bool bForceDeltaSerialization /* = false */, bool bOnlyHandleDirectSubObjects/* = false */, TMap<UObject*, UObject*>* OldToNewInstanceMap /*=nullptr*/, const TMap<UClass*,UClass*>* OldToNewClassMap /*=nullptr*/)
 {
 	SCOPED_LOADTIMER_ASSET_TEXT(*WriteToString<256>(TEXT("CopyPropertiesForUnrelatedObjects "), *GetPathNameSafe(NewObject)));
 	UEngine::FCopyPropertiesForUnrelatedObjectsParams Params;
@@ -3069,7 +3073,11 @@ void FBlueprintCompileReinstancer::CopyPropertiesForUnrelatedObjects(UObject* Ol
 	Params.OptionalReplacementMappings = OldToNewInstanceMap;
 	Params.bOnlyHandleDirectSubObjects = bOnlyHandleDirectSubObjects;
 	// Overridable serialization needs this to be able to merge back containers of subobjects.
-	Params.bReplaceInternalReferenceUponRead = FOverridableManager::Get().IsEnabled(*OldObject);
+	if (FOverridableManager::Get().IsEnabled(*OldObject))
+	{
+		Params.bReplaceInternalReferenceUponRead = true;
+		Params.OptionalOldToNewClassMappings = OldToNewClassMap;
+	}
 
 	UEngine::CopyPropertiesForUnrelatedObjects(OldObject, NewObject, Params);
 }
@@ -3079,18 +3087,21 @@ void FBlueprintCompileReinstancer::PreCreateSubObjectsForReinstantiation(const T
 	TSet<UObject*> OldInstancedSubObjects;
 	FReplaceReferenceHelper::GetOwnedSubobjectsRecursive(OldObject, OldInstancedSubObjects);
 
+	// Add the mapping from the old to the new object exists...
+	CreatedInstanceMap.Add(OldObject, NewUObject);
 	PreCreateSubObjectsForReinstantiation_Inner(OldInstancedSubObjects, OldToNewClassMap, OldObject, NewUObject, CreatedInstanceMap, OldToNewInstanceMap);
 }
 
 void FBlueprintCompileReinstancer::PreCreateSubObjectsForReinstantiation_Inner(const TSet<UObject*>& OldInstancedSubObjects, const TMap<UClass*, UClass*>& OldToNewClassMap, UObject* OldObject, UObject* NewUObject, TMap<UObject*, UObject*>& CreatedInstanceMap, const TMap<UObject*, UObject*>* OldToNewInstanceMap/* = nullptr*/)
 {
-	// Add the mapping from the old to the new object exists...
-	CreatedInstanceMap.Add(OldObject, NewUObject);
-
 	// Gather subobjects on old object and pre-create them if needed
 	TArray<UObject*> ContainedOldSubObjects;
 	GetObjectsWithOuter(OldObject, ContainedOldSubObjects, /*bIncludeNestedObjects*/false);
 
+	// Gather subobjects on old object and pre-create them if needed
+	TArray<UObject*> ContainedNewSubObjects;
+	GetObjectsWithOuter(NewUObject, ContainedNewSubObjects, /*bIncludeNestedObjects*/false);
+
 	// Pre-create all non default subobjects to prevent re-instancing them as an old classes 
 	TMap<UObject*, UObject*> ReferenceReplacementMap;
 	for (int32 i = 0; i < ContainedOldSubObjects.Num(); ++i)
@@ -3113,15 +3124,13 @@ void FBlueprintCompileReinstancer::PreCreateSubObjectsForReinstantiation_Inner(c
 		UObject* Archetype = UObject::GetArchetypeFromRequiredInfo(OldSubObjectClass, OldSubObjectOuter, SubObjectName, SubObjectFlags);
 		if (Archetype->HasAnyFlags(RF_ClassDefaultObject))
 		{
-			UObject* NewSubObjectOuter = CreatedInstanceMap.FindChecked(OldSubObjectOuter);
-
 			// Was it already re-instantiated
 			if (UObject* const* AlreadyCreatedSubObject = OldToNewInstanceMap ? OldToNewInstanceMap->Find(OldSubObject) : nullptr)
 			{
 				int32 AlreadyCreatedSubObjectIndex = ContainedOldSubObjects.Find(*AlreadyCreatedSubObject);
 				checkf(AlreadyCreatedSubObjectIndex > i, TEXT("Expecting the already created subobject to be in the old subobject list after this sub object"));
 				ContainedOldSubObjects.RemoveAt(AlreadyCreatedSubObjectIndex);
-				(*AlreadyCreatedSubObject)->Rename(nullptr, NewSubObjectOuter, REN_DoNotDirty | REN_DontCreateRedirectors | REN_ForceNoResetLoaders | REN_NonTransactional);
+				(*AlreadyCreatedSubObject)->Rename(nullptr, NewUObject, REN_DoNotDirty | REN_DontCreateRedirectors | REN_ForceNoResetLoaders | REN_NonTransactional);
 			}
 			else
 			{
@@ -3142,10 +3151,16 @@ void FBlueprintCompileReinstancer::PreCreateSubObjectsForReinstantiation_Inner(c
 				// Only pre-create object where the class does not have newer version of the it
 				if(!SubObjectClass->HasAnyClassFlags(CLASS_NewerVersionExists))
 				{
-					UObject* NewSubObject = NewObject<UObject>(NewSubObjectOuter, SubObjectClass, SubObjectName, SubObjectFlags);
-					PreCreateSubObjectsForReinstantiation(OldToNewClassMap, OldSubObject, NewSubObject, CreatedInstanceMap, OldToNewInstanceMap);
+					UObject* NewSubObject = NewObject<UObject>(NewUObject, SubObjectClass, SubObjectName, SubObjectFlags);
+					CreatedInstanceMap.Add(OldSubObject, NewSubObject);
+					PreCreateSubObjectsForReinstantiation_Inner(OldInstancedSubObjects, OldToNewClassMap, OldSubObject, NewSubObject, CreatedInstanceMap, OldToNewInstanceMap);
 				}
 			}
 		}
+		// There might be new subobjects attached to the sub object that are particular to this instance, let's traverse it to find them out.
+		else if (UObject** NewSubObject = ContainedNewSubObjects.FindByPredicate([SubObjectName](UObject* SubObject) { return SubObject && SubObject->GetName() == SubObjectName; }))
+		{
+			PreCreateSubObjectsForReinstantiation_Inner(OldInstancedSubObjects, OldToNewClassMap, OldSubObject, *NewSubObject, CreatedInstanceMap, OldToNewInstanceMap);
+		}
 	}
 }
diff --git a/Engine/Source/Editor/UnrealEd/Public/Kismet2/KismetReinstanceUtilities.h b/Engine/Source/Editor/UnrealEd/Public/Kismet2/KismetReinstanceUtilities.h
--- a/Engine/Source/Editor/UnrealEd/Public/Kismet2/KismetReinstanceUtilities.h
+++ b/Engine/Source/Editor/UnrealEd/Public/Kismet2/KismetReinstanceUtilities.h
@@ -287,8 +287,9 @@ public:
 	* @param bForceDeltaSerialization		If true the delta serialization will be used when copying
 	* @param bOnlyHandleDirectSubObjects	If true will only copy/handle immediate subobjects
 	* @param OldToNewInstanceMap			If != null it will be used to replace any references found in this object
+	* @param OldToNewClassMap				if != null it will be used to replace any class references found in this object
 	*/
-	static UNREALED_API void CopyPropertiesForUnrelatedObjects(UObject* OldObject, UObject* NewObject, bool bClearExternalReferences, bool bForceDeltaSerialization = false, bool bOnlyHandleDirectSubObjects = false, TMap<UObject*, UObject*>* OldToNewInstanceMap =nullptr);
+	static UNREALED_API void CopyPropertiesForUnrelatedObjects(UObject* OldObject, UObject* NewObject, bool bClearExternalReferences, bool bForceDeltaSerialization = false, bool bOnlyHandleDirectSubObjects = false, TMap<UObject*, UObject*>* OldToNewInstanceMap =nullptr, const TMap<UClass*,UClass*>* OldToNewClassMap =nullptr);
 
 	/**
 	 * This method will pre-create all non-default sub object needed for a re-instantiation, what is left is to CopyPropertiesForUnrelatedObjects on the created instances map to finish the re-instancing
diff --git a/Engine/Source/Runtime/Engine/Classes/Engine/Engine.h b/Engine/Source/Runtime/Engine/Classes/Engine/Engine.h
--- a/Engine/Source/Runtime/Engine/Classes/Engine/Engine.h
+++ b/Engine/Source/Runtime/Engine/Classes/Engine/Engine.h
@@ -2963,6 +2963,7 @@ public:
 		// In cases where the SourceObject will no longer be able to look up its correct Archetype, it can be supplied
 		UObject* SourceObjectArchetype;
 		TMap<UObject*, UObject*>* OptionalReplacementMappings;
+		const TMap<UClass*, UClass*>* OptionalOldToNewClassMappings; // Will be used along with the bReplaceInternalReferenceUponRead;
 
 		ENGINE_API FCopyPropertiesForUnrelatedObjectsParams();
 		ENGINE_API FCopyPropertiesForUnrelatedObjectsParams(const FCopyPropertiesForUnrelatedObjectsParams&);
diff --git a/Engine/Source/Runtime/Engine/Private/UnrealEngine.cpp b/Engine/Source/Runtime/Engine/Private/UnrealEngine.cpp
--- a/Engine/Source/Runtime/Engine/Private/UnrealEngine.cpp
+++ b/Engine/Source/Runtime/Engine/Private/UnrealEngine.cpp
@@ -17010,119 +17010,136 @@ private:
 };
 
 /* Responsible for applying the saved property data from a FCPFUOWriter to a specified object */
-struct FCPFUOReader : public FObjectReader, public FCPFUOArchive
+struct FObjectReaderWithReplacement : public FObjectReader
 {
 public:
-	FCPFUOReader(FCPFUOWriter& DataSrc, UObject* DstObject, TMap<UObject*, UObject*>* InReferenceReplacementMap = nullptr)
-		: FObjectReader(DataSrc.SavedPropertyData)
-		, FCPFUOArchive(DataSrc)
+	FObjectReaderWithReplacement(UObject* Obj, const TArray<uint8>& InBytes, bool bIgnoreClassRef = false, bool bIgnoreArchetypeRef = false, const TMap<UObject*, UObject*>* InReferenceReplacementMap = nullptr, const TMap<UClass*, UClass*>* InClassReferenceReplacementMap = nullptr)
+		: FObjectReader(InBytes)
 		, ReferenceReplacementMap(InReferenceReplacementMap)
+		, ClassReferenceReplacementMap(InClassReferenceReplacementMap)
 	{
-		ArIgnoreArchetypeRef = true;
-		ArIgnoreClassRef = true;
+		ArIgnoreClassRef = bIgnoreClassRef;
+		ArIgnoreArchetypeRef = bIgnoreArchetypeRef;
 
 #if USE_STABLE_LOCALIZATION_KEYS
 		if (GIsEditor && !(ArPortFlags & (PPF_DuplicateVerbatim | PPF_DuplicateForPIE)))
 		{
-			SetLocalizationNamespace(TextNamespaceUtil::EnsurePackageNamespace(DstObject));
+			SetLocalizationNamespace(TextNamespaceUtil::EnsurePackageNamespace(Obj));
 		}
 #endif // USE_STABLE_LOCALIZATION_KEYS
 
-		DstObject->Serialize(*this);
+		Obj->Serialize(*this);
 	}
 
-	//~ Begin FArchive Interface
-	virtual void Serialize(void* Data, int64 Num) override
+	FObjectReaderWithReplacement(const TArray<uint8>& InBytes, const TMap<UObject*, UObject*>* InReferenceReplacementMap = nullptr, const TMap<UClass*, UClass*>* InClassReferenceReplacementMap = nullptr)
+		: FObjectReader(InBytes)
+		, ReferenceReplacementMap(InReferenceReplacementMap)
+		, ClassReferenceReplacementMap(InClassReferenceReplacementMap)
 	{
-		if (IsSerializationEnabled())
-		{
-			FObjectReader::Serialize(Data, Num);
-		}
 	}
 
-	virtual void MarkScriptSerializationStart(const UObject* Object) override { OpenTaggedDataScope(); }
-	virtual void MarkScriptSerializationEnd(const UObject* Object) override   { CloseTaggedDataScope(); }
-	// ~End FArchive Interface
+	const TMap<UObject*, UObject*>* ReferenceReplacementMap = nullptr;
+	const TMap<UClass*, UClass*>* ClassReferenceReplacementMap = nullptr;
 
-	TMap<UObject*, UObject*>* ReferenceReplacementMap = nullptr;
-
-	FArchive& operator<<(FWeakObjectPtr& Value) override
+	template <typename VALUE_TYPE>
+	FORCEINLINE void ReplaceObjectReference(UObject* Object, VALUE_TYPE& Value) const
 	{
-		FObjectReader::operator<<(Value);
-		if (UObject* Object = Value.Get())
+		if(!Object)
+		{
+			return;
+		}
+		if (UObject*const* Reference = ReferenceReplacementMap ? ReferenceReplacementMap->Find(Object) : nullptr)
+		{
+			Value = const_cast<UObject*>(*Reference);
+		}
+		else if(ClassReferenceReplacementMap)
 		{
-			if (UObject** Reference = ReferenceReplacementMap ? ReferenceReplacementMap->Find(Object) : nullptr)
+			if (UClass* Class = Cast<UClass>(Object))
 			{
-				Value = *Reference;
+				if (UClass*const* ClassReference = ClassReferenceReplacementMap->Find(Class))
+				{
+					Value = const_cast<UClass*>(*ClassReference);
+				}
 			}
 		}
+	}
+
+	FArchive& operator<<(FWeakObjectPtr& Value) override
+	{
+		FObjectReader::operator<<(Value);
+		ReplaceObjectReference(Value.Get(), Value);
 		return *this;
 	}
 
 	FArchive& operator<<(FSoftObjectPtr& Value) override
 	{
 		FObjectReader::operator<<(Value);
-		if (UObject* Object = Value.Get())
-		{
-			if (UObject** Reference = ReferenceReplacementMap ? ReferenceReplacementMap->Find(Object) : nullptr)
-			{
-				Value = *Reference;
-			}
-		}
+		ReplaceObjectReference(Value.Get(), Value);
 		return *this;
 	}
 
 	FArchive& operator<<(FSoftObjectPath& Value) override
 	{
 		FObjectReader::operator<<(Value);
-		if (UObject* Object = Value.ResolveObject())
-		{
-			if (UObject** Reference = ReferenceReplacementMap ? ReferenceReplacementMap->Find(Object) : nullptr)
-			{
-				Value = *Reference;
-			}
-		}
+		ReplaceObjectReference(Value.ResolveObject(), Value);
 		return *this;
 	}
 
 	FArchive& operator<<(FLazyObjectPtr& Value) override
 	{
 		FObjectReader::operator<<(Value);
-		if (UObject* Object = Value.Get())
-		{
-			if (UObject** Reference = ReferenceReplacementMap ? ReferenceReplacementMap->Find(Object) : nullptr)
-			{
-				Value = *Reference;
-			}
-		}
+		ReplaceObjectReference(Value.Get(), Value);
 		return *this;
 	}
 
 	FArchive& operator<<(FObjectPtr& Value) override
 	{
 		FObjectReader::operator<<(Value);
-		if (UObject* Object = Value.Get())
-		{
-			if (UObject** Reference = ReferenceReplacementMap ? ReferenceReplacementMap->Find(Object) : nullptr)
-			{
-				Value = *Reference;
-			}
-		}
+		ReplaceObjectReference(Value.Get(), Value);
 		return *this;
 	}
 
 	FArchive& operator<<(UObject*& Value) override
 	{
 		FObjectReader::operator<<(Value);
-		if (UObject* Object = Value)
+		ReplaceObjectReference(Value, Value);
+		return *this;
+	}
+};
+
+/* Responsible for applying the saved property data from a FCPFUOWriter to a specified object */
+struct FCPFUOReader : public FObjectReaderWithReplacement, public FCPFUOArchive
+{
+public:
+	FCPFUOReader(FCPFUOWriter& DataSrc, UObject* DstObject, const TMap<UObject*, UObject*>* InReferenceReplacementMap = nullptr, const TMap<UClass*, UClass*>* InClassReferenceReplacementMap = nullptr)
+		: FObjectReaderWithReplacement(DataSrc.SavedPropertyData, InReferenceReplacementMap, InClassReferenceReplacementMap)
+		, FCPFUOArchive(DataSrc)
+	{
+		ArIgnoreArchetypeRef = true;
+		ArIgnoreClassRef = true;
+
+#if USE_STABLE_LOCALIZATION_KEYS
+		if (GIsEditor && !(ArPortFlags & (PPF_DuplicateVerbatim | PPF_DuplicateForPIE)))
 		{
-			if (UObject** Reference = ReferenceReplacementMap ? ReferenceReplacementMap->Find(Object) : nullptr)
-			{
-				Value = *Reference;
-			}
+			SetLocalizationNamespace(TextNamespaceUtil::EnsurePackageNamespace(DstObject));
 		}
-		return *this;
+#endif // USE_STABLE_LOCALIZATION_KEYS
+
+		DstObject->Serialize(*this);
 	}
+
+	//~ Begin FArchive Interface
+	virtual void Serialize(void* Data, int64 Num) override
+	{
+		if (IsSerializationEnabled())
+		{
+			FObjectReaderWithReplacement::Serialize(Data, Num);
+		}
+	}
+
+	virtual void MarkScriptSerializationStart(const UObject* Object) override { OpenTaggedDataScope(); }
+	virtual void MarkScriptSerializationEnd(const UObject* Object) override   { CloseTaggedDataScope(); }
+	// ~End FArchive Interface
 };
 
 PRAGMA_DISABLE_DEPRECATION_WARNINGS
@@ -17141,6 +17158,7 @@ UEngine::FCopyPropertiesForUnrelatedObjectsParams::FCopyPropertiesForUnrelatedOb
 	, bReplaceInternalReferenceUponRead(false)
 	, SourceObjectArchetype(nullptr)
 	, OptionalReplacementMappings(nullptr)
+	, OptionalOldToNewClassMappings(nullptr)
 {}
 UEngine::FCopyPropertiesForUnrelatedObjectsParams::FCopyPropertiesForUnrelatedObjectsParams(const FCopyPropertiesForUnrelatedObjectsParams&) = default;
 PRAGMA_ENABLE_DEPRECATION_WARNINGS
@@ -17371,17 +17389,18 @@ void UEngine::CopyPropertiesForUnrelatedObjects(UObject* OldObject, UObject* New
 			}
 		}
 
+		TMap<UObject*, UObject*>* OptionalReferenceReplacementMap = nullptr;
+		const TMap<UClass*, UClass*>* OptionalClassReferenceReplacementMap = nullptr;
+		if (Params.bReplaceInternalReferenceUponRead)
+		{
+			OptionalReferenceReplacementMap = Params.OptionalReplacementMappings ? Params.OptionalReplacementMappings : &ReferenceReplacementMap;
+			OptionalClassReferenceReplacementMap = Params.OptionalOldToNewClassMappings;
+		}
+
 		// Serialize in the modified properties from the old CDO to the new CDO
 		if (Writer.SavedPropertyData.Num() > 0)
 		{
-			if (Params.bReplaceInternalReferenceUponRead)
-			{
-				FCPFUOReader Reader(Writer, NewObject, Params.OptionalReplacementMappings ? Params.OptionalReplacementMappings : &ReferenceReplacementMap);
-			}
-			else
-			{
-				FCPFUOReader Reader(Writer, NewObject);
-			}
+			FCPFUOReader Reader(Writer, NewObject, OptionalReferenceReplacementMap, OptionalClassReferenceReplacementMap);
 		}
 
 		TArray<UObject*> EditInlineSubobjectsOfComponents;
@@ -17391,7 +17410,7 @@ void UEngine::CopyPropertiesForUnrelatedObjects(UObject* OldObject, UObject* New
 			{
 				// Restore modified properties into the new instance
 				FInstancedObjectRecord& Record = SavedInstances[*pOldInstanceIndex];
-				FObjectReader Reader(NewInstance, Record.SavedProperties, true, true);
+				FObjectReaderWithReplacement Reader(NewInstance, Record.SavedProperties, true, true, OptionalReferenceReplacementMap, OptionalClassReferenceReplacementMap);
 				FFindInstancedReferenceSubobjectHelper::Duplicate(Record.OldInstance, NewInstance, ReferenceReplacementMap, EditInlineSubobjectsOfComponents, Params.OptionalReplacementMappings);
 			}
 		}
