Map: Round corner support for bottom corner of the map.
Added code to compute the necessary information for the tile material, and updated the material to enable masking by default.
#tests In cooked build, opened the map both in fullscreen and windowed, rounded corners or positionned well. In PIE, tested both with in editor window PIE and additionnal window PIE. Played with the dimension and position of both windows. Rounded corner are well positionned overall, however, in some particular context, we can see the round a tiny bit different between the map and the background widget (should be fixed data side).: open the map with different combination, separated window or in viewport window, on different screens, at different positions. Also tested in cooked builds (windowed and fullscreen).

#virtualized

diff --git a/Engine/Source/Runtime/Engine/Private/Slate/SGameLayerManager.cpp b/Engine/Source/Runtime/Engine/Private/Slate/SGameLayerManager.cpp
--- a/Engine/Source/Runtime/Engine/Private/Slate/SGameLayerManager.cpp
+++ b/Engine/Source/Runtime/Engine/Private/Slate/SGameLayerManager.cpp
@@ -254,6 +254,11 @@ FGeometry SGameLayerManager::GetViewportWidgetHostGeometry() const
 	return WidgetHost->GetTickSpaceGeometry();
 }
 
+FGeometry SGameLayerManager::GetViewportWidgetHostPaintGeometry() const
+{
+	return WidgetHost->GetPaintSpaceGeometry();
+}
+
 FGeometry SGameLayerManager::GetPlayerWidgetHostGeometry(ULocalPlayer* Player) const
 {
 	TSharedPtr<FPlayerLayer> PlayerLayer = PlayerLayers.FindRef(Player);
diff --git a/Engine/Source/Runtime/Engine/Public/Slate/SGameLayerManager.h b/Engine/Source/Runtime/Engine/Public/Slate/SGameLayerManager.h
--- a/Engine/Source/Runtime/Engine/Public/Slate/SGameLayerManager.h
+++ b/Engine/Source/Runtime/Engine/Public/Slate/SGameLayerManager.h
@@ -60,6 +60,7 @@ public:
 	virtual void SetSceneViewport(FSceneViewport* SceneViewport) = 0;
 
 	virtual FGeometry GetViewportWidgetHostGeometry() const = 0;
+	virtual FGeometry GetViewportWidgetHostPaintGeometry() const = 0;
 	virtual FGeometry GetPlayerWidgetHostGeometry(ULocalPlayer* Player) const = 0;
 
 	virtual void NotifyPlayerAdded(int32 PlayerIndex, ULocalPlayer* AddedPlayer) = 0;
@@ -127,6 +128,7 @@ public:
 	// Begin IGameLayerManager
 	ENGINE_API virtual void SetSceneViewport(FSceneViewport* InSceneViewport) override;
 	ENGINE_API virtual FGeometry GetViewportWidgetHostGeometry() const override;
+	ENGINE_API virtual FGeometry GetViewportWidgetHostPaintGeometry() const override;
 	ENGINE_API virtual FGeometry GetPlayerWidgetHostGeometry(ULocalPlayer* Player) const override;
 
 	ENGINE_API virtual void NotifyPlayerAdded(int32 PlayerIndex, ULocalPlayer* AddedPlayer) override;
