Exposed ApplyRadialImpulse method

diff --git a/Engine/Source/Runtime/Engine/Classes/PhysicsEngine/PhysicsObjectBlueprintLibrary.h b/Engine/Source/Runtime/Engine/Classes/PhysicsEngine/PhysicsObjectBlueprintLibrary.h
--- a/Engine/Source/Runtime/Engine/Classes/PhysicsEngine/PhysicsObjectBlueprintLibrary.h
+++ b/Engine/Source/Runtime/Engine/Classes/PhysicsEngine/PhysicsObjectBlueprintLibrary.h
@@ -23,5 +23,5 @@ public:
 	static FTransform GetPhysicsObjectWorldTransform(UPrimitiveComponent* Component, FName BoneName);
 
 	UFUNCTION(BlueprintCallable, Category = "Physics Object")
-	static void ApplyRadialImpulse(UPrimitiveComponent* Component, FVector Origin, float Radius, float Strength, enum ERadialImpulseFalloff Falloff, bool bApplyStrain, float Strain, bool bVelChange = false);
+	static ENGINE_API void ApplyRadialImpulse(UPrimitiveComponent* Component, FVector Origin, float Radius, float Strength, enum ERadialImpulseFalloff Falloff, bool bApplyStrain, float Strain, bool bVelChange = false);
 };
\ No newline at end of file
