When starting audio tasks that do synchronous DDC requests, ensure they use the EQueuedWorkFlags::DoNotRunInsideBusyWait task flag so they don't get picked up by busy waits.

Owned by yuriy.odonnell in my absence.

diff --git a/Engine/Source/Runtime/Engine/Private/AudioStreaming.cpp b/Engine/Source/Runtime/Engine/Private/AudioStreaming.cpp
--- a/Engine/Source/Runtime/Engine/Private/AudioStreaming.cpp
+++ b/Engine/Source/Runtime/Engine/Private/AudioStreaming.cpp
@@ -336,7 +336,8 @@ void FStreamingWaveData::BeginPendingRequests(const TArray<uint32>& IndicesToLoa
 					MoveTemp(NullOnLoadCompletedCallback)
 				);
 				PendingAsyncStreamDerivedChunkTasks.Add(Task);
-				Task->StartBackgroundTask();
+				// This task may perform a long synchronous DDC request. Using DoNotRunInsideBusyWait prevents potentially delaying foreground tasks.
+				Task->StartBackgroundTask(GThreadPool, EQueuedWorkPriority::Normal, EQueuedWorkFlags::DoNotRunInsideBusyWait);
 			}
 			else
 #endif // #if WITH_EDITORONLY_DATA
diff --git a/Engine/Source/Runtime/Engine/Private/AudioStreamingCache.cpp b/Engine/Source/Runtime/Engine/Private/AudioStreamingCache.cpp
--- a/Engine/Source/Runtime/Engine/Private/AudioStreamingCache.cpp
+++ b/Engine/Source/Runtime/Engine/Private/AudioStreamingCache.cpp
@@ -1784,7 +1784,8 @@ void FAudioChunkCache::KickOffAsyncLoad(FCacheElement* CacheElement, const FChun
 			MoveTemp(OnLoadComplete)
 		));
 
-		CacheElement->DDCTask->StartBackgroundTask();
+		// This task may perform a long synchronous DDC request. Using DoNotRunInsideBusyWait prevents potentially delaying foreground tasks.
+		CacheElement->DDCTask->StartBackgroundTask(GThreadPool, EQueuedWorkPriority::Normal, EQueuedWorkFlags::DoNotRunInsideBusyWait);
 	}
 	else
 #endif // #if WITH_EDITORONLY_DATA
