Make Trim print correctly handle the uint64s

diff --git a/Engine/Source/Runtime/Engine/Private/AudioStreamingCache.cpp b/Engine/Source/Runtime/Engine/Private/AudioStreamingCache.cpp
--- a/Engine/Source/Runtime/Engine/Private/AudioStreamingCache.cpp
+++ b/Engine/Source/Runtime/Engine/Private/AudioStreamingCache.cpp
@@ -2049,7 +2049,7 @@ uint64 FCachedAudioStreamingManager::TrimMemory(uint64 NumBytesToFree)
 	// Freeing longer chunks will get us bigger gains and (presumably) have lower churn.
 	for (FAudioChunkCache& Cache : CacheArray)
 	{
-		uint64 NumBytesFreed = Cache.TrimMemory(NumBytesLeftToFree, false);
+		const uint64 NumBytesFreed = Cache.TrimMemory(NumBytesLeftToFree, false);
 
 		// NumBytesFreed could potentially be more than what we requested to free (since we delete whole chunks at once).
 		NumBytesLeftToFree -= FMath::Min(NumBytesFreed, NumBytesLeftToFree);
@@ -2062,9 +2062,9 @@ uint64 FCachedAudioStreamingManager::TrimMemory(uint64 NumBytesToFree)
 	}
 
 	check(NumBytesLeftToFree <= NumBytesToFree);
-	uint64 TotalBytesFreed = NumBytesToFree - NumBytesLeftToFree;
+	const uint64 TotalBytesFreed = NumBytesToFree - NumBytesLeftToFree;
 
-	UE_LOG(LogAudioStreamCaching, Display, TEXT("Call to IAudioStreamingManager::TrimMemory successfully freed %lu of the requested %lu bytes."), TotalBytesFreed, NumBytesToFree);
+	UE_LOG(LogAudioStreamCaching, Display, TEXT("Call to IAudioStreamingManager::TrimMemory successfully freed %llu of the requested %llu bytes."), TotalBytesFreed, NumBytesToFree);
 	return TotalBytesFreed;
 }
 
