[Engine]
* Fixed non unity compile errors

diff --git a/Engine/Source/Runtime/Engine/Private/ProfilingDebugging/LevelStreamingProfilingSubsystem.cpp b/Engine/Source/Runtime/Engine/Private/ProfilingDebugging/LevelStreamingProfilingSubsystem.cpp
--- a/Engine/Source/Runtime/Engine/Private/ProfilingDebugging/LevelStreamingProfilingSubsystem.cpp
+++ b/Engine/Source/Runtime/Engine/Private/ProfilingDebugging/LevelStreamingProfilingSubsystem.cpp
@@ -7,6 +7,7 @@
 #include "HAL/FileManager.h"
 #include "Misc/Paths.h"
 #include "Streaming/LevelStreamingDelegates.h"
+#include "UObject/Package.h"
 #include "WorldPartition/WorldPartition.h"
 #include "WorldPartition/WorldPartitionLevelStreamingDynamic.h"
 #include "Stats/Stats.h"
diff --git a/Engine/Source/Runtime/Engine/Private/WorldPartition/RuntimeHashSet/WorldPartitionRuntimeHashSetDebugDraw.cpp b/Engine/Source/Runtime/Engine/Private/WorldPartition/RuntimeHashSet/WorldPartitionRuntimeHashSetDebugDraw.cpp
--- a/Engine/Source/Runtime/Engine/Private/WorldPartition/RuntimeHashSet/WorldPartitionRuntimeHashSetDebugDraw.cpp
+++ b/Engine/Source/Runtime/Engine/Private/WorldPartition/RuntimeHashSet/WorldPartitionRuntimeHashSetDebugDraw.cpp
@@ -1,6 +1,7 @@
 // Copyright Epic Games, Inc. All Rights Reserved.
 
 #include "WorldPartition/RuntimeHashSet/WorldPartitionRuntimeHashSet.h"
+#include "UObject/Package.h"
 #include "WorldPartition/ContentBundle/ContentBundleWorldSubsystem.h"
 #include "WorldPartition/ContentBundle/ContentBundleDescriptor.h"
 #include "WorldPartition/ContentBundle/ContentBundleBase.h"
diff --git a/Engine/Source/Runtime/Engine/Private/WorldPartition/WorldPartitionReplay.cpp b/Engine/Source/Runtime/Engine/Private/WorldPartition/WorldPartitionReplay.cpp
--- a/Engine/Source/Runtime/Engine/Private/WorldPartition/WorldPartitionReplay.cpp
+++ b/Engine/Source/Runtime/Engine/Private/WorldPartition/WorldPartitionReplay.cpp
@@ -2,6 +2,7 @@
 
 #include "WorldPartition/WorldPartitionReplay.h"
 #include "Engine/GameInstance.h"
+#include "UObject/Package.h"
 #include "WorldPartition/WorldPartition.h"
 #include "Engine/World.h"
 #include "Net/UnrealNetwork.h"
diff --git a/Engine/Source/Runtime/Engine/Private/WorldPartition/WorldPartitionStreamingPolicy.cpp b/Engine/Source/Runtime/Engine/Private/WorldPartition/WorldPartitionStreamingPolicy.cpp
--- a/Engine/Source/Runtime/Engine/Private/WorldPartition/WorldPartitionStreamingPolicy.cpp
+++ b/Engine/Source/Runtime/Engine/Private/WorldPartition/WorldPartitionStreamingPolicy.cpp
@@ -8,6 +8,7 @@
 #include "Algo/Find.h"
 #include "Logging/LogScopedCategoryAndVerbosityOverride.h"
 #include "Misc/ScopeExit.h"
+#include "UObject/Package.h"
 #include "WorldPartition/WorldPartitionSubsystem.h"
 #include "WorldPartition/WorldPartitionLog.h"
 #include "WorldPartition/WorldPartitionReplay.h"
