WorldPartition: Regression fixes
- Make sure ContainerID is always properly set. This was causing WorldDataLayers actors to be loaded when it shouldn't be (IsMainWorldOnly relies on ContainerID)
- Fix Sub-WPs in a Non-WP world transforms

diff --git a/Engine/Source/Runtime/Engine/Private/LevelInstance/LevelInstanceEditorLevelStreaming.cpp b/Engine/Source/Runtime/Engine/Private/LevelInstance/LevelInstanceEditorLevelStreaming.cpp
--- a/Engine/Source/Runtime/Engine/Private/LevelInstance/LevelInstanceEditorLevelStreaming.cpp
+++ b/Engine/Source/Runtime/Engine/Private/LevelInstance/LevelInstanceEditorLevelStreaming.cpp
@@ -157,14 +157,9 @@ void ULevelStreamingLevelInstanceEditor::OnPreInitializeContainerInstance(UActor
 	// @todo_ow: Make sure LevelInstance exists (it is created after when creating a new Level Instance). This needs to be fixed properly but for now since the created LI is then unloaded it doesn't matter that the parent isn't set.
 	if (AActor* LevelInstanceActor = Cast<AActor>(GetLevelInstance()))
 	{
-		if (UWorldPartition* OwningWorldPartition = LevelInstanceActor->GetLevel()->GetWorldPartition())
-		{
-			if (FWorldPartitionActorDescInstance* ActorDescInstance = OwningWorldPartition->GetActorDescInstance(LevelInstanceActor->GetActorGuid()); ActorDescInstance && ActorDescInstance->IsChildContainerInstance())
-			{
-				// Add parenting info to init param
-				InInitParams.SetParent(ActorDescInstance->GetContainerInstance(), ActorDescInstance->GetGuid());
-			}
-		}
+		UWorldPartition* OwningWorldPartition = LevelInstanceActor->GetLevel()->GetWorldPartition();
+		// Add parenting info to init param
+		InInitParams.SetParent(OwningWorldPartition ? OwningWorldPartition->GetActorDescContainerInstance() : nullptr, LevelInstanceActor->GetActorGuid());
 	}
 }
 
diff --git a/Engine/Source/Runtime/Engine/Private/LevelInstance/LevelInstanceLevelStreaming.cpp b/Engine/Source/Runtime/Engine/Private/LevelInstance/LevelInstanceLevelStreaming.cpp
--- a/Engine/Source/Runtime/Engine/Private/LevelInstance/LevelInstanceLevelStreaming.cpp
+++ b/Engine/Source/Runtime/Engine/Private/LevelInstance/LevelInstanceLevelStreaming.cpp
@@ -259,16 +259,11 @@ void ULevelStreamingLevelInstance::OnLoadedActorsRemovedFromLevelPostEvent(const
 
 void ULevelStreamingLevelInstance::OnPreInitializeContainerInstance(UActorDescContainerInstance::FInitializeParams& InInitParams, UActorDescContainerInstance* InContainerInstance)
 {
-	// Apply override container
 	AActor* LevelInstanceActor = Cast<AActor>(GetLevelInstance());
-	if (UWorldPartition* OwningWorldPartition = LevelInstanceActor->GetLevel()->GetWorldPartition())
-	{
-		if (FWorldPartitionActorDescInstance* ActorDescInstance = OwningWorldPartition->GetActorDescInstance(LevelInstanceActor->GetActorGuid()); ActorDescInstance && ActorDescInstance->IsChildContainerInstance())
-		{
-			// Add parenting info to init param
-			InInitParams.SetParent(ActorDescInstance->GetContainerInstance(), ActorDescInstance->GetGuid());
-		}
-	}
+	UWorldPartition* OwningWorldPartition = LevelInstanceActor->GetLevel()->GetWorldPartition();
+	
+	// In Editor it is possible to have a non WP parent world in which case we pass in null to the SetParent method, this will ensure that in editor the Level Instance container ID won't be a IsMainContainer() and will properly handle IsMainWorldOnly actors
+	InInitParams.SetParent(OwningWorldPartition ? OwningWorldPartition->GetActorDescContainerInstance() : nullptr, LevelInstanceActor->GetActorGuid());
 }
 
 #endif
@@ -447,7 +442,13 @@ void ULevelStreamingLevelInstance::OnLevelLoadedChanged(ULevel* InLevel)
 			if (UWorldPartition* OuterWorldPartition = NewLoadedLevel->GetWorldPartition())
 			{
 				check(!OuterWorldPartition->IsInitialized());
-				OuterWorldPartition->OnActorDescContainerInstancePreInitialize.BindUObject(this, &ULevelStreamingLevelInstance::OnPreInitializeContainerInstance);
+				
+				// In Non-Editor worlds we want the container of a Level Instance to be considered the main container, it will do its own generate streaming if it is a World Partition
+				// so we don't need to be called on Pre init
+				if (IsEditorWorldMode())
+				{
+					OuterWorldPartition->OnActorDescContainerInstancePreInitialize.BindUObject(this, &ULevelStreamingLevelInstance::OnPreInitializeContainerInstance);
+				}
 
 				if (UWorldPartition* OwningWorldPartition = GetWorld()->GetWorldPartition(); OwningWorldPartition && OwningWorldPartition->IsStreamingEnabled())
 				{
diff --git a/Engine/Source/Runtime/Engine/Private/WorldPartition/ActorDescContainerInstance.cpp b/Engine/Source/Runtime/Engine/Private/WorldPartition/ActorDescContainerInstance.cpp
--- a/Engine/Source/Runtime/Engine/Private/WorldPartition/ActorDescContainerInstance.cpp
+++ b/Engine/Source/Runtime/Engine/Private/WorldPartition/ActorDescContainerInstance.cpp
@@ -69,9 +69,11 @@ void UActorDescContainerInstance::Initialize(const FInitializeParams& InParams)
 	FName OuterWorldContainerPackageName;
 
 	Transform = InParams.Transform;
-	if (InParams.ParentContainerInstance)
+	if (InParams.ContainerActorGuid.IsValid())
 	{
-		ContainerID = FActorContainerID(InParams.ParentContainerInstance->GetContainerID(), InParams.ContainerActorGuid);
+		// It is possible to not have a ParentContainerInstance if we are in a non-WP main world
+		// We want this ContainerInstance to still not have a Main ContainerID to properly handle IsMainWorldOnly() actors
+		ContainerID = FActorContainerID(InParams.ParentContainerInstance ? InParams.ParentContainerInstance->GetContainerID() : FActorContainerID(), InParams.ContainerActorGuid);
 	}
 				
 	// Only consider world if we are outered to a WorldPartition directly
@@ -423,11 +425,20 @@ const FLinkerInstancingContext* UActorDescContainerInstance::GetInstancingContex
 
 const FTransform& UActorDescContainerInstance::GetTransform() const
 {
+	// UActorDescContainerInstance outered to a UWorldPartition
 	if (UWorldPartition* WorldPartition = GetWorldPartition())
 	{
+		// GameWorld: Container instance is necessarly used for Streaming generation. In which case we want to return Identity and Transform of the World Partition will be applied on the LevelStreaming objects.
+		if (GetWorld()->IsGameWorld())
+		{
+			check(GetContainerID().IsMainContainer());
+			return FTransform::Identity;
+		}
+
 		return GetWorldPartition()->GetInstanceTransform();
 	}
 	
+	// Transform is set when this Container Instance is a Child Container Instance created for Streaming Generation
 	return Transform.IsSet() ? Transform.GetValue() : FTransform::Identity;
 }
 
diff --git a/Engine/Source/Runtime/Engine/Public/WorldPartition/ActorDescContainerInstance.h b/Engine/Source/Runtime/Engine/Public/WorldPartition/ActorDescContainerInstance.h
--- a/Engine/Source/Runtime/Engine/Public/WorldPartition/ActorDescContainerInstance.h
+++ b/Engine/Source/Runtime/Engine/Public/WorldPartition/ActorDescContainerInstance.h
@@ -46,7 +46,7 @@ public:
 				
 		FInitializeParams& SetParent(const UActorDescContainerInstance* InParentContainerInstance, const FGuid& InContainerActorGuid)
 		{
-			check(InParentContainerInstance && InContainerActorGuid.IsValid());
+			check(InContainerActorGuid.IsValid());
 			ParentContainerInstance = InParentContainerInstance;
 			ContainerActorGuid = InContainerActorGuid;
 			return *this;
