Fixed evaluation of attributes in a slot with time blend profiles and additive poses not working as expected

#rnx minor

diff --git a/Engine/Source/Runtime/Engine/Private/Animation/AnimInstanceProxy.cpp b/Engine/Source/Runtime/Engine/Private/Animation/AnimInstanceProxy.cpp
--- a/Engine/Source/Runtime/Engine/Private/Animation/AnimInstanceProxy.cpp
+++ b/Engine/Source/Runtime/Engine/Private/Animation/AnimInstanceProxy.cpp
@@ -1810,12 +1810,12 @@ void FAnimInstanceProxy::SlotEvaluatePoseWithBlendProfiles(const FName& SlotNode
 	}
 
 	// Additives.
-	for (int32 PoseIndex = 0; PoseIndex < Poses.Num(); ++PoseIndex)
+	for (int32 PoseIndex = 0; PoseIndex < AdditivePoses.Num(); ++PoseIndex)
 	{
-		FSlotEvaluationPose& AdditivePose = Poses[PoseIndex];
+		FSlotEvaluationPose& AdditivePose = AdditivePoses[PoseIndex];
 		const FAnimationPoseData AdditiveAnimationPoseData(AdditivePose);
 		OutBlendedAnimationPoseData.GetCurve().Accumulate(AdditiveAnimationPoseData.GetCurve(), AdditivePose.Weight);
-		UE::Anim::Attributes::AccumulateAttributes(AdditiveAnimationPoseData.GetAttributes(), OutBlendedAnimationPoseData.GetAttributes(), AdditivePose.Weight, Poses[PoseIndex].AdditiveType);
+		UE::Anim::Attributes::AccumulateAttributes(AdditiveAnimationPoseData.GetAttributes(), OutBlendedAnimationPoseData.GetAttributes(), AdditivePose.Weight, AdditivePose.AdditiveType);
 	}
 }
 
