Level Instance: see level streaming object with the level instance id to get a deterministic level color.

diff --git a/Engine/Source/Runtime/Engine/Private/LevelInstance/LevelInstanceLevelStreaming.cpp b/Engine/Source/Runtime/Engine/Private/LevelInstance/LevelInstanceLevelStreaming.cpp
--- a/Engine/Source/Runtime/Engine/Private/LevelInstance/LevelInstanceLevelStreaming.cpp
+++ b/Engine/Source/Runtime/Engine/Private/LevelInstance/LevelInstanceLevelStreaming.cpp
@@ -315,6 +315,7 @@ ULevelStreamingLevelInstance* ULevelStreamingLevelInstance::LoadInstance(ILevelI
 	if (bOutSuccess)
 	{
 		LevelStreaming->LevelInstanceID = LevelInstance->GetLevelInstanceID();
+		LevelStreaming->LevelColor = FLinearColor::MakeRandomSeededColor(GetTypeHash(LevelInstance->GetLevelInstanceID()));
 
 #if WITH_EDITOR
 		if (!World->IsGameWorld())
