LevelInstance: Fix ChildContainerInstance registration when switching from Embedded <-> Standalone

diff --git a/Engine/Source/Runtime/Engine/Private/WorldPartition/ActorDescContainerInstance.cpp b/Engine/Source/Runtime/Engine/Private/WorldPartition/ActorDescContainerInstance.cpp
--- a/Engine/Source/Runtime/Engine/Private/WorldPartition/ActorDescContainerInstance.cpp
+++ b/Engine/Source/Runtime/Engine/Private/WorldPartition/ActorDescContainerInstance.cpp
@@ -551,10 +551,10 @@ void UActorDescContainerInstance::OnActorDescUpdated(FWorldPartitionActorDesc* I
 	check(ActorDescInstance && ActorDescInstance->IsValid());
 	
 	// Update instance desc
-	(*ActorDescInstance)->UpdateActorDesc(InActorDesc);
+	ActorDescInstance->Get()->UpdateActorDesc(InActorDesc);
 
 	// Re-register container
-	if (bCreateChildContainerHierarchy && ChildContainerInstances.Contains(InActorDesc->GetGuid()))
+	if (bCreateChildContainerHierarchy && ActorDescInstance->Get()->IsChildContainerInstance())
 	{
 		ActorDescInstance->Get()->UpdateChildContainerInstance();
 	}
diff --git a/Engine/Source/Runtime/Engine/Private/WorldPartition/WorldPartitionActorDescInstance.cpp b/Engine/Source/Runtime/Engine/Private/WorldPartition/WorldPartitionActorDescInstance.cpp
--- a/Engine/Source/Runtime/Engine/Private/WorldPartition/WorldPartitionActorDescInstance.cpp
+++ b/Engine/Source/Runtime/Engine/Private/WorldPartition/WorldPartitionActorDescInstance.cpp
@@ -210,14 +210,19 @@ void FWorldPartitionActorDescInstance::UnregisterChildContainerInstance()
 
 void FWorldPartitionActorDescInstance::UpdateChildContainerInstance()
 {
-	check(ChildContainerInstance);
-		
-	// Create before unregistering so that we benefit from shared containers
-	UActorDescContainerInstance* NewChildContainerInstance = IsChildContainerInstance() ? GetActorDesc()->CreateChildContainerInstance(this) : nullptr;
-		
-	ContainerInstance->OnUnregisterChildContainerInstance(GetGuid());
-	ChildContainerInstance->Uninitialize();
+	// Create before unregistering so that we benefit from shared containers (use GetActorDesc->IsChildContainerInstance as we want to know if our updated desc should be a Container instance or not)
+	// ChildContainerInstance member might be non null and we don't want IsChildContainerInstance() to return true in this case if the ActorDesc isn't a Container anymore
+	UActorDescContainerInstance* NewChildContainerInstance = GetActorDesc()->IsChildContainerInstance() ? GetActorDesc()->CreateChildContainerInstance(this) : nullptr;
+	
+	// Unregister previous
+	if (ChildContainerInstance)
+	{
+		ContainerInstance->OnUnregisterChildContainerInstance(GetGuid());
+		ChildContainerInstance->Uninitialize();
+		ChildContainerInstance = nullptr;
+	}
 	
+	// Register new if it is valid
 	if (NewChildContainerInstance)
 	{
 		ChildContainerInstance = NewChildContainerInstance;
