Texture nop comments only

diff --git a/Engine/Source/Runtime/Engine/Classes/Engine/Texture.h b/Engine/Source/Runtime/Engine/Classes/Engine/Texture.h
--- a/Engine/Source/Runtime/Engine/Classes/Engine/Texture.h
+++ b/Engine/Source/Runtime/Engine/Classes/Engine/Texture.h
@@ -153,6 +153,14 @@ struct FTextureSource
 	* Init and copy in texture bits from Image
 	* 
 	* @param Image -  Image to initialize with
+	*
+	* FImageView has gamma information too that is lost
+	* TextureSource does not store gamma information (it's in the owning Texture)
+	* this function does NOT set Texture->SRGB , you must do so!
+	*
+	* Init() does UseHashAsGuid
+	* Init() must be done inside PreEdit/PostEdit on the owning Texture
+	*
 	*/
 	ENGINE_API void Init(const FImageView & Image);
 
diff --git a/Engine/Source/Runtime/Engine/Private/Texture.cpp b/Engine/Source/Runtime/Engine/Private/Texture.cpp
--- a/Engine/Source/Runtime/Engine/Private/Texture.cpp
+++ b/Engine/Source/Runtime/Engine/Private/Texture.cpp
@@ -2110,6 +2110,7 @@ void FTextureSource::Init(const FImageView & Image)
 
 	// FImageView has gamma information too that is lost
 	// TextureSource does not store gamma information (it's in the owning Texture)
+	// this function does NOT set Texture->SRGB , you must do so!
 
 	Init(Image.SizeX,Image.SizeY,Image.NumSlices,1,SourceFormat,(const uint8 *)Image.RawData);
 }
