Fix crash caused by pulling empty texture metadata from DDC

diff --git a/Engine/Source/Runtime/Engine/Private/TextureDerivedDataTask.cpp b/Engine/Source/Runtime/Engine/Private/TextureDerivedDataTask.cpp
--- a/Engine/Source/Runtime/Engine/Private/TextureDerivedDataTask.cpp
+++ b/Engine/Source/Runtime/Engine/Private/TextureDerivedDataTask.cpp
@@ -1475,7 +1475,7 @@ static void DDC1_FetchAndFillDerivedData(
 			if (bGotDDCData)
 			{
 				// We don't necessarily get the metadata until we force a texture rebuild.
-				if (MetadataBuffer.IsNull() == false)
+				if (MetadataBuffer.IsNull() == false && MetadataBuffer.GetSize() != 0)
 				{
 					PackTextureBuildMetadataInPlatformData(DerivedData, FCbObject(MetadataBuffer));
 				}
@@ -1517,7 +1517,7 @@ static void DDC1_FetchAndFillDerivedData(
 		{
 			// Only read the metadata if we actually got the main data.
 			// We don't necessarily get the metadata until we force a texture rebuild.
-			if (MetadataBuffer.IsNull() == false)
+			if (MetadataBuffer.IsNull() == false && MetadataBuffer.GetSize() != 0)
 			{
 				PackTextureBuildMetadataInPlatformData(DerivedData, FCbObject(MetadataBuffer));
 			}
