Add CSV metadata for the streaming pool size

diff --git a/Engine/Source/Runtime/Engine/Private/Streaming/StreamingManagerTexture.cpp b/Engine/Source/Runtime/Engine/Private/Streaming/StreamingManagerTexture.cpp
--- a/Engine/Source/Runtime/Engine/Private/Streaming/StreamingManagerTexture.cpp
+++ b/Engine/Source/Runtime/Engine/Private/Streaming/StreamingManagerTexture.cpp
@@ -14,10 +14,12 @@
 #include "Engine/Texture2DArray.h"
 #include "Materials/MaterialInterface.h"
 #include "Misc/ConfigCacheIni.h"
+#include "Misc/StringBuilder.h"
 #include "RenderedTextureStats.h"
 #include "UObject/UObjectIterator.h"
 #include "DeviceProfiles/DeviceProfile.h"
 #include "DeviceProfiles/DeviceProfileManager.h"
+#include "ProfilingDebugging/CsvProfiler.h"
 #include "RenderAssetUpdate.h"
 #include "RenderingThread.h"
 #include "Streaming/AsyncTextureStreaming.h"
@@ -1613,6 +1615,7 @@ void FRenderAssetStreamingManager::CheckUserSettings()
 		if (TexturePoolSize != GTexturePoolSize)
 		{
 			UE_LOG(LogContentStreaming,Log,TEXT("Texture pool size now %d MB"), int32(TexturePoolSize/1024/1024));
+			CSV_METADATA(TEXT("StreamingPoolSizeMB"), *WriteToString<32>(int32(TexturePoolSize / 1024 / 1024)));
 			GTexturePoolSize = TexturePoolSize;
 		}
 	}
