[NetPhysics] Soft-snap towards aligned target instead of source target in predictive interpolation.

diff --git a/Engine/Source/Runtime/Engine/Private/PhysicsEngine/PhysicsReplication.cpp b/Engine/Source/Runtime/Engine/Private/PhysicsEngine/PhysicsReplication.cpp
--- a/Engine/Source/Runtime/Engine/Private/PhysicsEngine/PhysicsReplication.cpp
+++ b/Engine/Source/Runtime/Engine/Private/PhysicsEngine/PhysicsReplication.cpp
@@ -161,15 +161,21 @@ namespace PhysicsReplicationCVars
 		static float SoftSnapRotStrength = 0.5f;
 		static FAutoConsoleVariableRef CVarSoftSnapRotStrength(TEXT("np2.PredictiveInterpolation.SoftSnapRotStrength"), SoftSnapRotStrength, TEXT("Value in percent between 0.0 - 1.0 representing how much to softsnap each tick of the remaining distance."));
 
+		static bool bSoftSnapToSource = false;
+		static FAutoConsoleVariableRef CVarSoftSnapToSource(TEXT("np2.PredictiveInterpolation.SoftSnapToSource"), bSoftSnapToSource, TEXT("If true, soft snap will be performed towards the source state of the current target instead of the predicted state of the current target."));
+
 		static float EarlyOutDistanceSqr = 1.0f;
 		static FAutoConsoleVariableRef CVarEarlyOutDistanceSqr(TEXT("np2.PredictiveInterpolation.EarlyOutDistanceSqr"), EarlyOutDistanceSqr, TEXT("Squared value. If object is within this distance from the source target, early out from replication and apply sleep if replicated."));
 		
-		static float EarlyOutAngle = 0.75f;
+		static float EarlyOutAngle = 1.5f;
 		static FAutoConsoleVariableRef CVarEarlyOutAngle(TEXT("np2.PredictiveInterpolation.EarlyOutAngle"), EarlyOutAngle, TEXT("If object is within this rotational angle (in degrees) from the source target, early out from replication and apply sleep if replicated."));
 		
 		static bool bEarlyOutWithVelocity = true;
 		static FAutoConsoleVariableRef CVarEarlyOutWithVelocity(TEXT("np2.PredictiveInterpolation.EarlyOutWithVelocity"), bEarlyOutWithVelocity, TEXT("If true, allow replication logic to early out if current velocities are driving replication well enough. If false, only early out if target velocity is zero."));
 
+		static bool bSkipVelocityRepOnPosEarlyOut = true;
+		static FAutoConsoleVariableRef CVarSkipVelocityRepOnPosEarlyOut(TEXT("np2.PredictiveInterpolation.SkipVelocityRepOnPosEarlyOut"), bSkipVelocityRepOnPosEarlyOut, TEXT("If true, don't run linear velocity replication if position can early out but angular can't early out."));
+
 		static bool bPostResimWaitForUpdate = false;
 		static FAutoConsoleVariableRef CVarPostResimWaitForUpdate(TEXT("np2.PredictiveInterpolation.PostResimWaitForUpdate"), bPostResimWaitForUpdate, TEXT("After a resimulation, wait for replicated states that correspond to post-resim state before processing replication again."));
 		
@@ -1578,7 +1584,7 @@ bool FPhysicsReplicationAsync::PredictiveInterpolation(Chaos::FPBDRigidParticleH
 		const float RotCorrectionTime = FMath::Max(PhysicsReplicationCVars::PredictiveInterpolationCVars::RotCorrectionTimeBase + AverageReceiveIntervalSeconds + RTT * PhysicsReplicationCVars::PredictiveInterpolationCVars::RotCorrectionTimeMultiplier,
 			DeltaSeconds + PhysicsReplicationCVars::PredictiveInterpolationCVars::RotCorrectionTimeMin);
 
-		if (!bXCanEarlyOut)
+		if ((bXCanEarlyOut && PhysicsReplicationCVars::PredictiveInterpolationCVars::bSkipVelocityRepOnPosEarlyOut) == false)
 		{	// --- Velocity Replication ---
 			
 			// Get PosDiff
@@ -1653,8 +1659,14 @@ bool FPhysicsReplicationAsync::PredictiveInterpolation(Chaos::FPBDRigidParticleH
 
 		if (bSoftSnap)
 		{
-			const FVector SoftSnapPos = FMath::Lerp(FVector(CurrentState.Position), Target.PrevPosTarget, FMath::Clamp(PhysicsReplicationCVars::PredictiveInterpolationCVars::SoftSnapPosStrength, 0.0f, 1.0f));
-			const FQuat SoftSnapRot = FQuat::Slerp(CurrentState.Quaternion, Target.PrevRotTarget, FMath::Clamp(PhysicsReplicationCVars::PredictiveInterpolationCVars::SoftSnapRotStrength, 0.0f, 1.0f));
+			const FVector SoftSnapPos = FMath::Lerp(FVector(CurrentState.Position),
+			PhysicsReplicationCVars::PredictiveInterpolationCVars::bSoftSnapToSource ? Target.PrevPosTarget : Target.TargetState.Position,
+			FMath::Clamp(PhysicsReplicationCVars::PredictiveInterpolationCVars::SoftSnapPosStrength, 0.0f, 1.0f));
+		
+			const FQuat SoftSnapRot = FQuat::Slerp(CurrentState.Quaternion,
+				PhysicsReplicationCVars::PredictiveInterpolationCVars::bSoftSnapToSource ? Target.PrevRotTarget : Target.TargetState.Quaternion,
+			FMath::Clamp(PhysicsReplicationCVars::PredictiveInterpolationCVars::SoftSnapRotStrength, 0.0f, 1.0f));
+		
 			Handle->SetX(SoftSnapPos);
 			Handle->SetP(SoftSnapPos);
 			Handle->SetR(SoftSnapRot);
