Add callbacks for set base material and assigned override material

diff --git a/Engine/Source/Runtime/Engine/Classes/Components/MeshComponent.h b/Engine/Source/Runtime/Engine/Classes/Components/MeshComponent.h
--- a/Engine/Source/Runtime/Engine/Classes/Components/MeshComponent.h
+++ b/Engine/Source/Runtime/Engine/Classes/Components/MeshComponent.h
@@ -178,6 +178,11 @@ public:
 		return (ParameterCache ? ParameterCache->ScalarParameterDefaultValue : 0.f);
 	}
 protected:
+
+	//~ Begin UObject Interface.
+	ENGINE_API virtual void BeginDestroy() override;
+	//~ End UObject Interface.
+
 	/** Get the default overlay material used by a mesh */
 	virtual UMaterialInterface* GetDefaultOverlayMaterial() const { return nullptr; };
 	
diff --git a/Engine/Source/Runtime/Engine/Classes/Materials/MaterialInstanceDynamic.h b/Engine/Source/Runtime/Engine/Classes/Materials/MaterialInstanceDynamic.h
--- a/Engine/Source/Runtime/Engine/Classes/Materials/MaterialInstanceDynamic.h
+++ b/Engine/Source/Runtime/Engine/Classes/Materials/MaterialInstanceDynamic.h
@@ -210,8 +210,12 @@ class UMaterialInstanceDynamic : public UMaterialInstance
 	// This overrides does the remapping before looking at the parent data.
 	ENGINE_API virtual float GetTextureDensity(FName TextureName, const struct FMeshUVChannelInfo& UVChannelData) const override;
 
-private:
+protected:
+
 	ENGINE_API void InitializeMID(class UMaterialInterface* ParentMaterial);
+
+private:
+	
 	ENGINE_API void UpdateCachedDataDynamic();
 };
 
diff --git a/Engine/Source/Runtime/Engine/Classes/Materials/MaterialInterface.h b/Engine/Source/Runtime/Engine/Classes/Materials/MaterialInterface.h
--- a/Engine/Source/Runtime/Engine/Classes/Materials/MaterialInterface.h
+++ b/Engine/Source/Runtime/Engine/Classes/Materials/MaterialInterface.h
@@ -5,6 +5,7 @@
 
 #include "Async/TaskGraphInterfaces.h"
 #include "CoreMinimal.h"
+#include "Delegates/Delegate.h"
 #include "MaterialTypes.h"
 #include "Containers/ArrayView.h"
 #include "UObject/ObjectMacros.h"
@@ -252,6 +253,8 @@ public:
 	bool bLoadedCachedExpressionData = false;
 };
 
+DECLARE_MULTICAST_DELEGATE_OneParam(FOnBaseMaterialIsSet, UMaterialInterface*);
+
 UCLASS(abstract, BlueprintType, MinimalAPI, HideCategories = (Thumbnail))
 class UMaterialInterface : public UObject, public IBlendableInterface, public IInterface_AssetUserData
 {
@@ -290,6 +293,9 @@ public:
 	UPROPERTY()
 	uint8 bIncludedInBaseGame : 1;
 
+	/** Event triggered when the base material is set */
+	FOnBaseMaterialIsSet OnBaseMaterialSetEvent;
+
 	/* -------------------------- */
 
 	/** A fence to track when the primitive is no longer used as a parent */
@@ -405,6 +411,11 @@ public:
 	UFUNCTION(BlueprintCallable, Category="Rendering|Material")
 	ENGINE_API UMaterial* GetBaseMaterial();
 
+	/** Callback triggered when the material has been assigned as an override material */
+	ENGINE_API virtual void OnAssignedAsOverride(const UObject* Owner);
+	/** Callback triggered when the material has been removed as an override material */
+	ENGINE_API virtual void OnRemovedAsOverride(const UObject* Owner);
+
 	/**
 	 * Get the material which we are instancing.
 	 * Walks up parent chain and finds the base Material that this is an instance of. 
diff --git a/Engine/Source/Runtime/Engine/Private/Components/MeshComponent.cpp b/Engine/Source/Runtime/Engine/Private/Components/MeshComponent.cpp
--- a/Engine/Source/Runtime/Engine/Private/Components/MeshComponent.cpp
+++ b/Engine/Source/Runtime/Engine/Private/Components/MeshComponent.cpp
@@ -85,9 +85,19 @@ void UMeshComponent::SetMaterial(int32 ElementIndex, UMaterialInterface* Materia
 				}
 			}	
 
+			if (UMaterialInterface* PreviousMaterial = OverrideMaterials[ElementIndex].Get())
+			{
+				PreviousMaterial->OnRemovedAsOverride(this);
+			}
+
 			// Set the material and invalidate things
 			OverrideMaterials[ElementIndex] = Material;
 
+			if (Material)
+			{
+				Material->OnAssignedAsOverride(this);
+			}
+
 			// Precache PSOs again
 			PrecachePSOs();
 
@@ -193,13 +203,24 @@ void UMeshComponent::PostEditChangeChainProperty(FPropertyChangedChainEvent& Pro
 void UMeshComponent::CleanUpOverrideMaterials()
 {
 	bool bUpdated = false;
+	int32 NumMaterials = GetNumMaterials();
+	int32 NumOverrideMaterials = OverrideMaterials.Num();
 
 	// We have to remove material override Ids that are bigger then the material list
-	if (OverrideMaterials.Num() > GetNumMaterials())
+	if (NumOverrideMaterials > NumMaterials)
 	{
 		//Remove the override material id that are superior to the static mesh materials number
-		int32 RemoveCount = OverrideMaterials.Num() - GetNumMaterials();
-		OverrideMaterials.RemoveAt(GetNumMaterials(), RemoveCount);
+		int32 RemoveCount = NumOverrideMaterials - NumMaterials;
+
+		for (int32 MatIndex = NumMaterials; MatIndex < NumOverrideMaterials; MatIndex++)
+		{
+			if (UMaterialInterface* MatInterface = OverrideMaterials[MatIndex].Get())
+			{
+				MatInterface->OnRemovedAsOverride(this);
+			}
+		}
+
+		OverrideMaterials.RemoveAt(NumMaterials, RemoveCount);
 		bUpdated = true;
 	}
 
@@ -214,6 +235,14 @@ void UMeshComponent::EmptyOverrideMaterials()
 {
 	if (OverrideMaterials.Num())
 	{
+		for (int32 MatIndex = 0; MatIndex < OverrideMaterials.Num(); MatIndex++)
+		{
+			if (UMaterialInterface* MatInterface = OverrideMaterials[MatIndex].Get())
+			{
+				MatInterface->OnRemovedAsOverride(this);
+			}
+		}
+
 		OverrideMaterials.Reset();
 		MarkRenderStateDirty();
 	}
@@ -472,6 +501,21 @@ void UMeshComponent::MarkCachedMaterialParameterNameIndicesDirty()
 	bCachedMaterialParameterIndicesAreDirty = true;
 }
 
+void UMeshComponent::BeginDestroy()
+{
+	for (int32 MatIndex = 0; MatIndex < OverrideMaterials.Num(); MatIndex++)
+	{
+		if (UMaterialInterface* MatInterface = OverrideMaterials[MatIndex].Get())
+		{
+			MatInterface->OnRemovedAsOverride(this);
+		}
+
+		OverrideMaterials[MatIndex] = nullptr;
+	}
+
+	Super::BeginDestroy();
+}
+
 void UMeshComponent::CacheMaterialParameterNameIndices()
 {
 	QUICK_SCOPE_CYCLE_COUNTER(STAT_CacheMaterialParameterNameIndices);
diff --git a/Engine/Source/Runtime/Engine/Private/Materials/MaterialInstance.cpp b/Engine/Source/Runtime/Engine/Private/Materials/MaterialInstance.cpp
--- a/Engine/Source/Runtime/Engine/Private/Materials/MaterialInstance.cpp
+++ b/Engine/Source/Runtime/Engine/Private/Materials/MaterialInstance.cpp
@@ -3438,6 +3438,12 @@ bool UMaterialInstance::SetParentInternal(UMaterialInterface* NewParent, bool Re
 			InitResources();
 		}
 	}
+
+	if (bSetParent)
+	{
+		OnBaseMaterialSetEvent.Broadcast(this);
+	}
+
 	return bSetParent;
 }
 
diff --git a/Engine/Source/Runtime/Engine/Private/Materials/MaterialInterface.cpp b/Engine/Source/Runtime/Engine/Private/Materials/MaterialInterface.cpp
--- a/Engine/Source/Runtime/Engine/Private/Materials/MaterialInterface.cpp
+++ b/Engine/Source/Runtime/Engine/Private/Materials/MaterialInterface.cpp
@@ -1220,6 +1220,16 @@ UMaterial* UMaterialInterface::GetBaseMaterial()
 	return GetMaterial();
 }
 
+void UMaterialInterface::OnAssignedAsOverride(const UObject* Owner)
+{
+
+}
+
+void UMaterialInterface::OnRemovedAsOverride(const UObject* Owner)
+{
+
+}
+
 bool DoesMaterialUseTexture(const UMaterialInterface* Material,const UTexture* CheckTexture)
 {
 	//Do not care if we're running dedicated server
