Fix: Saving Level Instance over existing Level Instance (while some instances of the existing map are loaded)

diff --git a/Engine/Source/Runtime/Engine/Private/WorldPartition/ActorDescContainerInstance.cpp b/Engine/Source/Runtime/Engine/Private/WorldPartition/ActorDescContainerInstance.cpp
--- a/Engine/Source/Runtime/Engine/Private/WorldPartition/ActorDescContainerInstance.cpp
+++ b/Engine/Source/Runtime/Engine/Private/WorldPartition/ActorDescContainerInstance.cpp
@@ -64,6 +64,24 @@ void UActorDescContainerInstance::OnContainerUpdated(FName ContainerPackage)
 	}
 }
 
+void UActorDescContainerInstance::OnContainerReplaced(UActorDescContainer* InOldContainer, UActorDescContainer* InNewContainer)
+{
+	if (ChildContainerInstances.Num())
+	{
+		TMap<FGuid, TObjectPtr<UActorDescContainerInstance>> CopyChildContainerInstances(ChildContainerInstances);
+		for (auto& [ContainerGuid, ContainerInstance] : CopyChildContainerInstances)
+		{
+			if (ContainerInstance->GetContainer() == InOldContainer)
+			{
+				FWorldPartitionActorDescInstance* ContainerDescInstance = GetActorDescInstance(ContainerGuid);
+				check(ContainerDescInstance);
+				ContainerDescInstance->UpdateChildContainerInstance();
+				check(ContainerDescInstance->GetChildContainerInstance()->GetContainer() == InNewContainer);
+			}
+		}
+	}
+}
+
 void UActorDescContainerInstance::Initialize(const FInitializeParams& InParams)
 {
 	FName OuterWorldContainerPackageName;
@@ -220,6 +238,7 @@ void UActorDescContainerInstance::RegisterDelegates()
 		if (bCreateChildContainerHierarchy)
 		{
 			UActorDescContainerSubsystem::GetChecked().ContainerUpdated().AddUObject(this, &UActorDescContainerInstance::OnContainerUpdated);
+			UActorDescContainerSubsystem::GetChecked().ContainerReplaced().AddUObject(this, &UActorDescContainerInstance::OnContainerReplaced);
 		}
 
 		// No need to register Added descs events for instanced worlds as they don't support it for now (Level Instances get reloaded after an edit)
@@ -252,6 +271,7 @@ void UActorDescContainerInstance::UnregisterDelegates()
 		if (UActorDescContainerSubsystem* ActorDescContainerSubsystem = UActorDescContainerSubsystem::Get())
 		{
 			ActorDescContainerSubsystem->ContainerUpdated().RemoveAll(this);
+			ActorDescContainerSubsystem->ContainerReplaced().RemoveAll(this);
 		}
 	}
 }
diff --git a/Engine/Source/Runtime/Engine/Private/WorldPartition/ActorDescContainerSubsystem.cpp b/Engine/Source/Runtime/Engine/Private/WorldPartition/ActorDescContainerSubsystem.cpp
--- a/Engine/Source/Runtime/Engine/Private/WorldPartition/ActorDescContainerSubsystem.cpp
+++ b/Engine/Source/Runtime/Engine/Private/WorldPartition/ActorDescContainerSubsystem.cpp
@@ -104,15 +104,32 @@ void UActorDescContainerSubsystem::FContainerManager::UpdateContainerBoundsFromP
 
 void UActorDescContainerSubsystem::FContainerManager::SetContainerPackage(UActorDescContainer* Container, FName PackageName)
 {
+	// Remove and Copy existing Registration (with previous ContainerName: ex: /Temp/Untitled)
 	FRegisteredContainer RegisteredContainer;
 	const bool bRegistered = RegisteredContainers.RemoveAndCopyValue(Container->GetContainerName(), RegisteredContainer);
 
+	// Update Name
 	Container->SetContainerPackage(PackageName);
-
+		
 	if(bRegistered)
 	{
+		// Check if we have an existing container registered with the new name, which means we are saving a map over another one
+		FRegisteredContainer ReplacedContainer;
+		if (RegisteredContainers.RemoveAndCopyValue(Container->GetContainerName(), ReplacedContainer))
+		{
+			// Move it out of the way with unique package name without losing the ref counts so that NotifyContainerReplaced listeners will be able to properly unregister
+			FString ReplacePackageName = PackageName.ToString() + TEXT("_Replaced_") + FGuid::NewGuid().ToString();
+			ReplacedContainer.Container->SetContainerPackage(*ReplacePackageName);
+			RegisteredContainers.Add(ReplacedContainer.Container->GetContainerName(), ReplacedContainer);
+		}
+		
 		check(RegisteredContainer.Container == Container);
 		RegisteredContainers.Add(Container->GetContainerName(), RegisteredContainer);
+
+		if (ReplacedContainer.Container)
+		{
+			Owner->NotifyContainerReplaced(ReplacedContainer.Container, RegisteredContainer.Container);
+		}
 	}
 }
 #endif
\ No newline at end of file
diff --git a/Engine/Source/Runtime/Engine/Public/WorldPartition/ActorDescContainerInstance.h b/Engine/Source/Runtime/Engine/Public/WorldPartition/ActorDescContainerInstance.h
--- a/Engine/Source/Runtime/Engine/Public/WorldPartition/ActorDescContainerInstance.h
+++ b/Engine/Source/Runtime/Engine/Public/WorldPartition/ActorDescContainerInstance.h
@@ -131,6 +131,7 @@ protected:
 	const UDataLayerManager* GetResolvingDataLayerManager() const;
 private:
 	void OnContainerUpdated(FName ContainerPackage);
+	void OnContainerReplaced(UActorDescContainer* InOldContainer, UActorDescContainer* InNewContainer);
 	void SetContainerPackage(FName InContainerPackageName);
 		
 	void RegisterDelegates();
diff --git a/Engine/Source/Runtime/Engine/Public/WorldPartition/ActorDescContainerSubsystem.h b/Engine/Source/Runtime/Engine/Public/WorldPartition/ActorDescContainerSubsystem.h
--- a/Engine/Source/Runtime/Engine/Public/WorldPartition/ActorDescContainerSubsystem.h
+++ b/Engine/Source/Runtime/Engine/Public/WorldPartition/ActorDescContainerSubsystem.h
@@ -33,9 +33,11 @@ public:
 	ENGINE_API static UActorDescContainerSubsystem* Get();
 	ENGINE_API static UActorDescContainerSubsystem& GetChecked();
 
-	DECLARE_EVENT_OneParam(UActorDescContainerSubsystemFOnContainerUpdated, FContainerUpdatedEvent, FName);
+	DECLARE_EVENT_OneParam(UActorDescContainerSubsystem, FContainerUpdatedEvent, FName);
+	DECLARE_EVENT_TwoParams(UActorDescContainerSubsystem, FContainerReplacedEvent, UActorDescContainer*, UActorDescContainer*);
 
 	FContainerUpdatedEvent& ContainerUpdated() { return OnContainerUpdated; }
+	FContainerReplacedEvent& ContainerReplaced() { return OnContainerReplaced; }
 	
 	template <class ContainerType = UActorDescContainer>
 	ContainerType* RegisterContainer(const typename UActorDescContainer::FInitializeParams& InitParams) { return ContainerManager.RegisterContainer<ContainerType>(InitParams); }
@@ -49,6 +51,10 @@ public:
 		ContainerManager.UpdateContainerBoundsFromPackage(ContainerPackage);
 		OnContainerUpdated.Broadcast(ContainerPackage);
 	}
+	void NotifyContainerReplaced(UActorDescContainer* OldContainer, UActorDescContainer* NewContainer)
+	{
+		OnContainerReplaced.Broadcast(OldContainer, NewContainer);
+	}
 #endif
 
 	//~ Begin UObject Interface
@@ -60,6 +66,7 @@ public:
 private:
 #if WITH_EDITOR
 	FContainerUpdatedEvent OnContainerUpdated;
+	FContainerReplacedEvent OnContainerReplaced;
 
 	class FContainerManager
 	{
@@ -133,7 +140,7 @@ private:
 
 	private:
 		TMap<FString, FRegisteredContainer> RegisteredContainers;
-		UObject* Owner = nullptr;
+		UActorDescContainerSubsystem* Owner = nullptr;
 	};
 
 	FContainerManager ContainerManager;
