World Partition - Data Layer: Prevent actors from being added to private data layers if they are not part of the same world

diff --git a/Engine/Source/Runtime/Engine/Private/WorldPartition/DataLayer/DataLayerInstancePrivate.cpp b/Engine/Source/Runtime/Engine/Private/WorldPartition/DataLayer/DataLayerInstancePrivate.cpp
--- a/Engine/Source/Runtime/Engine/Private/WorldPartition/DataLayer/DataLayerInstancePrivate.cpp
+++ b/Engine/Source/Runtime/Engine/Private/WorldPartition/DataLayer/DataLayerInstancePrivate.cpp
@@ -62,6 +62,17 @@ bool UDataLayerInstancePrivate::PerformAddActor(AActor* InActor) const
 	return FAssignActorDataLayer::AddDataLayerAsset(InActor, DataLayerAsset);
 }
 
+bool UDataLayerInstancePrivate::CanAddActor(AActor* InActor) const
+{
+	if (!Super::CanAddActor(InActor))
+	{
+		return false;
+	}
+
+	return GetOuterWorldDataLayers() == InActor->GetLevel()->GetWorldDataLayers() && // Make sure the instance is part of the same WorldDataLayers as the actor's level WorldDataLayer.
+		   UDataLayerManager::GetDataLayerManager(InActor)->GetDataLayerInstance(DataLayerAsset) != nullptr; // Make sure the DataLayerInstance exists for this level
+}
+
 bool UDataLayerInstancePrivate::PerformRemoveActor(AActor* InActor) const
 {
 	return FAssignActorDataLayer::RemoveDataLayerAsset(InActor, DataLayerAsset);
diff --git a/Engine/Source/Runtime/Engine/Public/WorldPartition/DataLayer/DataLayerInstancePrivate.h b/Engine/Source/Runtime/Engine/Public/WorldPartition/DataLayer/DataLayerInstancePrivate.h
--- a/Engine/Source/Runtime/Engine/Public/WorldPartition/DataLayer/DataLayerInstancePrivate.h
+++ b/Engine/Source/Runtime/Engine/Public/WorldPartition/DataLayer/DataLayerInstancePrivate.h
@@ -23,6 +23,7 @@ public:
 	virtual bool SupportsActorFilters() const override { return GetAsset()->SupportsActorFilters(); }
 	virtual bool IsIncludedInActorFilterDefault() const override { return bIsIncludedInActorFilterDefault; }
 	virtual bool CanEditDataLayerShortName() const override { return true; }
+	virtual bool CanAddActor(AActor* InActor) const override;
 #endif
 
 	virtual const UDataLayerAsset* GetAsset() const override { return DataLayerAsset; }
