World Partition - TActorDescContainerInstanceCollection: Renamed ForEachActorDescContainer to ForEachActorDescContainerInstance

diff --git a/Engine/Plugins/Editor/DataValidation/Source/DataValidation/Private/WorldPartitionChangelistValidator.cpp b/Engine/Plugins/Editor/DataValidation/Source/DataValidation/Private/WorldPartitionChangelistValidator.cpp
--- a/Engine/Plugins/Editor/DataValidation/Source/DataValidation/Private/WorldPartitionChangelistValidator.cpp
+++ b/Engine/Plugins/Editor/DataValidation/Source/DataValidation/Private/WorldPartitionChangelistValidator.cpp
@@ -214,7 +214,7 @@ void UWorldPartitionChangelistValidator::ValidateActorsAndDataLayersFromChangeLi
 			.SetErrorHandler(this)
 			.SetEnableStreaming(!ULevel::GetIsStreamingDisabledFromPackage(MapPath.GetPackageName()));
 
-		ContainersToValidate.ForEachActorDescContainer([&Params](const UActorDescContainerInstance* ContainerInstance)
+		ContainersToValidate.ForEachActorDescContainerInstance([&Params](const UActorDescContainerInstance* ContainerInstance)
 		{
 			for (UActorDescContainerInstance::TConstIterator<> Iterator(ContainerInstance); Iterator; ++Iterator)
 			{
diff --git a/Engine/Plugins/Editor/ObjectMixer/ObjectMixer/Source/ObjectMixer/Private/Views/List/Modes/ObjectMixerOutlinerMode.cpp b/Engine/Plugins/Editor/ObjectMixer/ObjectMixer/Source/ObjectMixer/Private/Views/List/Modes/ObjectMixerOutlinerMode.cpp
--- a/Engine/Plugins/Editor/ObjectMixer/ObjectMixer/Source/ObjectMixer/Private/Views/List/Modes/ObjectMixerOutlinerMode.cpp
+++ b/Engine/Plugins/Editor/ObjectMixer/ObjectMixer/Source/ObjectMixer/Private/Views/List/Modes/ObjectMixerOutlinerMode.cpp
@@ -2048,7 +2048,7 @@ bool FObjectMixerOutlinerMode::HasErrors() const
 		{
 			bool bHasErrors = 0;
 
-			WorldPartition->ForEachActorDescContainer([&bHasErrors](UActorDescContainerInstance* ActorDescContainerInstance)
+			WorldPartition->ForEachActorDescContainerInstance([&bHasErrors](UActorDescContainerInstance* ActorDescContainerInstance)
 			{
 				if (ActorDescContainerInstance->GetContainer()->HasInvalidActors())
 				{
@@ -2078,7 +2078,7 @@ void FObjectMixerOutlinerMode::RepairErrors() const
 			ISourceControlProvider& SourceControlProvider = SourceControlModule.GetProvider();
 
 			TArray<FAssetData> InvalidActorAssets;
-			WorldPartition->ForEachActorDescContainer([&InvalidActorAssets](UActorDescContainerInstance* ActorDescContainerInstance)
+			WorldPartition->ForEachActorDescContainerInstance([&InvalidActorAssets](UActorDescContainerInstance* ActorDescContainerInstance)
 			{
 				for (const FAssetData& InvalidActor : ActorDescContainerInstance->GetContainer()->GetInvalidActors())
 				{
diff --git a/Engine/Source/Editor/SceneOutliner/Private/ActorBrowsingMode.cpp b/Engine/Source/Editor/SceneOutliner/Private/ActorBrowsingMode.cpp
--- a/Engine/Source/Editor/SceneOutliner/Private/ActorBrowsingMode.cpp
+++ b/Engine/Source/Editor/SceneOutliner/Private/ActorBrowsingMode.cpp
@@ -1304,7 +1304,7 @@ bool FActorBrowsingMode::HasErrors() const
 		{
 			bool bHasErrors = 0;
 
-			WorldPartition->ForEachActorDescContainer([&bHasErrors](UActorDescContainerInstance* ActorDescContainerInstance)
+			WorldPartition->ForEachActorDescContainerInstance([&bHasErrors](UActorDescContainerInstance* ActorDescContainerInstance)
 			{
 				if (ActorDescContainerInstance->GetContainer()->HasInvalidActors())
 				{
@@ -1334,7 +1334,7 @@ void FActorBrowsingMode::RepairErrors() const
 			ISourceControlProvider& SourceControlProvider = SourceControlModule.GetProvider();
 
 			TArray<FAssetData> InvalidActorAssets;
-			WorldPartition->ForEachActorDescContainer([&InvalidActorAssets](UActorDescContainerInstance* ActorDescContainerInstance)
+			WorldPartition->ForEachActorDescContainerInstance([&InvalidActorAssets](UActorDescContainerInstance* ActorDescContainerInstance)
 			{
 				for (const FAssetData& InvalidActor : ActorDescContainerInstance->GetContainer()->GetInvalidActors())
 				{
diff --git a/Engine/Source/Editor/WorldPartitionEditor/Private/WorldPartition/HLOD/HLODEditorData.cpp b/Engine/Source/Editor/WorldPartitionEditor/Private/WorldPartition/HLOD/HLODEditorData.cpp
--- a/Engine/Source/Editor/WorldPartitionEditor/Private/WorldPartition/HLOD/HLODEditorData.cpp
+++ b/Engine/Source/Editor/WorldPartitionEditor/Private/WorldPartition/HLOD/HLODEditorData.cpp
@@ -21,7 +21,7 @@ FWorldPartitionHLODEditorData::FWorldPartitionHLODEditorData(UWorldPartition* In
 	WorldPartition->OnActorDescContainerInstanceUnregistered.AddRaw(this, &FWorldPartitionHLODEditorData::OnActorDescContainerInstanceUnregistered);
 
 	// Since this is created upon WP init, we missed the first broadcasts for existing container instances. Register them manually.
-	WorldPartition->ForEachActorDescContainer([this](UActorDescContainerInstance* InContainerInstance)
+	WorldPartition->ForEachActorDescContainerInstance([this](UActorDescContainerInstance* InContainerInstance)
 	{
 		OnActorDescContainerInstanceRegistered(InContainerInstance);
 	});	
diff --git a/Engine/Source/Runtime/Engine/Private/WorldPartition/WorldPartition.cpp b/Engine/Source/Runtime/Engine/Private/WorldPartition/WorldPartition.cpp
--- a/Engine/Source/Runtime/Engine/Private/WorldPartition/WorldPartition.cpp
+++ b/Engine/Source/Runtime/Engine/Private/WorldPartition/WorldPartition.cpp
@@ -621,7 +621,7 @@ void UWorldPartition::Initialize(UWorld* InWorld, const FTransform& InTransform)
 
 		ActorDescContainerInstance = RegisterActorDescContainerInstance(UActorDescContainerInstance::FInitializeParams(ContainerPackageName));
 
-		ForEachActorDescContainer([this, bIsEditor, bIsCooking](UActorDescContainerInstance* InActorDescContainerInstance)
+		ForEachActorDescContainerInstance([this, bIsEditor, bIsCooking](UActorDescContainerInstance* InActorDescContainerInstance)
 		{
 			InitializeActorDescContainerEditorStreaming(InActorDescContainerInstance, bIsEditor && !bIsCooking);
 		});
diff --git a/Engine/Source/Runtime/Engine/Private/WorldPartition/WorldPartitionStreamingGeneration.cpp b/Engine/Source/Runtime/Engine/Private/WorldPartition/WorldPartitionStreamingGeneration.cpp
--- a/Engine/Source/Runtime/Engine/Private/WorldPartition/WorldPartitionStreamingGeneration.cpp
+++ b/Engine/Source/Runtime/Engine/Private/WorldPartition/WorldPartitionStreamingGeneration.cpp
@@ -1753,7 +1753,7 @@ URuntimeHashExternalStreamingObjectBase* UWorldPartition::FlushStreamingToExtern
 
 void UWorldPartition::SetupHLODActors(const FSetupHLODActorsParams& Params)
 {
-	ForEachActorDescContainer([this, &Params](UActorDescContainerInstance* InContainerInstance)
+	ForEachActorDescContainerInstance([this, &Params](UActorDescContainerInstance* InContainerInstance)
 	{
 		TErrorHandlerSelector<FStreamingGenerationLogErrorHandler> ErrorHandlerSelector;
 
@@ -1899,7 +1899,7 @@ void UWorldPartition::CheckForErrors(IStreamingGenerationErrorHandler* ErrorHand
 		.SetIsValidHLODLayer([this](FName GridName, const FSoftObjectPath& HLODLayerPath) { return RuntimeHash->IsValidHLODLayer(GridName, HLODLayerPath); })
 		.SetEnableStreaming(IsStreamingEnabled());
 
-	ForEachActorDescContainer([&StreamingGeneratorParams](const UActorDescContainerInstance* InContainerInstance)
+	ForEachActorDescContainerInstance([&StreamingGeneratorParams](const UActorDescContainerInstance* InContainerInstance)
 	{
 		for (UActorDescContainerInstance::TConstIterator<> Iterator(InContainerInstance); Iterator; ++Iterator)
 		{
@@ -1908,7 +1908,7 @@ void UWorldPartition::CheckForErrors(IStreamingGenerationErrorHandler* ErrorHand
 		}
 	});
 
-	ForEachActorDescContainer([this, &StreamingGeneratorParams](const UActorDescContainerInstance* InContainerInstance)
+	ForEachActorDescContainerInstance([this, &StreamingGeneratorParams](const UActorDescContainerInstance* InContainerInstance)
 	{
 		check(StreamingGeneratorParams.WorldPartitionContext == InContainerInstance->GetWorldPartition());
 		FWorldPartitionStreamingGenerator StreamingGenerator(StreamingGeneratorParams);
@@ -1926,7 +1926,7 @@ void UWorldPartition::CheckForErrors(const FCheckForErrorsParams& Params)
 	FActorDescList ModifiedActorDescList;
 
 	TArray<const UActorDescContainerInstance*> Containers;
-	Params.ActorDescContainerInstanceCollection->ForEachActorDescContainer([&Containers](UActorDescContainerInstance* InContainerInstance) { Containers.Add(InContainerInstance); });
+	Params.ActorDescContainerInstanceCollection->ForEachActorDescContainerInstance([&Containers](UActorDescContainerInstance* InContainerInstance) { Containers.Add(InContainerInstance); });
 
 	FStreamingGenerationContainerInstanceCollection Collection{ Containers };
 	const UActorDescContainerInstance* MainContainerInstance = Collection.GetMainContainer();
diff --git a/Engine/Source/Runtime/Engine/Private/WorldPartitionBlueprintLibrary.cpp b/Engine/Source/Runtime/Engine/Private/WorldPartitionBlueprintLibrary.cpp
--- a/Engine/Source/Runtime/Engine/Private/WorldPartitionBlueprintLibrary.cpp
+++ b/Engine/Source/Runtime/Engine/Private/WorldPartitionBlueprintLibrary.cpp
@@ -86,7 +86,7 @@ bool UWorldPartitionBlueprintLibrary::GetActorDescs(const UWorldPartition* InWor
 {
 	bool bResult = true;
 
-	InWorldPartition->ForEachActorDescContainer([&](const UActorDescContainerInstance* InContainerInstance)
+	InWorldPartition->ForEachActorDescContainerInstance([&](const UActorDescContainerInstance* InContainerInstance)
 	{
 		bResult &= GetActorDescs(InContainerInstance, FTransform::Identity, OutActorDescs);
 	});
diff --git a/Engine/Source/Runtime/Engine/Public/WorldPartition/ActorDescContainerInstanceCollection.h b/Engine/Source/Runtime/Engine/Public/WorldPartition/ActorDescContainerInstanceCollection.h
--- a/Engine/Source/Runtime/Engine/Public/WorldPartition/ActorDescContainerInstanceCollection.h
+++ b/Engine/Source/Runtime/Engine/Public/WorldPartition/ActorDescContainerInstanceCollection.h
@@ -63,11 +63,11 @@ public:
 	DECLARE_EVENT_OneParam(UWorldPartition, FActorDescInstanceRemovedEvent, FWorldPartitionActorDescInstance*);
 	FActorDescInstanceRemovedEvent OnActorDescInstanceRemovedEvent;
 
-	void ForEachActorDescContainerBreakable(TFunctionRef<bool(ActorDescContPtrType)> Func);
-	void ForEachActorDescContainerBreakable(TFunctionRef<bool(ActorDescContPtrType)> Func) const;
+	void ForEachActorDescContainerInstanceBreakable(TFunctionRef<bool(ActorDescContPtrType)> Func);
+	void ForEachActorDescContainerInstanceBreakable(TFunctionRef<bool(ActorDescContPtrType)> Func) const;
 
-	void ForEachActorDescContainer(TFunctionRef<void(ActorDescContPtrType)> Func);
-	void ForEachActorDescContainer(TFunctionRef<void(ActorDescContPtrType)> Func) const;
+	void ForEachActorDescContainerInstance(TFunctionRef<void(ActorDescContPtrType)> Func);
+	void ForEachActorDescContainerInstance(TFunctionRef<void(ActorDescContPtrType)> Func) const;
 
 protected:
 	virtual void OnCollectionChanged() {};
@@ -205,7 +205,7 @@ template<class U>
 TActorDescContainerInstanceCollection<ActorDescContPtrType>::TActorDescContainerInstanceCollection(std::initializer_list<U> ActorDescContainerInstanceArray)
 	: ActorDescContainerInstanceCollection(ActorDescContainerInstanceArray)
 {
-	ForEachActorDescContainer([this](ActorDescContPtrType ActorDescContainerInstance)
+	ForEachActorDescContainerInstance([this](ActorDescContPtrType ActorDescContainerInstance)
 	{
 		RegisterDelegates(ActorDescContainerInstance);
 	});
@@ -216,7 +216,7 @@ template<class U>
 TActorDescContainerInstanceCollection<ActorDescContPtrType>::TActorDescContainerInstanceCollection(const TArray<U>& ActorDescContainerInstances)
 	: ActorDescContainerInstanceCollection(ActorDescContainerInstances)
 {
-	ForEachActorDescContainer([this](ActorDescContPtrType ActorDescContainerInstance)
+	ForEachActorDescContainerInstance([this](ActorDescContPtrType ActorDescContainerInstance)
 	{
 		RegisterDelegates(ActorDescContainerInstance);
 	});
@@ -283,7 +283,7 @@ ActorDescContPtrType TActorDescContainerInstanceCollection<ActorDescContPtrType>
 template<class ActorDescContPtrType>
 void TActorDescContainerInstanceCollection<ActorDescContPtrType>::Empty()
 {
-	ForEachActorDescContainer([this](ActorDescContPtrType ActorDescContainerInstance)
+	ForEachActorDescContainerInstance([this](ActorDescContPtrType ActorDescContainerInstance)
 	{
 		UnregisterDelegates(ActorDescContainerInstance);
 	});
@@ -310,7 +310,7 @@ template<class ActorDescContPtrType>
 const FWorldPartitionActorDescInstance* TActorDescContainerInstanceCollection<ActorDescContPtrType>::GetActorDescInstance(const FGuid& Guid) const
 {
 	const FWorldPartitionActorDescInstance* ActorDescInstance = nullptr;
-	ForEachActorDescContainerBreakable([&Guid, &ActorDescInstance](ActorDescContPtrType ActorDescContainerInstance)
+	ForEachActorDescContainerInstanceBreakable([&Guid, &ActorDescInstance](ActorDescContPtrType ActorDescContainerInstance)
 	{
 		ActorDescInstance = ActorDescContainerInstance->GetActorDescInstance(Guid);
 		return ActorDescInstance == nullptr;
@@ -329,7 +329,7 @@ template<class ActorDescContPtrType>
 const FWorldPartitionActorDescInstance* TActorDescContainerInstanceCollection<ActorDescContPtrType>::GetActorDescInstanceByPath(const FString& ActorPath) const
 {
 	const FWorldPartitionActorDescInstance* ActorDescInstance = nullptr;
-	ForEachActorDescContainerBreakable([&ActorPath, &ActorDescInstance](ActorDescContPtrType ActorDescContainerInstance)
+	ForEachActorDescContainerInstanceBreakable([&ActorPath, &ActorDescInstance](ActorDescContPtrType ActorDescContainerInstance)
 	{
 		ActorDescInstance = ActorDescContainerInstance->GetActorDescInstanceByPath(ActorPath);
 		return ActorDescInstance == nullptr;
@@ -342,7 +342,7 @@ template<class ActorDescContPtrType>
 const FWorldPartitionActorDescInstance* TActorDescContainerInstanceCollection<ActorDescContPtrType>::GetActorDescInstanceByPath(const FSoftObjectPath& ActorPath) const
 {
 	const FWorldPartitionActorDescInstance* ActorDescInstance = nullptr;
-	ForEachActorDescContainerBreakable([&ActorPath, &ActorDescInstance](ActorDescContPtrType ActorDescContainerInstance)
+	ForEachActorDescContainerInstanceBreakable([&ActorPath, &ActorDescInstance](ActorDescContPtrType ActorDescContainerInstance)
 	{
 		ActorDescInstance = ActorDescContainerInstance->GetActorDescInstanceByPath(ActorPath);
 		return ActorDescInstance == nullptr;
@@ -355,7 +355,7 @@ template<class ActorDescContPtrType>
 const FWorldPartitionActorDescInstance* TActorDescContainerInstanceCollection<ActorDescContPtrType>::GetActorDescInstanceByName(FName ActorName) const
 {
 	const FWorldPartitionActorDescInstance* ActorDescInstance = nullptr;
-	ForEachActorDescContainerBreakable([&ActorName, &ActorDescInstance](ActorDescContPtrType ActorDescContainerInstance)
+	ForEachActorDescContainerInstanceBreakable([&ActorName, &ActorDescInstance](ActorDescContPtrType ActorDescContainerInstance)
 	{
 		ActorDescInstance = ActorDescContainerInstance->GetActorDescInstanceByName(ActorName);
 		return ActorDescInstance == nullptr;
@@ -368,7 +368,7 @@ template<class ActorDescContPtrType>
 ActorDescContPtrType TActorDescContainerInstanceCollection<ActorDescContPtrType>::GetActorDescContainerInstance(const FGuid& ActorGuid) const
 {
 	ActorDescContPtrType ActorDescContainerInstance = nullptr;
-	ForEachActorDescContainerBreakable([&ActorGuid, &ActorDescContainerInstance](ActorDescContPtrType InActorDescContainerInstance)
+	ForEachActorDescContainerInstanceBreakable([&ActorGuid, &ActorDescContainerInstance](ActorDescContPtrType InActorDescContainerInstance)
 	{
 		if (InActorDescContainerInstance->GetActorDescInstance(ActorGuid) != nullptr)
 		{
@@ -391,7 +391,7 @@ ActorDescContPtrType TActorDescContainerInstanceCollection<ActorDescContPtrType>
 
 	// Actor is not yet stored in a container. Find which one should handle it.
 	ActorDescContPtrType ActorDescContainerInstance = nullptr;
-	ForEachActorDescContainerBreakable([&ActorDescContainerInstance, &Actor](ActorDescContPtrType InActorDescContainerInstance)
+	ForEachActorDescContainerInstanceBreakable([&ActorDescContainerInstance, &Actor](ActorDescContPtrType InActorDescContainerInstance)
 	{
 		if (InActorDescContainerInstance->IsActorDescHandled(Actor))
 		{
@@ -408,7 +408,7 @@ template<typename, typename>
 bool TActorDescContainerInstanceCollection<ActorDescContPtrType>::RemoveActor(const FGuid& ActorGuid)
 {
 	bool bRemoved = false;
-	ForEachActorDescContainerBreakable([&ActorGuid, &bRemoved](ActorDescContPtrType ContainerInstance)
+	ForEachActorDescContainerInstanceBreakable([&ActorGuid, &bRemoved](ActorDescContPtrType ContainerInstance)
 	{
 		bRemoved = ContainerInstance->GetContainer()->RemoveActor(ActorGuid);
 		return !bRemoved;
@@ -421,7 +421,7 @@ template<class ActorDescContPtrType>
 template<typename, typename>
 void TActorDescContainerInstanceCollection<ActorDescContPtrType>::OnPackageDeleted(UPackage* Package)
 {
-	ForEachActorDescContainer([Package](ActorDescContPtrType ContainerInstance)
+	ForEachActorDescContainerInstance([Package](ActorDescContPtrType ContainerInstance)
 	{
 		ContainerInstance->GetContainer()->OnPackageDeleted(Package);
 	});
@@ -431,7 +431,7 @@ template<class ActorDescContPtrType>
 template<typename, typename>
 void TActorDescContainerInstanceCollection<ActorDescContPtrType>::LoadAllActors(TArray<FWorldPartitionReference>& OutReferences)
 {
-	ForEachActorDescContainer([&OutReferences](ActorDescContPtrType ContainerInstance)
+	ForEachActorDescContainerInstance([&OutReferences](ActorDescContPtrType ContainerInstance)
 	{
 		ContainerInstance->LoadAllActors(OutReferences);
 	});
@@ -439,7 +439,7 @@ void TActorDescContainerInstanceCollection<ActorDescContPtrType>::LoadAllActors(
 
 
 template<class ActorDescContPtrType>
-void TActorDescContainerInstanceCollection<ActorDescContPtrType>::ForEachActorDescContainerBreakable(TFunctionRef<bool(ActorDescContPtrType)> Func) const
+void TActorDescContainerInstanceCollection<ActorDescContPtrType>::ForEachActorDescContainerInstanceBreakable(TFunctionRef<bool(ActorDescContPtrType)> Func) const
 {
 	for (ActorDescContPtrType ActorDescContainerInstance : ActorDescContainerInstanceCollection)
 	{
@@ -451,13 +451,13 @@ void TActorDescContainerInstanceCollection<ActorDescContPtrType>::ForEachActorDe
 }
 
 template<class ActorDescContPtrType>
-void TActorDescContainerInstanceCollection<ActorDescContPtrType>::ForEachActorDescContainerBreakable(TFunctionRef<bool(ActorDescContPtrType)> Func)
+void TActorDescContainerInstanceCollection<ActorDescContPtrType>::ForEachActorDescContainerInstanceBreakable(TFunctionRef<bool(ActorDescContPtrType)> Func)
 {
-	const_cast<const TActorDescContainerInstanceCollection*>(this)->ForEachActorDescContainerBreakable(Func);
+	const_cast<const TActorDescContainerInstanceCollection*>(this)->ForEachActorDescContainerInstanceBreakable(Func);
 }
 
 template<class ActorDescContPtrType>
-void TActorDescContainerInstanceCollection<ActorDescContPtrType>::ForEachActorDescContainer(TFunctionRef<void(ActorDescContPtrType)> Func) const
+void TActorDescContainerInstanceCollection<ActorDescContPtrType>::ForEachActorDescContainerInstance(TFunctionRef<void(ActorDescContPtrType)> Func) const
 {
 	for (ActorDescContPtrType ActorDescContainerInstance : ActorDescContainerInstanceCollection)
 	{
@@ -466,9 +466,9 @@ void TActorDescContainerInstanceCollection<ActorDescContPtrType>::ForEachActorDe
 }
 
 template<class ActorDescContPtrType>
-void TActorDescContainerInstanceCollection<ActorDescContPtrType>::ForEachActorDescContainer(TFunctionRef<void(ActorDescContPtrType)> Func)
+void TActorDescContainerInstanceCollection<ActorDescContPtrType>::ForEachActorDescContainerInstance(TFunctionRef<void(ActorDescContPtrType)> Func)
 {
-	const_cast<const TActorDescContainerInstanceCollection*>(this)->ForEachActorDescContainer(Func);
+	const_cast<const TActorDescContainerInstanceCollection*>(this)->ForEachActorDescContainerInstance(Func);
 }
 
 template<class ActorDescContPtrType>
diff --git a/Engine/Source/Runtime/Engine/Public/WorldPartition/WorldPartition.h b/Engine/Source/Runtime/Engine/Public/WorldPartition/WorldPartition.h
--- a/Engine/Source/Runtime/Engine/Public/WorldPartition/WorldPartition.h
+++ b/Engine/Source/Runtime/Engine/Public/WorldPartition/WorldPartition.h
@@ -491,9 +491,18 @@ private:
 public:
 	UE_DEPRECATED(5.4, "Use GetActorDescContainerInstance")
 	UActorDescContainer* GetActorDescContainer() const { return nullptr; }
-		
+
 #if WITH_EDITOR
 	UActorDescContainerInstance* GetActorDescContainerInstance() const { return ActorDescContainerInstance; }
+
+	UE_DEPRECATED(5.4, "Use ForEachActorDescContainerInstanceBreakable.")
+	void ForEachActorDescContainerBreakable(TFunctionRef<bool(UActorDescContainer*)> Func) {}
+	UE_DEPRECATED(5.4, "Use ForEachActorDescContainerInstanceBreakable.")
+	void ForEachActorDescContainerBreakable(TFunctionRef<bool(UActorDescContainer*)> Func) const {}
+	UE_DEPRECATED(5.4, "Use ForEachActorDescContainerInstance.")
+	void ForEachActorDescContainer(TFunctionRef<void(UActorDescContainer*)> Func) {}
+	UE_DEPRECATED(5.4, "Use ForEachActorDescContainerInstance.")
+	void ForEachActorDescContainer(TFunctionRef<void(UActorDescContainer*)> Func) const {}
 #endif			
 
 	UPROPERTY()
