Fix RegisterAllComponents to call redundant PostRegisterAllComponents functions even if OptimizeActorRegistration is enabled as editor code appears to depend on it, but skip the world callback

diff --git a/Engine/Source/Runtime/Engine/Private/Actor.cpp b/Engine/Source/Runtime/Engine/Private/Actor.cpp
--- a/Engine/Source/Runtime/Engine/Private/Actor.cpp
+++ b/Engine/Source/Runtime/Engine/Private/Actor.cpp
@@ -5366,13 +5366,6 @@ void AActor::PostUnregisterAllComponents()
 
 void AActor::RegisterAllComponents()
 {
-	if (bHasRegisteredAllComponents && GOptimizeActorRegistration > 0)
-	{
-		// Stops it from calling redundant Pre/Post registration functions, should never get here if it was deferred
-		ensureMsgf(!bHasDeferredComponentRegistration, TEXT("Actor %s has both bHasDeferredComponentRegistration and bHasRegisteredAllComponents"), *GetPathName());
-		return;
-	}
-
 	QUICK_SCOPE_CYCLE_COUNTER(STAT_AActor_RegisterAllComponents);
 	
 	PreRegisterAllComponents();
@@ -5489,13 +5482,9 @@ bool AActor::IncrementalRegisterComponents(int32 NumComponentsToRegister, FRegis
 #if PERF_TRACK_DETAILED_ASYNC_STATS
 		QUICK_SCOPE_CYCLE_COUNTER(STAT_AActor_IncrementalRegisterComponents_PostRegisterAllComponents);
 #endif
-
-#if !UE_BUILD_SHIPPING
-		if (GOptimizeActorRegistration == 2)
-		{
-			ensureMsgf(!bHasRegisteredAllComponents, TEXT("PostRegisterAllComponents called twice for %s!"), *GetPathName());
-		}
-#endif
+		
+		// Skip the world post register if optimizations are enabled and it was already called
+		const bool bCallWorldPostRegister = (!bHasRegisteredAllComponents || GOptimizeActorRegistration == 0);
 
 		// Clear this flag as it's no longer deferred
 		bHasDeferredComponentRegistration = false;
@@ -5504,8 +5493,11 @@ bool AActor::IncrementalRegisterComponents(int32 NumComponentsToRegister, FRegis
 		// Finally, call PostRegisterAllComponents
 		PostRegisterAllComponents();
 
-		// After all components have been registered the actor is considered fully added: notify the owning world.
-		World->NotifyPostRegisterAllActorComponents(this);
+		if (bCallWorldPostRegister)
+		{
+			// After all components have been registered the actor is considered fully added: notify the owning world.
+			World->NotifyPostRegisterAllActorComponents(this);
+		}
 		return true;
 	}
 	
