World Partition - HLOD: Deprecate UWorldPartitionHLODRuntimeSubsystem::SetHLODAlwaysLoadedCullDistance
* Moved logic to game side
* Fixed reference to possibly deleted HLOD actors at the end of replays

diff --git a/Engine/Source/Runtime/Engine/Private/WorldPartition/HLOD/HLODRuntimeSubsystem.cpp b/Engine/Source/Runtime/Engine/Private/WorldPartition/HLOD/HLODRuntimeSubsystem.cpp
--- a/Engine/Source/Runtime/Engine/Private/WorldPartition/HLOD/HLODRuntimeSubsystem.cpp
+++ b/Engine/Source/Runtime/Engine/Private/WorldPartition/HLOD/HLODRuntimeSubsystem.cpp
@@ -71,16 +71,6 @@ static TAutoConsoleVariable<int32> CVarHLODWarmupVTSizeClamp(
 	2048,
 	TEXT("Clamp VT warmup requests for safety."));
 
-static TAutoConsoleVariable<int32> CVarCullDistanceWorkSlicer(
-	TEXT("wp.Runtime.HLOD.CullDistanceWorkSlicer"),
-	60,
-	TEXT("Maximum amount of actors that can have it's CullDistance changed per frames."));
-
-static TAutoConsoleVariable<int32> CVarCullDistanceDifferenceNeeded(
-	TEXT("wp.Runtime.HLOD.CVarCullDistanceDifferenceNeeded"),
-	300,
-	TEXT("The difference needed for the CullDistance changed to be processed"));
-
 static void HLODRuntimeSubsystemCVarSinkFunction()
 {
 	for (UWorld* World : TObjectRange<UWorld>(RF_ClassDefaultObject | RF_ArchetypeObject, true, EInternalObjectFlags::Garbage))
@@ -146,112 +136,6 @@ UWorldPartitionHLODRuntimeSubsystem::UWorldPartitionHLODRuntimeSubsystem()
 {
 }
 
-UWorldPartitionHLODRuntimeSubsystem::~UWorldPartitionHLODRuntimeSubsystem()
-{
-	if (TickHandle.IsValid())
-	{
-		FTSTicker::GetCoreTicker().RemoveTicker(TickHandle);
-		TickHandle.Reset();
-	}
-}
-
-void UWorldPartitionHLODRuntimeSubsystem::SetHLODAlwaysLoadedCullDistance(int32 InCullDistance)
-{
-	const float MaxDrawDistance = FMath::Max(InCullDistance, 0);
-	const bool bNeverDistanceCull = InCullDistance <= 0;
-
-	//If the WorkSlicer is set to process 0 items, revert back to processing everything at once
-	if (CVarCullDistanceWorkSlicer.GetValueOnAnyThread() > 0)
-	{
-		int32 CullDiff = FMath::Abs(LastSetCullDistance - InCullDistance);
-		if (!OperationQueue.IsEmpty() || CullDiff < CVarCullDistanceDifferenceNeeded.GetValueOnAnyThread())
-		{
-			return;
-		}
-		LastSetCullDistance = InCullDistance;
-
-		OperationQueue.Reserve(AlwaysLoadedHLODActors.Num());
-		for (AWorldPartitionHLOD* HLODActor : AlwaysLoadedHLODActors)
-		{
-			FDrawDistanceQueue Item = FDrawDistanceQueue();
-			Item.DrawDistance = MaxDrawDistance;
-			Item.HLODActor = HLODActor;
-			OperationQueue.Add(Item);
-		}
-
-		if (!TickHandle.IsValid())
-		{
-			TickHandle = FTSTicker::GetCoreTicker().AddTicker(FTickerDelegate::CreateUObject(this, &UWorldPartitionHLODRuntimeSubsystem::Tick), 0.0f);
-		}
-	}
-	else
-	{
-		for (AWorldPartitionHLOD* HLODActor : AlwaysLoadedHLODActors)
-		{
-			HLODActor->ForEachComponent<UPrimitiveComponent>(false, [MaxDrawDistance, bNeverDistanceCull](UPrimitiveComponent* HLODComponent)
-			{
-				HLODComponent->bNeverDistanceCull = bNeverDistanceCull;
-				HLODComponent->SetCachedMaxDrawDistance(MaxDrawDistance);
-				// Enable per instance culling to avoid ISM components dissapearing as a whole
-				if (UInstancedStaticMeshComponent* ISMComponent = Cast<UInstancedStaticMeshComponent>(HLODComponent))
-				{
-					ISMComponent->SetCullDistances(0, MaxDrawDistance);
-				}
-			});
-		}
-	}
-	
-}
-
-void UWorldPartitionHLODRuntimeSubsystem::OnWorldBeginPlay(UWorld& InWorld)
-{
-	Super::OnWorldBeginPlay(InWorld);
-
-	check(InWorld.IsGameWorld());
-
-	AlwaysLoadedHLODActors.Reset();
-
-	// Cache always loaded HLODs
-	for (AActor* Actor : InWorld.PersistentLevel->Actors)
-	{
-		if (AWorldPartitionHLOD* HLODActor = Cast<AWorldPartitionHLOD>(Actor))
-		{
-			AlwaysLoadedHLODActors.Add(HLODActor);
-		}
-	}
-
-}
-
-bool UWorldPartitionHLODRuntimeSubsystem::Tick(float DeltaTime)
-{
-	//If the CVarCullDistanceAllowWorkSlicer is false we will never add this function to FTSTicker
-	if (OperationQueue.Num() > 0)
-	{
-		int32 Amount = FMath::Min(OperationQueue.Num(), CVarCullDistanceWorkSlicer.GetValueOnAnyThread());	
-		for (int32 i = OperationQueue.Num() - Amount; i < OperationQueue.Num(); i++)
-		{
-			const float MaxDrawDistance = OperationQueue[i].DrawDistance;
-			TObjectPtr<AWorldPartitionHLOD> HLODActor = OperationQueue[i].HLODActor;
-			if (HLODActor != nullptr)
-			{
-				HLODActor->ForEachComponent<UPrimitiveComponent>(false, [MaxDrawDistance](UPrimitiveComponent* HLODComponent)
-				{
-					HLODComponent->bNeverDistanceCull = false;
-					HLODComponent->SetCachedMaxDrawDistance(MaxDrawDistance);
-					// Enable per instance culling to avoid ISM components dissapearing as a whole
-					if (UInstancedStaticMeshComponent* ISMComponent = Cast<UInstancedStaticMeshComponent>(HLODComponent))
-					{
-						ISMComponent->SetCullDistances(0, MaxDrawDistance);
-					}
-				});
-			}
-			
-		}
-		OperationQueue.RemoveAt(OperationQueue.Num() - Amount, Amount);
-	}
-	return true;
-}
-
 bool UWorldPartitionHLODRuntimeSubsystem::WorldPartitionHLODEnabled = true;
 
 FAutoConsoleCommand UWorldPartitionHLODRuntimeSubsystem::EnableHLODCommand(
diff --git a/Engine/Source/Runtime/Engine/Public/WorldPartition/HLOD/HLODRuntimeSubsystem.h b/Engine/Source/Runtime/Engine/Public/WorldPartition/HLOD/HLODRuntimeSubsystem.h
--- a/Engine/Source/Runtime/Engine/Public/WorldPartition/HLOD/HLODRuntimeSubsystem.h
+++ b/Engine/Source/Runtime/Engine/Public/WorldPartition/HLOD/HLODRuntimeSubsystem.h
@@ -31,7 +31,6 @@ class UWorldPartitionHLODRuntimeSubsystem : public UWorldSubsystem
 
 public:
 	ENGINE_API UWorldPartitionHLODRuntimeSubsystem();
-	ENGINE_API virtual ~UWorldPartitionHLODRuntimeSubsystem();
 
 	//~ Begin USubsystem Interface.
 	ENGINE_API virtual void Initialize(FSubsystemCollectionBase& Collection) override;
@@ -40,11 +39,8 @@ public:
 
 	//~ Begin UWorldSubsystem Interface.
 	ENGINE_API virtual bool DoesSupportWorldType(const EWorldType::Type WorldType) const override;
-	ENGINE_API virtual void OnWorldBeginPlay(UWorld& InWorld) override;
 	//~ End UWorldSubsystem Interface.
 
-	ENGINE_API bool Tick(float DeltaTime);
-
 	ENGINE_API void RegisterHLODActor(AWorldPartitionHLOD* InWorldPartitionHLOD);
 	ENGINE_API void UnregisterHLODActor(AWorldPartitionHLOD* InWorldPartitionHLOD);
 
@@ -64,14 +60,15 @@ public:
 	ENGINE_API void OnExternalStreamingObjectInjected(URuntimeHashExternalStreamingObjectBase* ExternalStreamingObject);
 	ENGINE_API void OnExternalStreamingObjectRemoved(URuntimeHashExternalStreamingObjectBase* ExternalStreamingObject);
 
-	ENGINE_API void SetHLODAlwaysLoadedCullDistance(int32 InCullDistance);
-
 	ENGINE_API void OnCVarsChanged();
 
 #if WITH_EDITOR
 	uint32 GetNumOutdatedHLODActors() const { return OutdatedHLODActors.Num(); }
 	static ENGINE_API bool WriteHLODStatsCSV(UWorld* InWorld, const FString& InFilename);
 #endif
+
+	UE_DEPRECATED(5.4, "You should perform this logic on the game side.")
+	ENGINE_API void SetHLODAlwaysLoadedCullDistance(int32 InCullDistance) {}
 	
 private:
 	struct FCellData
@@ -85,19 +82,11 @@ private:
 		TMap<FGuid, FCellData> CellsData;
 	};
 
-	struct FDrawDistanceQueue
-	{
-		float DrawDistance;
-		TObjectPtr<AWorldPartitionHLOD> HLODActor;
-	};
-	
 	TMap<TObjectPtr<UWorldPartition>, FWorldPartitionHLODRuntimeData> WorldPartitionsHLODRuntimeData;
 	ENGINE_API const FCellData* GetCellData(const UWorldPartitionRuntimeCell* InCell) const;
 	ENGINE_API FCellData* GetCellData(const UWorldPartitionRuntimeCell* InCell);
 	ENGINE_API FCellData* GetCellData(AWorldPartitionHLOD* InWorldPartitionHLOD);
 
-	TArray<AWorldPartitionHLOD*> AlwaysLoadedHLODActors;
-
 	struct FWorldPartitionHLODWarmupState
 	{
 		uint32 WarmupLastRequestedFrame = INDEX_NONE;
@@ -114,10 +103,6 @@ private:
 	ENGINE_API bool ShouldPerformWarmupForCell(const UWorldPartitionRuntimeCell* InCell) const;
 	bool bCachedShouldPerformWarmup;
 
-	TArray<FDrawDistanceQueue> OperationQueue;
-	FTSTicker::FDelegateHandle TickHandle;
-	int32 LastSetCullDistance;
-
 	/** Console command used to turn on/off loading & rendering of world partition HLODs */
 	static ENGINE_API class FAutoConsoleCommand EnableHLODCommand;
 
