Iris
* Deleted Iris holder that allowed to transition the Iris system during a fork process. Unfortunately Iris doesn't correctly support realtime config changes rendering this feature dangerous.
* Instead it's recommended to destroy and recreate the Iris system entirely after forking or applying config/hotfix values.

#rnx

diff --git a/Engine/Source/Runtime/Engine/Classes/Engine/World.h b/Engine/Source/Runtime/Engine/Classes/Engine/World.h
--- a/Engine/Source/Runtime/Engine/Classes/Engine/World.h
+++ b/Engine/Source/Runtime/Engine/Classes/Engine/World.h
@@ -190,23 +190,6 @@ public:
 	static FOnWorldUnregisteredWithAudioDevice OnWorldUnregisteredWithAudioDevice;
 };
 
-#if UE_WITH_IRIS
-/**
-* Struct that temporarily holds the Iris replication system and bridge.
-*/
-struct FIrisSystemHolder
-{
-	bool IsHolding() const { return ReplicationSystem != nullptr; }
-
-	void Clear()
-	{
-		ReplicationSystem = nullptr;
-	}
-
-	class UReplicationSystem* ReplicationSystem = nullptr;
-};
-#endif // UE_WITH_IRIS
-
 /** class that encapsulates seamless world traveling */
 class FSeamlessTravelHandler
 {
@@ -1105,10 +1088,6 @@ public:
 
 	/** Returns BlockTillLevelStreamingCompletedEpoch. */
 	int32 GetBlockTillLevelStreamingCompletedEpoch() const { return BlockTillLevelStreamingCompletedEpoch; }
-#if UE_WITH_IRIS
-	/** Store the Iris managers from the GameNetDriver to they can be restored into a different NetDriver later */
-	void StoreIrisAndClearReferences();
-#endif // UE_WITH_IRIS
 
 	/** Prefix we used to rename streaming levels, non empty in PIE and standalone preview */
 	UPROPERTY()
@@ -3701,11 +3680,6 @@ public:
 
 private:
 
-#if UE_WITH_IRIS
-	/** Holds the Iris systems during the NetDriver transition that occurs when Forking */
-	FIrisSystemHolder IrisSystemHolder;
-#endif // UE_WITH_IRIS
-
 	/** Private version without inlining that does *not* check Dedicated server build flags (which should already have been done). */
 	ENetMode InternalGetNetMode() const;
 
diff --git a/Engine/Source/Runtime/Engine/Private/World.cpp b/Engine/Source/Runtime/Engine/Private/World.cpp
--- a/Engine/Source/Runtime/Engine/Private/World.cpp
+++ b/Engine/Source/Runtime/Engine/Private/World.cpp
@@ -1252,15 +1252,6 @@ void UWorld::BeginDestroy()
 	FAudioDeviceHandle EmptyHandle;
 	SetAudioDevice(EmptyHandle);
 	check(!AudioDeviceDestroyedHandle.IsValid());
-
-#if UE_WITH_IRIS
-	if (IrisSystemHolder.IsHolding())
-	{
-		UE::Net::FReplicationSystemFactory::DestroyReplicationSystem(IrisSystemHolder.ReplicationSystem);
-
-		IrisSystemHolder.Clear();
-	}
-#endif // UE_WITH_IRIS
 }
 
 void UWorld::ReleasePhysicsScene()
@@ -6811,16 +6802,6 @@ bool UWorld::Listen( FURL& InURL )
 	if (GEngine->CreateNamedNetDriver(this, NAME_GameNetDriver, NAME_GameNetDriver))
 	{
 		NetDriver = GEngine->FindNamedNetDriver(this, NAME_GameNetDriver);
-
-#if UE_WITH_IRIS
-		if (IrisSystemHolder.IsHolding())
-		{
-			NetDriver->RestoreIrisSystem(IrisSystemHolder.ReplicationSystem);
-			
-			IrisSystemHolder.Clear();
-		}
-#endif // UE_WITH_IRIS
-
 		NetDriver->SetWorld(this);
 		FLevelCollection* const SourceCollection = FindCollectionByType(ELevelCollectionType::DynamicSourceLevels);
 		if (SourceCollection)
@@ -8952,17 +8933,6 @@ void UWorld::DuplicateRequestedLevels(const FName MapName)
 	}
 }
 
-#if UE_WITH_IRIS
-void UWorld::StoreIrisAndClearReferences()
-{
-	if (NetDriver)
-	{
-		IrisSystemHolder.ReplicationSystem = NetDriver->GetReplicationSystem();
-		NetDriver->ClearIrisSystem();
-	}
-}
-#endif // UE_WITH_IRIS
-
 #if WITH_EDITOR
 void UWorld::ChangeFeatureLevel(ERHIFeatureLevel::Type InFeatureLevel, bool bShowSlowProgressDialog )
 {
