Minor change to reduce allocation when adding multiple points.

PR #10772: Avoid quadratic behavior in USplineComponent::AddPoints

diff --git a/Engine/Source/Runtime/Engine/Private/Components/SplineComponent.cpp b/Engine/Source/Runtime/Engine/Private/Components/SplineComponent.cpp
--- a/Engine/Source/Runtime/Engine/Private/Components/SplineComponent.cpp
+++ b/Engine/Source/Runtime/Engine/Private/Components/SplineComponent.cpp
@@ -728,8 +728,11 @@ void USplineComponent::AddPoint(const FSplinePoint& InSplinePoint, bool bUpdateS
 
 void USplineComponent::AddPoints(const TArray<FSplinePoint>& InSplinePoints, bool bUpdateSpline)
 {
-	const int32 NumPoints = SplineCurves.Position.Points.Num();
-	SplineCurves.Position.Points.Reserve(NumPoints + InSplinePoints.Num());
+	const int32 NumPoints = SplineCurves.Position.Points.Num() + InSplinePoints.Num();
+	// Position, Rotation, and Scale will all grow together.
+	SplineCurves.Position.Points.Reserve(NumPoints);
+	SplineCurves.Rotation.Points.Reserve(NumPoints);
+	SplineCurves.Scale.Points.Reserve(NumPoints);
 
 	for (const auto& SplinePoint : InSplinePoints)
 	{
