World Partition: ground work for actor descriptors mutators.
Also make sure FStreamingGenerationActorDescView can only be constructed from actor descriptor instances, not views.

diff --git a/Engine/Source/Runtime/Engine/Private/WorldPartition/WorldPartitionStreamingGeneration.cpp b/Engine/Source/Runtime/Engine/Private/WorldPartition/WorldPartitionStreamingGeneration.cpp
--- a/Engine/Source/Runtime/Engine/Private/WorldPartition/WorldPartitionStreamingGeneration.cpp
+++ b/Engine/Source/Runtime/Engine/Private/WorldPartition/WorldPartitionStreamingGeneration.cpp
@@ -101,7 +101,7 @@ public:
 	{
 		UActorDescContainerInstance* HandlingContainer = const_cast<UActorDescContainerInstance*>(InContainerInstanceCollection.FindHandlingContainerInstance(InActor).Get());
 		check(HandlingContainer);
-		return TUniquePtr<FStreamingGenerationUnsavedDirtyActorDescInstance>(new FStreamingGenerationUnsavedDirtyActorDescInstance(HandlingContainer, InActor));
+		return MakeUnique<FStreamingGenerationUnsavedDirtyActorDescInstance>(HandlingContainer, InActor);
 	}
 
 	//~ Begin FWorldPartitionActorDescInstance	
@@ -423,7 +423,7 @@ class FWorldPartitionStreamingGenerator
 					if (bContainsUnfilteredActors)
 					{
 						FActorSetInstance& ActorSetInstance = ActorSetInstances.Emplace_GetRef();
-						const FContainerCollectionInstanceDescriptor::FPerInstanceData& PerInstanceData = ContainerCollectionInstanceDescriptor.GetInstanceData(ReferenceActorDescView.GetGuid());
+						const FContainerCollectionInstanceDescriptor::FPerInstanceData& PerInstanceData = ContainerCollectionInstanceDescriptor.GetPerInstanceData(ReferenceActorDescView.GetGuid());
 
 						ActorSetInstance.ActorSetContainerInstance = &ActorSetContainer;
 						ActorSetInstance.ActorSet = &ActorSet;
@@ -558,7 +558,16 @@ class FWorldPartitionStreamingGenerator
 			}
 		};
 
-		FPerInstanceData& GetInstanceData(const FGuid& ActorGuid)
+		void AddPerInstanceData(const FGuid& ActorGuid, const FPerInstanceData& ActorInstanceData)
+		{
+			if (ActorInstanceData != InstanceData)
+			{
+				const FSetElementId UniquePerInstanceDataId = UniquePerInstanceData.Add(ActorInstanceData);
+				PerInstanceData.Emplace(ActorGuid, UniquePerInstanceDataId);
+			}
+		}
+
+		FPerInstanceData& GetPerInstanceData(const FGuid& ActorGuid)
 		{
 			if (const FSetElementId* PerInstanceDataId = PerInstanceData.Find(ActorGuid))
 			{
@@ -569,7 +578,7 @@ class FWorldPartitionStreamingGenerator
 			return InstanceData;
 		}
 
-		const FPerInstanceData& GetInstanceData(const FGuid& ActorGuid) const
+		const FPerInstanceData& GetPerInstanceData(const FGuid& ActorGuid) const
 		{
 			if (const FSetElementId* PerInstanceDataId = PerInstanceData.Find(ActorGuid))
 			{
@@ -646,23 +655,23 @@ class FWorldPartitionStreamingGenerator
 		// Consider all actors of a /Temp/ container package as Unsaved because loading them from disk will fail (Outer world name mismatch)
 		const bool bIsTempContainerPackage = FPackageName::IsTempPackage((InActorDescCollection.GetMainContainerPackageName().ToString()));
 
-		// Test whether an actor should be included in the ActorDescViewMap.
-		auto ShouldRegisterActorDesc = [this](const IWorldPartitionActorDescInstance* InActorDescView)
+		// Test whether an actor descriptor instance should be included in the ActorDescViewMap.
+		auto ShouldRegisterActorDesc = [this](const FWorldPartitionActorDescInstance* InActorDescInstance)
 		{
 			for (UClass* FilteredClass : FilteredClasses)
 			{
-				if (InActorDescView->GetActorNativeClass()->IsChildOf(FilteredClass))
+				if (InActorDescInstance->GetActorNativeClass()->IsChildOf(FilteredClass))
 				{
 					return false;
 				}
 			}
 
-			if (!InActorDescView->IsRuntimeRelevant())
+			if (!InActorDescInstance->IsRuntimeRelevant())
 			{
 				return false;
 			}
 
-			return InActorDescView->IsLoaded() ? !InActorDescView->GetActor()->IsEditorOnly() : !InActorDescView->GetActorIsEditorOnly();
+			return InActorDescInstance->IsLoaded() ? !InActorDescInstance->GetActor()->IsEditorOnly() : !InActorDescInstance->GetActorIsEditorOnly();
 		};
 
 		// Register the actor descriptor view
@@ -681,7 +690,7 @@ class FWorldPartitionStreamingGenerator
 
 		for (FStreamingGenerationContainerInstanceCollection::TConstIterator<> Iterator(&InActorDescCollection); Iterator; ++Iterator)
 		{
-			// @todo_pat: this is to validate that new parenting of container instance code is equivalent
+			// @todo_ow: this is to validate that new parenting of container instance code is equivalent
 			check(Iterator->GetContainerInstance()->GetContainerID() == InContainerID);
 			if (ShouldRegisterActorDesc(*Iterator))
 			{
@@ -723,10 +732,10 @@ class FWorldPartitionStreamingGenerator
 					&& !InActorDescCollection.GetActorDescInstance(Actor->GetActorGuid()))
 				{
 					TUniquePtr<FStreamingGenerationUnsavedDirtyActorDescInstance> UnsavedViewPtr = FStreamingGenerationUnsavedDirtyActorDescInstance::Create(Actor, InActorDescCollection);
-					FStreamingGenerationActorDescView ModifiedActorDescView(OutActorDescViewMap, UnsavedViewPtr.Get(), true);
-					if (ShouldRegisterActorDesc(&ModifiedActorDescView))
+					if (ShouldRegisterActorDesc(UnsavedViewPtr.Get()))
 					{
-						OutUnsavedDirtyInstances.Add(MoveTemp(UnsavedViewPtr));
+						FStreamingGenerationActorDescView ModifiedActorDescView(OutActorDescViewMap, UnsavedViewPtr.Get(), true);
+						OutUnsavedDirtyInstances.Add(MoveTemp(UnsavedViewPtr));						
 						RegisterActorDescView(MoveTemp(ModifiedActorDescView));
 					}
 				}
@@ -830,7 +839,7 @@ class FWorldPartitionStreamingGenerator
 				SubContainerInstanceDescriptor.ContainerInstanceCollection = MakeShared<FStreamingGenerationContainerInstanceCollection>(FStreamingGenerationContainerInstanceCollection{ SubContainerInstance.ContainerInstance });
 				SubContainerInstanceDescriptor.Transform = SubContainerInstance.ContainerInstance->GetTransform();
 
-				// @todo_pat: this is to validate that new parenting of container instance code is equivalent
+				// @todo_ow: this is to validate that new parenting of container instance code is equivalent
 				const FTransform ValidationTransform = SubContainerInstance.Transform * ContainerCollectionInstanceDescriptor.Transform;
 				check(SubContainerInstanceDescriptor.Transform.Equals(ValidationTransform));
 
@@ -869,16 +878,8 @@ class FWorldPartitionStreamingGenerator
 		ContainerCollectionInstanceDescriptor.ActorDescViewMap->ForEachActorDescView([&ContainerCollectionInstanceDescriptor, &InheritParentContainerPerInstanceData](FStreamingGenerationActorDescView& ActorDescView)
 		{
 			const FContainerCollectionInstanceDescriptor::FPerInstanceData PerInstanceData = InheritParentContainerPerInstanceData(ContainerCollectionInstanceDescriptor, ActorDescView);
-
-			if (PerInstanceData != ContainerCollectionInstanceDescriptor.InstanceData)
-			{
-				const FSetElementId UniquePerInstanceDataId = ContainerCollectionInstanceDescriptor.UniquePerInstanceData.Add(PerInstanceData);
-				ContainerCollectionInstanceDescriptor.PerInstanceData.Emplace(ActorDescView.GetGuid(), UniquePerInstanceDataId);
-			}
+			ContainerCollectionInstanceDescriptor.AddPerInstanceData(ActorDescView.GetGuid(), PerInstanceData);
 		});
-
-		// Validate container instance, fixing anything illegal, etc.
-		ValidateContainerInstanceDescriptor(ContainerCollectionInstanceDescriptor, ContainerCollectionInstanceDescriptor.ID.IsMainContainer());
 	}
 
 	/**
@@ -1183,7 +1184,7 @@ class FWorldPartitionStreamingGenerator
 						{
 							if (PassType == EPassType::ErrorReporting)
 							{
-								IWorldPartitionActorDescInstance* ReferencedActorDescInstance = nullptr;
+								FWorldPartitionActorDescInstance* ReferencedActorDescInstance = nullptr;
 
 								if (const UActorDescContainerInstance** ExistingReferencedContainerPtr = ActorGuidsToContainerInstanceMap.Find((Info.ReferenceGuid)))
 								{
@@ -1272,7 +1273,7 @@ class FWorldPartitionStreamingGenerator
 
 			ContainerCollectionInstanceDescriptor.ActorDescViewMap->ForEachActorDescView([this, &ContainerCollectionInstanceDescriptor, &NbErrorsDetected, PassType](FStreamingGenerationActorDescView& ActorDescView)
 			{
-				FContainerCollectionInstanceDescriptor::FPerInstanceData& PerInstanceData = ContainerCollectionInstanceDescriptor.GetInstanceData(ActorDescView.GetGuid());
+				FContainerCollectionInstanceDescriptor::FPerInstanceData& PerInstanceData = ContainerCollectionInstanceDescriptor.GetPerInstanceData(ActorDescView.GetGuid());
 
 PRAGMA_DISABLE_DEPRECATION_WARNINGS
 				if (ActorDescView.GetActorDesc()->ShouldValidateRuntimeGrid())
@@ -1347,17 +1348,19 @@ PRAGMA_ENABLE_DEPRECATION_WARNINGS
 	/**
 	 * Validate the streaming generator internal state.
 	 */
-	void ValidateInternalState()
+	void ValidateContainerInstanceDescriptors()
 	{
-		// Validate that all actors part of the same actor set share the same set of values
-		for (const auto& [ContainerID, ContainerCollectionInstanceDescriptor] : ContainerCollectionInstanceDescriptorsMap)
+		for (auto& [ContainerID, ContainerCollectionInstanceDescriptor] : ContainerCollectionInstanceDescriptorsMap)
 		{
+			ValidateContainerInstanceDescriptor(ContainerCollectionInstanceDescriptor, ContainerCollectionInstanceDescriptor.ID.IsMainContainer());
+
 			for (const TArray<FGuid>& Cluster : ContainerCollectionInstanceDescriptor.Clusters)
 			{
 				const FStreamingGenerationActorDescView& ReferenceActorDescView = ContainerCollectionInstanceDescriptor.ActorDescViewMap->FindByGuidChecked(Cluster[0]);
 
 				for (const FGuid& ActorGuid : Cluster)
 				{
+					// Validate that all actors part of the same actor set share the same set of values
 					const FStreamingGenerationActorDescView& ActorDescView = ContainerCollectionInstanceDescriptor.ActorDescViewMap->FindByGuidChecked(ActorGuid);
 					check(ActorDescView.GetRuntimeGrid() == ReferenceActorDescView.GetRuntimeGrid());
 					check(ActorDescView.GetIsSpatiallyLoaded() == ReferenceActorDescView.GetIsSpatiallyLoaded());
@@ -1427,7 +1430,7 @@ public:
 			CreateContainerResolver(ContainerInstanceCollection);
 		}
 
-		ValidateInternalState();
+		ValidateContainerInstanceDescriptors();
 	}
 
 	static TUniquePtr<FArchive> CreateDumpStateLogArchive(const TCHAR* Suffix, bool bTimeStamped = true)
@@ -1947,4 +1950,4 @@ void UWorldPartition::CheckForErrors(const FCheckForErrorsParams& Params)
 }
 #endif // WITH_EDITOR
 
-#undef LOCTEXT_NAMESPACE
\ No newline at end of file
+#undef LOCTEXT_NAMESPACE
diff --git a/Engine/Source/Runtime/Engine/Public/WorldPartition/WorldPartitionStreamingGeneration.h b/Engine/Source/Runtime/Engine/Public/WorldPartition/WorldPartitionStreamingGeneration.h
--- a/Engine/Source/Runtime/Engine/Public/WorldPartition/WorldPartitionStreamingGeneration.h
+++ b/Engine/Source/Runtime/Engine/Public/WorldPartition/WorldPartitionStreamingGeneration.h
@@ -7,7 +7,6 @@
 #include "WorldPartition/ActorDescContainerCollection.h"
 #include "WorldPartition/ActorDescContainerInstanceCollection.h"
 #include "WorldPartition/WorldPartitionActorDescInstance.h"
-#include "WorldPartition/WorldPartitionActorDescInstanceInterface.h"
 #include "WorldPartition/WorldPartitionActorDescInstanceViewInterface.h"
 
 // deprecated
@@ -22,7 +21,7 @@ class FStreamingGenerationActorDescView : public IWorldPartitionActorDescInstanc
 
 protected:
 	// Used for invalid reference error reporting
-	FStreamingGenerationActorDescView(const IWorldPartitionActorDescInstance* InActorDescInstance)
+	FStreamingGenerationActorDescView(const FWorldPartitionActorDescInstance* InActorDescInstance)
 		: ActorDescViewMap(nullptr)
 		, ActorDescInstance(InActorDescInstance)
 		, ParentView(nullptr)
@@ -35,7 +34,7 @@ protected:
 	}
 
 public:
-	FStreamingGenerationActorDescView(const FStreamingGenerationActorDescViewMap& InActorDescViewMap, const IWorldPartitionActorDescInstance* InActorDescInstance, bool bInUnsaved = false)
+	FStreamingGenerationActorDescView(const FStreamingGenerationActorDescViewMap& InActorDescViewMap, const FWorldPartitionActorDescInstance* InActorDescInstance, bool bInUnsaved = false)
 		: FStreamingGenerationActorDescView(InActorDescInstance)
 	{
 		ActorDescViewMap = &InActorDescViewMap;
@@ -140,7 +139,7 @@ public:
 
 private:
 	const FStreamingGenerationActorDescViewMap* ActorDescViewMap;
-	const IWorldPartitionActorDescInstance* ActorDescInstance;
+	const FWorldPartitionActorDescInstance* ActorDescInstance;
 		
 	const FStreamingGenerationActorDescView* ParentView;
 
