Iris
* Make CSV scopes match between Iris and other replication systems

#rnx

diff --git a/Engine/Source/Runtime/Engine/Classes/Engine/NetDriver.h b/Engine/Source/Runtime/Engine/Classes/Engine/NetDriver.h
--- a/Engine/Source/Runtime/Engine/Classes/Engine/NetDriver.h
+++ b/Engine/Source/Runtime/Engine/Classes/Engine/NetDriver.h
@@ -2192,6 +2192,14 @@ private:
 	/** Called from RPC processing code to forward RPC to other NetDrivers if ShouldForwardFunction returns true. */
 	void ForwardRemoteFunction(UObject* RootObject, UObject* SubObject, UFunction* Function, void* Parms);
 
+	/** Go over imported network guids and map them to the newly created object. */
+	void UpdateUnmappedObjects();
+
+	/** Periodically look for invalid dormant replicators tied to destroyed objects. */
+	void CleanupStaleDormantReplicators();
+
+private:
+
 	UPROPERTY(transient)
 	TObjectPtr<UReplicationDriver> ReplicationDriver;
 
diff --git a/Engine/Source/Runtime/Engine/Private/NetDriver.cpp b/Engine/Source/Runtime/Engine/Private/NetDriver.cpp
--- a/Engine/Source/Runtime/Engine/Private/NetDriver.cpp
+++ b/Engine/Source/Runtime/Engine/Private/NetDriver.cpp
@@ -925,7 +925,7 @@ void UNetDriver::TickFlush(float DeltaSeconds)
 	{
 		// Update all clients.
 #if WITH_SERVER_CODE
-		int32 Updated = 0;
+		CSV_SCOPED_TIMING_STAT_EXCLUSIVE(ServerReplicateActors);
 
 #if UE_WITH_IRIS
 		if (ReplicationSystem)
@@ -937,16 +937,8 @@ void UNetDriver::TickFlush(float DeltaSeconds)
 		else
 #endif // UE_WITH_IRIS
 		{
-			Updated = ServerReplicateActors(DeltaSeconds);
+			ServerReplicateActors(DeltaSeconds);
 		}
-
-		static int32 LastUpdateCount = 0;
-		// Only log the zero replicated actors once after replicating an actor
-		if ((LastUpdateCount && !Updated) || Updated)
-		{
-			UE_LOG(LogNetTraffic, Verbose, TEXT("%s replicated %d actors"), *GetDescription(), Updated);
-		}
-		LastUpdateCount = Updated;
 #endif // WITH_SERVER_CODE
 	}
 #if UE_WITH_IRIS
@@ -996,199 +988,219 @@ void UNetDriver::TickFlush(float DeltaSeconds)
 		DrawNetDriverDebug();
 	}
 
+	if (!IsUsingIrisReplication())
 	{
-		CSV_SCOPED_TIMING_STAT(Networking, UpdateUnmappedObjects);
-		SCOPE_CYCLE_COUNTER(STAT_NetUpdateUnmappedObjectsTime);
+		UpdateUnmappedObjects();
 
-		if (CVarOptimizedRemapping.GetValueOnAnyThread() && GuidCache.IsValid())
-		{
-			// Go over recently imported network guids, and see if there are any replicators that need to map them
-			TSet<FNetworkGUID>& ImportedNetGuidsRef = GuidCache->ImportedNetGuids;
-			TMap<FNetworkGUID, TSet<FNetworkGUID>>& PendingOuterNetGuidsRef = GuidCache->PendingOuterNetGuids;
-
-			GetMetrics()->SetInt(UE::Net::Metric::ImportedNetGuids, ImportedNetGuidsRef.Num());
-			GetMetrics()->SetInt(UE::Net::Metric::PendingOuterNetGuids, PendingOuterNetGuidsRef.Num());
-			GetMetrics()->SetInt(UE::Net::Metric::UnmappedReplicators, UnmappedReplicators.Num());
+		CleanupStaleDormantReplicators();
+	}
 
-			TSet<FObjectReplicator*> ForceUpdateReplicators;
+	UpdateNetworkStats();
 
-			for (FObjectReplicator* Replicator : UnmappedReplicators)
-			{
-				if (Replicator->bForceUpdateUnmapped)
-				{
-					Replicator->bForceUpdateUnmapped = false;
-					ForceUpdateReplicators.Add(Replicator);
-				}
-			}
+	// Send the current values of metrics to all of the metrics listeners.
+	GetMetrics()->ProcessListeners();
 
-			if (ImportedNetGuidsRef.Num() || ForceUpdateReplicators.Num())
-			{
-				int32 NumRemapTests = 0;
-				const int32 MaxRemaps = IsServer() ? 0 : CVarMaxClientGuidRemaps.GetValueOnAnyThread();
+	// Update the lag state
+	UpdateNetworkLagState();
+}
 
-				TArray<FNetworkGUID> UnmappedGuids;
-				TArray<FNetworkGUID> NewlyMappedGuids;
+void UNetDriver::UpdateUnmappedObjects()
+{
+	CSV_SCOPED_TIMING_STAT(Networking, UpdateUnmappedObjects);
+	SCOPE_CYCLE_COUNTER(STAT_NetUpdateUnmappedObjectsTime);
 
-				for (auto It = ImportedNetGuidsRef.CreateIterator(); It; ++It)
-				{
-					const FNetworkGUID NetworkGuid = *It;
-					bool bMappedOrBroken = false;
+	if (CVarOptimizedRemapping.GetValueOnAnyThread() && GuidCache.IsValid())
+	{
+		// Go over recently imported network guids, and see if there are any replicators that need to map them
+		TSet<FNetworkGUID>& ImportedNetGuidsRef = GuidCache->ImportedNetGuids;
+		TMap<FNetworkGUID, TSet<FNetworkGUID>>& PendingOuterNetGuidsRef = GuidCache->PendingOuterNetGuids;
 
-					if (GuidCache->GetObjectFromNetGUID(NetworkGuid, false) != nullptr)
-					{
-						if (UE::Net::Private::bRemapStableSubobjects)
-						{
-							QUICK_SCOPE_CYCLE_COUNTER(STAT_NetRemapStableSubobjects);
+		GetMetrics()->SetInt(UE::Net::Metric::ImportedNetGuids, ImportedNetGuidsRef.Num());
+		GetMetrics()->SetInt(UE::Net::Metric::PendingOuterNetGuids, PendingOuterNetGuidsRef.Num());
+		GetMetrics()->SetInt(UE::Net::Metric::UnmappedReplicators, UnmappedReplicators.Num());
 
-							// Import any unmapped, stably-named guids that are inners of the GUID that just mapped.
-							// These are tracked separately from the normal ImportedNetGuids since they are often
-							// default subobjects created in constructors and there's no other hook to import them.
-							const TArray<FNetworkGUID>* Inners = GuidCache->FindUnmappedStablyNamedGuidsWithOuter(NetworkGuid);
+		TSet<FObjectReplicator*> ForceUpdateReplicators;
 
-							if (Inners)
-							{
-								TSet<FNetworkGUID>& PendingGuidsRef = PendingOuterNetGuidsRef.FindOrAdd(NetworkGuid);
-								PendingGuidsRef.Append(*Inners);
-									
-								// Now that they're on the pending import list, they will stay there until mapped. Can remove from the outer-to-inner map.
-								GuidCache->RemoveUnmappedStablyNamedGuidsWithOuter(NetworkGuid);
-							}
-						}
+		for (FObjectReplicator* Replicator : UnmappedReplicators)
+		{
+			if (Replicator->bForceUpdateUnmapped)
+			{
+				Replicator->bForceUpdateUnmapped = false;
+				ForceUpdateReplicators.Add(Replicator);
+			}
+		}
 
-						NewlyMappedGuids.Add(NetworkGuid);
-						bMappedOrBroken = true;
-					}
+		if (ImportedNetGuidsRef.Num() || ForceUpdateReplicators.Num())
+		{
+			int32 NumRemapTests = 0;
+			const int32 MaxRemaps = IsServer() ? 0 : CVarMaxClientGuidRemaps.GetValueOnAnyThread();
 
-					if (GuidCache->IsGUIDBroken(NetworkGuid, false))
-					{
-						bMappedOrBroken = true;
-					}
+			TArray<FNetworkGUID> UnmappedGuids;
+			TArray<FNetworkGUID> NewlyMappedGuids;
 
-					It.RemoveCurrent();
+			for (auto It = ImportedNetGuidsRef.CreateIterator(); It; ++It)
+			{
+				const FNetworkGUID NetworkGuid = *It;
+				bool bMappedOrBroken = false;
 
-					if (!bMappedOrBroken)
+				if (GuidCache->GetObjectFromNetGUID(NetworkGuid, false) != nullptr)
+				{
+					if (UE::Net::Private::bRemapStableSubobjects)
 					{
-						const FNetworkGUID OuterGUID = GuidCache->GetOuterNetGUID(NetworkGuid);
+						QUICK_SCOPE_CYCLE_COUNTER(STAT_NetRemapStableSubobjects);
 
-						// we're missing the outer, stop checking until we map it
-						if ((UE::Net::FilterGuidRemapping != 0) && OuterGUID.IsValid() && !OuterGUID.IsDefault() && !GuidCache->IsGUIDLoaded(OuterGUID) && !GuidCache->IsGUIDPending(OuterGUID))
-						{
-							UE_LOG(LogNetPackageMap, Log, TEXT("Missing outer (%s) for unmapped guid (%s), marking pending"), *GuidCache->Describe(OuterGUID), *GuidCache->Describe(NetworkGuid));
+						// Import any unmapped, stably-named guids that are inners of the GUID that just mapped.
+						// These are tracked separately from the normal ImportedNetGuids since they are often
+						// default subobjects created in constructors and there's no other hook to import them.
+						const TArray<FNetworkGUID>* Inners = GuidCache->FindUnmappedStablyNamedGuidsWithOuter(NetworkGuid);
 
-							TSet<FNetworkGUID>& PendingGuidsRef = PendingOuterNetGuidsRef.FindOrAdd(OuterGUID);
-							PendingGuidsRef.Add(NetworkGuid);
-						}
-						else
+						if (Inners)
 						{
-							if (ensure(NetworkGuid.IsValid()))
-							{
-								UnmappedGuids.Add(NetworkGuid);
-							}
+							TSet<FNetworkGUID>& PendingGuidsRef = PendingOuterNetGuidsRef.FindOrAdd(NetworkGuid);
+							PendingGuidsRef.Append(*Inners);
+
+							// Now that they're on the pending import list, they will stay there until mapped. Can remove from the outer-to-inner map.
+							GuidCache->RemoveUnmappedStablyNamedGuidsWithOuter(NetworkGuid);
 						}
 					}
 
-					++NumRemapTests;
+					NewlyMappedGuids.Add(NetworkGuid);
+					bMappedOrBroken = true;
+				}
 
-					if ((MaxRemaps > 0) && (NumRemapTests >= MaxRemaps))
-					{
-						break;
-					}
+				if (GuidCache->IsGUIDBroken(NetworkGuid, false))
+				{
+					bMappedOrBroken = true;
 				}
 
-				// attempt to resolve dependent guids next tick (outer is now mapped)
-				for (const FNetworkGUID& NetGuid : NewlyMappedGuids)
+				It.RemoveCurrent();
+
+				if (!bMappedOrBroken)
 				{
-					if (TSet<FNetworkGUID>* DependentGuids = PendingOuterNetGuidsRef.Find(NetGuid))
+					const FNetworkGUID OuterGUID = GuidCache->GetOuterNetGUID(NetworkGuid);
+
+					// we're missing the outer, stop checking until we map it
+					if ((UE::Net::FilterGuidRemapping != 0) && OuterGUID.IsValid() && !OuterGUID.IsDefault() && !GuidCache->IsGUIDLoaded(OuterGUID) && !GuidCache->IsGUIDPending(OuterGUID))
 					{
-						UE_LOG(LogNetPackageMap, Log, TEXT("Newly mapped outer (%s) removing from pending"), *GuidCache->FullNetGUIDPath(NetGuid));
+						UE_LOG(LogNetPackageMap, Log, TEXT("Missing outer (%s) for unmapped guid (%s), marking pending"), *GuidCache->Describe(OuterGUID), *GuidCache->Describe(NetworkGuid));
 
-						ImportedNetGuidsRef.Append(*DependentGuids);
-						PendingOuterNetGuidsRef.Remove(NetGuid);
+						TSet<FNetworkGUID>& PendingGuidsRef = PendingOuterNetGuidsRef.FindOrAdd(OuterGUID);
+						PendingGuidsRef.Add(NetworkGuid);
+					}
+					else
+					{
+						if (ensure(NetworkGuid.IsValid()))
+						{
+							UnmappedGuids.Add(NetworkGuid);
+						}
 					}
 				}
 
-				// any tested guids that could not yet be mapped are added back to the list
-				if (UnmappedGuids.Num())
+				++NumRemapTests;
+
+				if ((MaxRemaps > 0) && (NumRemapTests >= MaxRemaps))
 				{
-					ImportedNetGuidsRef.CompactStable();
-					ImportedNetGuidsRef.Append(UnmappedGuids);
+					break;
 				}
+			}
 
-				if (NewlyMappedGuids.Num() || ForceUpdateReplicators.Num())
+			// attempt to resolve dependent guids next tick (outer is now mapped)
+			for (const FNetworkGUID& NetGuid : NewlyMappedGuids)
+			{
+				if (TSet<FNetworkGUID>* DependentGuids = PendingOuterNetGuidsRef.Find(NetGuid))
 				{
-					TSet<FObjectReplicator*> ReplicatorsToUpdate = MoveTemp(ForceUpdateReplicators);
+					UE_LOG(LogNetPackageMap, Log, TEXT("Newly mapped outer (%s) removing from pending"), *GuidCache->FullNetGUIDPath(NetGuid));
+
+					ImportedNetGuidsRef.Append(*DependentGuids);
+					PendingOuterNetGuidsRef.Remove(NetGuid);
+				}
+			}
 
-					for (const FNetworkGUID& NetGuid : NewlyMappedGuids)
+			// any tested guids that could not yet be mapped are added back to the list
+			if (UnmappedGuids.Num())
+			{
+				ImportedNetGuidsRef.CompactStable();
+				ImportedNetGuidsRef.Append(UnmappedGuids);
+			}
+
+			if (NewlyMappedGuids.Num() || ForceUpdateReplicators.Num())
+			{
+				TSet<FObjectReplicator*> ReplicatorsToUpdate = MoveTemp(ForceUpdateReplicators);
+
+				for (const FNetworkGUID& NetGuid : NewlyMappedGuids)
+				{
+					if (TSet<FObjectReplicator*>* Replicators = GuidToReplicatorMap.Find(NetGuid))
 					{
-						if (TSet<FObjectReplicator*>* Replicators = GuidToReplicatorMap.Find(NetGuid))
-						{
-							ReplicatorsToUpdate.Append(*Replicators);
-						}
+						ReplicatorsToUpdate.Append(*Replicators);
 					}
+				}
 
-					for (FObjectReplicator* Replicator : ReplicatorsToUpdate)
+				for (FObjectReplicator* Replicator : ReplicatorsToUpdate)
+				{
+					if (UnmappedReplicators.Contains(Replicator))
 					{
-						if (UnmappedReplicators.Contains(Replicator))
-						{
-							bool bHasMoreUnmapped = false;
-							Replicator->UpdateUnmappedObjects(bHasMoreUnmapped);
+						bool bHasMoreUnmapped = false;
+						Replicator->UpdateUnmappedObjects(bHasMoreUnmapped);
 
-							if (!bHasMoreUnmapped)
-							{
-								UnmappedReplicators.Remove(Replicator);
-							}
+						if (!bHasMoreUnmapped)
+						{
+							UnmappedReplicators.Remove(Replicator);
 						}
 					}
 				}
 			}
 		}
-		else
+	}
+	else
+	{
+		// Update properties that are unmapped, try to hook up the object pointers if they exist now
+		for (auto It = UnmappedReplicators.CreateIterator(); It; ++It)
 		{
-			// Update properties that are unmapped, try to hook up the object pointers if they exist now
-			for (auto It = UnmappedReplicators.CreateIterator(); It; ++It)
-			{
-				FObjectReplicator* Replicator = *It;
+			FObjectReplicator* Replicator = *It;
 
-				bool bHasMoreUnmapped = false;
+			bool bHasMoreUnmapped = false;
 
-				Replicator->UpdateUnmappedObjects(bHasMoreUnmapped);
+			Replicator->UpdateUnmappedObjects(bHasMoreUnmapped);
 
-				if (!bHasMoreUnmapped)
-				{
-					// If there are no more unmapped objects, we can also stop checking
-					It.RemoveCurrent();
-				}
+			if (!bHasMoreUnmapped)
+			{
+				// If there are no more unmapped objects, we can also stop checking
+				It.RemoveCurrent();
 			}
 		}
 	}
+}
 
+void UNetDriver::CleanupStaleDormantReplicators()
+{
 	// Go over RepChangedPropertyTrackerMap periodically, and remove entries that no longer have valid objects
 	// Unfortunately if you mark an object as pending kill, it will no longer find itself in this map,
 	// so we do this as a fail safe to make sure we never leak memory from this map
+
 	const double CleanupTimeSeconds = 10.0;
 	const double CurrentRealtimeSeconds = FPlatformTime::Seconds();
 
-	if ( CurrentRealtimeSeconds - LastCleanupTime > CleanupTimeSeconds )
+	if (CurrentRealtimeSeconds - LastCleanupTime > CleanupTimeSeconds)
 	{
 		QUICK_SCOPE_CYCLE_COUNTER(STAT_NetUpdateCleanupTime);
 
-		for ( auto It = ReplicationChangeListMap.CreateIterator(); It; ++It )
+		for (auto It = ReplicationChangeListMap.CreateIterator(); It; ++It)
 		{
-			if ( !It.Value().IsObjectValid() )
+			if (!It.Value().IsObjectValid())
 			{
 				It.RemoveCurrent();
 			}
 		}
 
-		for ( UNetConnection* const ClientConnection : ClientConnections )
+		for (UNetConnection* const ClientConnection : ClientConnections)
 		{
-			if ( ClientConnection )
+			if (ClientConnection)
 			{
 				ClientConnection->CleanupStaleDormantReplicators();
 			}
 		}
 
-		if ( ServerConnection )
+		if (ServerConnection)
 		{
 			ServerConnection->CleanupStaleDormantReplicators();
 		}
@@ -1201,14 +1213,6 @@ void UNetDriver::TickFlush(float DeltaSeconds)
 		UE_LOG(LogNet, Log, TEXT("UNetDriver::TickFlush: %u channel(s) exceeded net.QueuedBunchTimeFailsafeSeconds and flushed their entire queue(s) this frame."), QueuedBunchFailsafeNumChannels);
 		QueuedBunchFailsafeNumChannels = 0;
 	}
-
-	UpdateNetworkStats();
-
-	// Send the current values of metrics to all of the metrics listeners.
-	GetMetrics()->ProcessListeners();
-
-	// Update the lag state
-	UpdateNetworkLagState();
 }
 
 void UNetDriver::UpdateNetworkLagState()
@@ -5573,7 +5577,6 @@ struct FScopedNetDriverStats
 int32 UNetDriver::ServerReplicateActors(float DeltaSeconds)
 {
 	SCOPE_CYCLE_COUNTER(STAT_NetServerRepActorsTime);
-	CSV_SCOPED_TIMING_STAT_EXCLUSIVE(ServerReplicateActors);
 
 #if WITH_SERVER_CODE
 	if ( ClientConnections.Num() == 0 )
