USceneComponent: CreateSpriteComponent add RF_Transactional only if owning component has the flag

diff --git a/Engine/Source/Runtime/Engine/Private/Components/SceneComponent.cpp b/Engine/Source/Runtime/Engine/Private/Components/SceneComponent.cpp
--- a/Engine/Source/Runtime/Engine/Private/Components/SceneComponent.cpp
+++ b/Engine/Source/Runtime/Engine/Private/Components/SceneComponent.cpp
@@ -816,7 +816,8 @@ void USceneComponent::CreateSpriteComponent(class UTexture2D* SpriteTexture, boo
 		// Create a new billboard component to serve as a visualization of the actor until there is another primitive component
 		{
 			FCookLoadScope EditorOnlyLoadScope(ECookLoadType::EditorOnly);
-			SpriteComponent = NewObject<UBillboardComponent>(GetOwner(), NAME_None, RF_Transactional | RF_Transient | RF_TextExportTransient);
+			const EObjectFlags TransactionalFlag = GetFlags() & RF_Transactional;
+			SpriteComponent = NewObject<UBillboardComponent>(GetOwner(), NAME_None, TransactionalFlag | RF_Transient | RF_TextExportTransient);
 		}
 
 		SpriteComponent->Sprite = SpriteTexture? SpriteTexture : LoadObject<UTexture2D>(nullptr, TEXT("/Engine/EditorResources/EmptyActor.EmptyActor"));
