Fixed removal of leaf data layer instance (when using data layer parenting) that was not properly updating parent.
Fixed removal of a data layer instance which could trigger an infinite loop when reparenting children of the removed data layer instance.

diff --git a/Engine/Source/Runtime/Engine/Private/WorldPartition/DataLayer/DataLayerInstance.cpp b/Engine/Source/Runtime/Engine/Private/WorldPartition/DataLayer/DataLayerInstance.cpp
--- a/Engine/Source/Runtime/Engine/Private/WorldPartition/DataLayer/DataLayerInstance.cpp
+++ b/Engine/Source/Runtime/Engine/Private/WorldPartition/DataLayer/DataLayerInstance.cpp
@@ -403,18 +403,19 @@ bool UDataLayerInstance::SetParent(UDataLayerInstance* InParent)
 	return true;
 }
 
-void UDataLayerInstance::SetChildParent(UDataLayerInstance* InParent)
+void UDataLayerInstance::OnRemovedFromWorldDataLayers()
 {
-	if (this == InParent)
-	{
-		return;
-	}
-
-	check(!InParent || InParent->CanHaveChildDataLayers());
 	while (Children.Num())
 	{
-		Children[0]->SetParent(InParent);
+		// If can't reparent, move to root
+		UDataLayerInstance* ChildNewParent = Children[0]->CanBeChildOf(Parent) ? Parent : nullptr;
+		verify(Children[0]->SetParent(ChildNewParent));
 	};
+
+	if (Parent)
+	{
+		Parent->RemoveChild(this);
+	}
 }
 
 void UDataLayerInstance::RemoveChild(UDataLayerInstance* InDataLayer)
diff --git a/Engine/Source/Runtime/Engine/Private/WorldPartition/DataLayer/WorldDataLayers.cpp b/Engine/Source/Runtime/Engine/Private/WorldPartition/DataLayer/WorldDataLayers.cpp
--- a/Engine/Source/Runtime/Engine/Private/WorldPartition/DataLayer/WorldDataLayers.cpp
+++ b/Engine/Source/Runtime/Engine/Private/WorldPartition/DataLayer/WorldDataLayers.cpp
@@ -675,7 +675,7 @@ int32 AWorldDataLayers::RemoveDataLayers(const TArray<UDataLayerInstance*>& InDa
 			const bool bDirty = !bUseExternalPackageDataLayerInstances;
 			Modify(bDirty);
 			DataLayerInstance->Modify();
-			DataLayerInstance->SetChildParent(DataLayerInstance->GetParent());
+			DataLayerInstance->OnRemovedFromWorldDataLayers();
 			UsedDataLayerInstances.Remove(DataLayerInstance);
 			if (DataLayerInstance->IsA<UDeprecatedDataLayerInstance>())
 			{
diff --git a/Engine/Source/Runtime/Engine/Public/WorldPartition/DataLayer/DataLayerInstance.h b/Engine/Source/Runtime/Engine/Public/WorldPartition/DataLayer/DataLayerInstance.h
--- a/Engine/Source/Runtime/Engine/Public/WorldPartition/DataLayer/DataLayerInstance.h
+++ b/Engine/Source/Runtime/Engine/Public/WorldPartition/DataLayer/DataLayerInstance.h
@@ -93,7 +93,8 @@ public:
 	ENGINE_API bool CanBeChildOf(const UDataLayerInstance* InParent, FText* OutReason = nullptr) const;
 	ENGINE_API bool SetParent(UDataLayerInstance* InParent);
 
-	ENGINE_API void SetChildParent(UDataLayerInstance* InParent);
+	UE_DEPRECATED(5.4, "SetChildParent was removed")
+	ENGINE_API void SetChildParent(UDataLayerInstance* InParent) {}
 
 	static ENGINE_API FText GetDataLayerText(const UDataLayerInstance* InDataLayer);
 
@@ -199,6 +200,7 @@ private:
 	friend class FDataLayerUtils;
 
 #if WITH_EDITOR
+	ENGINE_API void OnRemovedFromWorldDataLayers();
 	ENGINE_API void RemoveChild(UDataLayerInstance* DataLayer);
 #endif
 
@@ -248,4 +250,6 @@ private:
 
 	UPROPERTY(Transient)
 	TArray<TObjectPtr<UDataLayerInstance>> Children;
+
+	friend class AWorldDataLayers;
 };
