ApplyCVars Fix: Move Effective values from globals (which could change between cook and runtime) to per SpatialHash values that are initialized only when GenerateStreaming is called

This also fixes ApplyCVars being called mulitple times in PIE when using Level Instances which could lead to different results if Level Instances were using different settings than main world

diff --git a/Engine/Source/Runtime/Engine/Private/WorldPartition/RuntimeSpatialHash/RuntimeSpatialHashGridHelper.cpp b/Engine/Source/Runtime/Engine/Private/WorldPartition/RuntimeSpatialHash/RuntimeSpatialHashGridHelper.cpp
--- a/Engine/Source/Runtime/Engine/Private/WorldPartition/RuntimeSpatialHash/RuntimeSpatialHashGridHelper.cpp
+++ b/Engine/Source/Runtime/Engine/Private/WorldPartition/RuntimeSpatialHash/RuntimeSpatialHashGridHelper.cpp
@@ -7,6 +7,8 @@
 #include "WorldPartition/WorldPartitionStreamingGeneration.h"
 #include "WorldPartition/WorldPartitionStreamingSource.h"
 
+#if WITH_EDITOR
+
 bool GRuntimeSpatialHashUseAlignedGridLevels = true;
 static FAutoConsoleVariableRef CVarRuntimeSpatialHashUseAlignedGridLevels(
 	TEXT("wp.Runtime.RuntimeSpatialHashUseAlignedGridLevels"),
@@ -31,12 +33,13 @@ static FAutoConsoleVariableRef CVarRuntimeSpatialHashPlacePartitionActorsUsingLo
 	GRuntimeSpatialHashPlacePartitionActorsUsingLocation,
 	TEXT("Set RuntimeSpatialHashPlacePartitionActorsUsingLocation to true to place partitioned actors into their corresponding cell using their location instead of their bounding box."));
 
-bool GRuntimeSpatialHashUseAlignedGridLevelsEffective = true;
-bool GRuntimeSpatialHashSnapNonAlignedGridLevelsToLowerLevelsEffective = true;
-bool GRuntimeSpatialHashPlaceSmallActorsUsingLocationEffective = false;
-bool GRuntimeSpatialHashPlacePartitionActorsUsingLocationEffective = true;
+#endif
 
 FSquare2DGridHelper::FSquare2DGridHelper(const FBox& InWorldBounds, const FVector& InOrigin, int64 InCellSize)
+	: FSquare2DGridHelper(InWorldBounds, InOrigin, InCellSize, true)
+{}
+
+FSquare2DGridHelper::FSquare2DGridHelper(const FBox& InWorldBounds, const FVector& InOrigin, int64 InCellSize, bool bUseAlignedGridLevels)
 	: WorldBounds(InWorldBounds)
 	, Origin(InOrigin)
 	, CellSize(InCellSize)
@@ -71,7 +74,7 @@ FSquare2DGridHelper::FSquare2DGridHelper(const FBox& InWorldBounds, const FVecto
 	{
 		int64 LevelGridSize = CurrentGridSize;
 
-		if (!GRuntimeSpatialHashUseAlignedGridLevelsEffective)
+		if (!bUseAlignedGridLevels)
 		{
 			// Except for top level, adding 1 to CurrentGridSize (which is always a power of 2) breaks the pattern of perfectly aligned cell edges between grid level cells.
 			// This will prevent weird artefact during actor promotion when an actor is placed using its bounds and which overlaps multiple cells.
@@ -138,20 +141,33 @@ int32 FSquare2DGridHelper::ForEachIntersectingCells(const FSphericalSector& InSh
 }
 
 #if WITH_EDITOR
+// Deprecated
 FSquare2DGridHelper GetGridHelper(const FBox& WorldBounds, const FVector& GridOrigin, int64 GridCellSize)
+{
+	return GetGridHelper(WorldBounds, GridOrigin, GridCellSize, true);
+}
+
+FSquare2DGridHelper GetGridHelper(const FBox& WorldBounds, const FVector& GridOrigin, int64 GridCellSize, bool bUseAlignedGridLevels)
 {
 	// Default grid to a minimum of 1 level and 1 cell, for always loaded actors
-	return FSquare2DGridHelper(WorldBounds, GridOrigin, GridCellSize);
+	return FSquare2DGridHelper(WorldBounds, GridOrigin, GridCellSize, bUseAlignedGridLevels);
 }
 
+// Deprecated
 FSquare2DGridHelper GetPartitionedActors(const FBox& WorldBounds, const FSpatialHashRuntimeGrid& Grid, const TArray<const IStreamingGenerationContext::FActorSetInstance*>& ActorSetInstances)
+{
+	FSpatialHashSettings Settings;
+	return GetPartitionedActors(WorldBounds, Grid, ActorSetInstances, Settings);
+}
+
+FSquare2DGridHelper GetPartitionedActors(const FBox& WorldBounds, const FSpatialHashRuntimeGrid& Grid, const TArray<const IStreamingGenerationContext::FActorSetInstance*>& ActorSetInstances, const FSpatialHashSettings& Settings)
 {
 	TRACE_CPUPROFILER_EVENT_SCOPE(GetPartitionedActors);
 
 	//
 	// Create the hierarchical grids for the game
 	//	
-	FSquare2DGridHelper PartitionedActors = GetGridHelper(WorldBounds, FVector(Grid.Origin, 0), Grid.CellSize);
+	FSquare2DGridHelper PartitionedActors = GetGridHelper(WorldBounds, FVector(Grid.Origin, 0), Grid.CellSize, Settings.bUseAlignedGridLevels);
 	const bool bWorldBoundsValid = WorldBounds.IsValid && WorldBounds.GetExtent().Size2D() > 0;
 	if (ensure(PartitionedActors.Levels.Num()) && bWorldBoundsValid)
 	{
@@ -168,15 +184,15 @@ FSquare2DGridHelper GetPartitionedActors(const FBox& WorldBounds, const FSpatial
 	const int64 CellSize = PartitionedActors.GetLowestLevel().CellSize;
 	const float CellArea = CellSize * CellSize;
 
-	auto ShouldActorUseLocationPlacement = [CellArea, CellSize, GridLevelCount](const IStreamingGenerationContext::FActorSetInstance* InActorSetInstance, const FBox2D& InActorSetInstanceBounds, int32& OutGridLevel)
+	auto ShouldActorUseLocationPlacement = [CellArea, CellSize, GridLevelCount, &Settings](const IStreamingGenerationContext::FActorSetInstance* InActorSetInstance, const FBox2D& InActorSetInstanceBounds, int32& OutGridLevel)
 	{
 		OutGridLevel = 0;
-		if (GRuntimeSpatialHashPlaceSmallActorsUsingLocationEffective && (InActorSetInstanceBounds.GetArea() <= CellArea))
+		if (Settings.bPlaceSmallActorsUsingLocation && (InActorSetInstanceBounds.GetArea() <= CellArea))
 		{
 			return true;
 		}
 
-		if (GRuntimeSpatialHashPlacePartitionActorsUsingLocationEffective)
+		if (Settings.bPlacePartitionActorsUsingLocation)
 		{
 			bool bUseLocation = true;
 			InActorSetInstance->ForEachActor([InActorSetInstance, &bUseLocation](const FGuid& ActorGuid)
diff --git a/Engine/Source/Runtime/Engine/Private/WorldPartition/RuntimeSpatialHash/RuntimeSpatialHashHLOD.cpp b/Engine/Source/Runtime/Engine/Private/WorldPartition/RuntimeSpatialHash/RuntimeSpatialHashHLOD.cpp
--- a/Engine/Source/Runtime/Engine/Private/WorldPartition/RuntimeSpatialHash/RuntimeSpatialHashHLOD.cpp
+++ b/Engine/Source/Runtime/Engine/Private/WorldPartition/RuntimeSpatialHash/RuntimeSpatialHashHLOD.cpp
@@ -112,12 +112,12 @@ static void GameTick(UWorld* InWorld)
 	}
 }
 
-static TArray<FGuid> GenerateHLODActorsForGrid(UWorldPartition* WorldPartition, const IStreamingGenerationContext* StreamingGenerationContext, const UWorldPartition::FSetupHLODActorsParams& Params, const FSpatialHashRuntimeGrid& RuntimeGrid, uint32 HLODLevel, FHLODCreationContext& Context, ISourceControlHelper* SourceControlHelper, const TArray<const IStreamingGenerationContext::FActorSetInstance*>& ActorSetInstances, TArray<FName>& NewActors, TArray<FName>& DirtyActors)
+static TArray<FGuid> GenerateHLODActorsForGrid(UWorldPartition* WorldPartition, const IStreamingGenerationContext* StreamingGenerationContext, const UWorldPartition::FSetupHLODActorsParams& Params, const FSpatialHashRuntimeGrid& RuntimeGrid, uint32 HLODLevel, FHLODCreationContext& Context, ISourceControlHelper* SourceControlHelper, const TArray<const IStreamingGenerationContext::FActorSetInstance*>& ActorSetInstances, TArray<FName>& NewActors, TArray<FName>& DirtyActors, const FSpatialHashSettings& Settings)
 {
 	const IStreamingGenerationContext::FActorSetContainerInstance* MainActorSetContainer = StreamingGenerationContext->GetMainWorldContainerInstance();
 	const FBox WorldBounds = StreamingGenerationContext->GetWorldBounds();
 
-	const FSquare2DGridHelper PartitionedActors = GetPartitionedActors(WorldBounds, RuntimeGrid, ActorSetInstances);
+	const FSquare2DGridHelper PartitionedActors = GetPartitionedActors(WorldBounds, RuntimeGrid, ActorSetInstances, Settings);
 	const FSquare2DGridHelper::FGridLevel::FGridCell& AlwaysLoadedCell = PartitionedActors.GetAlwaysLoadedCell();
 
 	auto ShouldGenerateHLODActors = [&AlwaysLoadedCell](const FSquare2DGridHelper::FGridLevel::FGridCell& GridCell, const FSquare2DGridHelper::FGridLevel::FGridCellDataChunk& GridCellDataChunk)
@@ -499,11 +499,11 @@ bool UWorldPartitionRuntimeSpatialHash::SetupHLODActors(const IStreamingGenerati
 
 	const UDataLayerManager* DataLayerManager = WorldPartition->GetDataLayerManager();
 
-	auto GenerateHLODActors = [&GridsHLODActors, &MainActorSetContainer, StreamingGenerationContext, &Params, WorldPartition, DataLayerManager, &Context, SourceControlHelper, &NewActors, &DirtyActors]
+	auto GenerateHLODActors = [&GridsHLODActors, &MainActorSetContainer, StreamingGenerationContext, &Params, WorldPartition, DataLayerManager, &Context, SourceControlHelper, &NewActors, &DirtyActors, this]
 	(const FSpatialHashRuntimeGrid& RuntimeGrid, uint32 HLODLevel, const TArray<const IStreamingGenerationContext::FActorSetInstance*>& ActorSetInstances)
 	{
 		// Generate HLODs for this grid
-		TArray<FGuid> HLODActors = GenerateHLODActorsForGrid(WorldPartition, StreamingGenerationContext, Params, RuntimeGrid, HLODLevel, Context, SourceControlHelper, ActorSetInstances, NewActors, DirtyActors);
+		TArray<FGuid> HLODActors = GenerateHLODActorsForGrid(WorldPartition, StreamingGenerationContext, Params, RuntimeGrid, HLODLevel, Context, SourceControlHelper, ActorSetInstances, NewActors, DirtyActors, Settings);
 
 		for (const FGuid& HLODActorGuid : HLODActors)
 		{
diff --git a/Engine/Source/Runtime/Engine/Private/WorldPartition/WorldPartitionRuntimeSpatialHash.cpp b/Engine/Source/Runtime/Engine/Private/WorldPartition/WorldPartitionRuntimeSpatialHash.cpp
--- a/Engine/Source/Runtime/Engine/Private/WorldPartition/WorldPartitionRuntimeSpatialHash.cpp
+++ b/Engine/Source/Runtime/Engine/Private/WorldPartition/WorldPartitionRuntimeSpatialHash.cpp
@@ -112,7 +112,7 @@ const FSquare2DGridHelper& FSpatialHashStreamingGrid::GetGridHelper() const
 {
 	if (!GridHelper)
 	{
-		GridHelper = new FSquare2DGridHelper(WorldBounds, Origin, CellSize);
+		GridHelper = new FSquare2DGridHelper(WorldBounds, Origin, CellSize, Settings.bUseAlignedGridLevels);
 	}
 
 	check(GridHelper->Levels.Num() == GridLevels.Num());
@@ -324,7 +324,7 @@ void FSpatialHashStreamingGrid::GetCells(const TArray<FWorldPartitionStreamingSo
 									break;
 								case EStreamingSourceTargetState::Activated:
 									OutActivateCells.AddCell(Cell, Source, Shape);
-									bAddedActivatedCell = !GRuntimeSpatialHashUseAlignedGridLevelsEffective && GRuntimeSpatialHashSnapNonAlignedGridLevelsToLowerLevelsEffective;
+									bAddedActivatedCell = !Settings.bUseAlignedGridLevels && Settings.bSnapNonAlignedGridLevelsToLowerLevels;
 									break;
 								default:
 									checkNoEntry();
@@ -348,7 +348,7 @@ void FSpatialHashStreamingGrid::GetCells(const TArray<FWorldPartitionStreamingSo
 
 	GetNonSpatiallyLoadedCells(OutActivateCells.GetCells(), OutLoadCells.GetCells());
 
-	if (!GRuntimeSpatialHashUseAlignedGridLevelsEffective && GRuntimeSpatialHashSnapNonAlignedGridLevelsToLowerLevelsEffective)
+	if (!Settings.bUseAlignedGridLevels && Settings.bSnapNonAlignedGridLevelsToLowerLevels)
 	{
 		auto FindIntersectingParents = [&Helper, this](const FIntersectingCells& InAllCells, const FIntersectingCells& InTestCells, FIntersectingCells& OutIntersectingCells)
 		{
@@ -1051,11 +1051,11 @@ UWorldPartitionRuntimeSpatialHash::UWorldPartitionRuntimeSpatialHash(const FObje
 #if WITH_EDITORONLY_DATA
 	, bPreviewGrids(false)
 	, PreviewGridLevel(0)
-#endif
 	, UseAlignedGridLevels(EWorldPartitionCVarProjectDefaultOverride::Disabled)
 	, SnapNonAlignedGridLevelsToLowerLevels(EWorldPartitionCVarProjectDefaultOverride::Disabled)
 	, PlaceSmallActorsUsingLocation(EWorldPartitionCVarProjectDefaultOverride::Enabled)
 	, PlacePartitionActorsUsingLocation(EWorldPartitionCVarProjectDefaultOverride::Enabled)
+#endif
 	, bIsNameToGridMappingDirty(true)
 {}
 
@@ -1087,7 +1087,9 @@ FString UWorldPartitionRuntimeSpatialHash::GetCellCoordString(const FGridCellCoo
 #if WITH_EDITOR
 void UWorldPartitionRuntimeSpatialHash::DrawPreview() const
 {
-	GridPreviewer.Draw(GetWorld(), Grids, bPreviewGrids, PreviewGridLevel);
+	// Use latest settings value for Preview
+	const bool bUseAlignedGridLevels = (UseAlignedGridLevels == EWorldPartitionCVarProjectDefaultOverride::ProjectDefault) ? GRuntimeSpatialHashUseAlignedGridLevels : (UseAlignedGridLevels == EWorldPartitionCVarProjectDefaultOverride::Enabled);
+	GridPreviewer.Draw(GetWorld(), Grids, bPreviewGrids, PreviewGridLevel, bUseAlignedGridLevels);
 }
 
 URuntimeHashExternalStreamingObjectBase* UWorldPartitionRuntimeSpatialHash::StoreToExternalStreamingObject(UObject* StreamingObjectOuter, FName StreamingObjectName)
@@ -1119,14 +1121,6 @@ URuntimeHashExternalStreamingObjectBase* UWorldPartitionRuntimeSpatialHash::Stor
 }
 #endif
 
-void UWorldPartitionRuntimeSpatialHash::ApplyCVars()
-{
-	GRuntimeSpatialHashUseAlignedGridLevelsEffective = (UseAlignedGridLevels == EWorldPartitionCVarProjectDefaultOverride::ProjectDefault) ? GRuntimeSpatialHashUseAlignedGridLevels : (UseAlignedGridLevels == EWorldPartitionCVarProjectDefaultOverride::Enabled);
-	GRuntimeSpatialHashSnapNonAlignedGridLevelsToLowerLevelsEffective = (SnapNonAlignedGridLevelsToLowerLevels == EWorldPartitionCVarProjectDefaultOverride::ProjectDefault) ? GRuntimeSpatialHashSnapNonAlignedGridLevelsToLowerLevels : (SnapNonAlignedGridLevelsToLowerLevels == EWorldPartitionCVarProjectDefaultOverride::Enabled);
-	GRuntimeSpatialHashPlaceSmallActorsUsingLocationEffective = (PlaceSmallActorsUsingLocation == EWorldPartitionCVarProjectDefaultOverride::ProjectDefault) ? GRuntimeSpatialHashPlaceSmallActorsUsingLocation : (PlaceSmallActorsUsingLocation == EWorldPartitionCVarProjectDefaultOverride::Enabled);
-	GRuntimeSpatialHashPlacePartitionActorsUsingLocationEffective = (PlacePartitionActorsUsingLocation == EWorldPartitionCVarProjectDefaultOverride::ProjectDefault) ? GRuntimeSpatialHashPlacePartitionActorsUsingLocation : (PlacePartitionActorsUsingLocation == EWorldPartitionCVarProjectDefaultOverride::Enabled);
-}
-
 void UWorldPartitionRuntimeSpatialHash::Serialize(FArchive& Ar)
 {
 	Ar.UsingCustomVersion(FFortniteMainBranchObjectVersion::GUID);
@@ -1142,8 +1136,16 @@ void UWorldPartitionRuntimeSpatialHash::Serialize(FArchive& Ar)
 		PlacePartitionActorsUsingLocation = EWorldPartitionCVarProjectDefaultOverride::ProjectDefault;
 	}
 #endif
+}
+
+void UWorldPartitionRuntimeSpatialHash::PostLoad()
+{
+	Super::PostLoad();
 
-	ApplyCVars();
+#if !WITH_EDITOR
+	UE_LOG(LogWorldPartition, Log, TEXT("UWorldPartitionRuntimeSpatialHash::PostLoad : UseAlignedGridLevels = %d, SnapNonAlignedGridLevelsToLowerLevels = %d"),
+		Settings.bUseAlignedGridLevels, Settings.bSnapNonAlignedGridLevelsToLowerLevels);
+#endif
 }
 
 #if WITH_EDITOR
@@ -1177,8 +1179,6 @@ void UWorldPartitionRuntimeSpatialHash::PostEditChangeProperty(FPropertyChangedE
 		}
 	}
 
-	ApplyCVars();
-
 	Super::PostEditChangeProperty(PropertyChangedEvent);
 }
 
@@ -1191,6 +1191,17 @@ void UWorldPartitionRuntimeSpatialHash::SetDefaultValues()
 	MainGrid.DebugColor = FLinearColor::Gray;
 }
 
+void UWorldPartitionRuntimeSpatialHash::UpdateSettings()
+{
+	check(StreamingGrids.IsEmpty());
+	Settings.bUseAlignedGridLevels = (UseAlignedGridLevels == EWorldPartitionCVarProjectDefaultOverride::ProjectDefault) ? GRuntimeSpatialHashUseAlignedGridLevels : (UseAlignedGridLevels == EWorldPartitionCVarProjectDefaultOverride::Enabled);
+	Settings.bSnapNonAlignedGridLevelsToLowerLevels = (SnapNonAlignedGridLevelsToLowerLevels == EWorldPartitionCVarProjectDefaultOverride::ProjectDefault) ? GRuntimeSpatialHashSnapNonAlignedGridLevelsToLowerLevels : (SnapNonAlignedGridLevelsToLowerLevels == EWorldPartitionCVarProjectDefaultOverride::Enabled);
+	Settings.bPlaceSmallActorsUsingLocation = (PlaceSmallActorsUsingLocation == EWorldPartitionCVarProjectDefaultOverride::ProjectDefault) ? GRuntimeSpatialHashPlaceSmallActorsUsingLocation : (PlaceSmallActorsUsingLocation == EWorldPartitionCVarProjectDefaultOverride::Enabled);
+	Settings.bPlacePartitionActorsUsingLocation = (PlacePartitionActorsUsingLocation == EWorldPartitionCVarProjectDefaultOverride::ProjectDefault) ? GRuntimeSpatialHashPlacePartitionActorsUsingLocation : (PlacePartitionActorsUsingLocation == EWorldPartitionCVarProjectDefaultOverride::Enabled);
+	UE_LOG(LogWorldPartition, Log, TEXT("UWorldPartitionRuntimeSpatialHash::UpdateSettings : UseAlignedGridLevels = %d, SnapNonAlignedGridLevelsToLowerLevels = %d, PlaceSmallActorsUsingLocation = %d, PlacePartitionActorsUsingLocation = %d"),
+		Settings.bUseAlignedGridLevels, Settings.bSnapNonAlignedGridLevelsToLowerLevels, Settings.bPlaceSmallActorsUsingLocation, Settings.bPlacePartitionActorsUsingLocation);
+}
+
 bool UWorldPartitionRuntimeSpatialHash::GenerateStreaming(UWorldPartitionStreamingPolicy* StreamingPolicy, const IStreamingGenerationContext* StreamingGenerationContext, TArray<FString>* OutPackagesToGenerate)
 {
 	verify(Super::GenerateStreaming(StreamingPolicy, StreamingGenerationContext, OutPackagesToGenerate));
@@ -1207,6 +1218,9 @@ bool UWorldPartitionRuntimeSpatialHash::GenerateStreaming(UWorldPartitionStreami
 	// Fix case where StreamingGrids might have been persisted.
 	bIsNameToGridMappingDirty = true;
 	StreamingGrids.Empty();
+	
+	// Apply Settings (should no longer change after streaming generation)
+	UpdateSettings();
 
 	// Append grids from ASpatialHashRuntimeGridInfo actors to runtime spatial hash grids
 	TArray<FSpatialHashRuntimeGrid> AllGrids;
@@ -1258,7 +1272,7 @@ bool UWorldPartitionRuntimeSpatialHash::GenerateStreaming(UWorldPartitionStreami
 	for (int32 GridIndex=0; GridIndex < AllGrids.Num(); GridIndex++)
 	{
 		const FSpatialHashRuntimeGrid& Grid = AllGrids[GridIndex];
-		const FSquare2DGridHelper PartionedActors = GetPartitionedActors(WorldBounds, Grid, GridActorSetInstances[GridIndex]);
+		const FSquare2DGridHelper PartionedActors = GetPartitionedActors(WorldBounds, Grid, GridActorSetInstances[GridIndex], Settings);
 		if (!CreateStreamingGrid(Grid, PartionedActors, StreamingPolicy, OutPackagesToGenerate))
 		{
 			return false;
@@ -1467,6 +1481,7 @@ bool UWorldPartitionRuntimeSpatialHash::CreateStreamingGrid(const FSpatialHashRu
 	const bool bIsMainWorldPartition = (World == OuterWorld);
 
 	FSpatialHashStreamingGrid& CurrentStreamingGrid = StreamingGrids.AddDefaulted_GetRef();
+	CurrentStreamingGrid.Settings = Settings;
 	CurrentStreamingGrid.GridName = RuntimeGrid.GridName;
 	CurrentStreamingGrid.CellSize = PartionedActors.CellSize;
 	CurrentStreamingGrid.WorldBounds = PartionedActors.WorldBounds;
@@ -1686,12 +1701,19 @@ bool UWorldPartitionRuntimeSpatialHash::InjectExternalStreamingObject(URuntimeHa
 	// Validate that there's a corresponding streaming grid for each streaming grid of this external streaming object
 	for (const FSpatialHashStreamingGrid& ExternalStreamingGrid : SpatialHashExternalStreamingObject->StreamingGrids)
 	{
-		if (!GetStreamingGridByName(ExternalStreamingGrid.GridName))
+		const FSpatialHashStreamingGrid* SourceGrid = GetStreamingGridByName(ExternalStreamingGrid.GridName);
+		if (!SourceGrid)
 		{
 			UE_LOG(LogWorldPartition, Error, TEXT("Failed to inject external streaming object %s, can't find matching streaming grid %s."),
 				*InExternalStreamingObject->GetName(), *ExternalStreamingGrid.GridName.ToString());
 			return false;
 		}
+		else if (SourceGrid->Settings != ExternalStreamingGrid.Settings)
+		{
+			UE_LOG(LogWorldPartition, Error, TEXT("Failed to inject external streaming object %s, miss matching settings."),
+				*InExternalStreamingObject->GetName());
+			return false;
+		}
 	}
 
 	for (const FSpatialHashStreamingGrid& ExternalStreamingGrid : SpatialHashExternalStreamingObject->StreamingGrids)
diff --git a/Engine/Source/Runtime/Engine/Private/WorldPartition/WorldPartitionRuntimeSpatialHashGridPreviewer.cpp b/Engine/Source/Runtime/Engine/Private/WorldPartition/WorldPartitionRuntimeSpatialHashGridPreviewer.cpp
--- a/Engine/Source/Runtime/Engine/Private/WorldPartition/WorldPartitionRuntimeSpatialHashGridPreviewer.cpp
+++ b/Engine/Source/Runtime/Engine/Private/WorldPartition/WorldPartitionRuntimeSpatialHashGridPreviewer.cpp
@@ -28,7 +28,7 @@ FWorldPartitionRuntimeSpatialHashGridPreviewer::FWorldPartitionRuntimeSpatialHas
 }
 
 #if WITH_EDITOR
-void FWorldPartitionRuntimeSpatialHashGridPreviewer::Draw(UWorld* World, const TArray<FSpatialHashRuntimeGrid>& Grids, bool bEnabled, int32 PreviewGridLevel)
+void FWorldPartitionRuntimeSpatialHashGridPreviewer::Draw(UWorld* World, const TArray<FSpatialHashRuntimeGrid>& Grids, bool bEnabled, int32 PreviewGridLevel, bool bUseAlignedGridLevels)
 {
 	if (bEnabled && Material)
 	{
@@ -83,7 +83,7 @@ void FWorldPartitionRuntimeSpatialHashGridPreviewer::Draw(UWorld* World, const T
 						CachedParameters.GridColor = Grid.DebugColor;
 					}
 
-					FVector GridOffset = FVector(Grid.Origin, 0) + (GRuntimeSpatialHashUseAlignedGridLevelsEffective ? FVector(0.5 * PreviewCellSize) : FVector::ZeroVector);
+					FVector GridOffset = FVector(Grid.Origin, 0) + (bUseAlignedGridLevels ? FVector(0.5 * PreviewCellSize) : FVector::ZeroVector);
 					if (CachedParameters.GridOffset != GridOffset)
 					{
 						MID->SetVectorParameterValue(*FString::Printf(TEXT("Grid%d_Offset"), i), GridOffset);
diff --git a/Engine/Source/Runtime/Engine/Public/WorldPartition/RuntimeSpatialHash/RuntimeSpatialHashGridHelper.h b/Engine/Source/Runtime/Engine/Public/WorldPartition/RuntimeSpatialHash/RuntimeSpatialHashGridHelper.h
--- a/Engine/Source/Runtime/Engine/Public/WorldPartition/RuntimeSpatialHash/RuntimeSpatialHashGridHelper.h
+++ b/Engine/Source/Runtime/Engine/Public/WorldPartition/RuntimeSpatialHash/RuntimeSpatialHashGridHelper.h
@@ -11,14 +11,14 @@
 #include "WorldPartition/DataLayer/DataLayersID.h"
 #include "WorldPartition/DataLayer/DataLayerInstance.h"
 
+#if WITH_EDITOR
+
 extern ENGINE_API bool GRuntimeSpatialHashUseAlignedGridLevels;
 extern ENGINE_API bool GRuntimeSpatialHashSnapNonAlignedGridLevelsToLowerLevels;
 extern ENGINE_API bool GRuntimeSpatialHashPlaceSmallActorsUsingLocation;
 extern ENGINE_API bool GRuntimeSpatialHashPlacePartitionActorsUsingLocation;
-extern ENGINE_API bool GRuntimeSpatialHashUseAlignedGridLevelsEffective;
-extern ENGINE_API bool GRuntimeSpatialHashSnapNonAlignedGridLevelsToLowerLevelsEffective;
-extern ENGINE_API bool GRuntimeSpatialHashPlaceSmallActorsUsingLocationEffective;
-extern ENGINE_API bool GRuntimeSpatialHashPlacePartitionActorsUsingLocationEffective;
+
+#endif
 
 /**
   * Square 2D grid helper
@@ -393,8 +393,11 @@ struct FSquare2DGridHelper
 		TMap<int64, int64> CellsMapping;
 	};
 
+	UE_DEPRECATED(5.4, "Use version with bUseAlignedGridLevels param")
 	ENGINE_API FSquare2DGridHelper(const FBox& InWorldBounds, const FVector& InOrigin, int64 InCellSize);
 
+	ENGINE_API FSquare2DGridHelper(const FBox& InWorldBounds, const FVector& InOrigin, int64 InCellSize, bool bUseAlignedGridLevels);
+
 	// Returns the lowest grid level
 	inline FGridLevel& GetLowestLevel() { return Levels[0]; }
 
@@ -475,6 +478,13 @@ public:
 };
 
 #if WITH_EDITOR
+
+FSquare2DGridHelper GetGridHelper(const FBox& WorldBounds, const FVector& GridOrigin, int64 GridCellSize, bool bUseAlignedGridLevels);
+FSquare2DGridHelper GetPartitionedActors(const FBox& WorldBounds, const FSpatialHashRuntimeGrid& Grid, const TArray<const IStreamingGenerationContext::FActorSetInstance*>& ActorSetInstances, const FSpatialHashSettings& Settings);
+
+UE_DEPRECATED(5.4, "Use version with bUseAlignedGridLevels param")
 FSquare2DGridHelper GetGridHelper(const FBox& WorldBounds, const FVector& GridOrigin, int64 GridCellSize);
+
+UE_DEPRECATED(5.4, "Use version with Settings param")
 FSquare2DGridHelper GetPartitionedActors(const FBox& WorldBounds, const FSpatialHashRuntimeGrid& Grid, const TArray<const IStreamingGenerationContext::FActorSetInstance*>& ActorSetInstances);
 #endif // #if WITH_EDITOR
diff --git a/Engine/Source/Runtime/Engine/Public/WorldPartition/WorldPartitionRuntimeSpatialHash.h b/Engine/Source/Runtime/Engine/Public/WorldPartition/WorldPartitionRuntimeSpatialHash.h
--- a/Engine/Source/Runtime/Engine/Public/WorldPartition/WorldPartitionRuntimeSpatialHash.h
+++ b/Engine/Source/Runtime/Engine/Public/WorldPartition/WorldPartitionRuntimeSpatialHash.h
@@ -27,6 +27,38 @@ enum class EWorldPartitionCVarProjectDefaultOverride : uint8
 	Enabled,
 };
 
+USTRUCT()
+struct FSpatialHashSettings
+{
+	GENERATED_USTRUCT_BODY()
+
+	FSpatialHashSettings()
+		: bUseAlignedGridLevels(true)
+		, bSnapNonAlignedGridLevelsToLowerLevels(true)
+#if WITH_EDITORONLY_DATA
+		, bPlaceSmallActorsUsingLocation(false)
+		, bPlacePartitionActorsUsingLocation(true)
+#endif
+	{}
+
+	friend bool operator==(const FSpatialHashSettings& Lhs, const FSpatialHashSettings& Rhs) = default;
+	friend bool operator!=(const FSpatialHashSettings& Lhs, const FSpatialHashSettings& Rhs) = default;
+	
+	UPROPERTY();
+	bool bUseAlignedGridLevels;
+
+	UPROPERTY();
+	bool bSnapNonAlignedGridLevelsToLowerLevels;
+
+#if WITH_EDITORONLY_DATA
+	UPROPERTY();
+	bool bPlaceSmallActorsUsingLocation;
+
+	UPROPERTY();
+	bool bPlacePartitionActorsUsingLocation;
+#endif
+};
+
 USTRUCT()
 struct FSpatialHashStreamingGridLayerCell
 {
@@ -136,7 +168,10 @@ struct FSpatialHashStreamingGrid
 
 	UPROPERTY()
 	int32 GridIndex;
-
+		
+	UPROPERTY()
+	FSpatialHashSettings Settings;
+			
 	ENGINE_API void InjectExternalStreamingObjectGrid(const FSpatialHashStreamingGrid& InExternalObjectStreamingGrid) const;
 	ENGINE_API void RemoveExternalStreamingObjectGrid(const FSpatialHashStreamingGrid& InExternalObjectStreamingGrid) const;
 
@@ -266,8 +301,10 @@ public:
 	virtual void PreSave(FObjectPreSaveContext ObjectSaveContext) override;
 
 	ENGINE_API virtual void Serialize(FArchive& Ar) override;
+	ENGINE_API virtual void PostLoad() override;
 
-	ENGINE_API void ApplyCVars();
+	UE_DEPRECATED(5.4, "ApplyCVars is deprecated")
+	ENGINE_API void ApplyCVars() {}
 
 #if WITH_EDITOR
 	ENGINE_API virtual void PostEditChangeProperty(FPropertyChangedEvent& PropertyChangedEvent) override;
@@ -335,6 +372,7 @@ private:
 	mutable FWorldPartitionRuntimeSpatialHashGridPreviewer GridPreviewer;
 #endif
 
+#if WITH_EDITORONLY_DATA
 	/** Disable to help break the pattern caused by world partition promotion of actors to upper grid levels that are always aligned on child levels. */
 	UPROPERTY(EditAnywhere, Config, AdvancedDisplay, Category = RuntimeSettings)
 	EWorldPartitionCVarProjectDefaultOverride UseAlignedGridLevels;
@@ -350,12 +388,16 @@ private:
 	/** Enable to place partitioned actors into their corresponding cell using their location instead of their bounding box. */
 	UPROPERTY(EditAnywhere, Config, AdvancedDisplay, Category = RuntimeSettings)
 	EWorldPartitionCVarProjectDefaultOverride PlacePartitionActorsUsingLocation;
-
+#endif
+		
 	/** Whether this hash enables Z culling. */
 	UPROPERTY(EditAnywhere, Config, Category = RuntimeSettings)
 	bool bEnableZCulling;
 
 protected:
+	UPROPERTY()
+	FSpatialHashSettings Settings;
+	
 	/** 
 	 * Represents the streaming grids (PIE or Game)
 	 */
@@ -374,6 +416,7 @@ private:
 	ENGINE_API void GetAlwaysLoadedStreamingCells(const FSpatialHashStreamingGrid& StreamingGrid, TSet<const UWorldPartitionRuntimeCell*>& Cells) const;
 	ENGINE_API const TMap<FName, const FSpatialHashStreamingGrid*>& GetNameToGridMapping() const;
 #if WITH_EDITOR
+	ENGINE_API void UpdateSettings();
 	ENGINE_API bool CreateStreamingGrid(const FSpatialHashRuntimeGrid& RuntimeGrid, const FSquare2DGridHelper& PartionedActors, UWorldPartitionStreamingPolicy* StreamingPolicy, TArray<FString>* OutPackagesToGenerate = nullptr);
 #endif
 	ENGINE_API TArray<const FSpatialHashStreamingGrid*> GetFilteredStreamingGrids() const;
diff --git a/Engine/Source/Runtime/Engine/Public/WorldPartition/WorldPartitionRuntimeSpatialHashGridPreviewer.h b/Engine/Source/Runtime/Engine/Public/WorldPartition/WorldPartitionRuntimeSpatialHashGridPreviewer.h
--- a/Engine/Source/Runtime/Engine/Public/WorldPartition/WorldPartitionRuntimeSpatialHashGridPreviewer.h
+++ b/Engine/Source/Runtime/Engine/Public/WorldPartition/WorldPartitionRuntimeSpatialHashGridPreviewer.h
@@ -17,7 +17,7 @@ struct FWorldPartitionRuntimeSpatialHashGridPreviewer
 	ENGINE_API FWorldPartitionRuntimeSpatialHashGridPreviewer();
 
 #if WITH_EDITOR
-	ENGINE_API void Draw(UWorld* World, const TArray<FSpatialHashRuntimeGrid>& Grids, bool bEnabled, int32 PreviewGridLevel);
+	ENGINE_API void Draw(UWorld* World, const TArray<FSpatialHashRuntimeGrid>& Grids, bool bEnabled, int32 PreviewGridLevel, bool bUseAlignedGridLevels = true);
 #endif
 
 #if WITH_EDITORONLY_DATA
