[Material Editor] Removed unused old path in material translation code now that DDC change has settled in.

diff --git a/Engine/Source/Runtime/Engine/Private/Materials/HLSLMaterialTranslator.cpp b/Engine/Source/Runtime/Engine/Private/Materials/HLSLMaterialTranslator.cpp
--- a/Engine/Source/Runtime/Engine/Private/Materials/HLSLMaterialTranslator.cpp
+++ b/Engine/Source/Runtime/Engine/Private/Materials/HLSLMaterialTranslator.cpp
@@ -202,9 +202,6 @@ static FAutoConsoleVariableRef CVarPreshaderGapInterval(
 	GPreshaderGapInterval,
 	TEXT("Insert an empty element in the preshader buffer every specified number of elements in the buffer.  Workaround for a shader compiler bug."));
 
-/* Controls whether to use the new GetMaterialShaderCode() and GetMaterialEnvironment() implementations. */
-static const bool GUseMaterialTranslationResultsGrouping = true;
-
 /* Controls whether DDC caching of material translation results should be forcefully disabled. */
 #define FORCE_DISABLE_MATERIAL_TRANSLATION_DDC false
 
@@ -1167,7 +1164,7 @@ bool FHLSLMaterialTranslator::Translate()
 #if FORCE_DISABLE_MATERIAL_TRANSLATION_DDC
 	const bool bDisableTranslationDDC = true;
 #else
-	const bool bDisableTranslationDDC = bNoMaterialTranslationDDC || Material->IsPreview() || !GUseMaterialTranslationResultsGrouping;
+	const bool bDisableTranslationDDC = bNoMaterialTranslationDDC || Material->IsPreview();
 #endif
 
 	// DDC query local data
@@ -2397,13 +2394,6 @@ void FHLSLMaterialTranslator::ValidateShadingModelsForFeatureLevel(const FMateri
 
 void FHLSLMaterialTranslator::GetMaterialEnvironment(EShaderPlatform InPlatform, FShaderCompilerEnvironment& OutEnvironment)
 {
-	// Use the old way of getting material environment if the setting is disabled.
-	if (!GUseMaterialTranslationResultsGrouping)
-	{
-		GetMaterialEnvironmentOld(InPlatform, OutEnvironment);
-		return;
-	}
-
 	check(InPlatform == GetShaderPlatform());
 
 	const bool bSubstrateEnabled = Substrate::IsSubstrateEnabled();
@@ -2718,518 +2708,35 @@ void FHLSLMaterialTranslator::GetMaterialEnvironment(EShaderPlatform InPlatform,
 	OutEnvironment.SetDefine(TEXT("TEXTURE_SAMPLE_DEBUG"), EnvironmentDefines->bTextureSampleDebug);
 }
 
-void FHLSLMaterialTranslator::GetMaterialEnvironmentOld(EShaderPlatform InPlatform, FShaderCompilerEnvironment& OutEnvironment)
+// Assign custom interpolators to slots, packing them as much as possible in unused slots.
+TBitArray<> FHLSLMaterialTranslator::GetVertexInterpolatorsOffsets(FString& VertexInterpolatorsOffsetsDefinitionCode) const
 {
-	const bool bSubstrateEnabled = Substrate::IsSubstrateEnabled();
-
-	if (bNeedsParticlePosition || Material->ShouldGenerateSphericalParticleNormals() || bUsesSphericalParticleOpacity)
-	{
-		OutEnvironment.SetDefine(TEXT("NEEDS_PARTICLE_POSITION"), 1);
-	}
-
-	if (bNeedsParticleVelocity || Material->IsUsedWithNiagaraMeshParticles())
-	{
-		OutEnvironment.SetDefine(TEXT("NEEDS_PARTICLE_VELOCITY"), 1);
-	}
-
-	if (DynamicParticleParameterMask)
-	{
-		OutEnvironment.SetDefine(TEXT("USE_DYNAMIC_PARAMETERS"), 1);
-		OutEnvironment.SetDefine(TEXT("DYNAMIC_PARAMETERS_MASK"), DynamicParticleParameterMask);
-	}
-
-	if (bNeedsParticleTime)
-	{
-		OutEnvironment.SetDefine(TEXT("NEEDS_PARTICLE_TIME"), 1);
-	}
-
-	if (bUsesParticleMotionBlur)
-	{
-		OutEnvironment.SetDefine(TEXT("USES_PARTICLE_MOTION_BLUR"), 1);
-	}
-
-	if (bNeedsParticleRandom)
-	{
-		OutEnvironment.SetDefine(TEXT("NEEDS_PARTICLE_RANDOM"), 1);
-	}
-
-	if (bUsesSphericalParticleOpacity)
-	{
-		OutEnvironment.SetDefine(TEXT("SPHERICAL_PARTICLE_OPACITY"), TEXT("1"));
-	}
-
-	if (bUsesParticleSubUVs)
-	{
-		OutEnvironment.SetDefine(TEXT("USE_PARTICLE_SUBUVS"), TEXT("1"));
-	}
-
-	if (bUsesLightmapUVs)
-	{
-		OutEnvironment.SetDefine(TEXT("LIGHTMAP_UV_ACCESS"),TEXT("1"));
-	}
-
-	if (bUsesAOMaterialMask)
-	{
-		OutEnvironment.SetDefine(TEXT("USES_AO_MATERIAL_MASK"),TEXT("1"));
-	}
-
-	if (bUsesSpeedTree)
-	{
-		OutEnvironment.SetDefine(TEXT("USES_SPEEDTREE"),TEXT("1"));
-	}
-
-	if (bNeedsWorldPositionExcludingShaderOffsets)
-	{
-		OutEnvironment.SetDefine(TEXT("NEEDS_WORLD_POSITION_EXCLUDING_SHADER_OFFSETS"), TEXT("1"));
-	}
-
-	if (bNeedsParticleSize)
-	{
-		OutEnvironment.SetDefine(TEXT("NEEDS_PARTICLE_SIZE"), TEXT("1"));
-	}
-
-	if (bNeedsParticleSpriteRotation)
-	{
-		OutEnvironment.SetDefine(TEXT("NEEDS_PARTICLE_SPRITE_ROTATION"), TEXT("1"));
-	}
-
-	if (MaterialCompilationOutput.bNeedsSceneTextures)
-	{
-		OutEnvironment.SetDefine(TEXT("NEEDS_SCENE_TEXTURES"), TEXT("1"));
-	}
-	if (MaterialCompilationOutput.bUsesEyeAdaptation)
-	{
-		OutEnvironment.SetDefine(TEXT("USES_EYE_ADAPTATION"), TEXT("1"));
-	}
-
-	if (MaterialCompilationOutput.bHasRuntimeVirtualTextureOutputNode)
-	{
-		OutEnvironment.SetDefine(TEXT("VIRTUAL_TEXTURE_OUTPUT"), 1);
-	}
-
-	OutEnvironment.SetDefine(TEXT("USES_PER_INSTANCE_CUSTOM_DATA"), MaterialCompilationOutput.bUsesPerInstanceCustomData && Material->IsUsedWithInstancedStaticMeshes());
-	// Whether we need to pass custom data paramaters to PS through interpolators
-	OutEnvironment.SetDefine(TEXT("USES_PER_INSTANCE_CUSTOM_DATA_PS"), bUsesPerInstanceCustomDataPS && Material->IsUsedWithInstancedStaticMeshes());
-
-	OutEnvironment.SetDefine(TEXT("USES_PER_INSTANCE_FADE_AMOUNT"), bUsesPerInstanceFadeAmount && Material->IsUsedWithInstancedStaticMeshes());
-	OutEnvironment.SetDefine(TEXT("USES_VERTEX_INTERPOLATOR"), MaterialCompilationOutput.bUsesVertexInterpolator);
-
-	OutEnvironment.SetDefine(TEXT("MATERIAL_SKY_ATMOSPHERE"), bUsesSkyAtmosphere);
-	OutEnvironment.SetDefine(TEXT("INTERPOLATE_VERTEX_COLOR"), bUsesVertexColor);
-	OutEnvironment.SetDefine(TEXT("NEEDS_PARTICLE_COLOR"), bUsesParticleColor); 
-	OutEnvironment.SetDefine(TEXT("NEEDS_PARTICLE_LOCAL_TO_WORLD"), bUsesParticleLocalToWorld);
-	OutEnvironment.SetDefine(TEXT("NEEDS_PARTICLE_WORLD_TO_LOCAL"), bUsesParticleWorldToLocal);
-	OutEnvironment.SetDefine(TEXT("NEEDS_INSTANCE_LOCAL_TO_WORLD_PS"), bUsesInstanceLocalToWorldPS);
-	OutEnvironment.SetDefine(TEXT("NEEDS_INSTANCE_WORLD_TO_LOCAL_PS"), bUsesInstanceWorldToLocalPS);
-	OutEnvironment.SetDefine(TEXT("NEEDS_PER_INSTANCE_RANDOM_PS"), bUsesPerInstanceRandomPS && Material->IsUsedWithInstancedStaticMeshes());
-	OutEnvironment.SetDefine(TEXT("USES_TRANSFORM_VECTOR"), bUsesTransformVector);
-	OutEnvironment.SetDefine(TEXT("WANT_PIXEL_DEPTH_OFFSET"), bUsesPixelDepthOffset);
-	
-	// we want USES_WORLD_POSITION_OFFSET to be readable as a bool compile argument, hence the != 0 comparison 
-	// (bUsesWorldPositionOffset is actually a 1-bit uint32 bitfield member)
-	OutEnvironment.SetDefineAndCompileArgument(TEXT("USES_WORLD_POSITION_OFFSET"), (bUsesWorldPositionOffset != 0));
-	OutEnvironment.SetDefineAndCompileArgument(TEXT("USES_DISPLACEMENT"), (bUsesDisplacement != 0));
-
-	OutEnvironment.SetDefine(TEXT("USES_EMISSIVE_COLOR"), bUsesEmissiveColor);
-	// Distortion uses tangent space transform 
-	OutEnvironment.SetDefine(TEXT("USES_DISTORTION"), Material->IsDistorted());
-	OutEnvironment.SetDefine(TEXT("DISTORTION_ACCOUNT_FOR_COVERAGE"), bSubstrateEnabled && Material->GetRefractionCoverageMode() == RCM_CoverageAccountedFor ? 1 : 0);
-
-	OutEnvironment.SetDefine(TEXT("MATERIAL_ENABLE_TRANSLUCENCY_FOGGING"), Material->ShouldApplyFogging());
-	OutEnvironment.SetDefine(TEXT("MATERIAL_ENABLE_TRANSLUCENCY_CLOUD_FOGGING"), Material->ShouldApplyCloudFogging());
-	OutEnvironment.SetDefine(TEXT("MATERIAL_IS_SKY"), Material->IsSky());
-	OutEnvironment.SetDefine(TEXT("MATERIAL_COMPUTE_FOG_PER_PIXEL"), Material->ComputeFogPerPixel());
-	OutEnvironment.SetDefine(TEXT("MATERIAL_FULLY_ROUGH"), bIsFullyRough || Material->IsFullyRough());
-	OutEnvironment.SetDefine(TEXT("MATERIAL_USES_ANISOTROPY"), bUsesAnisotropy && FDataDrivenShaderPlatformInfo::GetSupportsAnisotropicMaterials(InPlatform));
-
-	OutEnvironment.SetDefine(TEXT("MATERIAL_DECAL_READ_MASK"), MaterialCompilationOutput.UsedDBufferTextures);
-	OutEnvironment.SetDefine(TEXT("MATERIAL_USES_DECAL_LOOKUP"), MaterialCompilationOutput.bUsesDBufferTextureLookup);
-	OutEnvironment.SetDefine(TEXT("MATERIAL_PATH_TRACING_BUFFER_READ"), MaterialCompilationOutput.UsedPathTracingBufferTextures);
-	OutEnvironment.SetDefine(TEXT("MATERIAL_NEURAL_POST_PROCESS"), (MaterialCompilationOutput.bUsedWithNeuralNetworks || Material->IsUsedWithNeuralNetworks()) && Material->IsPostProcessMaterial());
-
-	// Count the number of VTStacks (each stack will allocate a feedback slot)
-	OutEnvironment.SetDefine(TEXT("NUM_VIRTUALTEXTURE_SAMPLES"), VTStacks.Num());
-
-	// Check if any feedback slots are in use. We can simplify shader and remove EARLYZ optimizations if none are.
-	bool bGenerateFeedback = false;
-	for (int i = 0; i < VTStacks.Num() && !bGenerateFeedback; ++i)
-	{
-		bGenerateFeedback |= VTStacks[i].bGenerateFeedback;
-	}
-	OutEnvironment.SetDefine(TEXT("MATERIAL_VIRTUALTEXTURE_FEEDBACK"), bGenerateFeedback);
-
-	// Setup defines to map each VT stack to either 1 or 2 page table textures, depending on how many layers it uses
-	for (int i = 0; i < VTStacks.Num(); ++i)
-	{
-		const FMaterialVirtualTextureStack& Stack = MaterialCompilationOutput.UniformExpressionSet.VTStacks[i];
-		FString PageTableValue = FString::Printf(TEXT("Material.VirtualTexturePageTable0_%d"), i);
-		if (Stack.GetNumLayers() > 4u)
-		{
-			PageTableValue += FString::Printf(TEXT(", Material.VirtualTexturePageTable1_%d"), i);
-		}
-		if (VTStacks[i].bAdaptive)
-		{
-			PageTableValue += FString::Printf(TEXT(", Material.VirtualTexturePageTableIndirection_%d"), i);
-		}
-		OutEnvironment.SetDefine(*FString::Printf(TEXT("VIRTUALTEXTURE_PAGETABLE_%d"), i), *PageTableValue);
-	}
-
-	for (int32 CollectionIndex = 0; CollectionIndex < ParameterCollections.Num(); CollectionIndex++)
-	{
-		// Add uniform buffer declarations for any parameter collections referenced
-		const FString CollectionName = FString::Printf(TEXT("MaterialCollection%u"), CollectionIndex);
-		// This can potentially become an issue for MaterialCollection Uniform Buffers if they ever get non-numeric resources (eg Textures), as
-		// OutEnvironment.ResourceTableMap has a map by name, and the N ParameterCollection Uniform Buffers ALL are names "MaterialCollection"
-		// (and the hlsl cbuffers are named MaterialCollection0, etc, so the names don't match the layout)
-		FShaderUniformBufferParameter::ModifyCompilationEnvironment(*CollectionName, ParameterCollections[CollectionIndex]->GetUniformBufferStruct(), InPlatform, OutEnvironment);
-	}
-	OutEnvironment.SetDefine(TEXT("IS_MATERIAL_SHADER"), TEXT("1"));
-
-	// Set all the shading models for this material here 
-	FMaterialShadingModelField ShadingModels = Material->GetShadingModels();
-
-	// If the material gets its shading model from the material expressions, then we use the result from the compilation (assuming it's valid).
-	// This result will potentially be tighter than what GetShadingModels() returns, because it only picks up the shading models from the expressions that get compiled for a specific feature level and quality level
-	// For example, the material might have shading models behind static switches. GetShadingModels() will return both the true and the false paths from that switch, whereas the shading model field from the compilation will only contain the actual shading model selected 
-	if (Material->IsShadingModelFromMaterialExpression() && ShadingModelsFromCompilation.IsValid())
-	{
-		// Shading models fetched from the compilation of the expression graph
-		ShadingModels = ShadingModelsFromCompilation;
-	}
-
-	ensure(ShadingModels.IsValid());
-
-	if (ShadingModels.IsLit())
-	{	
-		int NumSetMaterials = 0;
-		if (ShadingModels.HasShadingModel(MSM_DefaultLit))
-		{
-			OutEnvironment.SetDefine(TEXT("MATERIAL_SHADINGMODEL_DEFAULT_LIT"), TEXT("1"));
-			NumSetMaterials++;
-		}
-		if (ShadingModels.HasShadingModel(MSM_Subsurface))
-		{
-			OutEnvironment.SetDefine(TEXT("MATERIAL_SHADINGMODEL_SUBSURFACE"), TEXT("1"));
-			NumSetMaterials++;
-		}
-		if (ShadingModels.HasShadingModel(MSM_PreintegratedSkin))
-		{
-			OutEnvironment.SetDefine(TEXT("MATERIAL_SHADINGMODEL_PREINTEGRATED_SKIN"), TEXT("1"));
-			NumSetMaterials++;
-		}
-		if (ShadingModels.HasShadingModel(MSM_SubsurfaceProfile))
-		{
-			OutEnvironment.SetDefine(TEXT("MATERIAL_SHADINGMODEL_SUBSURFACE_PROFILE"), TEXT("1"));
-			NumSetMaterials++;
-
-			if (bUsesCurvature)
-			{
-				OutEnvironment.SetDefine(TEXT("MATERIAL_SUBSURFACE_PROFILE_USE_CURVATURE"), TEXT("1"));
-			}
-		}
-		if (ShadingModels.HasShadingModel(MSM_ClearCoat))
-		{
-			OutEnvironment.SetDefine(TEXT("MATERIAL_SHADINGMODEL_CLEAR_COAT"), TEXT("1"));
-			NumSetMaterials++;
-		}
-		if (ShadingModels.HasShadingModel(MSM_TwoSidedFoliage))
-		{
-			OutEnvironment.SetDefine(TEXT("MATERIAL_SHADINGMODEL_TWOSIDED_FOLIAGE"), TEXT("1"));
-			NumSetMaterials++;
-		}
-		if (ShadingModels.HasShadingModel(MSM_Hair))
-		{
-			OutEnvironment.SetDefine(TEXT("MATERIAL_SHADINGMODEL_HAIR"), TEXT("1"));
-			NumSetMaterials++;
-		}
-		if (ShadingModels.HasShadingModel(MSM_Cloth))
-		{
-			OutEnvironment.SetDefine(TEXT("MATERIAL_SHADINGMODEL_CLOTH"), TEXT("1"));
-			NumSetMaterials++;
-		}
-		if (ShadingModels.HasShadingModel(MSM_Eye))
-		{
-			OutEnvironment.SetDefine(TEXT("MATERIAL_SHADINGMODEL_EYE"), TEXT("1"));
-			NumSetMaterials++;
-
-			if (bUsesCurvature)
-			{
-				OutEnvironment.SetDefine(TEXT("MATERIAL_SHADINGMODEL_EYE_USE_CURVATURE"), TEXT("1"));
-			}
-		}
-		if (ShadingModels.HasShadingModel(MSM_SingleLayerWater))
-		{
-			OutEnvironment.SetDefine(TEXT("MATERIAL_SHADINGMODEL_SINGLELAYERWATER"), TEXT("1"));
-			NumSetMaterials++;
-		}
-		if (ShadingModels.HasShadingModel(MSM_ThinTranslucent))
-		{
-			OutEnvironment.SetDefine(TEXT("MATERIAL_SHADINGMODEL_THIN_TRANSLUCENT"), TEXT("1"));
-			NumSetMaterials++;
-		}
-
-		if (ShadingModels.HasShadingModel(MSM_SingleLayerWater) && FDataDrivenShaderPlatformInfo::GetRequiresDisableForwardLocalLights(Platform))
-		{
-			OutEnvironment.SetDefine(TEXT("DISABLE_FORWARD_LOCAL_LIGHTS"), TEXT("1"));
-		}
-
-		// This is to have platforms use the simple single layer water shading similar to mobile: no dynamic lights, only sun and sky, no distortion, no colored transmittance on background, no custom depth read.
-		const bool bSingleLayerWaterUsesSimpleShading = FDataDrivenShaderPlatformInfo::GetWaterUsesSimpleForwardShading(InPlatform) && IsForwardShadingEnabled(InPlatform);
-		if (ShadingModels.HasShadingModel(MSM_SingleLayerWater) && bSingleLayerWaterUsesSimpleShading)
-		{
-			// Value must match SINGLE_LAYER_WATER_SHADING_QUALITY_MOBILE_WITH_DEPTH_TEXTURE in SingleLayerWaterCommon.ush!
-			OutEnvironment.SetDefine(TEXT("SINGLE_LAYER_WATER_SHADING_QUALITY"), TEXT("1"));
-		}
-
-		const bool bIsWaterDistanceFieldShadowEnabled = IsWaterDistanceFieldShadowEnabled(InPlatform);
-		const bool bIsWaterVSMFilteringEnabled = IsWaterVirtualShadowMapFilteringEnabled(InPlatform);
-		if (ShadingModels.HasShadingModel(MSM_SingleLayerWater) && (bIsWaterDistanceFieldShadowEnabled || bIsWaterVSMFilteringEnabled))
-		{
-			OutEnvironment.SetDefine(TEXT("SINGLE_LAYER_WATER_SEPARATED_MAIN_LIGHT"), TEXT("1"));
-		}
-
-		if (NumSetMaterials == 1)
-		{
-			OutEnvironment.SetDefine(TEXT("MATERIAL_SINGLE_SHADINGMODEL"), TEXT("1"));
-		}
+	TBitArray<> AllocatedCoords = AllocatedUserTexCoords; // Don't mess with the already assigned sets of UV coords
 
-		ensure(NumSetMaterials != 0);
-		if (NumSetMaterials == 0)
-		{
-			// Should not really end up here
-			UE_LOG(LogMaterial, Warning, TEXT("Unknown material shading model(s). Setting to MSM_DefaultLit"));
-			OutEnvironment.SetDefine(TEXT("MATERIAL_SHADINGMODEL_DEFAULT_LIT"),TEXT("1"));
-		}
-	}
-	else
-	{
-		// Unlit shading model can only exist by itself
-		OutEnvironment.SetDefine(TEXT("MATERIAL_SINGLE_SHADINGMODEL"), TEXT("1"));
-		OutEnvironment.SetDefine(TEXT("MATERIAL_SHADINGMODEL_UNLIT"), TEXT("1"));
-	}
+	int32 CurrentSlot = INDEX_NONE;
+	int32 EndAllocatedSlot = INDEX_NONE;
 
-	if (Material->GetMaterialDomain() == MD_Volume)
+	auto GetNextUVSlot = [&CurrentSlot, &EndAllocatedSlot, &AllocatedCoords]() -> int32
 	{
-		TArray<const UMaterialExpressionVolumetricAdvancedMaterialOutput*> VolumetricAdvancedExpressions;
-		Material->GetMaterialInterface()->GetMaterial()->GetAllExpressionsOfType(VolumetricAdvancedExpressions);
-		if (VolumetricAdvancedExpressions.Num() > 0)
+		if (CurrentSlot == EndAllocatedSlot)
 		{
-			if (VolumetricAdvancedExpressions.Num() > 1)
-			{
-				UE_LOG(LogMaterial, Fatal, TEXT("Only a single UMaterialExpressionVolumetricAdvancedMaterialOutput node is supported."));
-			}
-
-			OutEnvironment.SetDefine(TEXT("MATERIAL_VOLUMETRIC_ADVANCED"), TEXT("1"));
-
-			const UMaterialExpressionVolumetricAdvancedMaterialOutput* VolumetricAdvancedNode = VolumetricAdvancedExpressions[0];
-			if (VolumetricAdvancedNode->GetEvaluatePhaseOncePerSample())
-			{
-				OutEnvironment.SetDefine(TEXT("MATERIAL_VOLUMETRIC_ADVANCED_PHASE_PERSAMPLE"), TEXT("1"));
-			}
-			else
+			CurrentSlot = AllocatedCoords.FindAndSetFirstZeroBit();
+			if (CurrentSlot == INDEX_NONE)
 			{
-				OutEnvironment.SetDefine(TEXT("MATERIAL_VOLUMETRIC_ADVANCED_PHASE_PERPIXEL"), TEXT("1"));
+				CurrentSlot = AllocatedCoords.Add(true);
 			}
 
-			OutEnvironment.SetDefine(TEXT("MATERIAL_VOLUMETRIC_ADVANCED_GRAYSCALE_MATERIAL"), VolumetricAdvancedNode->bGrayScaleMaterial ? TEXT("1") : TEXT("0"));
-			OutEnvironment.SetDefine(TEXT("MATERIAL_VOLUMETRIC_ADVANCED_RAYMARCH_VOLUME_SHADOW"), VolumetricAdvancedNode->bRayMarchVolumeShadow ? TEXT("1") : TEXT("0"));
-			OutEnvironment.SetDefine(TEXT("MATERIAL_VOLUMETRIC_ADVANCED_CLAMP_MULTISCATTERING_CONTRIBUTION"), VolumetricAdvancedNode->bClampMultiScatteringContribution ? TEXT("1") : TEXT("0"));
-
-			OutEnvironment.SetDefine(TEXT("MATERIAL_VOLUMETRIC_ADVANCED_MULTISCATTERING_OCTAVE_COUNT"), VolumetricAdvancedNode->GetMultiScatteringApproximationOctaveCount());
-
-			OutEnvironment.SetDefine(TEXT("MATERIAL_VOLUMETRIC_ADVANCED_CONSERVATIVE_DENSITY"),
-				VolumetricAdvancedNode->ConservativeDensity.IsConnected() ? TEXT("1") : TEXT("0"));
-
-			OutEnvironment.SetDefine(TEXT("MATERIAL_VOLUMETRIC_ADVANCED_OVERRIDE_AMBIENT_OCCLUSION"),
-				(Material->HasAmbientOcclusionConnected() || bSubstrateWritesAmbientOcclusion) ? TEXT("1") : TEXT("0"));
-
-			OutEnvironment.SetDefine(TEXT("MATERIAL_VOLUMETRIC_ADVANCED_GROUND_CONTRIBUTION"),
-				VolumetricAdvancedNode->bGroundContribution ? TEXT("1") : TEXT("0"));
-		}
-
-		TArray<const UMaterialExpressionVolumetricCloudEmptySpaceSkippingOutput*> EmptySpaceSkippingOutputExpressions;
-		Material->GetMaterialInterface()->GetMaterial()->GetAllExpressionsOfType(EmptySpaceSkippingOutputExpressions);
-		if (EmptySpaceSkippingOutputExpressions.Num() > 0)
-		{
-			if (VolumetricAdvancedExpressions.Num() > 1)
-			{
-				UE_LOG(LogMaterial, Fatal, TEXT("Only a single UMaterialExpressionEmptySpaceSkippingOutput node is supported."));
-			}
-			OutEnvironment.SetDefine(TEXT("MATERIAL_VOLUMETRIC_CLOUD_EMPTY_SPACE_SKIPPING_OUTPUT"), TEXT("1"));
+			// Track one slot per component (u,v)
+			const int32 NUM_COMPONENTS = 2;
+			CurrentSlot *= NUM_COMPONENTS;
+			EndAllocatedSlot = CurrentSlot + NUM_COMPONENTS;
 		}
-	}
-
-	OutEnvironment.SetDefine(TEXT("MATERIAL_IS_SUBSTRATE"), bMaterialIsSubstrate ? TEXT("1") : TEXT("0"));
-
-	// if duals source blending (colored transmittance) is not supported on a platform, it will fall back to standard alpha blending (grey scale transmittance)
-	OutEnvironment.SetDefine(TEXT("DUAL_SOURCE_COLOR_BLENDING_ENABLED"), Material->IsDualBlendingEnabled(Platform) ? TEXT("1") : TEXT("0"));
-
-	OutEnvironment.SetDefine(TEXT("SUBSTRATE_PREMULTIPLIED_ALPHA_OPACITY_OVERRIDEN"), bMaterialIsSubstrate&& bOpacityPropertyIsUsed ? TEXT("1") : TEXT("0"));
-
-	if (bMaterialIsSubstrate)
-	{
-		OutEnvironment.SetDefine(TEXT("SUBSTRATE_USES_CONVERSION_FROM_LEGACY"), bSubstrateUsesConversionFromLegacy ? 1 : 0);
-		OutEnvironment.SetDefine(TEXT("SUBSTRATE_MATERIAL_OUTPUT_OPAQUE_ROUGH_REFRACTIONS"), bSubstrateOutputsOpaqueRoughRefractions ? 1 : 0);
 
-		OutEnvironment.SetDefine(TEXT("SUBSTRATE_MATERIAL_EXPORT_TYPE"), (int32)GetSubstrateMaterialExportType());
-		OutEnvironment.SetDefine(TEXT("SUBSTRATE_MATERIAL_EXPORT_CONTEXT"), (int32)GetSubstrateMaterialExportContext());
-		OutEnvironment.SetDefine(TEXT("SUBSTRATE_MATERIAL_EXPORT_LEGACY_BLEND_MODE"), (int32)GetSubstrateMaterialExportLegacyBlendMode());
+		int32 ResultUVSlot = CurrentSlot / 2;
+		CurrentSlot++;
 
-		// Unlit cannot be combined with other BSDF so we can simply pick the default Substrate context
-		OutEnvironment.SetDefine(TEXT("SUBSTRATE_OPTIMIZED_UNLIT"), SubstrateCompilationContext[ESubstrateCompilationContext::SCC_Default].bSubstrateMaterialIsUnlitNode ? 1 : 0);
-
-		{
-			// For now, the fully simplified mode is used for Lumen or anything else supported inlined evaluation. The export is only valid for the default case.
-			// SUBSTRATE_TODO: generate an export for the different context (need to generate two export functions: the default one and the FullSimplification one)
-			FSubstrateCompilationContext& SubstrateCtx = SubstrateCompilationContext[ESubstrateCompilationContext::SCC_Default];
-			OutEnvironment.SetDefine(TEXT("SUBSTRATE_SINGLEPATH"), SubstrateCtx.SubstrateMaterialComplexity.IsSingle() ? TEXT("1") : TEXT("0"));
-			OutEnvironment.SetDefine(TEXT("SUBSTRATE_FASTPATH"), SubstrateCtx.SubstrateMaterialComplexity.IsSimple() ? TEXT("1") : TEXT("0"));
-			OutEnvironment.SetDefine(TEXT("SUBSTRATE_COMPLEXSPECIALPATH"), SubstrateCtx.SubstrateMaterialComplexity.IsComplexSpecial() ? TEXT("1") : TEXT("0"));
-			OutEnvironment.SetDefine(TEXT("SUBSTRATE_CLAMPED_CLOSURE_COUNT"), SubstrateCtx.SubstrateMaterialEffectiveClosureCount);
-		}
-
-
-		FString SubstrateMaterialDescription;
-		for (uint32 SubstrateCompilationContextIndex = 0; SubstrateCompilationContextIndex < ESubstrateCompilationContext::SCC_MAX; ++SubstrateCompilationContextIndex)
-		{
-			FSubstrateCompilationContext& SubstrateCtx = SubstrateCompilationContext[SubstrateCompilationContextIndex];
-			const FSubstrateSimplificationStatus& SubstrateSimplificationStatus = SubstrateCtx.SubstrateSimplificationStatus;
-
-			check(SubstrateCtx.SubstrateMaterialRootOperator);
-			uint32 RootMaximumDistanceToLeaves = SubstrateCtx.SubstrateMaterialRootOperator->MaxDistanceFromLeaves;
-
-			// Compute the shared local basis count and generate the hlsl shader code for it.
-			SubstrateCtx.SubstratePixelNormalInitializerValues = FString::Printf(TEXT("\n\n\n\t// Substrate normal and tangent\n"));
-			SubstrateCtx.FinalUsedSharedLocalBasesCount = 0;
-			uint8 RequestedSharedLocalBasesCount = 0;
-			SubstrateCtx.SubstrateEvaluateSharedLocalBases(this, RequestedSharedLocalBasesCount, &OutEnvironment);
-
-#if WITH_EDITOR
-			// Now write some feedback to the user, but only produce debug string if in editor
-			{
-				// Output some debug info as comment in code and in the material stat window
-				const uint32 SubstrateBytePerPixel_Platform = Substrate::GetBytePerPixel(InPlatform);
-				const uint32 SubstrateClosurePerPixel_Platform = Substrate::GetClosurePerPixel(InPlatform);
-				FString SubstrateMaterialContextDescription;
-
-				auto GetSubstrateCompilationContextName = [](uint32 Index)
-				{
-					switch (Index)
-					{
-					case ESubstrateCompilationContext::SCC_Default:
-						return TEXT("Default");
-					case ESubstrateCompilationContext::SCC_FullySimplified:
-						return TEXT("FullySimplified");
-					default:
-						check(false);
-					}
-					return TEXT("ERROR");
-				};
-				FString SubstrateCompilationContextName = GetSubstrateCompilationContextName(SubstrateCompilationContextIndex);
-				SubstrateMaterialContextDescription += FString::Printf(TEXT("----- SUBSTRATE - %s -----\n"), *SubstrateCompilationContextName);
-				SubstrateMaterialContextDescription += FString::Printf(TEXT("SubstrateCompilationInfo -\n"));
-				SubstrateMaterialContextDescription += FString::Printf(TEXT(" - Byte Per Pixel Budget                           %u\n"), SubstrateBytePerPixel_Platform);
-				SubstrateMaterialContextDescription += FString::Printf(TEXT(" - Closure Per Pixel Budget                        %u\n"), SubstrateClosurePerPixel_Platform);
-				SubstrateMaterialContextDescription += FString::Printf(TEXT(" - Requested Byte Size before simplification       %u (%d UINT32)\n"), SubstrateSimplificationStatus.OriginalRequestedByteSize, SubstrateSimplificationStatus.OriginalRequestedByteSize / 4);
-				SubstrateMaterialContextDescription += FString::Printf(TEXT(" - Requested Byte Size after simplification        %u (%d UINT32)\n"), SubstrateCtx.SubstrateMaterialRequestedSizeByte, SubstrateCtx.SubstrateMaterialRequestedSizeByte / 4);
-				SubstrateMaterialContextDescription += FString::Printf(TEXT(" - Requested Closure Count before simplification   %u\n"), SubstrateSimplificationStatus.OriginalRequestedClosureCount);
-				SubstrateMaterialContextDescription += FString::Printf(TEXT(" - Requested Closure Count after simplification    %u\n"), SubstrateCtx.SubstrateMaterialClosureCount);
-				SubstrateMaterialContextDescription += FString::Printf(TEXT(" - Material complexity                             %s\n"), *FSubstrateMaterialComplexity::ToString(SubstrateCtx.SubstrateMaterialComplexity.SubstrateMaterialType(), true /* Upper case */));
-				SubstrateMaterialContextDescription += FString::Printf(TEXT(" - BSDF Count                                      %i\n"), SubstrateCtx.SubstrateMaterialEffectiveClosureCount); // REMOVE?
-				if (RequestedSharedLocalBasesCount > SUBSTRATE_MAX_SHAREDLOCALBASES_REGISTERS)
-				{
-					SubstrateMaterialDescription += FString::Printf(TEXT(" - SharedLocalBasesCount                      %i (Requested:%i)\n"), SubstrateCtx.FinalUsedSharedLocalBasesCount, RequestedSharedLocalBasesCount);
-				}
-				else
-				{
-					SubstrateMaterialDescription += FString::Printf(TEXT(" - SharedLocalBasesCount                      %i\n"), SubstrateCtx.FinalUsedSharedLocalBasesCount);
-				}
-
-
-				for (int32 OpIt = 0; OpIt < SubstrateCtx.SubstrateMaterialExpressionRegisteredOperators.Num(); ++OpIt)
-				{
-					FSubstrateOperator& BSDFOperator = SubstrateCtx.SubstrateMaterialExpressionRegisteredOperators[OpIt];
-					if (BSDFOperator.BSDFIndex == INDEX_NONE || BSDFOperator.IsDiscarded())
-					{
-						continue;	// not a BSDF or if discarded (i.e. not the root of a parameter blending subtree), then there is no local basis to register
-					}
-					if (BSDFOperator.BSDFRegisteredSharedLocalBasis.NormalCodeChunk == INDEX_NONE && BSDFOperator.BSDFRegisteredSharedLocalBasis.TangentCodeChunk == INDEX_NONE)
-					{
-						continue;	// We skip null normal on certain BSDF, for instance unlit.
-					}
-					const FSubstrateSharedLocalBasesInfo& SubstrateSharedLocalBasesInfo = SubstrateCtx.SubstrateCompilationInfoGetMatchingSharedLocalBasisInfo(BSDFOperator.BSDFRegisteredSharedLocalBasis);
-
-					SubstrateMaterialContextDescription += FString::Printf(TEXT("     - %s - SharedLocalBasisIndexMacro = %s \n"), *GetSubstrateBSDFName(BSDFOperator.BSDFType), *GetSubstrateSharedLocalBasisIndexMacroInner(BSDFOperator.BSDFRegisteredSharedLocalBasis, SubstrateCompilationContextIndex));
-				}
-
-				SubstrateMaterialContextDescription += FString::Printf(TEXT("----------- SUBSTRATE TREE - %s -----------\n"), *SubstrateCompilationContextName);
-				SubstrateMaterialContextDescription += FString::Printf(TEXT("Graph maximum distance to leaves %u\n"), RootMaximumDistanceToLeaves);
-				// Debug print operators according to depth from root.
-				{
-
-					for (int32 DistanceToLeaves = RootMaximumDistanceToLeaves; DistanceToLeaves >= 0; --DistanceToLeaves)
-					{
-						SubstrateMaterialContextDescription += FString::Printf(TEXT("----- DistanceFromLeaves = %d -----\n"), DistanceToLeaves);
-						for (auto& It : SubstrateCtx.SubstrateMaterialExpressionRegisteredOperators)
-						{
-							if (!It.IsDiscarded() && It.MaxDistanceFromLeaves == DistanceToLeaves)
-							{
-								SubstrateMaterialContextDescription += FString::Printf(TEXT("\tIdx=%d Op=%s ParentIdx=%d LeftIndex=%d RightIndex=%d BSDFIdx=%d LayerDepth=%d IsTop=%d IsBot=%d BSDFType=%s SSS=%d MFP=%d F90=%d Rough2=%d Fuzz=%d Aniso=%d Glint=%d SpecularProfile=%d\n"),
-									It.Index, GetSubstrateOperatorStr(It.OperatorType), It.ParentIndex, It.LeftIndex, It.RightIndex, It.BSDFIndex, It.LayerDepth, It.bIsTop, It.bIsBottom,
-									*GetSubstrateBSDFName(It.BSDFType), It.bBSDFHasSSS, It.bBSDFHasMFPPluggedIn, It.bBSDFHasEdgeColor, It.bBSDFHasSecondRoughnessOrSimpleClearCoat, It.bBSDFHasFuzz, It.bBSDFHasAnisotropy, It.bBSDFHasGlint, It.bBSDFHasSpecularProfile);
-							}
-						}
-					}
-				}
-
-				ResourcesString += TEXT("/*");
-				ResourcesString += SubstrateMaterialContextDescription;
-				ResourcesString += TEXT("*/");
-
-				SubstrateMaterialDescription += SubstrateMaterialContextDescription;
-
-				if (SubstrateCompilationContextIndex == ESubstrateCompilationContext::SCC_Default)
-				{
-					MaterialCompilationOutput.SubstrateMaterialCompilationOutput.SharedLocalBasesCount = SubstrateCtx.FinalUsedSharedLocalBasesCount;
-				}
-			}
-#endif // WITH_EDITOR
-		}
-		MaterialCompilationOutput.SubstrateMaterialCompilationOutput.SubstrateMaterialDescription = SubstrateMaterialDescription;
-	}
-
-	OutEnvironment.SetDefine(TEXT("TEXTURE_SAMPLE_DEBUG"), IsDebugTextureSampleEnabled() ? TEXT("1") : TEXT("0"));
-}
-
-// Assign custom interpolators to slots, packing them as much as possible in unused slots.
-TBitArray<> FHLSLMaterialTranslator::GetVertexInterpolatorsOffsets(FString& VertexInterpolatorsOffsetsDefinitionCode) const
-{
-	TBitArray<> AllocatedCoords = AllocatedUserTexCoords; // Don't mess with the already assigned sets of UV coords
-
-	int32 CurrentSlot = INDEX_NONE;
-	int32 EndAllocatedSlot = INDEX_NONE;
-
-	auto GetNextUVSlot = [&CurrentSlot, &EndAllocatedSlot, &AllocatedCoords]() -> int32
-	{
-		if (CurrentSlot == EndAllocatedSlot)
-		{
-			CurrentSlot = AllocatedCoords.FindAndSetFirstZeroBit();
-			if (CurrentSlot == INDEX_NONE)
-			{
-				CurrentSlot = AllocatedCoords.Add(true);
-			}
-
-			// Track one slot per component (u,v)
-			const int32 NUM_COMPONENTS = 2;
-			CurrentSlot *= NUM_COMPONENTS;
-			EndAllocatedSlot = CurrentSlot + NUM_COMPONENTS;
-		}
-
-		int32 ResultUVSlot = CurrentSlot / 2;
-		CurrentSlot++;
-
-		return ResultUVSlot;
-	};
+		return ResultUVSlot;
+	};
 
 	TArray<UMaterialExpressionVertexInterpolator*> SortedInterpolators;
 	Algo::TransformIf(CustomVertexInterpolators, 
@@ -3238,406 +2745,121 @@ TBitArray<> FHLSLMaterialTranslator::GetVertexInterpolatorsOffsets(FString& Vert
 						[](UMaterialExpressionVertexInterpolator* Interpolator) { return Interpolator; });
 						
 	SortedInterpolators.Sort([](const UMaterialExpressionVertexInterpolator& LHS, const UMaterialExpressionVertexInterpolator& RHS)  { return LHS.InterpolatorOffset < RHS.InterpolatorOffset; });
-		
-	for (UMaterialExpressionVertexInterpolator* Interpolator : SortedInterpolators)
-	{
-		int32 Index = Interpolator->InterpolatorIndex;
-
-		const EMaterialValueType Type = Interpolator->InterpolatedType == MCT_Float ? MCT_Float1 : Interpolator->InterpolatedType;
-
-		VertexInterpolatorsOffsetsDefinitionCode += HLSL_LINE_TERMINATOR;
-		VertexInterpolatorsOffsetsDefinitionCode += FString::Printf(TEXT("#define VERTEX_INTERPOLATOR_%i_TEXCOORDS_X\t%i") HLSL_LINE_TERMINATOR, Index, GetNextUVSlot());
-
-		if (Type >= MCT_Float2)
-		{
-			VertexInterpolatorsOffsetsDefinitionCode += FString::Printf(TEXT("#define VERTEX_INTERPOLATOR_%i_TEXCOORDS_Y\t%i") HLSL_LINE_TERMINATOR, Index, GetNextUVSlot());
-
-			if (Type >= MCT_Float3)
-			{
-				VertexInterpolatorsOffsetsDefinitionCode += FString::Printf(TEXT("#define VERTEX_INTERPOLATOR_%i_TEXCOORDS_Z\t%i") HLSL_LINE_TERMINATOR, Index, GetNextUVSlot());
-
-				if (Type == MCT_Float4)
-				{
-					VertexInterpolatorsOffsetsDefinitionCode += FString::Printf(TEXT("#define VERTEX_INTERPOLATOR_%i_TEXCOORDS_W\t%i") HLSL_LINE_TERMINATOR, Index, GetNextUVSlot());
-				}
-			}
-		}
-			
-		VertexInterpolatorsOffsetsDefinitionCode += HLSL_LINE_TERMINATOR;
-	}
-
-	return AllocatedCoords;
-}
-
-void FHLSLMaterialTranslator::GetSharedInputsMaterialCode(FString& PixelMembersDeclaration, FString& NormalAssignment, FString& PixelMembersInitializationEpilog, ECompiledPartialDerivativeVariation DerivativeVariation)
-{
-	int32 LastProperty = -1;
-	FString PixelInputInitializerValues;
-	FString NormalInitializerValue;
-
-	for (int32 PropertyIndex = 0; PropertyIndex < MP_MAX; ++PropertyIndex)
-	{
-		// Skip non-shared properties
-		if (!SharedPixelProperties[PropertyIndex])
-		{
-			continue;
-		}
-
-		const EMaterialProperty Property = (EMaterialProperty)PropertyIndex;
-		check(FMaterialAttributeDefinitionMap::GetShaderFrequency(Property) == SF_Pixel);
-		// Special case MP_SubsurfaceColor as the actual property is a combination of the color and the profile but we don't want to expose the profile
-		const FString PropertyName = Property == MP_SubsurfaceColor ? "Subsurface" : FMaterialAttributeDefinitionMap::GetAttributeName(Property);
-		check(PropertyName.Len() > 0);				
-		const EMaterialValueType Type = Property == MP_SubsurfaceColor ? MCT_Float4 : FMaterialAttributeDefinitionMap::GetValueType(Property);
-
-		if (!bEnableExecutionFlow)
-		{
-		// Normal requires its own separate initializer
-		if (Property == MP_Normal)
-		{
-				NormalInitializerValue = FString::Printf(TEXT("\tPixelMaterialInputs.%s = %s;\n"), *PropertyName, *DerivativeVariations[DerivativeVariation].TranslatedCodeChunks[Property]);
-		}
-		else
-		{
-				if (DerivativeVariations[DerivativeVariation].TranslatedCodeChunkDefinitions[Property].Len() > 0)
-				{
-					LastProperty = Property;
-				}
-			}
-
-			PixelInputInitializerValues += FString::Printf(TEXT("\tPixelMaterialInputs.%s = %s;\n"), *PropertyName, *DerivativeVariations[DerivativeVariation].TranslatedCodeChunks[Property]);
-		}
-
-		PixelMembersDeclaration += FString::Printf(TEXT("\t%s %s;\n"), HLSLTypeString(Type), *PropertyName);
-	}
-
-	NormalAssignment = NormalInitializerValue;
-	if (LastProperty != -1)
-	{
-		PixelMembersInitializationEpilog += DerivativeVariations[DerivativeVariation].TranslatedCodeChunkDefinitions[LastProperty] + TEXT("\n");
-	}
-
-	if (Substrate::IsSubstrateEnabled())
-	{
-		PixelMembersDeclaration += FString::Printf(TEXT("\t#if SUBSTRATE_USE_FULLYSIMPLIFIED_MATERIAL == 1\n"));
-		PixelMembersDeclaration += FString::Printf(TEXT("\t%s FullySimplifiedFrontMaterial;\n"), HLSLTypeString(FMaterialAttributeDefinitionMap::GetValueType(MP_FrontMaterial)));
-		PixelMembersDeclaration += FString::Printf(TEXT("\t#endif\n"));
-		if (!FullySimplifiedSubstrateFrontMaterialTranslatedCodeChunks.IsEmpty())
-		{
-			PixelInputInitializerValues += FString::Printf(TEXT("\t#if SUBSTRATE_USE_FULLYSIMPLIFIED_MATERIAL == 1\n"));
-			PixelInputInitializerValues += FString::Printf(TEXT("\tPixelMaterialInputs.FullySimplifiedFrontMaterial = %s;\n"), *FullySimplifiedSubstrateFrontMaterialTranslatedCodeChunks);
-			PixelInputInitializerValues += FString::Printf(TEXT("\t#endif\n"));
-		}
-
-		for (uint32 SubstrateCompilationContextIndex = 0; SubstrateCompilationContextIndex < ESubstrateCompilationContext::SCC_MAX; ++SubstrateCompilationContextIndex)
-		{
-			FSubstrateCompilationContext& SubstrateCtx = SubstrateCompilationContext[SubstrateCompilationContextIndex];
-			if (SubstrateCtx.CodeChunkToSubstrateSharedLocalBasis.Num() > 0)
-			{
-				PixelInputInitializerValues += SubstrateCtx.SubstratePixelNormalInitializerValues;
-			}
-		}
-	}
-
-	PixelMembersInitializationEpilog += PixelInputInitializerValues;
-}
-
-FString FHLSLMaterialTranslator::GetMaterialShaderCode()
-{
-	if (!GUseMaterialTranslationResultsGrouping)
-	{
-		return GetMaterialShaderCodeOld();
-	}
-
-	// use "/Engine/Private/MaterialTemplate.ush" to create the functions to get data (e.g. material attributes) and code (e.g. material expressions to create specular color) from C++
-	int32 MaterialTemplateLineNumberNew;
-	FStringTemplateResolver Resolver = FMaterialSourceTemplate::Get().BeginResolve(GetShaderPlatform(), &MaterialTemplateLineNumberNew);
-
-	Resolver.SetParameterMap(&MaterialSourceTemplateParams);
-	MaterialSourceTemplateParams.Add({ TEXT("line_number"), FString::Printf(TEXT("%u"), MaterialTemplateLineNumberNew) });
-
-	return Resolver.Finalize();
-}
-
-FString FHLSLMaterialTranslator::GetMaterialShaderCodeOld()
-{	
-	// use "/Engine/Private/MaterialTemplate.ush" to create the functions to get data (e.g. material attributes) and code (e.g. material expressions to create specular color) from C++
-	TMap<FString, FString> Params;
-	int32 MaterialTemplateLineNumber;
-
-	FStringTemplateResolver Resolver = FMaterialSourceTemplate::Get().BeginResolve(GetShaderPlatform(), &MaterialTemplateLineNumber);
-	Resolver.SetParameterMap(&Params);
-
-	Params.Reserve(Resolver.GetTemplate().GetNumNamedParameters());
-	Params.Add({ TEXT("line_number"), FString::Printf(TEXT("%u"), MaterialTemplateLineNumber) });
-
-	// Assign slots to vertex interpolators
-	FString VertexInterpolatorsOffsetsDefinition;
-	TBitArray<> FinalAllocatedCoords = GetVertexInterpolatorsOffsets(VertexInterpolatorsOffsetsDefinition);
-
-	const uint32 NumUserVertexTexCoords = GetNumUserVertexTexCoords();
-	const uint32 NumUserTexCoords = GetNumUserTexCoords();
-	const uint32 NumCustomVectors = FMath::DivideAndRoundUp((uint32)CurrentCustomVertexInterpolatorOffset, 2u);
-	const uint32 NumTexCoordVectors = FinalAllocatedCoords.FindLast(true) + 1;
-
-	Params.Add({ TEXT("num_material_texcoords_vertex"), FString::Printf(TEXT("%u"), NumUserVertexTexCoords) });
-	Params.Add({ TEXT("num_material_texcoords"), FString::Printf(TEXT("%u"), NumUserTexCoords) });
-	Params.Add({ TEXT("num_custom_vertex_interpolators"), FString::Printf(TEXT("%u"), NumCustomVectors) });
-	Params.Add({ TEXT("num_tex_coord_interpolators"), FString::Printf(TEXT("%u"), NumTexCoordVectors) });
-	Params.Add({ TEXT("vertex_interpolators_offsets_definition"), VertexInterpolatorsOffsetsDefinition });
-
-	FString MaterialAttributesDeclaration;
-	FString MaterialAttributesUtilities;
-
-	// Reserve some bytes for fixed string content (e.g. "struct FMaterialAttributes { ... }")  and enough space for each declaration entry, using an estimation of the maximum length.
-	MaterialAttributesDeclaration.Reserve(30 + MaterialAttributesDeclaration.Len() * 128);
-	MaterialAttributesDeclaration.Append(TEXT("struct FMaterialAttributes\n{\n"));
-
-	// Reserve enough space for each utility entry, using an estimation of the maximum length.
-	MaterialAttributesUtilities.Reserve(MaterialAttributesDeclaration.Len() * 256);
-
-	const TArray<FGuid>& OrderedVisibleAttributes = FMaterialAttributeDefinitionMap::GetOrderedVisibleAttributeList();
-	for (const FGuid& AttributeID : OrderedVisibleAttributes)
-	{
-		const FString PropertyName = FMaterialAttributeDefinitionMap::GetAttributeName(AttributeID);
-		const EMaterialValueType PropertyType = FMaterialAttributeDefinitionMap::GetValueType(AttributeID);
-		const TCHAR* HLSLType = nullptr;
-
-		switch (PropertyType)
-		{
-		case MCT_Float1: case MCT_Float: HLSLType = TEXT("float"); break;
-		case MCT_Float2: HLSLType = TEXT("float2"); break;
-		case MCT_Float3: HLSLType = TEXT("float3"); break;
-		case MCT_Float4: HLSLType = TEXT("float4"); break;
-		case MCT_UInt: case MCT_UInt1: case MCT_ShadingModel: HLSLType = TEXT("uint"); break;
-		case MCT_UInt2: HLSLType = TEXT("uint2"); break;
-		case MCT_UInt3: HLSLType = TEXT("uint3"); break;
-		case MCT_UInt4: HLSLType = TEXT("uint4"); break;
-		case MCT_Substrate: HLSLType = TEXT("FSubstrateData"); break;
-		default: break;
-		}
-
-		if (HLSLType)
-		{
-			const FVector4f DefaultValue = FMaterialAttributeDefinitionMap::GetDefaultValue(AttributeID);
-
-			MaterialAttributesDeclaration.Appendf(TEXT("\t%s %s;") HLSL_LINE_TERMINATOR, HLSLType, *PropertyName);
-
-			// Chainable method to set the attribute
-			MaterialAttributesUtilities.Appendf(TEXT("FMaterialAttributes FMaterialAttributes_Set%s(FMaterialAttributes InAttributes, %s InValue) { InAttributes.%s = InValue; return InAttributes; }") HLSL_LINE_TERMINATOR,
-				*PropertyName, HLSLType, *PropertyName);
-		}
-	}
-
-	MaterialAttributesDeclaration.Append(TEXT("};\n"));
-
-	Params.Add({ TEXT("material_declarations"), MaterialAttributesDeclaration });
-	Params.Add({ TEXT("material_attributes_utilities"), MaterialAttributesUtilities });
-
-	// Stores the shared shader results member declarations
-	FString PixelMembersDeclaration[CompiledPDV_MAX];
-
-	FString NormalAssignment[CompiledPDV_MAX];
-
-	// Stores the code to initialize all inputs after MP_Normal
-	FString PixelMembersSetupAndAssignments[CompiledPDV_MAX];
-
-	for (int32 Index = 0; Index < CompiledPDV_MAX; Index++)
-	{
-		GetSharedInputsMaterialCode(PixelMembersDeclaration[Index], NormalAssignment[Index], PixelMembersSetupAndAssignments[Index], (ECompiledPartialDerivativeVariation)Index);
-	}
-
-	// PixelMembersDeclaration should be the same for all variations, but might change in the future. There are cases where work is shared
-	// between the pixel and vertex shader, but with Nanite all work has to be moved into the pixel shader, which means we will want
-	// different inputs. But for now, we are keeping them the same.
-	Params.Add({ TEXT("pixel_material_inputs"), PixelMembersDeclaration[CompiledPDV_FiniteDifferences] });
-
-	{
-		FString DerivativeHelpers = DerivativeAutogen.GenerateUsedFunctions(*this);
-		FString DerivativeHelpersAndResources = DerivativeHelpers + ResourcesString;
-		Params.Add({ TEXT("uniform_material_expressions"), DerivativeHelpersAndResources });
-	}
-
-	// Anything used bye the GenerationFunctionCode() like WorldPositionOffset shouldn't be using texures, right?
-	// Let those use the standard finite differences textures, since they should be the same. If we actually want
-	// those to handle texture reads properly, we'll have to make extra versions.
-	ECompiledPartialDerivativeVariation BaseDerivativeVariation = CompiledPDV_FiniteDifferences;
-
-	if (bCompileForComputeShader)
-	{
-		Params.Add({ TEXT("get_material_emissive_for_cs"), GenerateFunctionCode(CompiledMP_EmissiveColorCS, BaseDerivativeVariation) });
-	}
-	else
-	{
-		Params.Add({ TEXT("get_material_emissive_for_cs"), TEXT("return 0") });
-	}
-
-	Params.Add({ TEXT("get_material_translucency_directional_lighting_intensity"), FString::Printf(TEXT("return %.5f"), Material->GetTranslucencyDirectionalLightingIntensity()) });
-	Params.Add({ TEXT("get_material_translucent_shadow_density_scale"), FString::Printf(TEXT("return %.5f"), Material->GetTranslucentShadowDensityScale()) });
-	Params.Add({ TEXT("get_material_translucent_self_shadow_density_scale"), FString::Printf(TEXT("return %.5f"), Material->GetTranslucentSelfShadowDensityScale()) });
-	Params.Add({ TEXT("get_material_translucent_self_shadow_second_density_scale"), FString::Printf(TEXT("return %.5f"), Material->GetTranslucentSelfShadowSecondDensityScale()) });
-	Params.Add({ TEXT("get_material_translucent_self_shadow_second_opacity"), FString::Printf(TEXT("return %.5f"), Material->GetTranslucentSelfShadowSecondOpacity()) });
-	Params.Add({ TEXT("get_material_translucent_backscattering_exponent"), FString::Printf(TEXT("return %.5f"), Material->GetTranslucentBackscatteringExponent()) });
-
-	{
-		FLinearColor Extinction = Material->GetTranslucentMultipleScatteringExtinction();
-
-		Params.Add({ TEXT("get_material_translucent_multiple_scattering_extinction"), FString::Printf(TEXT("return MaterialFloat3(%.5f, %.5f, %.5f)"), Extinction.R, Extinction.G, Extinction.B) });
-	}
-
-	Params.Add({ TEXT("get_material_opacity_mask_clip_value"), FString::Printf(TEXT("return %.5f"), Material->GetOpacityMaskClipValue()) });
-
-	{
-		const FDisplacementScaling DisplacementScaling = Material->GetDisplacementScaling();
-		Params.Add({ TEXT("get_material_displacement_magnitude"), FString::Printf(TEXT("return %.5f"), FMath::Max(0.0f, DisplacementScaling.Magnitude)) });
-		Params.Add({ TEXT("get_material_displacement_center"), FString::Printf(TEXT("return %.5f"), FMath::Clamp(DisplacementScaling.Center, 0.0f, 1.0f)) });
-	}
-
-	Params.Add({ TEXT("get_material_world_position_offset_raw"), !bEnableExecutionFlow ? *GenerateFunctionCode(MP_WorldPositionOffset, BaseDerivativeVariation) : TEXT("return Parameters.MaterialAttributes.WorldPositionOffset") });
-	Params.Add({ TEXT("get_material_previous_world_position_offset_raw"), !bEnableExecutionFlow ? *GenerateFunctionCode(CompiledMP_PrevWorldPositionOffset, BaseDerivativeVariation) : TEXT("return 0.0f") });
-	Params.Add({ TEXT("get_material_custom_data0"), !bEnableExecutionFlow ? *GenerateFunctionCode(MP_CustomData0, BaseDerivativeVariation) : TEXT("return 0.0f") });
-	Params.Add({ TEXT("get_material_custom_data1"), !bEnableExecutionFlow ? *GenerateFunctionCode(MP_CustomData1, BaseDerivativeVariation) : TEXT("return 0.0f") });
-
-	// Print custom texture coordinate assignments, should be fine with regular derivatives
-	FString CustomUVAssignments;
-
-	int32 LastProperty = -1;
-	for (uint32 CustomUVIndex = 0; CustomUVIndex < NumUserTexCoords; CustomUVIndex++)
-	{
-		if (bEnableExecutionFlow)
-		{
-			const FString AttributeName = FMaterialAttributeDefinitionMap::GetAttributeName((EMaterialProperty)(MP_CustomizedUVs0 + CustomUVIndex));
-			CustomUVAssignments.Appendf(TEXT("\tOutTexCoords[%u] = Parameters.MaterialAttributes.%s;") HLSL_LINE_TERMINATOR, CustomUVIndex, *AttributeName);
-		}
-		else
-		{
-			if (CustomUVIndex == 0)
-			{
-				CustomUVAssignments.Append(DerivativeVariations[BaseDerivativeVariation].TranslatedCodeChunkDefinitions[MP_CustomizedUVs0 + CustomUVIndex]);
-			}
-
-			if (DerivativeVariations[BaseDerivativeVariation].TranslatedCodeChunkDefinitions[MP_CustomizedUVs0 + CustomUVIndex].Len() > 0)
-			{
-				LastProperty = MP_CustomizedUVs0 + CustomUVIndex;
-			}
-
-			CustomUVAssignments.Appendf(TEXT("\tOutTexCoords[%u] = %s;") HLSL_LINE_TERMINATOR, CustomUVIndex, *DerivativeVariations[BaseDerivativeVariation].TranslatedCodeChunks[MP_CustomizedUVs0 + CustomUVIndex]);
-		}
-	}
+		
+	for (UMaterialExpressionVertexInterpolator* Interpolator : SortedInterpolators)
+	{
+		int32 Index = Interpolator->InterpolatorIndex;
 
-	Params.Add({ TEXT("get_material_customized_u_vs"), CustomUVAssignments });
+		const EMaterialValueType Type = Interpolator->InterpolatedType == MCT_Float ? MCT_Float1 : Interpolator->InterpolatedType;
 
-	// Print custom vertex shader interpolator assignments
-	FString CustomInterpolatorAssignments;
+		VertexInterpolatorsOffsetsDefinitionCode += HLSL_LINE_TERMINATOR;
+		VertexInterpolatorsOffsetsDefinitionCode += FString::Printf(TEXT("#define VERTEX_INTERPOLATOR_%i_TEXCOORDS_X\t%i") HLSL_LINE_TERMINATOR, Index, GetNextUVSlot());
 
-	for (UMaterialExpressionVertexInterpolator* Interpolator : CustomVertexInterpolators)
-	{
-		if (Interpolator->InterpolatorOffset != INDEX_NONE)
+		if (Type >= MCT_Float2)
 		{
-			check(Interpolator->InterpolatorIndex != INDEX_NONE);
-			check(Interpolator->InterpolatedType & MCT_Float);
-
-			const EMaterialValueType Type = Interpolator->InterpolatedType == MCT_Float ? MCT_Float1 : Interpolator->InterpolatedType;
-			const TCHAR* Swizzle[2] = { TEXT("x"), TEXT("y") };
-			const int32 Offset = Interpolator->InterpolatorOffset;
-			const int32 Index = Interpolator->InterpolatorIndex;
+			VertexInterpolatorsOffsetsDefinitionCode += FString::Printf(TEXT("#define VERTEX_INTERPOLATOR_%i_TEXCOORDS_Y\t%i") HLSL_LINE_TERMINATOR, Index, GetNextUVSlot());
 
-			// Note: We reference the UV define directly to avoid having to pre-accumulate UV counts before property translation
-			CustomInterpolatorAssignments.Appendf(TEXT("\tOutTexCoords[VERTEX_INTERPOLATOR_%i_TEXCOORDS_X].%s = VertexInterpolator%i(Parameters).x;") HLSL_LINE_TERMINATOR, Index, Swizzle[Offset%2], Index);
-				
-			if (Type >= MCT_Float2)
+			if (Type >= MCT_Float3)
 			{
-				CustomInterpolatorAssignments.Appendf(TEXT("\tOutTexCoords[VERTEX_INTERPOLATOR_%i_TEXCOORDS_Y].%s = VertexInterpolator%i(Parameters).y;") HLSL_LINE_TERMINATOR, Index, Swizzle[(Offset+1)%2], Index);
+				VertexInterpolatorsOffsetsDefinitionCode += FString::Printf(TEXT("#define VERTEX_INTERPOLATOR_%i_TEXCOORDS_Z\t%i") HLSL_LINE_TERMINATOR, Index, GetNextUVSlot());
 
-				if (Type >= MCT_Float3)
+				if (Type == MCT_Float4)
 				{
-					CustomInterpolatorAssignments.Appendf(TEXT("\tOutTexCoords[VERTEX_INTERPOLATOR_%i_TEXCOORDS_Z].%s = VertexInterpolator%i(Parameters).z;") HLSL_LINE_TERMINATOR, Index, Swizzle[(Offset+2)%2], Index);
-
-					if (Type == MCT_Float4)
-					{
-						CustomInterpolatorAssignments.Appendf(TEXT("\tOutTexCoords[VERTEX_INTERPOLATOR_%i_TEXCOORDS_W].%s = VertexInterpolator%i(Parameters).w;") HLSL_LINE_TERMINATOR, Index, Swizzle[(Offset+3)%2], Index);
-					}
+					VertexInterpolatorsOffsetsDefinitionCode += FString::Printf(TEXT("#define VERTEX_INTERPOLATOR_%i_TEXCOORDS_W\t%i") HLSL_LINE_TERMINATOR, Index, GetNextUVSlot());
 				}
 			}
 		}
+			
+		VertexInterpolatorsOffsetsDefinitionCode += HLSL_LINE_TERMINATOR;
 	}
 
-	Params.Add({ TEXT("get_custom_interpolators"), CustomInterpolatorAssignments });
+	return AllocatedCoords;
+}
 
-	if (bEnableExecutionFlow)
-	{
-		FString EvaluateVertexCode;
+void FHLSLMaterialTranslator::GetSharedInputsMaterialCode(FString& PixelMembersDeclaration, FString& NormalAssignment, FString& PixelMembersInitializationEpilog, ECompiledPartialDerivativeVariation DerivativeVariation)
+{
+	int32 LastProperty = -1;
+	FString PixelInputInitializerValues;
+	FString NormalInitializerValue;
 
-		// Set default texcoords in the VS
-		for (uint32 TexCoordIndex = 0; TexCoordIndex < NumUserVertexTexCoords; ++TexCoordIndex)
+	for (int32 PropertyIndex = 0; PropertyIndex < MP_MAX; ++PropertyIndex)
+	{
+		// Skip non-shared properties
+		if (!SharedPixelProperties[PropertyIndex])
 		{
-			const FString AttributeName = FMaterialAttributeDefinitionMap::GetAttributeName((EMaterialProperty)(MP_CustomizedUVs0 + TexCoordIndex));
-
-			EvaluateVertexCode.Appendf(TEXT("\tDefaultMaterialAttributes.%s = Parameters.TexCoords[%d];") HLSL_LINE_TERMINATOR, *AttributeName, TexCoordIndex);
+			continue;
 		}
 
-		EvaluateVertexCode.Append(TranslatedAttributesCodeChunks[SF_Vertex]);
-
-		Params.Add({ TEXT("evaluate_material_attributes"), TranslatedAttributesCodeChunks[SF_Pixel] });
-
-		FString EvaluateMaterialAttributesCode = TEXT("    FMaterialAttributes MaterialAttributes = EvaluatePixelMaterialAttributes(Parameters);" HLSL_LINE_TERMINATOR);
+		const EMaterialProperty Property = (EMaterialProperty)PropertyIndex;
+		check(FMaterialAttributeDefinitionMap::GetShaderFrequency(Property) == SF_Pixel);
+		// Special case MP_SubsurfaceColor as the actual property is a combination of the color and the profile but we don't want to expose the profile
+		const FString PropertyName = Property == MP_SubsurfaceColor ? "Subsurface" : FMaterialAttributeDefinitionMap::GetAttributeName(Property);
+		check(PropertyName.Len() > 0);				
+		const EMaterialValueType Type = Property == MP_SubsurfaceColor ? MCT_Float4 : FMaterialAttributeDefinitionMap::GetValueType(Property);
 
-		for (int32 PropertyIndex = 0; PropertyIndex < MP_MAX; ++PropertyIndex)
+		if (!bEnableExecutionFlow)
 		{
-			// Skip non-shared properties
-			if (!SharedPixelProperties[PropertyIndex])
-			{
-				continue;
+		// Normal requires its own separate initializer
+		if (Property == MP_Normal)
+		{
+				NormalInitializerValue = FString::Printf(TEXT("\tPixelMaterialInputs.%s = %s;\n"), *PropertyName, *DerivativeVariations[DerivativeVariation].TranslatedCodeChunks[Property]);
+		}
+		else
+		{
+				if (DerivativeVariations[DerivativeVariation].TranslatedCodeChunkDefinitions[Property].Len() > 0)
+				{
+					LastProperty = Property;
+				}
 			}
 
-			const EMaterialProperty Property = (EMaterialProperty)PropertyIndex;
-			check(FMaterialAttributeDefinitionMap::GetShaderFrequency(Property) == SF_Pixel);
-			// Special case MP_SubsurfaceColor as the actual property is a combination of the color and the profile but we don't want to expose the profile
-			const FString PropertyName = FMaterialAttributeDefinitionMap::GetAttributeName(Property);
-
-			if (PropertyIndex == MP_SubsurfaceColor)
-			{
-				// TODO - properly handle subsurface profile
-				EvaluateMaterialAttributesCode.Appendf("    PixelMaterialInputs.Subsurface = float4(MaterialAttributes.%s, 0.0f);" HLSL_LINE_TERMINATOR, *PropertyName);
-			}
-			else
-			{
-				EvaluateMaterialAttributesCode.Appendf("    PixelMaterialInputs.%s = MaterialAttributes.%s;" HLSL_LINE_TERMINATOR, *PropertyName, *PropertyName);
-			}
+			PixelInputInitializerValues += FString::Printf(TEXT("\tPixelMaterialInputs.%s = %s;\n"), *PropertyName, *DerivativeVariations[DerivativeVariation].TranslatedCodeChunks[Property]);
 		}
 
-		Params.Add({ TEXT("calc_pixel_material_inputs_initial_calculations"), EvaluateMaterialAttributesCode });
-		Params.Add({ TEXT("calc_pixel_material_inputs_normal"), TEXT("") });
-		Params.Add({ TEXT("calc_pixel_material_inputs_other_inputs"), TEXT("") });
-
-		Params.Add({ TEXT("calc_pixel_material_inputs_analytic_derivatives_initial"), EvaluateMaterialAttributesCode });
-		Params.Add({ TEXT("calc_pixel_material_inputs_analytic_derivatives_normal"), TEXT("") });
-		Params.Add({ TEXT("calc_pixel_material_inputs_analytic_derivatives_other_inputs"), TEXT("") });
+		PixelMembersDeclaration += FString::Printf(TEXT("\t%s %s;\n"), HLSLTypeString(Type), *PropertyName);
 	}
-	else
+
+	NormalAssignment = NormalInitializerValue;
+	if (LastProperty != -1)
 	{
-		// skip material attributes code
-		Params.Add({ TEXT("evaluate_material_attributes"), TEXT("") });
+		PixelMembersInitializationEpilog += DerivativeVariations[DerivativeVariation].TranslatedCodeChunkDefinitions[LastProperty] + TEXT("\n");
+	}
 
+	if (Substrate::IsSubstrateEnabled())
+	{
+		PixelMembersDeclaration += FString::Printf(TEXT("\t#if SUBSTRATE_USE_FULLYSIMPLIFIED_MATERIAL == 1\n"));
+		PixelMembersDeclaration += FString::Printf(TEXT("\t%s FullySimplifiedFrontMaterial;\n"), HLSLTypeString(FMaterialAttributeDefinitionMap::GetValueType(MP_FrontMaterial)));
+		PixelMembersDeclaration += FString::Printf(TEXT("\t#endif\n"));
+		if (!FullySimplifiedSubstrateFrontMaterialTranslatedCodeChunks.IsEmpty())
 		{
-			// Initializers required for Normal
-			Params.Add({ TEXT("calc_pixel_material_inputs_initial_calculations"), DerivativeVariations[CompiledPDV_FiniteDifferences].TranslatedCodeChunkDefinitions[MP_Normal] });
-			Params.Add({ TEXT("calc_pixel_material_inputs_normal"), NormalAssignment[CompiledPDV_FiniteDifferences] });
-			// Finally the rest of common code followed by assignment into each input
-			Params.Add({ TEXT("calc_pixel_material_inputs_other_inputs"), PixelMembersSetupAndAssignments[CompiledPDV_FiniteDifferences] });
+			PixelInputInitializerValues += FString::Printf(TEXT("\t#if SUBSTRATE_USE_FULLYSIMPLIFIED_MATERIAL == 1\n"));
+			PixelInputInitializerValues += FString::Printf(TEXT("\tPixelMaterialInputs.FullySimplifiedFrontMaterial = %s;\n"), *FullySimplifiedSubstrateFrontMaterialTranslatedCodeChunks);
+			PixelInputInitializerValues += FString::Printf(TEXT("\t#endif\n"));
 		}
+
+		for (uint32 SubstrateCompilationContextIndex = 0; SubstrateCompilationContextIndex < ESubstrateCompilationContext::SCC_MAX; ++SubstrateCompilationContextIndex)
 		{
-			// Initializers required for Normal
-			Params.Add({ TEXT("calc_pixel_material_inputs_analytic_derivatives_initial"), DerivativeVariations[CompiledPDV_Analytic].TranslatedCodeChunkDefinitions[MP_Normal] });
-			Params.Add({ TEXT("calc_pixel_material_inputs_analytic_derivatives_normal"), NormalAssignment[CompiledPDV_Analytic] });
-			// Finally the rest of common code followed by assignment into each input
-			Params.Add({ TEXT("calc_pixel_material_inputs_analytic_derivatives_other_inputs"), PixelMembersSetupAndAssignments[CompiledPDV_Analytic] });
+			FSubstrateCompilationContext& SubstrateCtx = SubstrateCompilationContext[SubstrateCompilationContextIndex];
+			if (SubstrateCtx.CodeChunkToSubstrateSharedLocalBasis.Num() > 0)
+			{
+				PixelInputInitializerValues += SubstrateCtx.SubstratePixelNormalInitializerValues;
+			}
 		}
 	}
 
-	return Resolver.Finalize();
+	PixelMembersInitializationEpilog += PixelInputInitializerValues;
 }
 
+FString FHLSLMaterialTranslator::GetMaterialShaderCode()
+{
+	// use "/Engine/Private/MaterialTemplate.ush" to create the functions to get data (e.g. material attributes) and code (e.g. material expressions to create specular color) from C++
+	int32 MaterialTemplateLineNumberNew;
+	FStringTemplateResolver Resolver = FMaterialSourceTemplate::Get().BeginResolve(GetShaderPlatform(), &MaterialTemplateLineNumberNew);
+
+	Resolver.SetParameterMap(&MaterialSourceTemplateParams);
+	MaterialSourceTemplateParams.Add({ TEXT("line_number"), FString::Printf(TEXT("%u"), MaterialTemplateLineNumberNew) });
+
+	return Resolver.Finalize();
+}
 
 // ========== PROTECTED: ========== //
 
@@ -12370,132 +11592,10 @@ static FString GetParametersSubstrateTreeName(ESubstrateCompilationContext Compi
 	return TEXT("ERROR");
 }
 
-void FHLSLMaterialTranslator::FSubstrateCompilationContext::SubstrateEvaluateSharedLocalBases(
-	FHLSLMaterialTranslator* Compiler,
-	uint8& OutRequestedSharedLocalBasesCount,
-	FShaderCompilerEnvironment* OutEnvironment)
-{
-	/*
-	* The final output code/workflow for shared tangent basis should look like
-	*
-	* #define SHAREDLOCALBASIS_INDEX_0 0		// default, unused
-	* #define SHAREDLOCALBASIS_INDEX_1 0		// default, unused
-	* #define SHAREDLOCALBASIS_INDEX_2 0
-	*
-	* FSubstrateData BSDF0 = GetSubstrateSlabBSDF(... SHAREDLOCALBASIS_0, NormalCode0 ...)
-	* FSubstrateData BSDF1 = GetSubstrateSlabBSDF(... SHAREDLOCALBASIS_1, NormalCode1 ...)
-	*
-	* float3 NormalCode2 = lerp(NormalCode0, NormalCode1, mix)
-	* FSubstrateData BSDF2 = SubstrateHorizontalMixingParameterBlending(BSDF0, BSDF1, mix, NormalCode2, SHAREDLOCALBASIS_INDEX_2, SharedLocalBases.Types) // will internally create NormalCode2
-	*
-	* tParameters.SharedLocalBases.Normals[SHAREDLOCALBASIS_INDEX_2] = NormalCode2;
-	* #if MATERIAL_TANGENTSPACENORMAL
-	* Parameters.SharedLocalBases.Normals[SHAREDLOCALBASIS_INDEX_2] *= Parameters.TwoSidedSign;
-	* #endif
-	*/
-
-	FSubstrateRegisteredSharedLocalBasis UsedSharedLocalBasesInfo[SUBSTRATE_MAX_SHAREDLOCALBASES_REGISTERS];
-	FinalUsedSharedLocalBasesCount = 0;
-	OutRequestedSharedLocalBasesCount = 0;
-	const FString ParameterSharedLocalBasesName = GetParametersSharedLocalBasesName(CompilationContextIndex);
-
-	for (int32 OpIt = 0; OpIt < SubstrateMaterialExpressionRegisteredOperators.Num(); ++OpIt)
-	{
-		const FSubstrateOperator& BSDFOperator = SubstrateMaterialExpressionRegisteredOperators[OpIt];
-		if (BSDFOperator.BSDFIndex == INDEX_NONE || BSDFOperator.IsDiscarded())
-		{
-			continue;	// not a BSDF or if discarded (i.e. not the root of a parameter blending subtree), then there is no local basis to register
-		}
-
-		if (BSDFOperator.BSDFRegisteredSharedLocalBasis.NormalCodeChunk == INDEX_NONE && BSDFOperator.BSDFRegisteredSharedLocalBasis.TangentCodeChunk == INDEX_NONE)
-		{
-			continue;	// We skip null normal on certain BSDF, for instance unlit.
-		}
-		const FSubstrateSharedLocalBasesInfo& SubstrateSharedLocalBasesInfo = SubstrateCompilationInfoGetMatchingSharedLocalBasisInfo(BSDFOperator.BSDFRegisteredSharedLocalBasis);
-
-		// First, we check that the normal/tangent has not already written out (avoid 2 BSDFs sharing the same normal to note generate the same code twice)
-		bool bAlreadyProcessed = false;
-		for (uint8 i = 0; i < FinalUsedSharedLocalBasesCount; ++i)
-		{
-			if (UsedSharedLocalBasesInfo[i].NormalCodeChunkHash == SubstrateSharedLocalBasesInfo.SharedData.NormalCodeChunkHash &&
-				(UsedSharedLocalBasesInfo[i].TangentCodeChunkHash == SubstrateSharedLocalBasesInfo.SharedData.TangentCodeChunkHash || BSDFOperator.BSDFRegisteredSharedLocalBasis.TangentCodeChunk == INDEX_NONE))
-			{
-				bAlreadyProcessed = true;
-				break;
-			}
-		}
-		if (bAlreadyProcessed)
-		{
-			continue;
-		}
-
-		++OutRequestedSharedLocalBasesCount;
-		if (FinalUsedSharedLocalBasesCount >= SUBSTRATE_MAX_SHAREDLOCALBASES_REGISTERS)
-		{
-			continue;
-		}
-
-		const uint8 FinalSharedLocalBasisIndex = FinalUsedSharedLocalBasesCount++;
-		UsedSharedLocalBasesInfo[FinalSharedLocalBasisIndex] = SubstrateSharedLocalBasesInfo.SharedData;
-
-		if (CompilationContextIndex == SCC_FullySimplified)
-		{
-			SubstratePixelNormalInitializerValues += FString::Printf(TEXT("\t#if SUBSTRATE_USE_FULLYSIMPLIFIED_MATERIAL == 1\n"));
-		}
-
-		// Write out normals
-		SubstratePixelNormalInitializerValues += FString::Printf(TEXT("\tParameters.%s.Normals[%u] = %s;\n"), *ParameterSharedLocalBasesName, FinalSharedLocalBasisIndex, *SubstrateSharedLocalBasesInfo.NormalCode);
-		SubstratePixelNormalInitializerValues += FString::Printf(TEXT("\t#if MATERIAL_TANGENTSPACENORMAL\n"));
-		SubstratePixelNormalInitializerValues += FString::Printf(TEXT("\tParameters.%s.Normals[%u] *= Parameters.TwoSidedSign;\n"), *ParameterSharedLocalBasesName, FinalSharedLocalBasisIndex);
-		SubstratePixelNormalInitializerValues += FString::Printf(TEXT("\t#endif\n"));
-
-		// Write out tangents
-		if (SubstrateSharedLocalBasesInfo.SharedData.TangentCodeChunk != INDEX_NONE)
-		{
-			SubstratePixelNormalInitializerValues += FString::Printf(TEXT("\tParameters.%s.Tangents[%u] = %s;\n"), *ParameterSharedLocalBasesName, FinalSharedLocalBasisIndex, *SubstrateSharedLocalBasesInfo.TangentCode);
-		}
-		else
-		{
-			SubstratePixelNormalInitializerValues += FString::Printf(TEXT("\tParameters.%s.Tangents[%u] = Parameters.TangentToWorld[0];\n"), *ParameterSharedLocalBasesName, FinalSharedLocalBasisIndex);
-		}
-		SubstratePixelNormalInitializerValues += FString::Printf(TEXT("\t#if MATERIAL_TANGENTSPACENORMAL\n"));
-		SubstratePixelNormalInitializerValues += FString::Printf(TEXT("\tParameters.%s.Tangents[%u] *= Parameters.TwoSidedSign;\n"), *ParameterSharedLocalBasesName, FinalSharedLocalBasisIndex);
-		SubstratePixelNormalInitializerValues += FString::Printf(TEXT("\t#endif\n"));
-	}
-
-	SubstratePixelNormalInitializerValues += FString::Printf(TEXT("\tParameters.%s.Count = %u;\n"), *ParameterSharedLocalBasesName, FinalUsedSharedLocalBasesCount);
-
-	if (CompilationContextIndex == SCC_FullySimplified)
-	{
-		SubstratePixelNormalInitializerValues += FString::Printf(TEXT("\t#endif // SUBSTRATE_USE_FULLYSIMPLIFIED_MATERIAL\n"));
-	}
-
-	if (OutEnvironment)
-	{
-		// Now write out all the macros, them mapping from the BSDF to the effective position/index in the shared local basis array they should write to.
-		for (TMultiMap<uint64, FSubstrateSharedLocalBasesInfo>::TConstIterator It(CodeChunkToSubstrateSharedLocalBasis); It; ++It)
-		{
-			// The default linear output index will be 0 by default, and different if in fact the shared local basis points to one that is effectively in used in the array of shared local bases.
-			uint8 LinearIndex = 0;
-			for (uint8 i = 0; i < FinalUsedSharedLocalBasesCount; ++i)
-			{
-				if (UsedSharedLocalBasesInfo[i].NormalCodeChunkHash == It->Value.SharedData.NormalCodeChunkHash &&
-					(UsedSharedLocalBasesInfo[i].TangentCodeChunkHash == It->Value.SharedData.TangentCodeChunkHash || It->Value.SharedData.TangentCodeChunk == INDEX_NONE))
-				{
-					LinearIndex = i;
-					break;
-				}
-			}
-
-			OutEnvironment->SetDefine(*Compiler->GetSubstrateSharedLocalBasisIndexMacroInner(It->Value.SharedData, CompilationContextIndex), LinearIndex);
-		}
-	}
-}
-
 /**
  * This function is temporary. It's only been duplicated for validation purposes and it will be removed in a few days when the latest changes settle without problems.
  */
-void FHLSLMaterialTranslator::FSubstrateCompilationContext::SubstrateEvaluateSharedLocalBasesNew(
+void FHLSLMaterialTranslator::FSubstrateCompilationContext::SubstrateEvaluateSharedLocalBases(
 	FHLSLMaterialTranslator* Compiler,
 	uint8& OutRequestedSharedLocalBasesCount,
 	FEnvironmentDefines* OutEnvironment)
@@ -15312,7 +14412,7 @@ void FHLSLMaterialTranslator::PrepareEnvironmentDefines()
 			SubstrateCtx.SubstratePixelNormalInitializerValues = FString::Printf(TEXT("\n\n\n\t// Substrate normal and tangent\n"));
 			SubstrateCtx.FinalUsedSharedLocalBasesCount = 0;
 			uint8 RequestedSharedLocalBasesCount = 0;
-			SubstrateCtx.SubstrateEvaluateSharedLocalBasesNew(this, RequestedSharedLocalBasesCount, EnvironmentDefines.Get());
+			SubstrateCtx.SubstrateEvaluateSharedLocalBases(this, RequestedSharedLocalBasesCount, EnvironmentDefines.Get());
 
 #if WITH_EDITOR
 			// Now write some feedback to the user, but only produce debug string if in editor
diff --git a/Engine/Source/Runtime/Engine/Private/Materials/HLSLMaterialTranslator.h b/Engine/Source/Runtime/Engine/Private/Materials/HLSLMaterialTranslator.h
--- a/Engine/Source/Runtime/Engine/Private/Materials/HLSLMaterialTranslator.h
+++ b/Engine/Source/Runtime/Engine/Private/Materials/HLSLMaterialTranslator.h
@@ -541,8 +541,7 @@ protected:
 
 		bool SubstrateGenerateDerivedMaterialOperatorData(FHLSLMaterialTranslator* Compiler);
 
-		void SubstrateEvaluateSharedLocalBases(FHLSLMaterialTranslator* Compiler, uint8& RequestedSharedLocalBasesCount, FShaderCompilerEnvironment* OutEnvironment);
-		void SubstrateEvaluateSharedLocalBasesNew(FHLSLMaterialTranslator* Compiler, uint8& RequestedSharedLocalBasesCount, FEnvironmentDefines* OutEnvironment);
+		void SubstrateEvaluateSharedLocalBases(FHLSLMaterialTranslator* Compiler, uint8& RequestedSharedLocalBasesCount, FEnvironmentDefines* OutEnvironment);
 
 		FSubstrateSharedLocalBasesInfo SubstrateCompilationInfoGetMatchingSharedLocalBasisInfo(const FSubstrateRegisteredSharedLocalBasis& SearchedSharedLocalBasis);
 
@@ -649,12 +648,6 @@ public:
 	virtual FString GetParameterCodeDeriv(int32 Index, ECompiledPartialDerivativeVariation Variation);
 
 protected:
-	// These -Old functions are temporarily preserved in order to switch back to the old way of
-	// generating material translation results. We intend to remove these as soon as we are sure
-	// the new implementations of these functions work well.
-	void GetMaterialEnvironmentOld(EShaderPlatform InPlatform, FShaderCompilerEnvironment& OutEnvironment);
-	FString GetMaterialShaderCodeOld();
-
 	uint64 GetParameterHash(int32 Index);
 
 	uint64 GetParameterMaterialAttributeMask(int32 Index);
