Fix r.DefaultFeature.LocalExposure.HighlightContrastScale incorrectly affecting ShadowContrast

diff --git a/Engine/Source/Runtime/Engine/Private/SceneView.cpp b/Engine/Source/Runtime/Engine/Private/SceneView.cpp
--- a/Engine/Source/Runtime/Engine/Private/SceneView.cpp
+++ b/Engine/Source/Runtime/Engine/Private/SceneView.cpp
@@ -1899,7 +1899,7 @@ void FSceneView::StartFinalPostprocessSettings(FVector InViewLocation)
 		}
 
 		{
-			const float ShadowContrast = FMath::Clamp(CVarDefaultLocalExposureHighlightContrast.GetValueOnGameThread(), 0.0f, 1.0f);
+			const float ShadowContrast = FMath::Clamp(CVarDefaultLocalExposureShadowContrast.GetValueOnGameThread(), 0.0f, 1.0f);
 
 			FinalPostProcessSettings.LocalExposureShadowContrastScale = ShadowContrast;
 		}
