Level Instance: Fix partial loading issue when editing parent, then editing child and exiting edit. Child wouldn't reload its actors properly because of of desync between ILevelInstanceInterface::SupportsPartialEditorLoading and UWorldPartition::IsStreamingEnabledInEditor

- Now always rely on UWorldPartition::IsStreamingEnabledInEditor once WorldPartition is loaded
- Also simplify Loader code / api to get the loaded WorldPartitions from a ActorDescInstance

diff --git a/Engine/Source/Runtime/Engine/Private/LevelInstance/LevelInstanceInterface.cpp b/Engine/Source/Runtime/Engine/Private/LevelInstance/LevelInstanceInterface.cpp
--- a/Engine/Source/Runtime/Engine/Private/LevelInstance/LevelInstanceInterface.cpp
+++ b/Engine/Source/Runtime/Engine/Private/LevelInstance/LevelInstanceInterface.cpp
@@ -3,6 +3,7 @@
 #include "LevelInstance/LevelInstanceInterface.h"
 #include "LevelInstance/LevelInstanceSubsystem.h"
 #include "LevelInstance/LevelInstanceLevelStreaming.h"
+#include "WorldPartition/WorldPartition.h"
 #include "GameFramework/Actor.h"
 #include "Engine/Level.h"
 #include "Engine/World.h"
@@ -44,6 +45,15 @@ bool ILevelInstanceInterface::SupportsPartialEditorLoading() const
 				return false;
 			}
 
+			// If the level is loaded, check that it has editor streaming enabled
+			if (ULevel* Level = GetLoadedLevel())
+			{
+				if (UWorldPartition* WorldPartition = Level->GetWorldPartition(); WorldPartition && WorldPartition->IsInitialized() && !WorldPartition->IsStreamingEnabledInEditor())
+				{
+					return false;
+				}
+			}
+
 			if (ILevelInstanceInterface* Parent = LevelInstanceSubsystem->GetParentLevelInstance(Actor))
 			{
 				if (!Parent->SupportsPartialEditorLoading())
diff --git a/Engine/Source/Runtime/Engine/Private/WorldPartition/LevelInstance/LevelInstanceActorDesc.cpp b/Engine/Source/Runtime/Engine/Private/WorldPartition/LevelInstance/LevelInstanceActorDesc.cpp
--- a/Engine/Source/Runtime/Engine/Private/WorldPartition/LevelInstance/LevelInstanceActorDesc.cpp
+++ b/Engine/Source/Runtime/Engine/Private/WorldPartition/LevelInstance/LevelInstanceActorDesc.cpp
@@ -212,19 +212,17 @@ void FLevelInstanceActorDesc::TransferFrom(const FWorldPartitionActorDesc* From)
 	FromLevelInstanceActorDesc->UnregisterChildContainer();
 }
 
-bool FLevelInstanceActorDesc::GetLoadedChildContainerInstance(const FWorldPartitionActorDescInstance* InActorDescInstance, FLoadedContainerInstance& OutContainerInstance) const
+UWorldPartition* FLevelInstanceActorDesc::GetLoadedChildWorldPartition(const FWorldPartitionActorDescInstance* InActorDescInstance) const
 {
 	if (const ILevelInstanceInterface* LevelInstance = Cast<ILevelInstanceInterface>(InActorDescInstance->GetActor()))
 	{
 		if (ULevel* LoadedLevel = LevelInstance->GetLoadedLevel())
 		{
-			OutContainerInstance.LoadedLevel = LoadedLevel;
-			OutContainerInstance.bSupportsPartialEditorLoading = LevelInstance->SupportsPartialEditorLoading();
-			return true;
+			return LoadedLevel->GetWorldPartition();
 		}
 	}
 
-	return false;
+	return nullptr;
 }
 
 UActorDescContainerInstance* FLevelInstanceActorDesc::CreateChildContainerInstance(const FWorldPartitionActorDescInstance* InActorDescInstance) const
diff --git a/Engine/Source/Runtime/Engine/Private/WorldPartition/LoaderAdapter/LoaderAdapterActorList.cpp b/Engine/Source/Runtime/Engine/Private/WorldPartition/LoaderAdapter/LoaderAdapterActorList.cpp
--- a/Engine/Source/Runtime/Engine/Private/WorldPartition/LoaderAdapter/LoaderAdapterActorList.cpp
+++ b/Engine/Source/Runtime/Engine/Private/WorldPartition/LoaderAdapter/LoaderAdapterActorList.cpp
@@ -67,19 +67,12 @@ void FLoaderAdapterActorList::RemoveActors(const TArray<FWorldPartitionHandle>&
 	{
 		if (ActorHandle.GetInstance()->IsChildContainerInstance())
 		{
-			FWorldPartitionActorDesc::FLoadedContainerInstance ContainerInstance;
-			if (IWorldPartitionActorLoaderInterface::GetLoadedChildContainerInstance(ActorHandle, ContainerInstance))
+			if (UWorldPartition* ContainerWorldPartition = GetLoadedChildWorldPartition(ActorHandle); ContainerWorldPartition && ContainerWorldPartition->IsStreamingEnabledInEditor())
 			{
-				if (ContainerInstance.bSupportsPartialEditorLoading)
+				for (FActorDescContainerInstanceCollection::TIterator<> Iterator(ContainerWorldPartition); Iterator; ++Iterator)
 				{
-					if (UWorldPartition* ContainerWorldPartition = ContainerInstance.LoadedLevel ? ContainerInstance.LoadedLevel->GetWorldPartition() : nullptr)
-					{
-						for (FActorDescContainerInstanceCollection::TIterator<> Iterator(ContainerWorldPartition); Iterator; ++Iterator)
-						{
-							FWorldPartitionHandle SubActorHandle(ContainerWorldPartition, Iterator->GetGuid());
-							ActorsToRemove.Add(SubActorHandle);
-						}
-					}
+					FWorldPartitionHandle SubActorHandle(ContainerWorldPartition, Iterator->GetGuid());
+					ActorsToRemove.Add(SubActorHandle);
 				}
 			}
 		}
diff --git a/Engine/Source/Runtime/Engine/Private/WorldPartition/LoaderAdapter/LoaderAdapterList.cpp b/Engine/Source/Runtime/Engine/Private/WorldPartition/LoaderAdapter/LoaderAdapterList.cpp
--- a/Engine/Source/Runtime/Engine/Private/WorldPartition/LoaderAdapter/LoaderAdapterList.cpp
+++ b/Engine/Source/Runtime/Engine/Private/WorldPartition/LoaderAdapter/LoaderAdapterList.cpp
@@ -28,20 +28,13 @@ void FLoaderAdapterList::HandleActorContainer(const FWorldPartitionHandle& InAct
 {
 	if (InActor.GetInstance()->IsChildContainerInstance())
 	{
-		FWorldPartitionActorDesc::FLoadedContainerInstance ContainerInstance;
-		if (IWorldPartitionActorLoaderInterface::GetLoadedChildContainerInstance(InActor, ContainerInstance))
+		if (UWorldPartition* ContainerWorldPartition = GetLoadedChildWorldPartition(InActor); ContainerWorldPartition && ContainerWorldPartition->IsStreamingEnabledInEditor())
 		{
-			if (ContainerInstance.bSupportsPartialEditorLoading)
+			for (FActorDescContainerInstanceCollection::TIterator<> Iterator(ContainerWorldPartition); Iterator; ++Iterator)
 			{
-				if (UWorldPartition* ContainerWorldPartition = ContainerInstance.LoadedLevel ? ContainerInstance.LoadedLevel->GetWorldPartition() : nullptr)
-				{
-					for (FActorDescContainerInstanceCollection::TIterator<> Iterator(ContainerWorldPartition); Iterator; ++Iterator)
-					{
-						FWorldPartitionHandle ActorHandle(ContainerWorldPartition, Iterator->GetGuid());
-						InOperation(ActorHandle);
-						HandleActorContainer(ActorHandle, InOperation);
-					}
-				}
+				FWorldPartitionHandle ActorHandle(ContainerWorldPartition, Iterator->GetGuid());
+				InOperation(ActorHandle);
+				HandleActorContainer(ActorHandle, InOperation);
 			}
 		}
 	}
diff --git a/Engine/Source/Runtime/Engine/Private/WorldPartition/LoaderAdapter/LoaderAdapterSpatial.cpp b/Engine/Source/Runtime/Engine/Private/WorldPartition/LoaderAdapter/LoaderAdapterSpatial.cpp
--- a/Engine/Source/Runtime/Engine/Private/WorldPartition/LoaderAdapter/LoaderAdapterSpatial.cpp
+++ b/Engine/Source/Runtime/Engine/Private/WorldPartition/LoaderAdapter/LoaderAdapterSpatial.cpp
@@ -43,17 +43,14 @@ void ILoaderAdapterSpatial::HandleIntersectingContainer(UWorldPartition* InWorld
 
 			if (ActorDescInstance->GetIsSpatiallyLoaded() && ActorDescInstance->IsChildContainerInstance())
 			{
-				FWorldPartitionActorDesc::FLoadedContainerInstance ContainerInstance;
-				if (IWorldPartitionActorLoaderInterface::GetLoadedChildContainerInstance(ActorHandle, ContainerInstance))
+				if (UWorldPartition* ContainerWorldPartition = GetLoadedChildWorldPartition(ActorHandle))
 				{
-					if (UWorldPartition* ContainerWorldPartition = ContainerInstance.LoadedLevel ? ContainerInstance.LoadedLevel->GetWorldPartition() : nullptr)
-					{
-						const FBox InnerBoundingBox = ContainerInstance.bSupportsPartialEditorLoading ? *GetBoundingBox() : FBox(FVector(-HALF_WORLD_MAX), FVector(HALF_WORLD_MAX));
-						HandleIntersectingContainer(ContainerWorldPartition, InnerBoundingBox, InOperation);
-					}
+					const FBox InnerBoundingBox = ContainerWorldPartition->IsStreamingEnabledInEditor() ? *GetBoundingBox() : FBox(FVector(-HALF_WORLD_MAX), FVector(HALF_WORLD_MAX));
+					HandleIntersectingContainer(ContainerWorldPartition, InnerBoundingBox, InOperation);
 				}
 			}
 		}
 	}, ForEachIntersectingActorParams);
 }
+
 #endif
diff --git a/Engine/Source/Runtime/Engine/Private/WorldPartition/WorldPartitionActorLoaderInterface.cpp b/Engine/Source/Runtime/Engine/Private/WorldPartition/WorldPartitionActorLoaderInterface.cpp
--- a/Engine/Source/Runtime/Engine/Private/WorldPartition/WorldPartitionActorLoaderInterface.cpp
+++ b/Engine/Source/Runtime/Engine/Private/WorldPartition/WorldPartitionActorLoaderInterface.cpp
@@ -351,14 +351,19 @@ const IWorldPartitionActorLoaderInterface::FActorReferenceMap* IWorldPartitionAc
 	return ContainerActorReferences.Find(InContainerInstance);
 }
 
-bool IWorldPartitionActorLoaderInterface::GetLoadedChildContainerInstance(const FWorldPartitionHandle& Handle, FWorldPartitionActorDesc::FLoadedContainerInstance& OutContainerInstance)
+UWorldPartition* IWorldPartitionActorLoaderInterface::ILoaderAdapter::GetLoadedChildWorldPartition(const FWorldPartitionHandle& Handle) const
+{
+	return IWorldPartitionActorLoaderInterface::GetLoadedChildWorldPartition(Handle);
+}
+
+UWorldPartition* IWorldPartitionActorLoaderInterface::GetLoadedChildWorldPartition(const FWorldPartitionHandle& Handle)
 {
 	if (FWorldPartitionActorDescInstance* Instance = Handle.GetInstance())
 	{
-		return Instance->GetLoadedChildContainerInstance(OutContainerInstance);
+		return Instance->GetLoadedChildWorldPartition();
 	}
 
-	return false;
+	return nullptr;
 }
 #endif
 
diff --git a/Engine/Source/Runtime/Engine/Public/WorldPartition/LevelInstance/LevelInstanceActorDesc.h b/Engine/Source/Runtime/Engine/Public/WorldPartition/LevelInstance/LevelInstanceActorDesc.h
--- a/Engine/Source/Runtime/Engine/Public/WorldPartition/LevelInstance/LevelInstanceActorDesc.h
+++ b/Engine/Source/Runtime/Engine/Public/WorldPartition/LevelInstance/LevelInstanceActorDesc.h
@@ -47,7 +47,7 @@ protected:
 	ENGINE_API virtual void SetContainer(UActorDescContainer* InContainer) override;
 	ENGINE_API virtual UActorDescContainerInstance* CreateChildContainerInstance(const FWorldPartitionActorDescInstance* InActorDescInstance) const override;
 	ENGINE_API virtual bool GetChildContainerInstance(const FWorldPartitionActorDescInstance* InActorDescInstance, FContainerInstance& OutContainerInstance) const override;
-	ENGINE_API virtual bool GetLoadedChildContainerInstance(const FWorldPartitionActorDescInstance* InActorDescInstance, FLoadedContainerInstance& OutContainerInstance) const override;
+	ENGINE_API virtual UWorldPartition* GetLoadedChildWorldPartition(const FWorldPartitionActorDescInstance* InActorDescInstance) const override;
 	ENGINE_API FTransform GetChildContainerTransform() const;
 
 	FSoftObjectPath WorldAsset;
diff --git a/Engine/Source/Runtime/Engine/Public/WorldPartition/WorldPartitionActorDesc.h b/Engine/Source/Runtime/Engine/Public/WorldPartition/WorldPartitionActorDesc.h
--- a/Engine/Source/Runtime/Engine/Public/WorldPartition/WorldPartitionActorDesc.h
+++ b/Engine/Source/Runtime/Engine/Public/WorldPartition/WorldPartitionActorDesc.h
@@ -97,12 +97,6 @@ class FWorldPartitionActorDesc
 	template<class U> friend class TActorDescContainerCollection;
 
 public:
-	struct FLoadedContainerInstance
-	{
-		ULevel* LoadedLevel = nullptr;
-		bool bSupportsPartialEditorLoading = false;
-	};
-
 	struct FContainerInstance
 	{
 		UActorDescContainerInstance* ContainerInstance = nullptr;
@@ -342,7 +336,7 @@ protected:
 
 	virtual bool GetChildContainerInstance(const FWorldPartitionActorDescInstance* InActorDescInstance, FContainerInstance& OutContainerInstance) const { return false; }
 	virtual UActorDescContainerInstance* CreateChildContainerInstance(const FWorldPartitionActorDescInstance* InActorDescInstance) const { return nullptr; }
-	virtual bool GetLoadedChildContainerInstance(const FWorldPartitionActorDescInstance* InActorDescInstance, FLoadedContainerInstance& OutContainerInstance) const { return false; }
+	virtual UWorldPartition* GetLoadedChildWorldPartition(const FWorldPartitionActorDescInstance* InActorDescInstance) const { return nullptr; }
 	ENGINE_API virtual bool IsRuntimeRelevant(const FWorldPartitionActorDescInstance* InActorDescInstance) const;
 	ENGINE_API virtual bool IsEditorRelevant(const FWorldPartitionActorDescInstance* InActorDescInstance) const;
 
diff --git a/Engine/Source/Runtime/Engine/Public/WorldPartition/WorldPartitionActorDescInstance.h b/Engine/Source/Runtime/Engine/Public/WorldPartition/WorldPartitionActorDescInstance.h
--- a/Engine/Source/Runtime/Engine/Public/WorldPartition/WorldPartitionActorDescInstance.h
+++ b/Engine/Source/Runtime/Engine/Public/WorldPartition/WorldPartitionActorDescInstance.h
@@ -27,6 +27,7 @@ class FWorldPartitionActorDescInstance : public IWorldPartitionActorDescInstance
 	friend struct FWorldPartitionActorDescUnitTestAcccessor;
 	friend class UDataLayerManager;
 	friend class UWorldPartition;
+	friend class IWorldPartitionActorLoaderInterface;
 
 protected:
 	ENGINE_API FWorldPartitionActorDescInstance();
@@ -99,8 +100,7 @@ public:
 	//~ End
 
 	inline UActorDescContainerInstance* GetChildContainerInstance() const { return ChildContainerInstance; }
-	bool GetLoadedChildContainerInstance(FWorldPartitionActorDesc::FLoadedContainerInstance& OutContainerInstance) const { return GetActorDesc()->GetLoadedChildContainerInstance(this, OutContainerInstance); }
-		
+
 	ENGINE_API virtual TWeakObjectPtr<AActor>* GetActorPtr(bool bEvenIfPendingKill = true, bool bEvenIfUnreachable = false) const;
 	ENGINE_API virtual bool IsValid() const;
 				
@@ -116,6 +116,7 @@ public:
 	ENGINE_API AActor* Load();
 	ENGINE_API void Unload();
 protected:
+	UWorldPartition* GetLoadedChildWorldPartition() const { return GetActorDesc()->GetLoadedChildWorldPartition(this); }
 	ENGINE_API void UpdateActorDesc(FWorldPartitionActorDesc* InActorDesc);
 	void Invalidate();
 		
diff --git a/Engine/Source/Runtime/Engine/Public/WorldPartition/WorldPartitionActorLoaderInterface.h b/Engine/Source/Runtime/Engine/Public/WorldPartition/WorldPartitionActorLoaderInterface.h
--- a/Engine/Source/Runtime/Engine/Public/WorldPartition/WorldPartitionActorLoaderInterface.h
+++ b/Engine/Source/Runtime/Engine/Public/WorldPartition/WorldPartitionActorLoaderInterface.h
@@ -75,6 +75,7 @@ public:
 		// Helpers
 		ENGINE_API FActorReferenceMap& GetContainerReferences(UActorDescContainerInstance* InContainerInstance);
 		ENGINE_API const FActorReferenceMap* GetContainerReferencesConst(UActorDescContainerInstance* InContainerInstance) const;
+		ENGINE_API UWorldPartition* GetLoadedChildWorldPartition(const FWorldPartitionHandle& Handle) const;
 
 	private:
 		UWorld* World;
@@ -102,9 +103,9 @@ public:
 
 	static ENGINE_API void RefreshLoadedState(bool bIsFromUserChange);
 
-	static ENGINE_API bool GetLoadedChildContainerInstance(const FWorldPartitionHandle& Handle, FWorldPartitionActorDesc::FLoadedContainerInstance& OutContainerInstance);
-
 private:
+	static ENGINE_API UWorldPartition* GetLoadedChildWorldPartition(const FWorldPartitionHandle& Handle);
+	
 	DECLARE_EVENT_OneParam(IWorldPartitionActorLoaderInterface, FOnActorLoaderInterfaceRefreshState, bool /*bIsFromUserChange*/);
 	static ENGINE_API FOnActorLoaderInterfaceRefreshState ActorLoaderInterfaceRefreshState;
 
