Level Instance: Move some code from OnLevelLoadedChanged to OnLevelStreamingStateChanged

Fixes trying to mark objects to early and skipping over actors like AWorldDataLayers which can still have AsyncLoading flag on its package when OnLevelLoadedChanged gets called (circular dep) which prevents calling ForEachObjectInPackage

This was causing some Actors/Objects to remain RF_Transactional and prevent proper unloading of Level Instance without clearing Undo stack

diff --git a/Engine/Source/Runtime/Engine/Private/LevelInstance/LevelInstanceLevelStreaming.cpp b/Engine/Source/Runtime/Engine/Private/LevelInstance/LevelInstanceLevelStreaming.cpp
--- a/Engine/Source/Runtime/Engine/Private/LevelInstance/LevelInstanceLevelStreaming.cpp
+++ b/Engine/Source/Runtime/Engine/Private/LevelInstance/LevelInstanceLevelStreaming.cpp
@@ -339,6 +339,24 @@ void ULevelStreamingLevelInstance::OnLevelStreamingStateChanged(UWorld* InWorld,
 		ULevel* Level = GetLoadedLevel();
 		check(GetLevelStreamingState() == ELevelStreamingState::LoadedVisible);
 
+		// Flag all Level Objects as non RF_Transactional & RF_Transient so that they can't be added to the Transaction Buffer and will be allowed to be unloaded / reloaded without clearing the Transaction Buffer
+		FLevelInstanceLevelStreamingUtils::MarkObjectsInPackageAsTransientAndNonTransactional(Level->GetPackage());
+
+		ForEachObjectWithOuter(Level, [](UObject* InObject)
+		{
+			// Skip actors as they are already handled in OnLoadedActorsAddedToLevelPreEvent
+			if (InObject && InObject->IsPackageExternal() && !InObject->IsA<AActor>())
+			{
+				FLevelInstanceLevelStreamingUtils::MarkObjectsInPackageAsTransientAndNonTransactional(InObject->GetPackage());
+			}
+		}, /*bIncludeNestedObjects*/ true);
+
+		OnLoadedActorsAddedToLevelPreEvent(Level->Actors);
+
+		Level->OnLoadedActorAddedToLevelPreEvent.AddUObject(this, &ULevelStreamingLevelInstance::OnLoadedActorsAddedToLevelPreEvent);
+		Level->OnLoadedActorAddedToLevelPostEvent.AddUObject(this, &ULevelStreamingLevelInstance::OnLoadedActorsAddedToLevelPostEvent);
+		Level->OnLoadedActorRemovedFromLevelPreEvent.AddUObject(this, &ULevelStreamingLevelInstance::OnLoadedActorsRemovedFromLevelPostEvent);
+
 		// Create special actor that will handle selection and transform
 		LevelInstanceEditorInstanceActor = ALevelInstanceEditorInstanceActor::Create(LevelInstance, Level);
 
@@ -402,28 +420,6 @@ void ULevelStreamingLevelInstance::OnLevelLoadedChanged(ULevel* InLevel)
 		if (IsEditorWorldMode())
 		{
 			FLevelStreamingDelegates::OnLevelStreamingStateChanged.AddUObject(this, &ULevelStreamingLevelInstance::OnLevelStreamingStateChanged);
-
-			// Most of the code here is meant to allow partial support for undo/redo of LevelInstance Instance Loading:
-			// by setting the objects RF_Transient and !RF_Transactional we can check when unloading if those flags
-			// have been changed and figure out if we need to clear the transaction buffer or not.
-			// It might not be the final solution to support Undo/Redo in LevelInstances but it handles most of the non-editing part
-			check(!InLevel->bAreComponentsCurrentlyRegistered);
-			FLevelInstanceLevelStreamingUtils::MarkObjectsInPackageAsTransientAndNonTransactional(NewLoadedLevel->GetPackage());
-
-			ForEachObjectWithOuter(NewLoadedLevel, [](UObject* InObject)
-			{
-				// Skip actors as they are already handled in OnLoadedActorsAddedToLevelPreEvent
-				if (InObject && InObject->IsPackageExternal() && !InObject->IsA<AActor>())
-				{
-					FLevelInstanceLevelStreamingUtils::MarkObjectsInPackageAsTransientAndNonTransactional(InObject->GetPackage());
-				}
-			}, /*bIncludeNestedObjects*/ true);
-
-			OnLoadedActorsAddedToLevelPreEvent(InLevel->Actors);
-
-			NewLoadedLevel->OnLoadedActorAddedToLevelPreEvent.AddUObject(this, &ULevelStreamingLevelInstance::OnLoadedActorsAddedToLevelPreEvent);
-			NewLoadedLevel->OnLoadedActorAddedToLevelPostEvent.AddUObject(this, &ULevelStreamingLevelInstance::OnLoadedActorsAddedToLevelPostEvent);
-			NewLoadedLevel->OnLoadedActorRemovedFromLevelPreEvent.AddUObject(this, &ULevelStreamingLevelInstance::OnLoadedActorsRemovedFromLevelPostEvent);
 		}
 #endif
 
