Fix Linux compile: missing forward declares
#rnx

diff --git a/Engine/Source/Runtime/Engine/Public/WorldPartition/WorldPartitionActorDesc.h b/Engine/Source/Runtime/Engine/Public/WorldPartition/WorldPartitionActorDesc.h
--- a/Engine/Source/Runtime/Engine/Public/WorldPartition/WorldPartitionActorDesc.h
+++ b/Engine/Source/Runtime/Engine/Public/WorldPartition/WorldPartitionActorDesc.h
@@ -38,6 +38,7 @@ struct FWorldPartitionAssetDataPatcher
 };
 
 class AActor;
+class UWorldPartition;
 class UActorDescContainerInstance;
 class FWorldPartitionActorDescInstance;
 class UActorDescContainer;
diff --git a/Engine/Source/Runtime/Engine/Public/WorldPartition/WorldPartitionActorDescInstance.h b/Engine/Source/Runtime/Engine/Public/WorldPartition/WorldPartitionActorDescInstance.h
--- a/Engine/Source/Runtime/Engine/Public/WorldPartition/WorldPartitionActorDescInstance.h
+++ b/Engine/Source/Runtime/Engine/Public/WorldPartition/WorldPartitionActorDescInstance.h
@@ -11,6 +11,7 @@
 #include "UObject/WeakObjectPtr.h"
 
 class AActor;
+class UWorldPartition;
 class UActorDescContainerInstance;
 
 #endif // WITH_EDITOR
