World Partition: Expose Per User Settings into Editor Preferences (under World Partition (Local) section)

diff --git a/Engine/Source/Editor/WorldPartitionEditor/Private/WorldPartition/Customizations/WorldPartitionEditorPerProjectUserSettingsDetails.cpp b/Engine/Source/Editor/WorldPartitionEditor/Private/WorldPartition/Customizations/WorldPartitionEditorPerProjectUserSettingsDetails.cpp
new file mode 100644
--- /dev/null
+++ b/Engine/Source/Editor/WorldPartitionEditor/Private/WorldPartition/Customizations/WorldPartitionEditorPerProjectUserSettingsDetails.cpp
@@ -0,0 +1,56 @@
+// Copyright Epic Games, Inc. All Rights Reserved.
+
+#include "WorldPartitionEditorPerProjectUserSettingsDetails.h"
+
+#include "DetailLayoutBuilder.h"
+#include "DetailCategoryBuilder.h"
+#include "DetailWidgetRow.h"
+#include "Widgets/SBoxPanel.h"
+#include "Widgets/Input/SButton.h"
+#include "Widgets/Text/STextBlock.h"
+#include "WorldPartition/WorldPartitionEditorPerProjectUserSettings.h"
+#include "Misc/MessageDialog.h"
+
+#define LOCTEXT_NAMESPACE "FWorldPartitionEditorPerProjectUserSettingsCustomization"
+
+TSharedRef<IDetailCustomization> FWorldPartitionEditorPerProjectUserSettingsCustomization::MakeInstance()
+{
+	return MakeShareable(new FWorldPartitionEditorPerProjectUserSettingsCustomization);
+}
+
+FReply FWorldPartitionEditorPerProjectUserSettingsCustomization::OnClearPerWorldSettings()
+{
+	const FText Message = LOCTEXT("ClearPerWorldSettingsWarningMessage", "This will Clear all local Per World Settings: Loaded Regions, Loaded Location Volumes, Loaded/Visibility Data Layers states.");
+	const FText Title = LOCTEXT("ClearPerWorldSettingsWarningMessageTitle", "Continue?");
+	if (FMessageDialog::Open(EAppMsgType::YesNo, Message, Title) == EAppReturnType::Yes)
+	{
+		GetMutableDefault<UWorldPartitionEditorPerProjectUserSettings>()->ClearPerWorldSettings();
+	}
+	return FReply::Handled();
+}
+
+void FWorldPartitionEditorPerProjectUserSettingsCustomization::CustomizeDetails(IDetailLayoutBuilder& InDetailBuilder)
+{
+	IDetailCategoryBuilder& Category = InDetailBuilder.EditCategory("Default");
+	Category.AddCustomRow(LOCTEXT("PerWorldSettingsRow", "Per World Settings"), false)
+	.NameContent()
+	[
+		SNew(STextBlock)
+		.Text(LOCTEXT("PerWorldSettingsText", "Per World Settings"))
+		.Font(InDetailBuilder.GetDetailFont())
+	]
+	.ValueContent()
+	[
+		SNew(SVerticalBox)
+		+ SVerticalBox::Slot()
+		[
+			SNew(SButton)
+			.Text(LOCTEXT("ClearPerWorldSettingsText", "Clear"))
+			.HAlign(EHorizontalAlignment::HAlign_Center)
+			.ToolTipText(LOCTEXT("ClearPerWorldSettingsTooltipText", "Clear Per World Settings (Loaded Regions, Loaded Location Volumes, Loaded/Visibility Data Layers states)"))
+			.OnClicked(this, &FWorldPartitionEditorPerProjectUserSettingsCustomization::OnClearPerWorldSettings)
+		]
+	];
+}
+
+#undef LOCTEXT_NAMESPACE
\ No newline at end of file
diff --git a/Engine/Source/Editor/WorldPartitionEditor/Private/WorldPartition/Customizations/WorldPartitionEditorPerProjectUserSettingsDetails.h b/Engine/Source/Editor/WorldPartitionEditor/Private/WorldPartition/Customizations/WorldPartitionEditorPerProjectUserSettingsDetails.h
new file mode 100644
--- /dev/null
+++ b/Engine/Source/Editor/WorldPartitionEditor/Private/WorldPartition/Customizations/WorldPartitionEditorPerProjectUserSettingsDetails.h
@@ -0,0 +1,23 @@
+// Copyright Epic Games, Inc. All Rights Reserved.
+
+#pragma once
+
+#include "CoreMinimal.h"
+#include "Input/Reply.h"
+#include "IDetailCustomization.h"
+
+class FWorldPartitionEditorPerProjectUserSettingsCustomization : public IDetailCustomization
+{
+public:
+	/** Makes a new instance of this detail layout class for a specific detail view requesting it */
+	static TSharedRef<IDetailCustomization> MakeInstance();
+
+protected:
+	FWorldPartitionEditorPerProjectUserSettingsCustomization() {}
+
+	FReply OnClearPerWorldSettings();
+
+private:
+	/** IDetailCustomization interface */
+	virtual void CustomizeDetails(IDetailLayoutBuilder& InDetailBuilder);
+};
\ No newline at end of file
diff --git a/Engine/Source/Editor/WorldPartitionEditor/Private/WorldPartitionEditorModule.cpp b/Engine/Source/Editor/WorldPartitionEditor/Private/WorldPartitionEditorModule.cpp
--- a/Engine/Source/Editor/WorldPartitionEditor/Private/WorldPartitionEditorModule.cpp
+++ b/Engine/Source/Editor/WorldPartitionEditor/Private/WorldPartitionEditorModule.cpp
@@ -13,6 +13,7 @@
 #include "WorldPartition/Customizations/WorldPartitionHLODDetailsCustomization.h"
 #include "WorldPartition/Customizations/WorldPartitionRuntimeSpatialHashDetailsCustomization.h"
 #include "WorldPartition/Customizations/WorldDataLayersActorDetails.h"
+#include "WorldPartition/Customizations/WorldPartitionEditorPerProjectUserSettingsDetails.h"
 #include "WorldPartition/SWorldPartitionConvertDialog.h"
 #include "WorldPartition/WorldPartitionConvertOptions.h"
 #include "WorldPartition/WorldPartitionEditorSettings.h"
@@ -206,6 +207,7 @@ void FWorldPartitionEditorModule::StartupModule()
 	PropertyEditor.RegisterCustomClassLayout("WorldPartitionRuntimeSpatialHash", FOnGetDetailCustomizationInstance::CreateStatic(&FWorldPartitionRuntimeSpatialHashDetails::MakeInstance));
 	PropertyEditor.RegisterCustomClassLayout("WorldPartitionHLOD", FOnGetDetailCustomizationInstance::CreateStatic(&FWorldPartitionHLODDetailsCustomization::MakeInstance));
 	PropertyEditor.RegisterCustomClassLayout("WorldDataLayers", FOnGetDetailCustomizationInstance::CreateStatic(&FWorldDataLayersActorDetails::MakeInstance));
+	PropertyEditor.RegisterCustomClassLayout("WorldPartitionEditorPerProjectUserSettings", FOnGetDetailCustomizationInstance::CreateStatic(&FWorldPartitionEditorPerProjectUserSettingsCustomization::MakeInstance));
 
 	FWorldPartitionClassDescRegistry().Get().Initialize();
 
diff --git a/Engine/Source/Runtime/Engine/Public/WorldPartition/WorldPartitionEditorPerProjectUserSettings.h b/Engine/Source/Runtime/Engine/Public/WorldPartition/WorldPartitionEditorPerProjectUserSettings.h
--- a/Engine/Source/Runtime/Engine/Public/WorldPartition/WorldPartitionEditorPerProjectUserSettings.h
+++ b/Engine/Source/Runtime/Engine/Public/WorldPartition/WorldPartitionEditorPerProjectUserSettings.h
@@ -8,6 +8,7 @@
 #include "UObject/Object.h"
 #include "UObject/SoftObjectPtr.h"
 #include "Engine/World.h"
+#include "Engine/DeveloperSettings.h"
 #include "WorldPartitionEditorPerProjectUserSettings.generated.h"
 
 USTRUCT()
@@ -43,8 +44,8 @@ struct FWorldPartitionPerWorldSettings
 #endif
 };
 
-UCLASS(config = EditorPerProjectUserSettings, MinimalAPI)
-class UWorldPartitionEditorPerProjectUserSettings : public UObject
+UCLASS(config = EditorPerProjectUserSettings, MinimalAPI, meta = (DisplayName = "World Partition (Local)"))
+class UWorldPartitionEditorPerProjectUserSettings : public UDeveloperSettings
 {
 	GENERATED_BODY()
 
@@ -118,11 +119,17 @@ public:
 		}
 	}
 
+	void ClearPerWorldSettings()
+	{
+		PerWorldEditorSettings.Empty();
+		SaveConfig();
+	}
+
 	ENGINE_API TArray<FName> GetWorldDataLayersNotLoadedInEditor(UWorld* InWorld) const;
 	ENGINE_API TArray<FName> GetWorldDataLayersLoadedInEditor(UWorld* InWorld) const;
 	
 	ENGINE_API void SetWorldDataLayersNonDefaultEditorLoadStates(UWorld* InWorld, const TArray<FName>& InDataLayersLoadedInEditor, const TArray<FName>& InDataLayersNotLoadedInEditor);
-
+		
 private:
 	const FWorldPartitionPerWorldSettings* GetWorldPartitionPerWorldSettings(UWorld* InWorld) const;
 #endif
@@ -130,31 +137,31 @@ private:
 #if WITH_EDITORONLY_DATA
 public:
 	/** True when the Data Layer Outliner is displaying Editor Data Layers */
-	UPROPERTY(config)
+	UPROPERTY(config, EditAnywhere, Category = "Data Layer")
 	uint32 bHideEditorDataLayers : 1;
 
 	/** True when the Data Layer Outliner is displaying Runtime Data Layers */
-	UPROPERTY(config)
+	UPROPERTY(config, EditAnywhere, Category = "Data Layer")
 	uint32 bHideRuntimeDataLayers : 1;
 
 	/** True when the Data Layer Outliner is not displaying actors */
-	UPROPERTY(config)
+	UPROPERTY(config, EditAnywhere, Category = "Data Layer")
 	uint32 bHideDataLayerActors : 1;
 
 	/** True when the Data Layer Outliner is not displaying unloaded actors */
-	UPROPERTY(config)
+	UPROPERTY(config, EditAnywhere, Category = "Data Layer")
 	uint32 bHideUnloadedActors : 1;
 
 	/** True when the Data Layer Outliner is only displaying actors and datalayers for selected actors */
-	UPROPERTY(config)
+	UPROPERTY(config, EditAnywhere, Category = "Data Layer")
 	uint32 bShowOnlySelectedActors : 1;
 
 	/** True when the Data Layer Outliner highlights Data Layers containing actors that are currently selected */
-	UPROPERTY(config)
+	UPROPERTY(config, EditAnywhere, Category = "Data Layer")
 	uint32 bHighlightSelectedDataLayers : 1;
 
 	/** True when the Data Layer Outliner is not displaying Level Instance content */
-	UPROPERTY(config)
+	UPROPERTY(config, EditAnywhere, Category = "Data Layer")
 	uint32 bHideLevelInstanceContent : 1;
 
 private:
@@ -162,16 +169,16 @@ private:
 
 	bool ShouldLoadSettings(const UWorld* InWorld) const;
 
-	UPROPERTY(config)
+	UPROPERTY(config, EditAnywhere, Category = Default)
 	uint32 bDisableLoadingOfLastLoadedRegions : 1;
 
-	UPROPERTY(config)
+	UPROPERTY(config, EditAnywhere, Category = Default)
 	uint32 bBugItGoLoadRegion : 1;
 
-	UPROPERTY(config)
+	UPROPERTY(config, EditAnywhere, Category = Default)
 	uint32 bShowCellCoords : 1;
 
-	UPROPERTY(config)
+	UPROPERTY(config, EditAnywhere, Category = Default, meta = (ClampMin=0.0, ClampMax=1.0))
 	float MinimapUnloadedOpacity;
 
 	UPROPERTY(config)
