Fix warning V1101: The default argument of the first virtual function parameter was changed in the derived class from '1' to '0'. Calling such function via a pointer or a reference to the base class may result in unexpected behavior.

diff --git a/Engine/Source/Runtime/Engine/Classes/Engine/Brush.h b/Engine/Source/Runtime/Engine/Classes/Engine/Brush.h
--- a/Engine/Source/Runtime/Engine/Classes/Engine/Brush.h
+++ b/Engine/Source/Runtime/Engine/Classes/Engine/Brush.h
@@ -154,7 +154,7 @@ public:
 	ENGINE_API virtual void PostEditChangeProperty(FPropertyChangedEvent& PropertyChangedEvent) override;
 	ENGINE_API virtual void PostEditChangeChainProperty(FPropertyChangedChainEvent& PropertyChangedEvent) override;
 	ENGINE_API virtual FName GetCustomIconName() const override;
-	ENGINE_API virtual bool Modify(bool bAlwaysMarkDirty = false) override;
+	ENGINE_API virtual bool Modify(bool bAlwaysMarkDirty = true) override;
 #endif // WITH_EDITOR
 
 	virtual bool NeedsLoadForClient() const override
diff --git a/Engine/Source/Runtime/Engine/Classes/Engine/Polys.h b/Engine/Source/Runtime/Engine/Classes/Engine/Polys.h
--- a/Engine/Source/Runtime/Engine/Classes/Engine/Polys.h
+++ b/Engine/Source/Runtime/Engine/Classes/Engine/Polys.h
@@ -429,7 +429,7 @@ class UPolys : public UObject
 
 	//~ Begin UObject Interface
 #if WITH_EDITOR
-	ENGINE_API virtual bool Modify(bool bAlwaysMarkDirty = false) override;	
+	ENGINE_API virtual bool Modify(bool bAlwaysMarkDirty = true) override;	
 #endif
 	ENGINE_API virtual void Serialize( FArchive& Ar ) override;
 	virtual bool IsAsset() const override { return false; }
diff --git a/Engine/Source/Runtime/Engine/Public/Model.h b/Engine/Source/Runtime/Engine/Public/Model.h
--- a/Engine/Source/Runtime/Engine/Public/Model.h
+++ b/Engine/Source/Runtime/Engine/Public/Model.h
@@ -479,7 +479,7 @@ public:
 	virtual void PreSave(FObjectPreSaveContext SaveContext) override;
 #if WITH_EDITOR
 	virtual void PostEditUndo() override;
-	virtual bool Modify(bool bAlwaysMarkDirty = false) override;
+	virtual bool Modify(bool bAlwaysMarkDirty = true) override;
 #endif // WITH_EDITOR
 	virtual bool Rename( const TCHAR* InName=NULL, UObject* NewOuter=NULL, ERenameFlags Flags=REN_None ) override;
 	static void AddReferencedObjects(UObject* InThis, FReferenceCollector& Collector);
