Remove long deprecated constructor causing warning V557: Array overrun is possible.

diff --git a/Engine/Source/Runtime/Engine/Private/Collision/Collision.cpp b/Engine/Source/Runtime/Engine/Private/Collision/Collision.cpp
--- a/Engine/Source/Runtime/Engine/Private/Collision/Collision.cpp
+++ b/Engine/Source/Runtime/Engine/Private/Collision/Collision.cpp
@@ -235,8 +235,6 @@ void FCollisionQueryParams::SetNumIgnoredComponents(int32 NewNum)
 //////////////////////////////////////////////////////////////////////////
 // FSeparatingAxisPointCheck
 
-TArray<FVector> FSeparatingAxisPointCheck::TriangleVertices;
-
 bool FSeparatingAxisPointCheck::TestSeparatingAxisCommon(const FVector& Axis, float ProjectedPolyMin, float ProjectedPolyMax)
 {
 	const float ProjectedCenter = FVector::DotProduct(Axis, BoxCenter);
diff --git a/Engine/Source/Runtime/Engine/Public/Collision.h b/Engine/Source/Runtime/Engine/Public/Collision.h
--- a/Engine/Source/Runtime/Engine/Public/Collision.h
+++ b/Engine/Source/Runtime/Engine/Public/Collision.h
@@ -77,31 +77,6 @@ public:
 		bHit = (PolyVertices.Num() == 3) ? FindSeparatingAxisTriangle() : FindSeparatingAxisGeneric();
 	}
 
-	/**
-	 *	Legacy constructor for the class (deprecated)
-	 */
-	FSeparatingAxisPointCheck(
-		const FVector& InV0,
-		const FVector& InV1,
-		const FVector& InV2,
-		const FVector& InBoxCenter,
-		const FVector& InBoxExtent,
-		float InBestDist
-		)
-		: HitNormal(FVector::ZeroVector),
-		  BestDist(InBestDist),
-		  PolyVertices(TriangleVertices),
-		  BoxCenter(InBoxCenter),
-		  BoxExtent(InBoxExtent),
-		  bCalcLeastPenetration(true)
-	{
-		TriangleVertices.Empty(3);
-		TriangleVertices[0] = InV0;
-		TriangleVertices[1] = InV1;
-		TriangleVertices[2] = InV2;
-		bHit = FindSeparatingAxisTriangle();
-	}
-
 private:
 
 	/**
@@ -158,9 +133,6 @@ private:
 
 	/** Flag specifying whether the least penetration should be calculated. */
 	bool bCalcLeastPenetration;
-
-	/** Array into which triangle vertices are placed (legacy use only) */
-	static ENGINE_API TArray<FVector> TriangleVertices;
 };
 
 /**
