Fix IPrimitiveComponent::GetUsedMaterials failing to return the values due to passing array by value instead of reference
#rnx

diff --git a/Engine/Source/Runtime/Engine/Classes/Components/ActorPrimitiveComponentInterface.h b/Engine/Source/Runtime/Engine/Classes/Components/ActorPrimitiveComponentInterface.h
--- a/Engine/Source/Runtime/Engine/Classes/Components/ActorPrimitiveComponentInterface.h
+++ b/Engine/Source/Runtime/Engine/Classes/Components/ActorPrimitiveComponentInterface.h
@@ -16,7 +16,7 @@ public:
 	UWorld* GetWorld() const override;
 	FSceneInterface* GetScene() const override;
 	FPrimitiveSceneProxy* GetSceneProxy() const override;	
-	void GetUsedMaterials(TArray<UMaterialInterface*> OutMaterials, bool bGetDebugMaterials) const;
+	void GetUsedMaterials(TArray<UMaterialInterface*>& OutMaterials, bool bGetDebugMaterials) const;
 	void MarkRenderStateDirty() override;
 	void DestroyRenderState() override;
 	void CreateRenderState(FRegisterComponentContext* Context) override;
diff --git a/Engine/Source/Runtime/Engine/Classes/Components/ComponentInterfaces.h b/Engine/Source/Runtime/Engine/Classes/Components/ComponentInterfaces.h
--- a/Engine/Source/Runtime/Engine/Classes/Components/ComponentInterfaces.h
+++ b/Engine/Source/Runtime/Engine/Classes/Components/ComponentInterfaces.h
@@ -38,7 +38,7 @@ public:
 	virtual UWorld* GetWorld() const = 0;
 	virtual FSceneInterface* GetScene() const = 0;
 	virtual FPrimitiveSceneProxy* GetSceneProxy() const = 0;
-	virtual void GetUsedMaterials(TArray<UMaterialInterface*> OutMaterials, bool bGetDebugMaterials = false) const = 0;		
+	virtual void GetUsedMaterials(TArray<UMaterialInterface*>& OutMaterials, bool bGetDebugMaterials = false) const = 0;		
 	virtual void MarkRenderStateDirty() = 0;
 	virtual void DestroyRenderState() = 0;
 	virtual void CreateRenderState(FRegisterComponentContext* Context) = 0;	
diff --git a/Engine/Source/Runtime/Engine/Private/Components/PrimitiveComponent.cpp b/Engine/Source/Runtime/Engine/Private/Components/PrimitiveComponent.cpp
--- a/Engine/Source/Runtime/Engine/Private/Components/PrimitiveComponent.cpp
+++ b/Engine/Source/Runtime/Engine/Private/Components/PrimitiveComponent.cpp
@@ -4920,7 +4920,7 @@ FPrimitiveSceneProxy* FActorPrimitiveComponentInterface::GetSceneProxy() const
 	return UPrimitiveComponent::GetPrimitiveComponent(this)->SceneProxy;
 }
 
-void FActorPrimitiveComponentInterface::GetUsedMaterials(TArray<UMaterialInterface*> OutMaterials, bool bGetDebugMaterials) const
+void FActorPrimitiveComponentInterface::GetUsedMaterials(TArray<UMaterialInterface*>& OutMaterials, bool bGetDebugMaterials) const
 {
 	UPrimitiveComponent::GetPrimitiveComponent(this)->GetUsedMaterials(OutMaterials, bGetDebugMaterials);
 }
