- Recache PSOs on all components when FGlobalComponentRecreateRenderStateContext is used because parameters which drive MDC creation can influence used PSOs and might need to be precached/compiled
- Add recursion check on FGlobalComponentRecreateRenderStateContext to reduce multiple calls when applying game user settings - render state and PSO precaching will only called on final context goes out of scope

diff --git a/Engine/Source/Runtime/Engine/Classes/Components/ActorComponent.h b/Engine/Source/Runtime/Engine/Classes/Components/ActorComponent.h
--- a/Engine/Source/Runtime/Engine/Classes/Components/ActorComponent.h
+++ b/Engine/Source/Runtime/Engine/Classes/Components/ActorComponent.h
@@ -686,6 +686,9 @@ protected:
 	 */
 	ENGINE_API virtual void OnUnregister();
 
+	/** Precache all PSOs which can be used by the actor component */
+	virtual void PrecachePSOs() {}
+
 	/** Return true if CreateRenderState() should be called */
 	virtual bool ShouldCreateRenderState() const 
 	{
diff --git a/Engine/Source/Runtime/Engine/Classes/Components/PrimitiveComponent.h b/Engine/Source/Runtime/Engine/Classes/Components/PrimitiveComponent.h
--- a/Engine/Source/Runtime/Engine/Classes/Components/PrimitiveComponent.h
+++ b/Engine/Source/Runtime/Engine/Classes/Components/PrimitiveComponent.h
@@ -914,7 +914,7 @@ public:
 	virtual void CollectPSOPrecacheData(const FPSOPrecacheParams& BasePrecachePSOParams, FComponentPSOPrecacheParamsList& OutParams) {}
 
 	/** Precache all PSOs which can be used by the primitive component */
-	ENGINE_API virtual void PrecachePSOs();
+	ENGINE_API virtual void PrecachePSOs() override;
 
 	/** Schedule task to mark render state dirty when the PSO precaching tasks are done */
 	ENGINE_API void RequestRecreateRenderStateWhenPSOPrecacheFinished(const FGraphEventArray& PSOPrecacheCompileEvents);
diff --git a/Engine/Source/Runtime/Engine/Classes/Components/TextRenderComponent.h b/Engine/Source/Runtime/Engine/Classes/Components/TextRenderComponent.h
--- a/Engine/Source/Runtime/Engine/Classes/Components/TextRenderComponent.h
+++ b/Engine/Source/Runtime/Engine/Classes/Components/TextRenderComponent.h
@@ -178,6 +178,7 @@ class UTextRenderComponent : public UPrimitiveComponent
 
 	//~ Begin UActorComponent Interface.
 	ENGINE_API virtual bool RequiresGameThreadEndOfFrameUpdates() const override;
+	ENGINE_API virtual void PrecachePSOs() override;
 	//~ End UActorComponent Interface.
 
 	//~ Begin UObject Interface.
diff --git a/Engine/Source/Runtime/Engine/Classes/Materials/MaterialInterface.h b/Engine/Source/Runtime/Engine/Classes/Materials/MaterialInterface.h
--- a/Engine/Source/Runtime/Engine/Classes/Materials/MaterialInterface.h
+++ b/Engine/Source/Runtime/Engine/Classes/Materials/MaterialInterface.h
@@ -1054,6 +1054,11 @@ public:
 	 */
 	ENGINE_API static void InitDefaultMaterials();
 
+	/**
+	 * Precache PSOs for all default materials.
+	 */
+	ENGINE_API static void PrecacheDefaultMaterialPSOs();
+
 	/** Checks to see if an input property should be active, based on the state of the material */
 	ENGINE_API virtual bool IsPropertyActive(EMaterialProperty InProperty) const;
 
diff --git a/Engine/Source/Runtime/Engine/Private/Components/ActorComponent.cpp b/Engine/Source/Runtime/Engine/Private/Components/ActorComponent.cpp
--- a/Engine/Source/Runtime/Engine/Private/Components/ActorComponent.cpp
+++ b/Engine/Source/Runtime/Engine/Private/Components/ActorComponent.cpp
@@ -37,6 +37,8 @@
 #include "SceneInterface.h"
 #include "UObject/FortniteReleaseBranchCustomObjectVersion.h"
 #include "UObject/FrameworkObjectVersion.h"
+#include "PSOPrecacheMaterial.h"
+#include "Materials/MaterialInterface.h"
 
 #if WITH_EDITOR
 #include "Kismet2/ComponentEditorUtils.h"
@@ -89,6 +91,10 @@ FAutoConsoleVariableRef GTickComponentLatentActionsWithTheComponentCVar(
 /** Static var indicating activity of reregister context */
 int32 FGlobalComponentReregisterContext::ActiveGlobalReregisterContextCount = 0;
 
+/** Static var indicating activity of recreate render state context */
+int32 FGlobalComponentRecreateRenderStateContext::ActiveGlobalRecreateRenderStateContextCount = 0;
+
+
 bool GDefaultUseSubObjectReplicationList = false;
 static FAutoConsoleVariableRef CVarDefaultUseSubObjectReplicationList(
 	TEXT("net.SubObjects.DefaultUseSubObjectReplicationList"),
@@ -298,6 +304,8 @@ FGlobalComponentRecreateRenderStateContext::FGlobalComponentRecreateRenderStateC
 	{
 		TRACE_CPUPROFILER_EVENT_SCOPE(FGlobalComponentRecreateRenderStateContext::FGlobalComponentRecreateRenderStateContext);
 
+		ActiveGlobalRecreateRenderStateContextCount++;
+
 		// wait until resources are released
 		FlushRenderingCommands();
 
@@ -316,7 +324,7 @@ FGlobalComponentRecreateRenderStateContext::FGlobalComponentRecreateRenderStateC
 
 FGlobalComponentRecreateRenderStateContext::FGlobalComponentRecreateRenderStateContext(const TArray<UActorComponent*>& InComponents)
 {
-	if (FApp::CanEverRender())
+	if (FApp::CanEverRender() && ++ActiveGlobalRecreateRenderStateContextCount == 1)
 	{
 		TRACE_CPUPROFILER_EVENT_SCOPE(FGlobalComponentRecreateRenderStateContext::FGlobalComponentRecreateRenderStateContext);
 
@@ -340,9 +348,22 @@ FGlobalComponentRecreateRenderStateContext::~FGlobalComponentRecreateRenderState
 {
 	TRACE_CPUPROFILER_EVENT_SCOPE(FGlobalComponentRecreateRenderStateContext::~FGlobalComponentRecreateRenderStateContext);
 
-	ComponentContexts.Empty();
+	if (FApp::CanEverRender())
+	{
+		check(ActiveGlobalRecreateRenderStateContextCount > 0);
+
+		// Check if this is the last active context
+		if (--ActiveGlobalRecreateRenderStateContextCount == 0)
+		{
+			// Clear the PSO material request cache to make sure PSO collection happens again on possible changed data
+			ClearMaterialPSORequests();
+			UMaterialInterface::PrecacheDefaultMaterialPSOs();
 
-	UpdateAllPrimitiveSceneInfos();
+			ComponentContexts.Empty();
+
+			UpdateAllPrimitiveSceneInfos();
+		}
+	}
 }
 
 void FGlobalComponentRecreateRenderStateContext::UpdateAllPrimitiveSceneInfos()
diff --git a/Engine/Source/Runtime/Engine/Private/Components/TextRenderComponent.cpp b/Engine/Source/Runtime/Engine/Private/Components/TextRenderComponent.cpp
--- a/Engine/Source/Runtime/Engine/Private/Components/TextRenderComponent.cpp
+++ b/Engine/Source/Runtime/Engine/Private/Components/TextRenderComponent.cpp
@@ -1520,6 +1520,13 @@ void UTextRenderComponent::PostLoad()
 		}
 	}
 
+	PrecachePSOs();
+
+	Super::PostLoad();
+}
+
+void UTextRenderComponent::PrecachePSOs()
+{
 	if (IsComponentPSOPrecachingEnabled() && TextMaterial
 		// FIXME: need to collect an actual vertex declaration for non-MVF path
 		&& RHISupportsManualVertexFetch(GMaxRHIShaderPlatform))
@@ -1533,8 +1540,6 @@ void UTextRenderComponent::PostLoad()
 
 		TextMaterial->PrecachePSOs(&FLocalVertexFactory::StaticType, PrecachePSOParams);
 	}
-
-	Super::PostLoad();
 }
 
 void UTextRenderComponent::InitializeMIDCache()
diff --git a/Engine/Source/Runtime/Engine/Private/GameEngine.cpp b/Engine/Source/Runtime/Engine/Private/GameEngine.cpp
--- a/Engine/Source/Runtime/Engine/Private/GameEngine.cpp
+++ b/Engine/Source/Runtime/Engine/Private/GameEngine.cpp
@@ -65,6 +65,7 @@
 #include "RenderTargetPool.h"
 #include "RenderGraphBuilder.h"
 #include "CustomResourcePool.h"
+#include "ComponentRecreateRenderStateContext.h"
 
 #if WITH_EDITOR
 #include "PIEPreviewDeviceProfileSelectorModule.h"
@@ -1119,6 +1120,10 @@ void UGameEngine::Init(IEngineLoop* InEngineLoop)
 	// Load and apply user game settings
 	{
 		TRACE_CPUPROFILER_EVENT_SCOPE(InitGameUserSettings);
+
+		// Push recreate render state context to force single recreate instead of multiple recreates for each changed cvar
+		FGlobalComponentRecreateRenderStateContext Context;
+
 		GetGameUserSettings()->LoadSettings();
 		GetGameUserSettings()->ApplyNonResolutionSettings();
 	}
diff --git a/Engine/Source/Runtime/Engine/Private/GameUserSettings.cpp b/Engine/Source/Runtime/Engine/Private/GameUserSettings.cpp
--- a/Engine/Source/Runtime/Engine/Private/GameUserSettings.cpp
+++ b/Engine/Source/Runtime/Engine/Private/GameUserSettings.cpp
@@ -17,6 +17,7 @@
 #include "HAL/PlatformFramePacer.h"
 #include "HDRHelper.h"
 #include "UnrealClient.h"
+#include "ComponentRecreateRenderStateContext.h"
 
 // For sandboxing migration of UserSettings from AppData\Local to AppData\LocalLow.
 #if PLATFORM_WINDOWS
@@ -505,8 +506,12 @@ void UGameUserSettings::ApplyResolutionSettings(bool bCheckForCommandLineOverrid
 
 void UGameUserSettings::ApplySettings(bool bCheckForCommandLineOverrides)
 {
-	ApplyResolutionSettings(bCheckForCommandLineOverrides);
-	ApplyNonResolutionSettings();
+	{
+		// Push recreate render state context to force single recreate instead of multiple recreates for each changed cvar
+		FGlobalComponentRecreateRenderStateContext Context;
+		ApplyResolutionSettings(bCheckForCommandLineOverrides);
+		ApplyNonResolutionSettings();
+	}
 	RequestUIUpdate();
 
 	SaveSettings();
diff --git a/Engine/Source/Runtime/Engine/Private/Materials/Material.cpp b/Engine/Source/Runtime/Engine/Private/Materials/Material.cpp
--- a/Engine/Source/Runtime/Engine/Private/Materials/Material.cpp
+++ b/Engine/Source/Runtime/Engine/Private/Materials/Material.cpp
@@ -585,28 +585,33 @@ void UMaterialInterface::InitDefaultMaterials()
 		// Skip platforms that do not support MVF, non-MVF path needs mesh information for PSO 
 		if (bInitialized && PipelineStateCache::IsPSOPrecachingEnabled() && RHISupportsManualVertexFetch(GMaxRHIShaderPlatform))
 		{
-			TArray<FMaterialPSOPrecacheRequestID> MaterialPrecacheRequestIDs;
+			PrecacheDefaultMaterialPSOs();
+		}
+	}
+}
 
-			FPSOPrecacheParams PrecachePSOParams;
-			PrecachePSOParams.bDefaultMaterial = true;
-			FPSOPrecacheVertexFactoryDataList AllVertexFactoryTypes;
-			for (TLinkedList<FVertexFactoryType*>::TIterator It(FVertexFactoryType::GetTypeList()); It; It.Next())
-			{
-				FPSOPrecacheVertexFactoryData VFData;
-				VFData.VertexFactoryType = *It;
-				AllVertexFactoryTypes.Add(VFData);
-			}
-			for (int32 Domain = 0; Domain < MD_MAX; ++Domain)
-			{
-				if (GDefaultMaterials[Domain])
-				{
-					PrecachePSOParams.Mobility = EComponentMobility::Static;
-					GDefaultMaterials[Domain]->PrecachePSOs(AllVertexFactoryTypes, PrecachePSOParams, EPSOPrecachePriority::High, MaterialPrecacheRequestIDs);
+void UMaterialInterface::PrecacheDefaultMaterialPSOs()
+{
+	TArray<FMaterialPSOPrecacheRequestID> MaterialPrecacheRequestIDs;
 
-					PrecachePSOParams.Mobility = EComponentMobility::Movable;
-					GDefaultMaterials[Domain]->PrecachePSOs(AllVertexFactoryTypes, PrecachePSOParams, EPSOPrecachePriority::High, MaterialPrecacheRequestIDs);
-				}
-			}
+	FPSOPrecacheParams PrecachePSOParams;
+	PrecachePSOParams.bDefaultMaterial = true;
+	FPSOPrecacheVertexFactoryDataList AllVertexFactoryTypes;
+	for (TLinkedList<FVertexFactoryType*>::TIterator It(FVertexFactoryType::GetTypeList()); It; It.Next())
+	{
+		FPSOPrecacheVertexFactoryData VFData;
+		VFData.VertexFactoryType = *It;
+		AllVertexFactoryTypes.Add(VFData);
+	}
+	for (int32 Domain = 0; Domain < MD_MAX; ++Domain)
+	{
+		if (GDefaultMaterials[Domain])
+		{
+			PrecachePSOParams.Mobility = EComponentMobility::Static;
+			GDefaultMaterials[Domain]->PrecachePSOs(AllVertexFactoryTypes, PrecachePSOParams, EPSOPrecachePriority::High, MaterialPrecacheRequestIDs);
+
+			PrecachePSOParams.Mobility = EComponentMobility::Movable;
+			GDefaultMaterials[Domain]->PrecachePSOs(AllVertexFactoryTypes, PrecachePSOParams, EPSOPrecachePriority::High, MaterialPrecacheRequestIDs);
 		}
 	}
 }
diff --git a/Engine/Source/Runtime/Engine/Private/Materials/MaterialShared.cpp b/Engine/Source/Runtime/Engine/Private/Materials/MaterialShared.cpp
--- a/Engine/Source/Runtime/Engine/Private/Materials/MaterialShared.cpp
+++ b/Engine/Source/Runtime/Engine/Private/Materials/MaterialShared.cpp
@@ -3171,6 +3171,12 @@ TArray<FMaterialPSOPrecacheRequestID> FMaterial::GetMaterialPSOPrecacheRequestID
 	return TmpPrecachedPSORequestIDs;
 }
 
+void FMaterial::ClearPrecachedPSORequestIDs()
+{
+	FScopeLock ScopeLock(&PrecachedPSORequestIDsCS);
+	PrecachedPSORequestIDs.Empty();
+}
+
 #if WITH_EDITOR
 
 void FMaterial::BeginCacheShaders(EShaderPlatform Platform, EMaterialShaderPrecompileMode PrecompileMode, const ITargetPlatform* TargetPlatform, TUniqueFunction<void(bool bSuccess)>&& CompletionCallback)
diff --git a/Engine/Source/Runtime/Engine/Private/PSOPrecacheMaterial.cpp b/Engine/Source/Runtime/Engine/Private/PSOPrecacheMaterial.cpp
--- a/Engine/Source/Runtime/Engine/Private/PSOPrecacheMaterial.cpp
+++ b/Engine/Source/Runtime/Engine/Private/PSOPrecacheMaterial.cpp
@@ -15,6 +15,14 @@
 #include "VertexFactory.h"
 #include "SceneInterface.h"
 
+int32 GPSOUseBackgroundThreadForCollection = 1;
+static FAutoConsoleVariableRef CVarPSOUseBackgroundThreadForCollection(
+	TEXT("r.PSOPrecache.UseBackgroundThreadForCollection"),
+	GPSOUseBackgroundThreadForCollection,
+	TEXT("Use background threads for PSO precache data collection on the mesh pass processors.\n"),
+	ECVF_ReadOnly
+);
+
 FPSOCollectorCreateManager::FPSOCollectorData FPSOCollectorCreateManager::PSOCollectors[(int32)EShadingPath::Num][FPSOCollectorCreateManager::MaxPSOCollectorCount] = {};
 
 int32 FPSOCollectorCreateManager::GetIndex(EShadingPath ShadingPath, const TCHAR* Name)
@@ -119,10 +127,12 @@ public:
 	explicit FMaterialPSOPrecacheCollectionTask(
 		TStrongObjectPtr<UMaterialInterface>* InMaterialInterface,
 		const FMaterialPSOPrecacheParams& InPrecacheParams,
-		FGraphEventRef& InCollectionGraphEvent)
+		FGraphEventRef& InCollectionGraphEvent,
+		uint32 InRequestLifecycleID)
 		: MaterialInterface(InMaterialInterface)
 		, PrecacheParams(InPrecacheParams)
 		, CollectionGraphEvent(InCollectionGraphEvent)
+		, RequestLifecycleID(InRequestLifecycleID)
 	{
 	}
 
@@ -133,6 +143,7 @@ public:
 	TStrongObjectPtr<UMaterialInterface>* MaterialInterface;
 	FMaterialPSOPrecacheParams PrecacheParams;
 	FGraphEventRef CollectionGraphEvent;
+	uint32 RequestLifecycleID;	
 
 	static ESubsequentsMode::Type	GetSubsequentsMode() { return ESubsequentsMode::TrackSubsequents; }
 	ENamedThreads::Type				GetDesiredThread() { return ENamedThreads::AnyBackgroundThreadNormalTask; }
@@ -165,7 +176,7 @@ public:
 		// Offload to background job task graph if threading is enabled
 		// Don't use background thread in editor because shader maps and material resources could be destroyed while the task is running
 		// If it's a perf problem at some point then FMaterialPSOPrecacheRequestID has to be used at material level in the correct places to wait for
-		bool bUseBackgroundTask = FApp::ShouldUseThreadingForPerformance() && !GIsEditor;
+		bool bUseBackgroundTask = GPSOUseBackgroundThreadForCollection && FApp::ShouldUseThreadingForPerformance() && !GIsEditor;
 
 		FGraphEventRef CollectionGraphEvent;
 
@@ -216,10 +227,11 @@ public:
 					PrecacheData.CollectionGraphEvent = CollectionGraphEvent;
 
 					// Create task the clear mark fully complete in the cache when done
+					uint32 RequestLifecycleID = LifecycleID;
 					FFunctionGraphTask::CreateAndDispatchWhenReady(
-						[this, Params]
+						[this, Params, RequestLifecycleID]
 						{
-							MarkCompilationComplete(Params);
+							MarkCompilationComplete(Params, RequestLifecycleID);
 						},
 						TStatId{}, CollectionGraphEvent
 						);
@@ -234,7 +246,7 @@ public:
 			TStrongObjectPtr<UMaterialInterface>* MaterialInterface = new TStrongObjectPtr<UMaterialInterface>(Params.Material->GetMaterialInterface());
 
 			// Create and kick the collection task
-			TGraphTask<FMaterialPSOPrecacheCollectionTask>::CreateTask().ConstructAndDispatchWhenReady(MaterialInterface, Params, CollectionGraphEvent);
+			TGraphTask<FMaterialPSOPrecacheCollectionTask>::CreateTask().ConstructAndDispatchWhenReady(MaterialInterface, Params, CollectionGraphEvent, LifecycleID);
 			
 			// Need to wait for collection task which will be extented during run with the actual async compile events
 			OutGraphEvents.Add(CollectionGraphEvent);
@@ -248,7 +260,7 @@ public:
 			FPSOPrecacheRequestResultArray PrecacheResults = RequestPrecachePSOs(PSOPrecacheData);
 						
 			// Mark collection complete
-			MarkCollectionComplete(Params, PSOPrecacheData, PrecacheResults);
+			MarkCollectionComplete(Params, PSOPrecacheData, PrecacheResults, LifecycleID);
 			
 			// Add the graph events to wait for
 			for (FPSOPrecacheRequestResult& Result : PrecacheResults)
@@ -261,11 +273,18 @@ public:
 		return RequestID;
 	}
 
-	void MarkCollectionComplete(const FMaterialPSOPrecacheParams& Params, const FPSOPrecacheDataArray& PrecacheData, const FPSOPrecacheRequestResultArray& PrecacheRequestResults)
+	void MarkCollectionComplete(const FMaterialPSOPrecacheParams& Params, const FPSOPrecacheDataArray& PrecacheData, const FPSOPrecacheRequestResultArray& PrecacheRequestResults, uint32 RequestLifecycleID)
 	{
 		FRWScopeLock WriteLock(RWLock, SLT_Write);
+		
+		// Ignore requests not coming from current life cycle ID
+		if (RequestLifecycleID != LifecycleID)
+		{
+			return;
+		}
+
 		FPrecacheData* FindResult = MaterialPSORequestData.Find(Params);
-		check(FindResult && FindResult->State == EState::Collecting);
+		check(FindResult->State == EState::Collecting);
 		check(FindResult->ActivePSOPrecacheRequests.IsEmpty());
 		FindResult->ActivePSOPrecacheRequests = PrecacheRequestResults;
 #if PSO_PRECACHING_TRACKING
@@ -306,6 +325,12 @@ public:
 
 		{
 			FRWScopeLock ReadLock(RWLock, SLT_ReadOnly);
+
+			if (MaterialPSORequestID >= (uint32)MaterialPSORequests.Num())
+			{
+				return;
+			}
+
 			const FMaterialPSOPrecacheParams& Params = MaterialPSORequests[MaterialPSORequestID];
 			FPrecacheData* FindResult = MaterialPSORequestData.Find(Params);
 			check(FindResult);
@@ -327,13 +352,43 @@ public:
 		CheckCompilingPSOs(*FindResult, true /*bBoostPriority*/);
 	}
 
+	void ClearMaterialPSORequests()
+	{
+		check(IsInGameThread());
+
+		FRWScopeLock WriteLock(RWLock, SLT_Write);
+
+		// Increment the life cycle ID - all current active collection tasks are 'not important' anymore and can either be skipped or ignored
+		LifecycleID++;
+
+		TSet<FMaterial*> Materials;
+		for (FMaterialPSOPrecacheParams& Params : MaterialPSORequests)
+		{
+			if (Params.Material)
+			{
+				Materials.Add(Params.Material);
+			}
+		}
+
+		for (FMaterial* Material : Materials)
+		{
+			Material->ClearPrecachedPSORequestIDs();
+		}
+
+		// Clear the current cached pso requests so we gather the PSOs to compile again (usually called on cvar changes which could influence MDC and thus PSOs)			
+		MaterialPSORequests.Empty(MaterialPSORequests.Num());
+		MaterialPSORequestData.Empty(MaterialPSORequestData.Num());
+	}
+
+	uint32 GetLifecycleID() const { return LifecycleID; }
+
 #if PSO_PRECACHING_TRACKING
 
 	FMaterialPSOPrecacheParams GetMaterialPSOPrecacheParams(FMaterialPSOPrecacheRequestID MaterialPSORequestID)
 	{
 		check(MaterialPSORequestID != INDEX_NONE);
 
-		FRWScopeLock ReadLock(RWLock, SLT_Write);
+		FRWScopeLock ReadLock(RWLock, SLT_ReadOnly);
 		return MaterialPSORequests[MaterialPSORequestID];
 	}
 
@@ -341,7 +396,7 @@ public:
 	{
 		check(MaterialPSORequestID != INDEX_NONE);
 
-		FRWScopeLock ReadLock(RWLock, SLT_Write);
+		FRWScopeLock ReadLock(RWLock, SLT_ReadOnly);
 		const FMaterialPSOPrecacheParams&  Params = MaterialPSORequests[MaterialPSORequestID];
 		FPrecacheData* FindResult = MaterialPSORequestData.Find(Params);
 		check(FindResult);
@@ -412,11 +467,11 @@ private:
 		return (PrecacheData.State != EState::Completed);
 	}	
 
-	void MarkCompilationComplete(const FMaterialPSOPrecacheParams& Params)
+	void MarkCompilationComplete(const FMaterialPSOPrecacheParams& Params, uint32 RequestLifecycleID)
 	{
 		FRWScopeLock WriteLock(RWLock, SLT_Write);
 		FPrecacheData* FindResult = MaterialPSORequestData.Find(Params);
-		if (FindResult)
+		if (FindResult && RequestLifecycleID == LifecycleID)
 		{
 			verify(!CheckCompilingPSOs(*FindResult, false /*bBoostPriority*/));
 		}
@@ -425,6 +480,7 @@ private:
 	FRWLock RWLock;
 	TArray<FMaterialPSOPrecacheParams> MaterialPSORequests;	
 	TMap<FMaterialPSOPrecacheParams, FPrecacheData> MaterialPSORequestData;
+	uint32 LifecycleID = 0; //< ID to check current outstanding requests are still valid - incremented on re-precache of all current requests
 };
 
 // The global request manager - only used locally in a few global function to precache, release or boost PSO precache requests
@@ -434,6 +490,13 @@ void FMaterialPSOPrecacheCollectionTask::DoTask(ENamedThreads::Type CurrentThrea
 {
 	TRACE_CPUPROFILER_EVENT_SCOPE(FMaterialPSOPrecacheCollectionTask);
 
+	// Make sure task is still relevant
+	if (RequestLifecycleID != GMaterialPSORequestManager.GetLifecycleID())
+	{
+		CollectionGraphEvent->DispatchSubsequents();
+		return;
+	}
+
 	FTaskTagScope ParallelGTScope(ETaskTag::EParallelGameThread);
 
 	// Collect pso data
@@ -450,7 +513,7 @@ void FMaterialPSOPrecacheCollectionTask::DoTask(ENamedThreads::Type CurrentThrea
 	FPSOPrecacheRequestResultArray PrecacheResults = RequestPrecachePSOs(PSOPrecacheData);
 
 	// Mark collection complete
-	GMaterialPSORequestManager.MarkCollectionComplete(PrecacheParams, PSOPrecacheData, PrecacheResults);
+	GMaterialPSORequestManager.MarkCollectionComplete(PrecacheParams, PSOPrecacheData, PrecacheResults, RequestLifecycleID);
 
 	// Extend MyCompletionGraphEvent to wait for all the async compile events
 	if (PrecacheResults.Num() > 0)
@@ -488,6 +551,13 @@ void BoostPSOPriority(const TArray<FMaterialPSOPrecacheRequestID>& MaterialPSORe
 	}
 }
 
+void ClearMaterialPSORequests()
+{
+	TRACE_CPUPROFILER_EVENT_SCOPE(ClearMaterialPSORequests);
+
+	return GMaterialPSORequestManager.ClearMaterialPSORequests();
+}
+
 #if PSO_PRECACHING_TRACKING
 
 FMaterialPSOPrecacheParams GetMaterialPSOPrecacheParams(FMaterialPSOPrecacheRequestID RequestID)
diff --git a/Engine/Source/Runtime/Engine/Public/ComponentRecreateRenderStateContext.h b/Engine/Source/Runtime/Engine/Public/ComponentRecreateRenderStateContext.h
--- a/Engine/Source/Runtime/Engine/Public/ComponentRecreateRenderStateContext.h
+++ b/Engine/Source/Runtime/Engine/Public/ComponentRecreateRenderStateContext.h
@@ -77,6 +77,7 @@ public:
 	{
 		if (Component && !Component->IsRenderStateCreated() && Component->IsRegistered())
 		{
+			Component->PrecachePSOs();
 			Component->CreateRenderState_Concurrent(nullptr);
 
 			UpdateAllPrimitiveSceneInfosForSingleComponent(Component, ScenesToUpdateAllPrimitiveSceneInfos);
@@ -109,6 +110,9 @@ public:
 	/** Destructor */
 	ENGINE_API ~FGlobalComponentRecreateRenderStateContext();
 
+	/** Indicates that a FGlobalComponentRecreateRenderStateContext is currently active */
+	static int32 ActiveGlobalRecreateRenderStateContextCount;
+
 private:
 	/** The recreate contexts for the individual components. */
 	TArray<FComponentRecreateRenderStateContext> ComponentContexts;
diff --git a/Engine/Source/Runtime/Engine/Public/MaterialShared.h b/Engine/Source/Runtime/Engine/Public/MaterialShared.h
--- a/Engine/Source/Runtime/Engine/Public/MaterialShared.h
+++ b/Engine/Source/Runtime/Engine/Public/MaterialShared.h
@@ -1848,6 +1848,11 @@ public:
 	 */
 	ENGINE_API TArray<FMaterialPSOPrecacheRequestID> GetMaterialPSOPrecacheRequestIDs() const;
 
+	/**
+	 * Clear all cached PSO data because manager has been reset
+	 */
+	void ClearPrecachedPSORequestIDs();
+
 	/**
 	 * Should the shader for this material with the given platform, shader type and vertex 
 	 * factory type combination be compiled
diff --git a/Engine/Source/Runtime/Engine/Public/PSOPrecacheMaterial.h b/Engine/Source/Runtime/Engine/Public/PSOPrecacheMaterial.h
--- a/Engine/Source/Runtime/Engine/Public/PSOPrecacheMaterial.h
+++ b/Engine/Source/Runtime/Engine/Public/PSOPrecacheMaterial.h
@@ -136,6 +136,11 @@ extern ENGINE_API void ReleasePSOPrecacheData(const TArray<FMaterialPSOPrecacheR
  */
 extern ENGINE_API void BoostPSOPriority(const TArray<FMaterialPSOPrecacheRequestID>& MaterialPSORequestIDs);
 
+/**
+ * Invalidate & clear all the current material PSO requests
+ */
+extern ENGINE_API void ClearMaterialPSORequests();
+
 /**
  * Get original FMaterialPSOPrecacheParams from precache request
  */
