Remove disable optimization

diff --git a/Engine/Source/Runtime/Engine/Private/PSOPrecacheValidation.cpp b/Engine/Source/Runtime/Engine/Private/PSOPrecacheValidation.cpp
--- a/Engine/Source/Runtime/Engine/Private/PSOPrecacheValidation.cpp
+++ b/Engine/Source/Runtime/Engine/Private/PSOPrecacheValidation.cpp
@@ -12,8 +12,6 @@
 
 #if PSO_PRECACHING_VALIDATE && UE_WITH_PSO_PRECACHING
 
-UE_DISABLE_OPTIMIZATION
-
 /**
 * Different IHVs and drivers can have different opinions on what subset of a PSO
 * matters for caching. We track multiple PSO subsets, ranging from shaders-only
