Water: added new flag to box component to control if it should be shown when ShowFlags.Volumes is enabled or not.

diff --git a/Engine/Plugins/Experimental/Water/Source/Runtime/Private/WaterZoneActor.cpp b/Engine/Plugins/Experimental/Water/Source/Runtime/Private/WaterZoneActor.cpp
--- a/Engine/Plugins/Experimental/Water/Source/Runtime/Private/WaterZoneActor.cpp
+++ b/Engine/Plugins/Experimental/Water/Source/Runtime/Private/WaterZoneActor.cpp
@@ -61,6 +61,10 @@ AWaterZone::AWaterZone(const FObjectInitializer& Initializer)
 		// Bounds component extent is half-extent, ZoneExtent is full extent.
 		BoundsComponent->SetBoxExtent(FVector(ZoneExtent / 2., 8192.));
 		BoundsComponent->bIsEditorOnly = true;
+
+		FEngineShowFlags BoundsComponentShowFlags = BoundsComponent->GetShowFlags();
+		BoundsComponentShowFlags.SetVolumes(true);
+		BoundsComponent->SetShowFlags(BoundsComponentShowFlags);
 	}
 
 	if (GIsEditor && !IsTemplate())
diff --git a/Engine/Source/Runtime/Engine/Classes/Components/BoxComponent.h b/Engine/Source/Runtime/Engine/Classes/Components/BoxComponent.h
--- a/Engine/Source/Runtime/Engine/Classes/Components/BoxComponent.h
+++ b/Engine/Source/Runtime/Engine/Classes/Components/BoxComponent.h
@@ -6,6 +6,7 @@
 #include "CoreMinimal.h"
 #include "UObject/ObjectMacros.h"
 #include "Components/ShapeComponent.h"
+#include "ShowFlags.h"
 #include "BoxComponent.generated.h"
 
 class FPrimitiveSceneProxy;
@@ -23,6 +24,11 @@ protected:
 	UPROPERTY(EditAnywhere, BlueprintReadOnly, export, Category=Shape)
 	FVector BoxExtent;
 
+#if !UE_BUILD_SHIPPING
+	/** List of all show flags this box component visualizer should respect. */
+	FEngineShowFlags ShowFlags;
+#endif // !UE_BUILD_SHIPPING
+
 public:
 	/** 
 	 * Change the box extent size. This is the unscaled size, before component scale is applied.
@@ -56,6 +62,11 @@ public:
 
 	// Sets the box extents without triggering a render or physics update.
 	FORCEINLINE void InitBoxExtent(const FVector& InBoxExtent) { BoxExtent = InBoxExtent; }
+
+#if WITH_EDITOR
+	ENGINE_API FEngineShowFlags GetShowFlags() const { return ShowFlags; }
+	ENGINE_API void SetShowFlags(const FEngineShowFlags& InShowFlags);
+#endif // WITH_EDITOR
 };
 
 
diff --git a/Engine/Source/Runtime/Engine/Private/Components/BoxComponent.cpp b/Engine/Source/Runtime/Engine/Private/Components/BoxComponent.cpp
--- a/Engine/Source/Runtime/Engine/Private/Components/BoxComponent.cpp
+++ b/Engine/Source/Runtime/Engine/Private/Components/BoxComponent.cpp
@@ -14,6 +14,9 @@
 
 UBoxComponent::UBoxComponent(const FObjectInitializer& ObjectInitializer)
 	: Super(ObjectInitializer)
+#if WITH_EDITOR
+	, ShowFlags(ESFIM_All0)
+#endif // WITH_EDITOR
 {
 	BoxExtent = FVector(32.0f, 32.0f, 32.0f);
 
@@ -102,6 +105,12 @@ void UBoxComponent::UpdateBodySetup()
 	}
 }
 
+void UBoxComponent::SetShowFlags(const FEngineShowFlags& InShowFlags)
+{
+	ShowFlags = InShowFlags;
+	MarkRenderStateDirty();
+}
+
 bool UBoxComponent::IsZeroExtent() const
 {
 	return BoxExtent.IsZero();
@@ -134,6 +143,46 @@ FPrimitiveSceneProxy* UBoxComponent::CreateSceneProxy()
 			,	LineThickness( InComponent->LineThickness )
 		{
 			bWillEverBeLit = false;
+
+#if !UE_BUILD_SHIPPING
+			struct FIterSink
+			{
+				FIterSink(const FEngineShowFlags InSelectedShowFlags)
+					: SelectedShowFlags(InSelectedShowFlags)
+				{
+					SelectedShowFlagIndices.SetNum(FEngineShowFlags::SF_FirstCustom, false);
+				}
+
+				bool HandleShowFlag(uint32 InIndex, const FString& InName)
+				{
+					if (SelectedShowFlags.GetSingleFlag(InIndex) == true)
+					{
+						SelectedShowFlagIndices.PadToNum(InIndex + 1, false);
+						SelectedShowFlagIndices[InIndex] = true;
+					}
+
+					return true;
+				}
+
+				bool OnEngineShowFlag(uint32 InIndex, const FString& InName)
+				{
+					return HandleShowFlag(InIndex, InName);
+				}
+
+				bool OnCustomShowFlag(uint32 InIndex, const FString& InName)
+				{
+					return HandleShowFlag(InIndex, InName);
+				}
+
+				const FEngineShowFlags SelectedShowFlags;
+
+				TBitArray<> SelectedShowFlagIndices;
+			};
+
+			FIterSink Sink(InComponent->ShowFlags);
+			FEngineShowFlags::IterateAllFlags(Sink);
+			SelectedShowFlagIndices = MoveTemp(Sink.SelectedShowFlagIndices);
+#endif // !UE_BUILD_SHIPPING
 		}
 
 		virtual void GetDynamicMeshElements(const TArray<const FSceneView*>& Views, const FSceneViewFamily& ViewFamily, uint32 VisibilityMap, FMeshElementCollector& Collector) const override
@@ -165,6 +214,15 @@ FPrimitiveSceneProxy* UBoxComponent::CreateSceneProxy()
 
 			FPrimitiveViewRelevance Result;
 			Result.bDrawRelevance = (IsShown(View) && bProxyVisible) || bShowForCollision;
+#if !UE_BUILD_SHIPPING
+			bool bAreAllSelectedFlagsEnabled = true;
+			for (TConstSetBitIterator<> It(SelectedShowFlagIndices); It; ++It)
+			{
+				bAreAllSelectedFlagsEnabled &= View->Family->EngineShowFlags.GetSingleFlag(It.GetIndex());
+			}
+
+			Result.bDrawRelevance &= bAreAllSelectedFlagsEnabled;
+#endif // !UE_BUILD_SHIPPING
 			Result.bDynamicRelevance = true;
 			Result.bShadowRelevance = IsShadowCast(View);
 			Result.bEditorPrimitiveRelevance = UseEditorCompositing(View);
@@ -178,6 +236,9 @@ FPrimitiveSceneProxy* UBoxComponent::CreateSceneProxy()
 		const FVector	BoxExtents;
 		const FColor	BoxColor;
 		const float		LineThickness;
+#if !UE_BUILD_SHIPPING
+		TBitArray<>		SelectedShowFlagIndices;
+#endif // !UE_BUILD_SHIPPING
 	};
 
 	return new FBoxSceneProxy( this );
