Fixed incorrect preprocessor blocks

diff --git a/Engine/Source/Runtime/Engine/Classes/Components/BoxComponent.h b/Engine/Source/Runtime/Engine/Classes/Components/BoxComponent.h
--- a/Engine/Source/Runtime/Engine/Classes/Components/BoxComponent.h
+++ b/Engine/Source/Runtime/Engine/Classes/Components/BoxComponent.h
@@ -63,10 +63,10 @@ public:
 	// Sets the box extents without triggering a render or physics update.
 	FORCEINLINE void InitBoxExtent(const FVector& InBoxExtent) { BoxExtent = InBoxExtent; }
 
-#if WITH_EDITOR
+#if !UE_BUILD_SHIPPING
 	ENGINE_API FEngineShowFlags GetShowFlags() const { return ShowFlags; }
 	ENGINE_API void SetShowFlags(const FEngineShowFlags& InShowFlags);
-#endif // WITH_EDITOR
+#endif // !UE_BUILD_SHIPPING
 };
 
 
diff --git a/Engine/Source/Runtime/Engine/Private/Components/BoxComponent.cpp b/Engine/Source/Runtime/Engine/Private/Components/BoxComponent.cpp
--- a/Engine/Source/Runtime/Engine/Private/Components/BoxComponent.cpp
+++ b/Engine/Source/Runtime/Engine/Private/Components/BoxComponent.cpp
@@ -105,11 +105,13 @@ void UBoxComponent::UpdateBodySetup()
 	}
 }
 
+#if !UE_BUILD_SHIPPING
 void UBoxComponent::SetShowFlags(const FEngineShowFlags& InShowFlags)
 {
 	ShowFlags = InShowFlags;
 	MarkRenderStateDirty();
 }
+#endif // !UE_BUILD_SHIPPING
 
 bool UBoxComponent::IsZeroExtent() const
 {
