Print out total shader compile worker memory usage on OOM

Usually the first thing I do when I see the shader compile worker memory report is to sum it up. We might as well just add the sum to the report. Summing up maximum values of course isn't going to reflect the actual maximum, but it does provide an upper bound for the maximum memory usage over the lifetime of the shader compile workers, which can be useful to rule out shader compilation as a source of memory issues.

diff --git a/Engine/Source/Runtime/Engine/Private/ShaderCompiler/ShaderCompiler.cpp b/Engine/Source/Runtime/Engine/Private/ShaderCompiler/ShaderCompiler.cpp
--- a/Engine/Source/Runtime/Engine/Private/ShaderCompiler/ShaderCompiler.cpp
+++ b/Engine/Source/Runtime/Engine/Private/ShaderCompiler/ShaderCompiler.cpp
@@ -4050,6 +4050,8 @@ void FShaderCompileThreadRunnable::CompileDirectlyThroughDll()
 void FShaderCompileThreadRunnable::PrintWorkerMemoryUsage()
 {
 	FScopeLock WorkerScopeLock(&WorkerInfosLock);
+	FPlatformProcessMemoryStats TotalMemoryStats{};
+	int32 NumValidWorkers = 0;
 	for (int32 Iter = 0, End = WorkerInfos.Num(); Iter < End; Iter++)
 	{
 		const TUniquePtr<FShaderCompileWorkerInfo>& WorkerInfo = WorkerInfos[Iter];
@@ -4061,22 +4063,43 @@ void FShaderCompileThreadRunnable::PrintWorkerMemoryUsage()
 		FPlatformProcessMemoryStats MemoryStats;
 		if (FPlatformProcess::TryGetMemoryUsage(ProcHandle, MemoryStats))
 		{
+			NumValidWorkers++;
 			UE_LOG(LogShaderCompilers, Display,
 				TEXT("ShaderCompileWorker [%d/%d] MemoryStats:")
 				TEXT("\n\t     UsedPhysical %llu")
 				TEXT("\n\t PeakUsedPhysical %llu")
 				TEXT("\n\t      UsedVirtual %llu")
 				TEXT("\n\t  PeakUsedVirtual %llu"),
-				Iter,
+				Iter + 1,
 				End,
 				MemoryStats.UsedPhysical,
 				MemoryStats.PeakUsedPhysical,
 				MemoryStats.UsedVirtual,
 				MemoryStats.PeakUsedVirtual
 			);
+			TotalMemoryStats.UsedPhysical += MemoryStats.UsedPhysical;
+			TotalMemoryStats.PeakUsedPhysical += MemoryStats.PeakUsedPhysical;
+			TotalMemoryStats.UsedVirtual += MemoryStats.PeakUsedVirtual;
+			TotalMemoryStats.PeakUsedVirtual += MemoryStats.PeakUsedVirtual;
 		}
 		LogQueuedCompileJobs(WorkerInfo->QueuedJobs, -1);
 	}
+
+	if (NumValidWorkers > 0)
+	{
+		UE_LOG(LogShaderCompilers, Display,
+			TEXT("Sum of MemoryStats for %d ShaderCompileWorker(s):")
+			TEXT("\n\t     UsedPhysical %llu")
+			TEXT("\n\t PeakUsedPhysical %llu")
+			TEXT("\n\t      UsedVirtual %llu")
+			TEXT("\n\t  PeakUsedVirtual %llu"),
+			NumValidWorkers,
+			TotalMemoryStats.UsedPhysical,
+			TotalMemoryStats.PeakUsedPhysical,
+			TotalMemoryStats.UsedVirtual,
+			TotalMemoryStats.PeakUsedVirtual
+		);
+	}
 }
 
 void FShaderCompileUtilities::ExecuteShaderCompileJob(FShaderCommonCompileJob& Job)
