The debug name provided for the DDC requests has been modified to append a "[Fast]" and/or "[Meta]" tag to make DDC logs clearer as to what key is being requested and hit/missed.

Note: Preflight will be added once a review has happened since incorrect changes could potentially corrupt the shared cache.

diff --git a/Engine/Source/Runtime/Engine/Private/TextureDerivedDataTask.cpp b/Engine/Source/Runtime/Engine/Private/TextureDerivedDataTask.cpp
--- a/Engine/Source/Runtime/Engine/Private/TextureDerivedDataTask.cpp
+++ b/Engine/Source/Runtime/Engine/Private/TextureDerivedDataTask.cpp
@@ -1431,6 +1431,9 @@ static void DDC1_FetchAndFillDerivedData(
 
 	FSharedBuffer RawDerivedData;
 	const FSharedString SharedTexturePathName(TexturePathName);
+	const FSharedString SharedTextureMetaPathName(WriteToString<256>(TexturePathName, TEXTVIEW(" [Meta]")));
+	const FSharedString SharedTextureFastPathName(WriteToString<256>(TexturePathName, TEXTVIEW(" [Fast]")));
+	const FSharedString SharedTextureFastMetaPathName(WriteToString<256>(TexturePathName, TEXTVIEW(" [Fast][Meta]")));
 
 	FString LocalDerivedDataKeySuffix;
 	FString LocalDerivedDataKey;
@@ -1453,7 +1456,7 @@ static void DDC1_FetchAndFillDerivedData(
 
 			TArray<FCacheGetValueRequest, TInlineAllocator<2>> Requests;
 			Requests.Add({ SharedTexturePathName, ConvertLegacyCacheKey(FetchFirstKey), ECachePolicy::Default, 0 /* UserData */});
-			Requests.Add({ SharedTexturePathName, GetTextureDerivedMetadataKeyFromSuffix(FetchFirstKeySuffix), ECachePolicy::Default, 1 /* UserData */});
+			Requests.Add({ SharedTextureMetaPathName, GetTextureDerivedMetadataKeyFromSuffix(FetchFirstKeySuffix), ECachePolicy::Default, 1 /* UserData */});
 
 			FRequestOwner BlockingOwner(EPriority::Blocking);
 			FSharedBuffer MetadataBuffer;
@@ -1493,8 +1496,8 @@ static void DDC1_FetchAndFillDerivedData(
 		GetTextureDerivedDataKeyFromSuffix(LocalDerivedDataKeySuffix, LocalDerivedDataKey);
 
 		TArray<FCacheGetValueRequest, TInlineAllocator<2>> Requests;
-		Requests.Add({ SharedTexturePathName, ConvertLegacyCacheKey(LocalDerivedDataKey), ECachePolicy::Default, 0 /* UserData */ });
-		Requests.Add({ SharedTexturePathName, GetTextureDerivedMetadataKeyFromSuffix(LocalDerivedDataKeySuffix), ECachePolicy::Default, 1 /* UserData */ });
+		Requests.Add({ SharedTextureFastPathName, ConvertLegacyCacheKey(LocalDerivedDataKey), ECachePolicy::Default, 0 /* UserData */ });
+		Requests.Add({ SharedTextureFastMetaPathName, GetTextureDerivedMetadataKeyFromSuffix(LocalDerivedDataKeySuffix), ECachePolicy::Default, 1 /* UserData */ });
 
 		FRequestOwner BlockingOwner(EPriority::Blocking);
 		FSharedBuffer MetadataBuffer;
