Fixed another incorrect preprocessor block causing EngineShowFlags member of BoxComponent to default-construct and crash.

diff --git a/Engine/Source/Runtime/Engine/Private/Components/BoxComponent.cpp b/Engine/Source/Runtime/Engine/Private/Components/BoxComponent.cpp
--- a/Engine/Source/Runtime/Engine/Private/Components/BoxComponent.cpp
+++ b/Engine/Source/Runtime/Engine/Private/Components/BoxComponent.cpp
@@ -14,9 +14,9 @@
 
 UBoxComponent::UBoxComponent(const FObjectInitializer& ObjectInitializer)
 	: Super(ObjectInitializer)
-#if WITH_EDITOR
+#if !UE_BUILD_SHIPPING
 	, ShowFlags(ESFIM_All0)
-#endif // WITH_EDITOR
+#endif // !UE_BUILD_SHIPPING
 {
 	BoxExtent = FVector(32.0f, 32.0f, 32.0f);
 
