Fix montage preview of non-default tracks not firing notifies

diff --git a/Engine/Source/Runtime/Engine/Private/Animation/AnimSingleNodeInstance.cpp b/Engine/Source/Runtime/Engine/Private/Animation/AnimSingleNodeInstance.cpp
--- a/Engine/Source/Runtime/Engine/Private/Animation/AnimSingleNodeInstance.cpp
+++ b/Engine/Source/Runtime/Engine/Private/Animation/AnimSingleNodeInstance.cpp
@@ -80,7 +80,10 @@ void UAnimSingleNodeInstance::SetAnimationAsset(class UAnimationAsset* NewAsset,
 		Proxy.ReinitializeSlotNodes();
 		if ( Montage->SlotAnimTracks.Num() > 0 )
 		{
-			Proxy.RegisterSlotNodeWithAnimInstance(Montage->SlotAnimTracks[0].SlotName);
+			for (const FSlotAnimationTrack& SlotAnimationTrack : Montage->SlotAnimTracks)
+			{
+				Proxy.RegisterSlotNodeWithAnimInstance(SlotAnimationTrack.SlotName);
+			}
 			Proxy.SetMontagePreviewSlot(Montage->SlotAnimTracks[0].SlotName);
 		}
 		RestartMontage( Montage );
