Added cook-time precomputed spatial hash for instance culling, enable with read-only cvar r.SceneCulling.Precomputed (initially disabled).
- removed vestiges of legacy cooked instance data (unused).
- optimized instance transform calculation with specialized vector ops.

diff --git a/Engine/Source/Runtime/Engine/Classes/Components/InstancedStaticMeshComponent.h b/Engine/Source/Runtime/Engine/Classes/Components/InstancedStaticMeshComponent.h
--- a/Engine/Source/Runtime/Engine/Classes/Components/InstancedStaticMeshComponent.h
+++ b/Engine/Source/Runtime/Engine/Classes/Components/InstancedStaticMeshComponent.h
@@ -412,13 +412,6 @@ private:
 	bool bUseConservativeBounds = false;
 	/** Current cached conservativ bounds. */
 	FBox CachedConservativeInstanceBounds;
-
-	/**
-	 *  Buffers with per-instance data laid out for rendering. 
-	 *  Serialized for cooked content. Used to create PerInstanceRenderData. 
-	 *  Alive between Serialize and PostLoad calls 
-	 */
-	TUniquePtr<FStaticMeshInstanceData> InstanceDataBufferSerializationTmp;
 public:
 
 #if WITH_EDITOR
@@ -507,6 +500,16 @@ public:
 #if WITH_EDITOR
 	ENGINE_API virtual void PostEditChangeChainProperty(FPropertyChangedChainEvent& PropertyChangedEvent) override;
 	ENGINE_API virtual void PostEditUndo() override;
+	/**
+	 * See: UObject::BeginCacheForCookedPlatformData
+	 */
+	ENGINE_API virtual void BeginCacheForCookedPlatformData( const ITargetPlatform* TargetPlatform ) override;
+	
+	/**
+	 * See: UObject::IsCachedCookedPlatformDataLoaded
+	 */
+	ENGINE_API virtual bool IsCachedCookedPlatformDataLoaded( const ITargetPlatform* TargetPlatform ) override;
+
 #endif
 	//~ End UObject Interface
 
diff --git a/Engine/Source/Runtime/Engine/Private/InstanceDataSceneProxy.cpp b/Engine/Source/Runtime/Engine/Private/InstanceDataSceneProxy.cpp
--- a/Engine/Source/Runtime/Engine/Private/InstanceDataSceneProxy.cpp
+++ b/Engine/Source/Runtime/Engine/Private/InstanceDataSceneProxy.cpp
@@ -200,6 +200,13 @@ void FInstanceSceneDataBuffers::ValidateData() const
 	// ValidateArray(Flags.bHasPerInstancePayloadExtension, InstancePayloadExtension, GetNumInstances());
 }
 
+void FInstanceSceneDataBuffers::SetImmutable(FInstanceSceneDataImmutable &&ImmutableData, FAccessTag AccessTag)
+{
+	ValidateAccess(AccessTag);
+	Flags.bHasCompressedSpatialHash = !ImmutableData.GetCompressedInstanceSpatialHashes().IsEmpty();
+	Immutable = MakeShared<FInstanceSceneDataImmutable>(MoveTemp(ImmutableData));
+}
+
 FSingleInstanceDataBuffers::FSingleInstanceDataBuffers()
 {
 	PrimitiveToRelativeWorld.SetIdentity();
diff --git a/Engine/Source/Runtime/Engine/Private/InstancedStaticMesh.cpp b/Engine/Source/Runtime/Engine/Private/InstancedStaticMesh.cpp
--- a/Engine/Source/Runtime/Engine/Private/InstancedStaticMesh.cpp
+++ b/Engine/Source/Runtime/Engine/Private/InstancedStaticMesh.cpp
@@ -2433,6 +2433,7 @@ void UInstancedStaticMeshComponent::BuildComponentInstanceData(ERHIFeatureLevel:
 
 	OutData.PrimitiveMaterialDesc = GetUsedMaterialPropertyDesc(FeatureLevel);
 	OutData.Flags = MakeInstanceDataFlags(OutData.PrimitiveMaterialDesc.bAnyMaterialHasPerInstanceRandom, OutData.PrimitiveMaterialDesc.bAnyMaterialHasPerInstanceCustomData);
+	OutData.ComponentMobility = Mobility;
 	OutData.PrimitiveLocalToWorld = GetRenderMatrix();
 	OutData.StaticMeshBounds = GetStaticMesh()->GetBounds();
 	OutData.NumProxyInstances = PerInstanceSMData.Num();
@@ -2966,40 +2967,19 @@ void UInstancedStaticMeshComponent::SerializeRenderData(FArchive& Ar)
 {
 	if (Ar.IsLoading())
 	{
-		uint64 RenderDataSizeBytes = 0;
-		Ar << RenderDataSizeBytes; // TODO: can skip serialization if we know that data will be discarded
-
-		if (RenderDataSizeBytes > 0)
-		{
-			// Serialize legacy format.
-			InstanceDataBufferSerializationTmp = MakeUnique<FStaticMeshInstanceData>();
-			InstanceDataBufferSerializationTmp->Serialize(Ar);
-		}
+		PrimitiveInstanceDataManager.ReadCookedRenderData(Ar);
 	}
 	else if (Ar.IsSaving())
 	{
-		uint64 RenderDataSizePos = Ar.Tell();
-		
-		// write render data size, will write real size later
-		uint64 RenderDataSizeBytes = 0;
-		Ar << RenderDataSizeBytes;
-
-		bool bSaveRenderData = NeedRenderDataForTargetPlatform(Ar.CookingTarget()) && InstancingRandomSeed != 0;
-		if (bSaveRenderData && !HasAnyFlags(RF_ClassDefaultObject | RF_ArchetypeObject))
+#if WITH_EDITOR
+		FInstanceUpdateComponentDesc ComponentData;
+		if (GetStaticMesh() != nullptr)
 		{
-			uint64 RenderDataPos = Ar.Tell();
-
-			FStaticMeshInstanceData StaticMeshInstanceDataTmp(Ar.CookingTarget()->SupportsFeature(ETargetPlatformFeatures::HalfFloatVertexFormat));
-			BuildLegacyRenderData(StaticMeshInstanceDataTmp);
-			StaticMeshInstanceDataTmp.Serialize(Ar);
-
-			// save render data real size
-			uint64 CurPos = Ar.Tell();
-			RenderDataSizeBytes = CurPos - RenderDataPos;
-			Ar.Seek(RenderDataSizePos);
-			Ar << RenderDataSizeBytes;
-			Ar.Seek(CurPos);
+			UpdateComponentToWorld();
+			BuildComponentInstanceData(GMaxRHIFeatureLevel, ComponentData);
 		}
+		PrimitiveInstanceDataManager.WriteCookedRenderData(Ar, MoveTemp(ComponentData), MakeStridedView(PerInstanceSMData, &FInstancedStaticMeshInstanceData::Transform));
+#endif
 	}
 }
 
@@ -4387,7 +4367,6 @@ void UInstancedStaticMeshComponent::ClearInstances()
 	PerInstanceSMData.Empty();
 	PerInstanceSMCustomData.Empty();
 	InstanceReorderTable.Empty();
-	InstanceDataBufferSerializationTmp.Reset();
 
 	ProxySize = 0;
 
@@ -4859,15 +4838,17 @@ void UInstancedStaticMeshComponent::PostLoad()
 
 	if (!HasAnyFlags(RF_ClassDefaultObject|RF_ArchetypeObject))
 	{
-		PrimitiveInstanceDataManager.PostLoad(PerInstanceSMData.Num(), MoveTemp(InstanceDataBufferSerializationTmp));
+		FInstanceUpdateComponentDesc ComponentData;
+		if (GetStaticMesh() != nullptr)
+		{
+			UpdateComponentToWorld();
+			BuildComponentInstanceData(GMaxRHIFeatureLevel, ComponentData);
+		}
+		PrimitiveInstanceDataManager.PostLoad(PerInstanceSMData.Num(), MoveTemp(ComponentData));
 	}
 
 	// Has different implementation in HISMC
 	OnPostLoadPerInstanceData();
-
-
-	// release InstanceDataBuffers
-	InstanceDataBufferSerializationTmp.Reset();
 }
 
 void UInstancedStaticMeshComponent::CollectPSOPrecacheData(const FPSOPrecacheParams& BasePrecachePSOParams, FComponentPSOPrecacheParamsList& OutParams)
@@ -5350,6 +5331,25 @@ void UInstancedStaticMeshComponent::PostEditUndo()
 
 	MarkRenderStateDirty();
 }
+
+void UInstancedStaticMeshComponent::BeginCacheForCookedPlatformData( const ITargetPlatform* TargetPlatform )
+{
+	Super::BeginCacheForCookedPlatformData(TargetPlatform);
+
+	FInstanceUpdateComponentDesc ComponentData;
+	if (GetStaticMesh() != nullptr)
+	{
+		UpdateComponentToWorld();
+		BuildComponentInstanceData(GMaxRHIFeatureLevel, ComponentData);
+	}
+	PrimitiveInstanceDataManager.BeginCacheForCookedPlatformData(TargetPlatform, MoveTemp(ComponentData), MakeStridedView(PerInstanceSMData, &FInstancedStaticMeshInstanceData::Transform));
+}
+	
+bool UInstancedStaticMeshComponent::IsCachedCookedPlatformDataLoaded( const ITargetPlatform* TargetPlatform )
+{
+	return Super::IsCachedCookedPlatformDataLoaded(TargetPlatform);
+}
+
 #endif
 
 bool UInstancedStaticMeshComponent::IsInstanceSelected(int32 InInstanceIndex) const
diff --git a/Engine/Source/Runtime/Engine/Private/InstancedStaticMesh/ISMInstanceDataManager.cpp b/Engine/Source/Runtime/Engine/Private/InstancedStaticMesh/ISMInstanceDataManager.cpp
--- a/Engine/Source/Runtime/Engine/Private/InstancedStaticMesh/ISMInstanceDataManager.cpp
+++ b/Engine/Source/Runtime/Engine/Private/InstancedStaticMesh/ISMInstanceDataManager.cpp
@@ -9,6 +9,8 @@
 #include "Rendering/MotionVectorSimulation.h"
 #include "SceneInterface.h"
 
+//UE_DISABLE_OPTIMIZATION
+
 #define IDPROXY_ENABLE_ASYNC_TASK 1
 
 #if 0
@@ -34,6 +36,29 @@ static TAutoConsoleVariable<int32> CVarInstanceDataResetTrackingOnRegister(
 	1,
 	TEXT("Chicken switch to disable the new code to reset tracking & instance count during OnRegister, if this causes problems.\nTODO: Remove."));
 
+static TAutoConsoleVariable<int32> CVarInstanceDataMinInstanceCountToOptimize(
+	TEXT("r.InstanceData.MinInstanceCountToOptimize"),
+	2,
+	TEXT("Minimum number of instances to perform optimized build for (if enabled), can be used to disable optimized build for small ISMs as there is some overhead from doing so."));
+
+namespace RenderingSpatialHash
+{
+	template <typename ScalarType>
+	FArchive& operator<<(FArchive& Ar, TLocation<ScalarType>& Item)
+	{
+		Ar << Item.Coord;
+		Ar << Item.Level;
+		return Ar;
+	}
+}
+
+FArchive& operator<<(FArchive& Ar, FInstanceSceneDataBuffers::FCompressedSpatialHashItem& Item)
+{
+	Ar << Item.Location;
+	Ar << Item.NumInstances;
+
+	return Ar;
+}
 
 FPrimitiveInstanceDataManager::FPrimitiveInstanceDataManager(UPrimitiveComponent* InPrimitiveComponent) 
 	: PrimitiveComponent(InPrimitiveComponent) 
@@ -326,31 +351,11 @@ void FPrimitiveInstanceDataManager::PrimitiveTransformChanged()
 
 bool FPrimitiveInstanceDataManager::HasAnyInstanceChanges() const
 {
-	return InstanceUpdateTracker.HasAnyChanges()
+	return bAnyInstanceChange || bNumCustomDataChanged || bBakedLightingDataChanged || bTransformChangedAllInstances
 #if WITH_EDITOR
 		|| bAnyEditorDataChanged 
 #endif
-		|| bNumCustomDataChanged || bBakedLightingDataChanged || bTransformChangedAllInstances;
-}
-
-
-void FPrimitiveInstanceDataManager::SerializeRenderData(FArchive& Ar, bool bCooked)
-{
-	if (bCooked)
-	{
-		bool bHasCookedData = false;//CachedCookedData != nullptr;
-		Ar << bHasCookedData;
-#if 0
-		if (bHasCookedData)
-		{
-			if (Ar.IsLoading())
-			{
-				CachedCookedData = new FInstanceSceneDataBuffers;
-			}
-			CachedCookedData->Serialize(Ar);
-		}
-#endif
-	}
+		|| InstanceUpdateTracker.HasAnyChanges();
 }
 
 /**
@@ -504,6 +509,7 @@ void FPrimitiveInstanceDataManager::InitChangeSet(const FChangeDesc &ChangeDesc,
 	bAnyEditorDataChanged = false;
 #endif	
 
+	bAnyInstanceChange = false;
 	bTransformChangedAllInstances = false;
 	bNumCustomDataChanged = false;
 	bBakedLightingDataChanged = false;
@@ -633,21 +639,21 @@ bool FPrimitiveInstanceDataManager::FlushChanges(FInstanceUpdateComponentDesc &&
 
 	// TODO: We may decide to do so if other conditions are met (e.g., large change-set or marked for full invalidation).
 	// TODO: Need to figure this out in some other way, e.g., attachment counter or whatnot, since the creation has been moved up we no longer know if this is a fresh one.
-	ChangeDesc.bUntrackedState = GetState() != ETrackingState::Tracked;
+	ChangeDesc.bUntrackedState = GetState() != ETrackingState::Tracked && GetState() != ETrackingState::Optimized;
 
 	// Figure out the deltas.
 	if (!ChangeDesc.bUntrackedState)
 	{
 		ChangeDesc.bInstancesChanged =  HasAnyInstanceChanges();
 
-		ChangeDesc.bPrimitiveTransformChanged = bPrimitiveTransformChanged || !ComponentData.PrimitiveLocalToWorld.Equals(PrimitiveLocalToWorld);
+		ChangeDesc.bPrimitiveTransformChanged = !ComponentData.PrimitiveLocalToWorld.Equals(PrimitiveLocalToWorld);
 		ChangeDesc.bMaterialUsageFlagsChanged = Flags != ComponentData.Flags;
 		ChangeDesc.bMaxDisplacementChanged = AbsMaxDisplacement != NewAbsMaxDisplacement;
 		ChangeDesc.bStaticMeshBoundsChanged = !StaticMeshBounds.Equals(ComponentData.StaticMeshBounds);
 	}
 	
 	// Yet another special case to handle externally managed data from landscape grass
-	if (Mode == EMode::ExternalLegacyData && (bPrimitiveTransformChanged || !ComponentData.PrimitiveLocalToWorld.Equals(PrimitiveLocalToWorld)))
+	if (Mode == EMode::ExternalLegacyData && (!ComponentData.PrimitiveLocalToWorld.Equals(PrimitiveLocalToWorld)))
 	{
 		ChangeDesc.bPrimitiveTransformChanged = true;
 	}
@@ -676,11 +682,19 @@ bool FPrimitiveInstanceDataManager::FlushChanges(FInstanceUpdateComponentDesc &&
 		return bWasUpdateQueued;
 	}
 
+	// If we got here & the state is "optimized" then we know the precomputed data is now invalid and we ditch it.
+	if (GetState() == ETrackingState::Optimized)
+	{
+		UE_LOG(LogInstanceProxy, Log, TEXT("Discarded PrecomputedOptimizationData"));
+		PrecomputedOptimizationData.Reset();
+	}
+
+	ETrackingState SuccessorTrackingState = ETrackingState::Tracked;
 	// 
 	{
 		// TODO: Maybe specialize for only bPrimitiveTransformChanged (no need to send/modify anything _other_ than the transform data)
 		// TODO: The states other than bCreateNewProxy _can_ be handled through delta updates (e.g., add instance + offset all the rest).
-		bool bNeedFullUpdate = ChangeDesc.bUntrackedState || ChangeDesc.bMaterialUsageFlagsChanged;
+		bool bNeedFullUpdate = ChangeDesc.bUntrackedState || ChangeDesc.bMaterialUsageFlagsChanged || GetState() == ETrackingState::Optimized;
 
 		// NOTE: Moving the update tracker to the change set implicitly resets it.
 		FISMInstanceUpdateChangeSet ChangeSet(bNeedFullUpdate, MoveTemp(InstanceUpdateTracker));
@@ -741,6 +755,17 @@ bool FPrimitiveInstanceDataManager::FlushChanges(FInstanceUpdateComponentDesc &&
 		else if (bNeedFullUpdate)
 		{
 			LOG_INST_DATA(TEXT("Full Build %s"), TEXT(""));
+			if (PrecomputedOptimizationData.IsValid())
+			{
+				check(Mode == EMode::Default);
+				check(ComponentData.ComponentMobility == EComponentMobility::Static || ComponentData.ComponentMobility == EComponentMobility::Stationary);
+
+				LOG_INST_DATA(TEXT("  Optimized Build (%s)"), PrecomputedOptimizationData.IsValid() ? TEXT("Precomputed") : TEXT(""));
+				Proxy->bBuildOptimized = true;
+				Proxy->PrecomputedOptimizationData = PrecomputedOptimizationData;
+				SuccessorTrackingState = ETrackingState::Optimized;
+			}
+
 			DispatchUpdateTask(bIsUnattached, InstanceDataBufferHeader, [ChangeSet = MoveTemp(ChangeSet), Proxy = Proxy] () mutable 
 			{
 				Proxy->Build(MoveTemp(ChangeSet));
@@ -758,43 +783,18 @@ bool FPrimitiveInstanceDataManager::FlushChanges(FInstanceUpdateComponentDesc &&
 	}
 
 	// After an update has been sent, we need to track all deltas.
-	TrackingState = ETrackingState::Tracked;
+	TrackingState = SuccessorTrackingState;
 
 	return true;
 }
 
-void FPrimitiveInstanceDataManager::PostLoad(int32 InNumInstances, TUniquePtr<FStaticMeshInstanceData> &&InStaticMeshInstanceData)
+void FPrimitiveInstanceDataManager::PostLoad(int32 InNumInstances, FInstanceUpdateComponentDesc &&ComponentData)
 {
 	if (GetState() == ETrackingState::Disabled)
 	{
 		return;
 	}
 
-	if (LegacyStaticMeshInstanceData.IsValid())
-	{
-		check(NumInstances == LegacyStaticMeshInstanceData->GetNumInstances());
-		// kick off conversion job to extract the needed data? When do we need what?
-
-		// Not sure about this whole thing... when do we need it, should it belong with the HISM implementation (probably)
-#ifdef NEW_INSTANCE_DATA_PATH_TODO
-		MarkForRebuildFromExternal([
-			LegacyInstanceData = MoveTemp(LegacyStaticMeshInstanceData)] (TArray<TRefCountPtr<HHitProxy>> &OutHitProxies) mutable
-		{
-			// TODO: Figure out the flow of these things again.
-			// NEW_INSTANCE_DATA_PATH_TODO: Need to differentiate wrt HISM such that we can send along the reorder table if needed...
-
-			OutHitProxies.Reset(); 
-			FPrimitiveInstanceDataManager::FExternalUpdateData ExternalUpdateData;
-			ExternalUpdateData.NumInstances = LegacyInstanceData ? LegacyInstanceData->GetNumInstances() : 0;
-			ExternalUpdateData.UpdateProxy = [LegacyInstanceDataInner = MoveTemp(LegacyInstanceData)](FInstanceDataSceneProxy &InstanceDataSceneProxy, const FRenderBounds &InstanceLocalBounds) mutable
-			{
-				// No reorder table needed since we never perform delta updates (is that true? how do we know that?).
-				InstanceDataSceneProxy.BuildFromLegacyData(MoveTemp(LegacyInstanceDataInner), InstanceLocalBounds, TArray<int32>());
-			};
-			return ExternalUpdateData;
-		});
-#endif
-	}
 	NumInstances = InNumInstances;
 }
 
@@ -875,6 +875,7 @@ void FPrimitiveInstanceDataManager::MarkChangeHelper(int32 InstanceIndex)
 
 	if (GetState() != ETrackingState::Tracked)
 	{
+		bAnyInstanceChange = true;
 		MarkComponentRenderInstancesDirty();
 		return;
 	}
@@ -889,6 +890,7 @@ void FPrimitiveInstanceDataManager::MarkChangeHelper(FPrimitiveInstanceId Instan
 
 	if (GetState() != ETrackingState::Tracked)
 	{
+		bAnyInstanceChange = true;
 		MarkComponentRenderInstancesDirty();
 		return;
 	}
@@ -1059,9 +1061,243 @@ void FPrimitiveInstanceDataManager::OnRegister(int32 InNumInstances)
 		{
 			ClearIdTracking(InNumInstances);
 		}
+	}
+}
+
+static bool ShouldUsePrecomputed()
+{
+	static const auto CVarPrecomputed = IConsoleManager::Get().FindTConsoleVariableDataInt(TEXT("r.SceneCulling.Precomputed"));
+	static const auto CVarSceneCull = IConsoleManager::Get().FindTConsoleVariableDataInt(TEXT("r.SceneCulling"));
+
+	return CVarSceneCull && CVarSceneCull->GetValueOnAnyThread() != 0
+		&& CVarPrecomputed && CVarPrecomputed->GetValueOnAnyThread() != 0;
+}
+
+#if WITH_EDITOR
+
+bool FPrimitiveInstanceDataManager::ShouldWriteCookedData(const ITargetPlatform* TargetPlatform, int32 NumInstancesToBuildFor)
+{
+	EComponentMobility::Type Mobility = PrimitiveComponent.IsValid() ? PrimitiveComponent->Mobility.GetValue() : EComponentMobility::Type::Movable;
+
+	// Only cook for static & stationary(?) regular ISMs
+	bool bValidTypeAndMobility = Mode == EMode::Default && (Mobility == EComponentMobility::Static || Mobility == EComponentMobility::Stationary);
+
+	return bValidTypeAndMobility
+		&& NumInstancesToBuildFor >= CVarInstanceDataMinInstanceCountToOptimize.GetValueOnAnyThread()
+		&& ShouldUsePrecomputed()
+		&& DoesTargetPlatformSupportNanite(TargetPlatform);
+
+}
+
+void FPrimitiveInstanceDataManager::BeginCacheForCookedPlatformData(const ITargetPlatform* TargetPlatform, FInstanceUpdateComponentDesc &&ComponentData, TStridedView<FMatrix> InstanceTransforms)
+{
+	// Already precomputed, we don't need to do it twice (could add checks to see that it is not incorrect for some obscure reason)
+	if (PrecomputedOptimizationData.IsValid())
+	{
+		return;
+	}
+	
+	bool bShouldBuild = ShouldWriteCookedData(TargetPlatform, InstanceTransforms.Num());
+
+	// TODO: we could kick an async thread here if that is preferrable for the cooker?
+	if (bShouldBuild && ComponentData.BuildChangeSet)
+	{
+		uint32 StartTime = FPlatformTime::Cycles();
+		PrecomputedOptimizationData = MakeShared<FISMPrecomputedSpatialHashData>(PrecomputeOptimizationData(MoveTemp(ComponentData), InstanceTransforms));
+		uint32 EndtTime = FPlatformTime::Cycles();
+
+		UE_LOG(LogInstanceProxy, Log, TEXT("Build Instance Spatial Hashes (%.2fms), Instances: %d, Hashes: %d, Remap Size: %d"), FPlatformTime::ToMilliseconds( EndtTime - StartTime ), InstanceTransforms.Num(), PrecomputedOptimizationData->Hashes.Num(), PrecomputedOptimizationData->ProxyIndexToComponentIndexRemap.Num());
+	}	
+}
+
+class FSpatialHashSortBuilder
+{
+public:
+	//int32 FirstLevel;
+	struct FSortedInstanceItem
+	{
+		RenderingSpatialHash::FLocation64 InstanceLoc;
+		int32 InstanceIndex;
+	};
+
+	template <typename GetWorldSpaceInstanceSphereFuncType>
+	void BuildOptimizedSpatialHashOrder(int32 NumInstances, GetWorldSpaceInstanceSphereFuncType &&GetWorldSpaceInstanceSphere)
+	{
+		// TODO: Fix the cvar interaction here, need to unify code path with SceneCulling, or something maybe even come up with some heuristic to allow using fewer bits?
+		int32 FirstLevel = 0;
+		{
+			static const auto CVarInstanceHierarchyMinCellSize = IConsoleManager::Get().FindTConsoleVariableDataFloat(TEXT("r.SceneCulling.MinCellSize"));
+			if (CVarInstanceHierarchyMinCellSize)
+			{
+				// TODO: only one code path to compute this value!!!
+				FirstLevel = RenderingSpatialHash::CalcLevel(CVarInstanceHierarchyMinCellSize->GetValueOnAnyThread() - 1.0);
+			}
+		}
+
+		SortedInstances.Reserve(NumInstances);
+		for (int32 InstanceIndex = 0; InstanceIndex < NumInstances; ++InstanceIndex)
+		{
+			FSphere InstanceWorldSpaceSphere = GetWorldSpaceInstanceSphere(InstanceIndex);
+
+			RenderingSpatialHash::FLocation64 InstanceLoc = RenderingSpatialHash::CalcLevelAndLocationClamped(InstanceWorldSpaceSphere.Center, InstanceWorldSpaceSphere.W, FirstLevel);
+
+			FSortedInstanceItem& Item = SortedInstances.AddDefaulted_GetRef();
+			Item.InstanceLoc = InstanceLoc;
+			Item.InstanceIndex = InstanceIndex;
+		}
+	
+		// Sort the instances according to hash location (first level, then coordinate) and last on instance Index.
+		SortedInstances.Sort(
+			[](const FSortedInstanceItem& A, const FSortedInstanceItem& B) -> bool
+			{
+				if (A.InstanceLoc.Level != B.InstanceLoc.Level)
+				{
+					return A.InstanceLoc.Level < B.InstanceLoc.Level;
+				}
+				if (A.InstanceLoc.Coord.X != B.InstanceLoc.Coord.X)
+				{
+					return A.InstanceLoc.Coord.X < B.InstanceLoc.Coord.X;
+				}
+				if (A.InstanceLoc.Coord.Y != B.InstanceLoc.Coord.Y)
+				{
+					return A.InstanceLoc.Coord.Y < B.InstanceLoc.Coord.Y;
+				}
+				if (A.InstanceLoc.Coord.Z != B.InstanceLoc.Coord.Z)
+				{
+					return A.InstanceLoc.Coord.Z < B.InstanceLoc.Coord.Z;
+				}
+				return A.InstanceIndex < B.InstanceIndex;
+			}
+		);
+	}
+	TArray<FSortedInstanceItem> SortedInstances;
+};
+
+FISMPrecomputedSpatialHashData FPrimitiveInstanceDataManager::PrecomputeOptimizationData(FInstanceUpdateComponentDesc &&ComponentData, TStridedView<FMatrix> InstanceTransforms)
+{
+	FSpatialHashSortBuilder SortBuilder;
+	const float LocalAbsMaxDisplacement = FMath::Max(-ComponentData.PrimitiveMaterialDesc.MinMaxMaterialDisplacement.X, ComponentData.PrimitiveMaterialDesc.MinMaxMaterialDisplacement.Y)
+									+ ComponentData.PrimitiveMaterialDesc.MaxWorldPositionOffsetDisplacement;
+
+	const FVector3f PadExtent = FISMCInstanceDataSceneProxy::GetLocalBoundsPadExtent(ComponentData.PrimitiveLocalToWorld, LocalAbsMaxDisplacement);
+	FRenderBounds InstanceLocalBounds = ComponentData.StaticMeshBounds;
+	InstanceLocalBounds.Min -= PadExtent;
+	InstanceLocalBounds.Max += PadExtent;
+
+	FSphere LocalInstanceSphere = InstanceLocalBounds.ToBoxSphereBounds().GetSphere();
+
+	SortBuilder.BuildOptimizedSpatialHashOrder(InstanceTransforms.Num(),
+		[&](int32 InstanceIndex) -> FSphere
+		{
+			FMatrix InstanceLocalToWorld = InstanceTransforms[InstanceIndex] * ComponentData.PrimitiveLocalToWorld;
+			return LocalInstanceSphere.TransformBy(InstanceLocalToWorld);
+		}				
+	);
+
+	FISMPrecomputedSpatialHashData Result;
+
+	// Pack down the spatial hashes & index remap
+	Result.ProxyIndexToComponentIndexRemap.SetNumUninitialized(InstanceTransforms.Num());
+
+	FInstanceSceneDataBuffers::FCompressedSpatialHashItem CurrentItem;
+	CurrentItem.NumInstances = 0;
+
+	bool bIsIdentityIndexMap = true;
+
+	for (int32 InstanceIndex = 0; InstanceIndex < SortBuilder.SortedInstances.Num(); ++InstanceIndex)
+	{
+		int32 ComponentInstanceIndex = SortBuilder.SortedInstances[InstanceIndex].InstanceIndex;
+		bIsIdentityIndexMap = bIsIdentityIndexMap && InstanceIndex == ComponentInstanceIndex;
+		Result.ProxyIndexToComponentIndexRemap[InstanceIndex] = ComponentInstanceIndex;
+
+		bool bSameLoc = CurrentItem.NumInstances > 0 && CurrentItem.Location == SortBuilder.SortedInstances[InstanceIndex].InstanceLoc;
+		if (bSameLoc)
+		{
+			CurrentItem.NumInstances += 1;
+		}
+		else
+		{
+			if (CurrentItem.NumInstances > 0)
+			{
+				Result.Hashes.Add(CurrentItem);
+			}
+			CurrentItem.Location = SortBuilder.SortedInstances[InstanceIndex].InstanceLoc;
+			CurrentItem.NumInstances = 1;
+		}
+	}
+	if (CurrentItem.NumInstances > 0)
+	{
+		Result.Hashes.Add(CurrentItem);
+	}
+
+	// Don't store a 1:1 mapping
+	if (bIsIdentityIndexMap)
+	{
+		Result.ProxyIndexToComponentIndexRemap.Reset();
+	}
+
+	return Result;
+}
+
+void FPrimitiveInstanceDataManager::WriteCookedRenderData(FArchive& Ar, FInstanceUpdateComponentDesc &&ComponentData, TStridedView<FMatrix> InstanceTransforms)
+{
+	bool bHasCookedData = false;
+
+	bool bShouldBuild = ShouldWriteCookedData(Ar.CookingTarget(), InstanceTransforms.Num());
+
+	if (bShouldBuild)
+	{
+		if (!PrecomputedOptimizationData.IsValid())
+		{
+			if (ComponentData.BuildChangeSet)
+			{
+				PrecomputedOptimizationData = MakeShared<const FISMPrecomputedSpatialHashData>(PrecomputeOptimizationData(MoveTemp(ComponentData), InstanceTransforms));
+			}
+		}
+		
+		if (PrecomputedOptimizationData.IsValid())
+		{
+			// We have to copy the whole thing to be able to serialize?
+			FISMPrecomputedSpatialHashData OptData = *PrecomputedOptimizationData;
+
+			// Serialize the stuff we need.
+			bHasCookedData = true;
+			Ar << bHasCookedData;
+
+			OptData.Hashes.BulkSerialize(Ar);
+			OptData.ProxyIndexToComponentIndexRemap.BulkSerialize(Ar);
+		}
+	}
+
+	if (!bHasCookedData)
+	{
+		// write the bool if we didn't write any data previously
+		Ar << bHasCookedData;
+	}
+}
+#endif // WITH_EDITOR
+
+void FPrimitiveInstanceDataManager::ReadCookedRenderData(FArchive& Ar)
+{
+	bool bHasCookedData = false;
+	Ar << bHasCookedData;
+	if (bHasCookedData)
+	{
+		FISMPrecomputedSpatialHashData Tmp;
+
+		// TODO: Pack the data representation to far fewer bits
+		Tmp.Hashes.BulkSerialize(Ar);
+		// TODO: RLE-compress
+		Tmp.ProxyIndexToComponentIndexRemap.BulkSerialize(Ar);
+
+		// Ditch the precomputed data if it has been disabled (in the runtime), even if the cook was done with the data enabled.
+		if (ShouldUsePrecomputed())
+		{
+			PrecomputedOptimizationData = MakeShared<const FISMPrecomputedSpatialHashData>(Tmp);
+		}
 		else
 		{
-			ClearChangeTracking();
+			PrecomputedOptimizationData.Reset();
 		}
 	}
 }
diff --git a/Engine/Source/Runtime/Engine/Private/InstancedStaticMesh/ISMInstanceDataSceneProxy.cpp b/Engine/Source/Runtime/Engine/Private/InstancedStaticMesh/ISMInstanceDataSceneProxy.cpp
--- a/Engine/Source/Runtime/Engine/Private/InstancedStaticMesh/ISMInstanceDataSceneProxy.cpp
+++ b/Engine/Source/Runtime/Engine/Private/InstancedStaticMesh/ISMInstanceDataSceneProxy.cpp
@@ -8,6 +8,8 @@
 #include "Rendering/RenderingSpatialHash.h"
 #include "Rendering/MotionVectorSimulation.h"
 
+//UE_DISABLE_OPTIMIZATION
+
 DEFINE_LOG_CATEGORY(LogInstanceProxy);
 
 #if 0
@@ -16,6 +18,172 @@ DEFINE_LOG_CATEGORY(LogInstanceProxy);
 	#define LOG_INST_DATA(_Format_, ...) 
 #endif
 
+/**
+ * Vector register version of FRenderTransform, used to preload the primitive to world transform into registers
+ */
+struct FRenderTransformVectorRegister
+{
+	VectorRegister4f R0;
+	VectorRegister4f R1;
+	VectorRegister4f R2;
+	VectorRegister4f Origin;
+
+	FORCEINLINE FRenderTransformVectorRegister(const FRenderTransform &RenderTransform)
+	{
+		//  we can use unaligmed vectorized load since we know there is data beyond the three rows (the origin), so it is ok to load whatever into the 4th component.
+		R0 = VectorLoad(&RenderTransform.TransformRows[0].X);
+		R1 = VectorLoad(&RenderTransform.TransformRows[1].X);
+		R2 = VectorLoad(&RenderTransform.TransformRows[2].X);
+		// But not for the origin
+		Origin = VectorLoadFloat3(&RenderTransform.Origin);
+	}
+};
+
+FORCEINLINE_DEBUGGABLE FRenderTransform VectorMatrixMultiply(const FRenderTransform& LocalToPrimitive, const FRenderTransformVectorRegister& PrimitiveToWorld)
+{
+	FRenderTransform Result;
+
+	// First row of result (Matrix1[0] * Matrix2).
+	{
+		//  we can use unaligmed vectorized load since we know there is data beyond the three rows (the origin), so it is ok to load whatever into the 4th component.
+		const VectorRegister4Float ARow = VectorLoad(&LocalToPrimitive.TransformRows[0].X);
+		VectorRegister4Float R0 = VectorMultiply(VectorReplicate(ARow, 0), PrimitiveToWorld.R0);
+		R0 = VectorMultiplyAdd(VectorReplicate(ARow, 1), PrimitiveToWorld.R1, R0);
+		R0 = VectorMultiplyAdd(VectorReplicate(ARow, 2), PrimitiveToWorld.R2, R0);
+	
+		// We can use unaligmed vectorized store since we know there is data beyond the three floats that is written later
+		// Note: stomps the X of the TransformRows[1]
+		VectorStore(R0, &Result.TransformRows[0].X);		
+	}
+
+	// Second row of result (Matrix1[1] * Matrix2).
+	{
+		//  we can use unaligmed vectorized load since we know there is data beyond the three rows (the origin), so it is ok to load whatever into the 4th component.
+		const VectorRegister4Float ARow = VectorLoad(&LocalToPrimitive.TransformRows[1].X);
+		VectorRegister4Float R1 = VectorMultiply(VectorReplicate(ARow, 0), PrimitiveToWorld.R0);
+		R1 = VectorMultiplyAdd(VectorReplicate(ARow, 1), PrimitiveToWorld.R1, R1);
+		R1 = VectorMultiplyAdd(VectorReplicate(ARow, 2), PrimitiveToWorld.R2, R1);
+
+		// We can use unaligmed vectorized store since we know there is data beyond the three floats that is written later
+		// Note: stomps the X of the TransformRows[2]
+		VectorStore(R1, &Result.TransformRows[1].X);
+	}
+
+	// Third row of result (Matrix1[2] * Matrix2).
+	{
+		//  we can use unaligmed vectorized load since we know there is data beyond the three rows (the origin), so it is ok to load whatever into the 4th component.
+		const VectorRegister4Float ARow = VectorLoad(&LocalToPrimitive.TransformRows[2].X);
+		VectorRegister4Float R2 = VectorMultiply(VectorReplicate(ARow, 0), PrimitiveToWorld.R0);
+		R2 = VectorMultiplyAdd(VectorReplicate(ARow, 1), PrimitiveToWorld.R1, R2);
+		R2 = VectorMultiplyAdd(VectorReplicate(ARow, 2), PrimitiveToWorld.R2, R2);
+
+		// We can use unaligmed vectorized store since we know there is data beyond the three floats that is written later
+		// Note: stomps the X of the Origin
+		VectorStore(R2, &Result.TransformRows[2].X);
+	}
+
+	// Fourth row of result (Matrix1[3] * Matrix2).
+	{
+		//  can _NOT_ use VectorLoad, or we'll run off the end of the FRenderTransform struct.
+		const VectorRegister4Float ARow = VectorLoadFloat3(&LocalToPrimitive.Origin);
+		
+		// Add B3 at once (instead of mult by 1.0 which would have been the fourth value in the 4x4 version of the matrix)
+		VectorRegister4Float R3 = VectorMultiplyAdd(VectorReplicate(ARow, 0), PrimitiveToWorld.R0, PrimitiveToWorld.Origin);
+		R3 = VectorMultiplyAdd(VectorReplicate(ARow, 1), PrimitiveToWorld.R1, R3);
+		R3 = VectorMultiplyAdd(VectorReplicate(ARow, 2), PrimitiveToWorld.R2, R3);
+
+		VectorStoreFloat3(R3, &Result.Origin);
+	}
+	return Result;
+}
+
+
+/**
+ * Helper function to apply transform update that selectively performs Orthogonalize only if the primitive transform has any non-uniform scale.
+ */
+template <typename DeltaType, typename IndexRemapType, typename TransformOutputType>
+FORCEINLINE void ApplyTransformUpdatesEx(const DeltaType &DeltaRange, const IndexRemapType &IndexRemap, const FRenderTransform &PrimitiveToRelativeWorld, const TArray<FRenderTransform> &InstanceTransforms, int32 PostUpdateNumTransforms, TransformOutputType &OutInstanceToPrimitiveRelative)
+{
+	OutInstanceToPrimitiveRelative.SetNum(PostUpdateNumTransforms);
+
+	if (DeltaRange.IsEmpty())
+	{
+		return;
+	}
+
+	if (PrimitiveToRelativeWorld.IsScaleNonUniform())
+	{
+		FRenderTransformVectorRegister PrimitiveToRelativeWorldVR(PrimitiveToRelativeWorld);
+		for (auto It = DeltaRange.GetIterator(); It; ++It)
+		{
+			int32 ItemIndex = It.GetItemIndex();
+			int32 InstanceIndex = It.GetIndex();
+
+			if (IndexRemap.Remap(ItemIndex, InstanceIndex))
+			{
+				FRenderTransform LocalToPrimitiveRelativeWorld = VectorMatrixMultiply(InstanceTransforms[ItemIndex], PrimitiveToRelativeWorldVR);
+				// Remove shear
+				LocalToPrimitiveRelativeWorld.Orthogonalize();
+				OutInstanceToPrimitiveRelative.Set(InstanceIndex, LocalToPrimitiveRelativeWorld);
+			}
+		}
+	}
+	else
+	{
+		FRenderTransformVectorRegister PrimitiveToRelativeWorldVR(PrimitiveToRelativeWorld);
+		for (auto It = DeltaRange.GetIterator(); It; ++It)
+		{
+			int32 ItemIndex = It.GetItemIndex();
+			int32 InstanceIndex = It.GetIndex();
+
+			if (IndexRemap.Remap(ItemIndex, InstanceIndex))
+			{
+				OutInstanceToPrimitiveRelative.Set(InstanceIndex, VectorMatrixMultiply(InstanceTransforms[ItemIndex], PrimitiveToRelativeWorldVR));
+			}
+		}
+	}
+}
+
+template <typename DeltaType, typename IndexRemapType>
+FORCEINLINE void ApplyTransformUpdates(const DeltaType &DeltaRange, const IndexRemapType &IndexRemap, const FRenderTransform &PrimitiveToRelativeWorld, const TArray<FRenderTransform> &InstanceTransforms, int32 PostUpdateNumTransforms, TArray<FRenderTransform> &OutInstanceToPrimitiveRelative)
+{
+	struct FArrayTransformCollector
+	{
+		FORCEINLINE void SetNum(int32 Num) { OutArray.SetNumUninitialized(Num); }
+		FORCEINLINE void Set(int32 Index, const FRenderTransform &Transform) { OutArray[Index] = Transform; }
+		TArray<FRenderTransform> &OutArray;
+	};
+
+	FArrayTransformCollector TransformCollector {OutInstanceToPrimitiveRelative};
+
+	ApplyTransformUpdatesEx(DeltaRange, IndexRemap, PrimitiveToRelativeWorld, InstanceTransforms, PostUpdateNumTransforms, TransformCollector);
+};
+/**
+ * Helper class to apply transform concatenation that selectively performs Orthogonalize only if the primitive transform has any non-uniform scale.
+ */
+struct FInstanceTransformApplyHelper
+{
+	FORCEINLINE_DEBUGGABLE FInstanceTransformApplyHelper(bool bHasWork, const FRenderTransform &InPrimitiveToRelativeWorld)
+	: PrimitiveToRelativeWorldVR(InPrimitiveToRelativeWorld)
+	, bNeedsOrthogonalization(bHasWork && InPrimitiveToRelativeWorld.IsScaleNonUniform())
+	{
+	}
+
+	FORCEINLINE_DEBUGGABLE FRenderTransform Apply(const FRenderTransform &InstanceTransform)
+	{
+		FRenderTransform LocalToPrimitiveRelativeWorld = VectorMatrixMultiply(InstanceTransform, PrimitiveToRelativeWorldVR);
+		if (bNeedsOrthogonalization)
+		{
+			// Remove shear
+			LocalToPrimitiveRelativeWorld.Orthogonalize();
+		}
+		return LocalToPrimitiveRelativeWorld;
+	}
+	FRenderTransformVectorRegister PrimitiveToRelativeWorldVR;
+	bool bNeedsOrthogonalization = false;
+};
+
+
 FISMCInstanceDataSceneProxy::FISMCInstanceDataSceneProxy(FStaticShaderPlatform InShaderPlatform, ERHIFeatureLevel::Type InFeatureLevel) 
 	: ShaderPlatform(InShaderPlatform)
 	, FeatureLevel(InFeatureLevel) 
@@ -25,6 +193,8 @@ FISMCInstanceDataSceneProxy::FISMCInstanceDataSceneProxy(FStaticShaderPlatform I
 
 struct FIdentityIndexRemap
 {
+	FORCEINLINE constexpr bool IsIdentity() const { return true; }
+
 	inline int32 operator[](int32 InIndex) const { return InIndex; }
 
 	template <typename DeltaType, typename ValueType>
@@ -40,7 +210,9 @@ struct FIdentityIndexRemap
 		}
 	}
 
-	FORCEINLINE constexpr bool RemapIndex(int32 Index) const { return true; }
+	FORCEINLINE constexpr bool RemapDestIndex(int32 Index) const { return true; }
+
+	FORCEINLINE bool Remap(int32 &SrcIndex, int32 &DstIndex) const  { return true; }
 };
 
 struct FReorderTableIndexRemap
@@ -65,12 +237,18 @@ struct FReorderTableIndexRemap
 		return ClampValidIndex(InIndex); 
 	}
 
-	FORCEINLINE bool RemapIndex(int32 &Index) const 
+	FORCEINLINE bool RemapDestIndex(int32 &Index) const 
 	{ 
 		Index = operator[](Index);
 		return Index != INDEX_NONE; 
 	}
 
+	FORCEINLINE bool Remap(int32 &SrcIndex, int32 &DstIndex) const 
+	{ 
+		RemapDestIndex(DstIndex);
+		return DstIndex != INDEX_NONE; 
+	}
+
 	template <typename DeltaType, typename ValueType>
 	inline void Scatter(bool bHasData, const DeltaType &Delta, TArray<ValueType> &DestData, int32 NumOutElements, TArray<ValueType> &&InData, int32 ElementStride = 1) const
 	{
@@ -103,65 +281,12 @@ FVector3f FISMCInstanceDataSceneProxy::GetLocalBoundsPadExtent(const FRenderTran
 }
 
 template <typename IndexRemapType>
-void FISMCInstanceDataSceneProxy::ApplyDataChanges(FISMInstanceUpdateChangeSet &ChangeSet, const IndexRemapType &IndexRemap, int32 PostUpdateNumInstances, FInstanceSceneDataBuffers::FWriteView &ProxyData)
+void FISMCInstanceDataSceneProxy::ApplyAttributeChanges(FISMInstanceUpdateChangeSet &ChangeSet, const IndexRemapType &IndexRemap, FInstanceSceneDataBuffers::FWriteView &ProxyData)
 {
-	ProxyData.PrimitiveToRelativeWorld = ChangeSet.PrimitiveToRelativeWorld;
-	ProxyData.PrimitiveWorldSpaceOffset = ChangeSet.PrimitiveWorldSpaceOffset;
-		
-	check(!ChangeSet.Flags.bHasPerInstanceLocalBounds);
-	// TODO: delta support & always assume all bounds changed, and that there is in fact only one
-	ProxyData.InstanceLocalBounds = MoveTemp(ChangeSet.InstanceLocalBounds);
-
-	// TODO: DISP - Fix me (this comment came along from FPrimitiveSceneProxy::SetInstanceLocalBounds and is probably still true...)
-	const FVector3f PadExtent = GetLocalBoundsPadExtent(ProxyData.PrimitiveToRelativeWorld, ChangeSet.AbsMaxDisplacement);
-	for (FRenderBounds& Bounds : ProxyData.InstanceLocalBounds)
-	{
-		Bounds.Min -= PadExtent;
-		Bounds.Max += PadExtent;
-	}
-
-	// unpack transform deltas
-	ProxyData.InstanceToPrimitiveRelative.SetNumUninitialized(PostUpdateNumInstances);
-	auto TransformDelta = ChangeSet.GetTransformDelta();
-	for (auto It = TransformDelta.GetIterator(); It; ++It)
-	{
-		int32 ItemIndex = It.GetItemIndex();
-		int32 InstanceIndex = It.GetIndex();
-
-		if (IndexRemap.RemapIndex(InstanceIndex))
-		{
-			FRenderTransform LocalToPrimitiveRelativeWorld = ChangeSet.Transforms[ItemIndex] * ChangeSet.PrimitiveToRelativeWorld;
-			// Remove shear
-			LocalToPrimitiveRelativeWorld.Orthogonalize();
-			ProxyData.InstanceToPrimitiveRelative[InstanceIndex] = LocalToPrimitiveRelativeWorld;
-		}
-	}
-
-	if (ChangeSet.Flags.bHasPerInstanceDynamicData)
-	{
-		FRenderTransform PrevPrimitiveToRelativeWorld = ChangeSet.PreviousPrimitiveToRelativeWorld.Get(ChangeSet.PrimitiveToRelativeWorld);
-		ProxyData.PrevInstanceToPrimitiveRelative.SetNumUninitialized(PostUpdateNumInstances);
-		for (auto It = TransformDelta.GetIterator(); It; ++It)
-		{
-			int32 ItemIndex = It.GetItemIndex();
-			int32 InstanceIndex = It.GetIndex();
-			if (IndexRemap.RemapIndex(InstanceIndex))
-			{
-				FRenderTransform PrevLocalToPrimitiveRelativeWorld = ChangeSet.PrevTransforms[ItemIndex] * PrevPrimitiveToRelativeWorld;
-				PrevLocalToPrimitiveRelativeWorld.Orthogonalize();
-				ProxyData.PrevInstanceToPrimitiveRelative[InstanceIndex] = PrevLocalToPrimitiveRelativeWorld;
-			}
-		}
-	}
-	else
-	{
-		ProxyData.PrevInstanceToPrimitiveRelative.Reset();
-	}
-
 	if (ChangeSet.Flags.bHasPerInstanceCustomData)
 	{
 		ProxyData.NumCustomDataFloats = ChangeSet.NumCustomDataFloats;
-		IndexRemap.Scatter(ChangeSet.Flags.bHasPerInstanceCustomData, ChangeSet.GetCustomDataDelta(), ProxyData.InstanceCustomData, PostUpdateNumInstances, MoveTemp(ChangeSet.PerInstanceCustomData), ProxyData.NumCustomDataFloats);
+		IndexRemap.Scatter(ChangeSet.Flags.bHasPerInstanceCustomData, ChangeSet.GetCustomDataDelta(), ProxyData.InstanceCustomData, ChangeSet.PostUpdateNumInstances, MoveTemp(ChangeSet.PerInstanceCustomData), ProxyData.NumCustomDataFloats);
 	}
 	else
 	{
@@ -169,9 +294,9 @@ void FISMCInstanceDataSceneProxy::ApplyDataChanges(FISMInstanceUpdateChangeSet &
 		ProxyData.InstanceCustomData.Reset();
 	}
 
-	IndexRemap.Scatter(ChangeSet.Flags.bHasPerInstanceLMSMUVBias, ChangeSet.GetInstanceLightShadowUVBiasDelta(), ProxyData.InstanceLightShadowUVBias, PostUpdateNumInstances, MoveTemp(ChangeSet.InstanceLightShadowUVBias));
+	IndexRemap.Scatter(ChangeSet.Flags.bHasPerInstanceLMSMUVBias, ChangeSet.GetInstanceLightShadowUVBiasDelta(), ProxyData.InstanceLightShadowUVBias, ChangeSet.PostUpdateNumInstances, MoveTemp(ChangeSet.InstanceLightShadowUVBias));
 #if WITH_EDITOR
-	IndexRemap.Scatter(ChangeSet.Flags.bHasPerInstanceEditorData, ChangeSet.GetInstanceEditorDataDelta(), ProxyData.InstanceEditorData, PostUpdateNumInstances, MoveTemp(ChangeSet.InstanceEditorData));
+	IndexRemap.Scatter(ChangeSet.Flags.bHasPerInstanceEditorData, ChangeSet.GetInstanceEditorDataDelta(), ProxyData.InstanceEditorData, ChangeSet.PostUpdateNumInstances, MoveTemp(ChangeSet.InstanceEditorData));
 
 	// replace the HP container.
 	if (ChangeSet.HitProxyContainer)
@@ -187,19 +312,19 @@ void FISMCInstanceDataSceneProxy::ApplyDataChanges(FISMInstanceUpdateChangeSet &
 		// TODO: only need to process added instances? No help for ISM since the move path would be taken.
 		// TODO: OTOH for HISM there is no meaningful data, so just skipping and letting the SetNumZeroed fill in the blanks is fine.
 
-		ProxyData.InstanceRandomIDs.SetNumZeroed(PostUpdateNumInstances);
+		ProxyData.InstanceRandomIDs.SetNumZeroed(ChangeSet.PostUpdateNumInstances);
 		if (ChangeSet.GeneratePerInstanceRandomIds)
 		{
 			// NOTE: this is not super efficient(!)
 			TArray<float> TmpInstanceRandomIDs;
-			TmpInstanceRandomIDs.SetNumZeroed(PostUpdateNumInstances);
+			TmpInstanceRandomIDs.SetNumZeroed(ChangeSet.PostUpdateNumInstances);
 			ChangeSet.GeneratePerInstanceRandomIds(TmpInstanceRandomIDs);
 			FIdentityDeltaRange PerInstanceRandomDelta(TmpInstanceRandomIDs.Num());
-			IndexRemap.Scatter(true, PerInstanceRandomDelta, ProxyData.InstanceRandomIDs, PostUpdateNumInstances, MoveTemp(TmpInstanceRandomIDs));
+			IndexRemap.Scatter(true, PerInstanceRandomDelta, ProxyData.InstanceRandomIDs, ChangeSet.PostUpdateNumInstances, MoveTemp(TmpInstanceRandomIDs));
 		}
 		//else 
 		//{
-		//	IndexRemap.Scatter(true, PerInstanceRandomDelta, ProxyData.InstanceRandomIDs, PostUpdateNumInstances, MoveTemp(ChangeSet.InstanceRandomIDs));
+		//	IndexRemap.Scatter(true, PerInstanceRandomDelta, ProxyData.InstanceRandomIDs, ChangeSet.PostUpdateNumInstances, MoveTemp(ChangeSet.InstanceRandomIDs));
 		//}
 	}
 	else
@@ -208,6 +333,39 @@ void FISMCInstanceDataSceneProxy::ApplyDataChanges(FISMInstanceUpdateChangeSet &
 	}
 }
 
+template <typename IndexRemapType>
+void FISMCInstanceDataSceneProxy::ApplyDataChanges(FISMInstanceUpdateChangeSet &ChangeSet, const IndexRemapType &IndexRemap, int32 PostUpdateNumInstances, FInstanceSceneDataBuffers::FWriteView &ProxyData)
+{
+	ProxyData.PrimitiveToRelativeWorld = ChangeSet.PrimitiveToRelativeWorld;
+	ProxyData.PrimitiveWorldSpaceOffset = ChangeSet.PrimitiveWorldSpaceOffset;
+		
+	check(!ChangeSet.Flags.bHasPerInstanceLocalBounds);
+	// TODO: delta support & always assume all bounds changed, and that there is in fact only one
+	ProxyData.InstanceLocalBounds = MoveTemp(ChangeSet.InstanceLocalBounds);
+
+	// TODO: DISP - Fix me (this comment came along from FPrimitiveSceneProxy::SetInstanceLocalBounds and is probably still true...)
+	const FVector3f PadExtent = GetLocalBoundsPadExtent(ProxyData.PrimitiveToRelativeWorld, ChangeSet.AbsMaxDisplacement);
+	for (FRenderBounds& Bounds : ProxyData.InstanceLocalBounds)
+	{
+		Bounds.Min -= PadExtent;
+		Bounds.Max += PadExtent;
+	}
+
+	// unpack transform deltas
+	ApplyTransformUpdates(ChangeSet.GetTransformDelta(), IndexRemap, ChangeSet.PrimitiveToRelativeWorld, ChangeSet.Transforms, PostUpdateNumInstances, ProxyData.InstanceToPrimitiveRelative);
+	if (ChangeSet.Flags.bHasPerInstanceDynamicData)
+	{
+		FRenderTransform PrevPrimitiveToRelativeWorld = ChangeSet.PreviousPrimitiveToRelativeWorld.Get(ChangeSet.PrimitiveToRelativeWorld);
+		ApplyTransformUpdates(ChangeSet.GetTransformDelta(), IndexRemap, PrevPrimitiveToRelativeWorld, ChangeSet.PrevTransforms, PostUpdateNumInstances, ProxyData.PrevInstanceToPrimitiveRelative);
+	}
+	else
+	{
+		ProxyData.PrevInstanceToPrimitiveRelative.Reset();
+	}
+
+	ApplyAttributeChanges(ChangeSet, IndexRemap, ProxyData);
+}
+
 template<typename ValueType>
 void CondMove(bool bCondition, TArray<ValueType> &Data, int32 FromIndex, int32 ToIndex, int32 NumElements = 1)
 {
@@ -217,8 +375,89 @@ void CondMove(bool bCondition, TArray<ValueType> &Data, int32 FromIndex, int32 T
 	}
 }
 
+struct FSrcIndexRemap
+{
+	FORCEINLINE constexpr bool IsIdentity() const { return false; }
+
+	FSrcIndexRemap(const TArray<int32> &InIndexRemap) : IndexRemap(InIndexRemap) {}
+
+	FORCEINLINE bool RemapDestIndex(int32 &Index) const 
+	{ 
+		return true;
+	}
+
+	FORCEINLINE bool Remap(int32 &SrcIndex, int32 &DstIndex) const 
+	{ 
+		SrcIndex = IndexRemap[SrcIndex];
+		return true; 
+	}
+
+	template <typename DeltaType, typename ValueType>
+	FORCEINLINE void Scatter(bool bHasData, const DeltaType &Delta, TArray<ValueType> &DestData, int32 NumOutElements, TArray<ValueType> &&InData, int32 ElementStride = 1) const
+	{
+		if (bHasData)
+		{
+			::Scatter(Delta, DestData, NumOutElements, MoveTemp(InData), *this, ElementStride);
+		}
+		else
+		{
+			DestData.Reset();
+		}
+	}
+
+	const TArray<int32> &IndexRemap;
+};
+
+void FISMCInstanceDataSceneProxy::BuildFromOptimizedDataBuffers(FISMInstanceUpdateChangeSet& ChangeSet, FInstanceIdIndexMap &OutInstanceIdIndexMap, FInstanceSceneDataBuffers::FWriteView &ProxyData)
+{
+	SCOPED_NAMED_EVENT(FISMCInstanceDataSceneProxy_BuildFromOptimizedDataBuffers, FColor::Emerald);
+
+	ProxyData.PrimitiveToRelativeWorld = ChangeSet.PrimitiveToRelativeWorld;
+	ProxyData.PrimitiveWorldSpaceOffset = ChangeSet.PrimitiveWorldSpaceOffset;
+		
+	check(!ChangeSet.Flags.bHasPerInstanceLocalBounds);
+	
+	// TODO: delta support & always assume all bounds changed, and that there is in fact only one
+	ProxyData.InstanceLocalBounds = MoveTemp(ChangeSet.InstanceLocalBounds);
+
+	// TODO: DISP - Fix me (this comment came along from FPrimitiveSceneProxy::SetInstanceLocalBounds and is probably still true...)
+	const FVector3f PadExtent = GetLocalBoundsPadExtent(ProxyData.PrimitiveToRelativeWorld, ChangeSet.AbsMaxDisplacement);
+	for (FRenderBounds& Bounds : ProxyData.InstanceLocalBounds)
+	{
+		Bounds.Min -= PadExtent;
+		Bounds.Max += PadExtent;
+	}
+
+	// If preoptimized:
+	if (PrecomputedOptimizationData.IsValid())
+	{
+		if (PrecomputedOptimizationData->ProxyIndexToComponentIndexRemap.IsEmpty())
+		{
+			ApplyTransformUpdates(ChangeSet.GetTransformDelta(), FIdentityIndexRemap(), ChangeSet.PrimitiveToRelativeWorld, ChangeSet.Transforms, ChangeSet.PostUpdateNumInstances, ProxyData.InstanceToPrimitiveRelative);
+			ApplyAttributeChanges(ChangeSet, FIdentityIndexRemap(), ProxyData);
+		}
+		else
+		{
+			FSrcIndexRemap SortedInstancesRemap(PrecomputedOptimizationData->ProxyIndexToComponentIndexRemap);
+			ApplyTransformUpdates(ChangeSet.GetTransformDelta(), SortedInstancesRemap, ChangeSet.PrimitiveToRelativeWorld, ChangeSet.Transforms, ChangeSet.PostUpdateNumInstances, ProxyData.InstanceToPrimitiveRelative);
+			ApplyAttributeChanges(ChangeSet, SortedInstancesRemap, ProxyData);
+		}
+
+		// We don't store an ID mapping for this case, since we assume a full rebuild is needed to handle any changes at all.
+		InstanceIdIndexMap.Reset(ChangeSet.PostUpdateNumInstances);
+	
+		InstanceSceneDataBuffers.SetImmutable(FInstanceSceneDataImmutable(PrecomputedOptimizationData->Hashes), ProxyData.AccessTag);
+
+		// Clear the data, we're done with it and it is never coming back (until it is loaded again)
+		PrecomputedOptimizationData.Reset();
+		return;
+	}
+}
+
 void FISMCInstanceDataSceneProxy::Build(FISMInstanceUpdateChangeSet&& ChangeSet)
 {
+	SCOPED_NAMED_EVENT(FISMCInstanceDataSceneProxy_Build, FColor::Emerald);
+
 	DecStatCounters();
 	check(ChangeSet.IsFullUpdate());
 	checkSlow(!ChangeSet.GetTransformDelta().IsDelta());
@@ -228,28 +467,38 @@ void FISMCInstanceDataSceneProxy::Build(FISMInstanceUpdateChangeSet&& ChangeSet)
 	checkSlow(!ChangeSet.GetInstanceEditorDataDelta().IsDelta() || ChangeSet.GetInstanceEditorDataDelta().IsEmpty());
 #endif
 
-
 	FInstanceSceneDataBuffers::FAccessTag AccessTag(PointerHash(this));
 	FInstanceSceneDataBuffers::FWriteView WriteView = InstanceSceneDataBuffers.BeginWriteAccess(AccessTag);
 
 	WriteView.Flags = ChangeSet.Flags;
 
-	UpdateIdMapping(ChangeSet);
-	check(ChangeSet.PostUpdateNumInstances == InstanceIdIndexMap.GetMaxInstanceIndex());
+	if (bBuildOptimized && ChangeSet.PostUpdateNumInstances)
+	{
+		BuildFromOptimizedDataBuffers(ChangeSet, InstanceIdIndexMap, WriteView);
+	}
+	else
+	{
+		UpdateIdMapping(ChangeSet, FIdentityIndexRemap());
+		check(ChangeSet.PostUpdateNumInstances == InstanceIdIndexMap.GetMaxInstanceIndex());
 
-	FIdentityIndexRemap IndexRemap;
-	ApplyDataChanges(ChangeSet, IndexRemap, InstanceIdIndexMap.GetMaxInstanceIndex(), WriteView);
+		FIdentityIndexRemap IndexRemap;
+		ApplyDataChanges(ChangeSet, IndexRemap, InstanceIdIndexMap.GetMaxInstanceIndex(), WriteView);
+	}
 	InstanceSceneDataBuffers.EndWriteAccess(AccessTag);
 
 	InstanceSceneDataBuffers.ValidateData();
 
 	IncStatCounters();
+
+	// This auto-resets such that following builds are _NOT_ doing the opt (these are symptoms of something that was expected to be static, was built anyway)
+	bBuildOptimized = false;
 }
 
-void FISMCInstanceDataSceneProxy::UpdateIdMapping(FISMInstanceUpdateChangeSet& ChangeSet)
+template <typename IndexRemapType>
+void FISMCInstanceDataSceneProxy::UpdateIdMapping(FISMInstanceUpdateChangeSet& ChangeSet, const IndexRemapType &IndexRemap)
 {
 	// update mapping, create explicit mapping if needed
-	if (ChangeSet.bIdentityIdMap)
+	if (ChangeSet.bIdentityIdMap && IndexRemap.IsIdentity())
 	{
 		// Reset to identity mapping with the new number of instances
 		InstanceIdIndexMap.Reset(ChangeSet.PostUpdateNumInstances);
@@ -275,7 +524,10 @@ void FISMCInstanceDataSceneProxy::UpdateIdMapping(FISMInstanceUpdateChangeSet& C
 		{
 			int32 NewInstanceIndex = It.GetIndex();
 			int32 ItemIndex = It.GetItemIndex();
-			FPrimitiveInstanceId InstanceId = ChangeSet.IndexToIdMapDeltaData[ItemIndex];
+			
+			IndexRemap.Remap(ItemIndex, NewInstanceIndex);
+
+			FPrimitiveInstanceId InstanceId = ChangeSet.bIdentityIdMap ? FPrimitiveInstanceId{ItemIndex} : ChangeSet.IndexToIdMapDeltaData[ItemIndex];
 			InstanceIdIndexMap.Update(InstanceId, NewInstanceIndex);
 		}
 	}
@@ -283,6 +535,7 @@ void FISMCInstanceDataSceneProxy::UpdateIdMapping(FISMInstanceUpdateChangeSet& C
 
 void FISMCInstanceDataSceneProxy::Update(FISMInstanceUpdateChangeSet&& ChangeSet)
 {
+	SCOPED_NAMED_EVENT(FISMCInstanceDataSceneProxy_Update, FColor::Emerald);
 	check(!ChangeSet.IsFullUpdate());
 
 	DecStatCounters();
@@ -323,7 +576,7 @@ void FISMCInstanceDataSceneProxy::Update(FISMInstanceUpdateChangeSet&& ChangeSet
 		}
 	}
 
-	UpdateIdMapping(ChangeSet);
+	UpdateIdMapping(ChangeSet, FIdentityIndexRemap());
 	check(ChangeSet.PostUpdateNumInstances == InstanceIdIndexMap.GetMaxInstanceIndex());
 
 	FIdentityIndexRemap IndexRemap;
@@ -382,14 +635,14 @@ void FISMCInstanceDataSceneProxyLegacyReordered::Update(FISMInstanceUpdateChange
 		{
 			// This is somewhat nonintuitive, but the current instance->index map is where we retain knowledge of where the instance used to be placed (in the component address space at last update)
 			int32 InstanceIndex = InstanceIdIndexMap.IdToIndex(FPrimitiveInstanceId{It.GetIndex()});
-			if (IndexRemapOld.RemapIndex(InstanceIndex))
+			if (IndexRemapOld.RemapDestIndex(InstanceIndex))
 			{
 				LOG_INST_DATA(TEXT("Update/HideInstance, ID: %d, IDX: %d"), It.GetIndex(), InstanceIndex);
 				ProxyData.VisibleInstances[InstanceIndex] = false;
 			}
 		}
 	}
-	UpdateIdMapping(ChangeSet);
+	UpdateIdMapping(ChangeSet, FIdentityIndexRemap());
 
 	LegacyInstanceReorderTable = MoveTemp(ChangeSet.LegacyInstanceReorderTable);
 	FReorderTableIndexRemap IndexRemap(LegacyInstanceReorderTable, ChangeSet.PostUpdateNumInstances);
@@ -422,7 +675,7 @@ void FISMCInstanceDataSceneProxyLegacyReordered::Build(FISMInstanceUpdateChangeS
 	LegacyInstanceReorderTable = MoveTemp(ChangeSet.LegacyInstanceReorderTable);
 	ProxyData.Flags = ChangeSet.Flags;
 
-	UpdateIdMapping(ChangeSet);
+	UpdateIdMapping(ChangeSet, FIdentityIndexRemap());
 
 	FReorderTableIndexRemap IndexRemap(LegacyInstanceReorderTable, ChangeSet.PostUpdateNumInstances);
 	ApplyDataChanges(ChangeSet, IndexRemap, ChangeSet.PostUpdateNumInstances, ProxyData);
@@ -434,7 +687,7 @@ void FISMCInstanceDataSceneProxyLegacyReordered::Build(FISMInstanceUpdateChangeS
 		ProxyData.VisibleInstances.SetNum(ChangeSet.PostUpdateNumInstances, false);
 		for (int32 InstanceIndex : LegacyInstanceReorderTable)
 		{
-			if (IndexRemap.RemapIndex(InstanceIndex))
+			if (IndexRemap.RemapDestIndex(InstanceIndex))
 			{
 				ProxyData.VisibleInstances[InstanceIndex] = true;
 			}
@@ -580,7 +833,6 @@ FISMCInstanceDataSceneProxyNoGPUScene::~FISMCInstanceDataSceneProxyNoGPUScene()
 template <typename IndexRemapType>
 void FISMCInstanceDataSceneProxyNoGPUScene::ApplyDataChanges(FISMInstanceUpdateChangeSet &ChangeSet, const IndexRemapType &IndexRemap, int32 PostUpdateNumInstances, FInstanceSceneDataBuffers::FWriteView &ProxyData, FStaticMeshInstanceData &LegacyInstanceData)
 {
-
 	ProxyData.NumCustomDataFloats = ChangeSet.Flags.bHasPerInstanceCustomData ? ChangeSet.NumCustomDataFloats : 0;
 	LegacyInstanceData.AllocateInstances(PostUpdateNumInstances, ProxyData.NumCustomDataFloats, GIsEditor ? EResizeBufferFlags::AllowSlackOnGrow|EResizeBufferFlags::AllowSlackOnReduce : EResizeBufferFlags::None, false); // In Editor always permit overallocation, to prevent too much realloc
 
@@ -622,7 +874,7 @@ void FISMCInstanceDataSceneProxyNoGPUScene::ApplyDataChanges(FISMInstanceUpdateC
 		int32 PackedIndex = It.GetItemIndex();
 		int32 InstanceIndex = It.GetIndex();
 
-		if (IndexRemap.RemapIndex(InstanceIndex))
+		if (IndexRemap.RemapDestIndex(InstanceIndex))
 		{
 			LegacyInstanceData.SetInstance(InstanceIndex, ChangeSet.Transforms[PackedIndex].ToMatrix44f(), ChangeSet.Flags.bHasPerInstanceRandom ? InstanceRandomIDs[InstanceIndex] : 0.0f);
 
@@ -641,7 +893,7 @@ void FISMCInstanceDataSceneProxyNoGPUScene::ApplyDataChanges(FISMInstanceUpdateC
 		{
 			int32 PackedIndex = It.GetItemIndex();
 			int32 InstanceIndex = It.GetIndex();
-			if (IndexRemap.RemapIndex(InstanceIndex))
+			if (IndexRemap.RemapDestIndex(InstanceIndex))
 			{
 				for (int32 j = 0; j < ProxyData.NumCustomDataFloats; ++j)
 				{
@@ -657,7 +909,7 @@ void FISMCInstanceDataSceneProxyNoGPUScene::ApplyDataChanges(FISMInstanceUpdateC
 		{
 			int32 PackedIndex = It.GetItemIndex();
 			int32 InstanceIndex = It.GetIndex();
-			if (IndexRemap.RemapIndex(InstanceIndex))
+			if (IndexRemap.RemapDestIndex(InstanceIndex))
 			{
 				FVector4f Packed = ChangeSet.InstanceLightShadowUVBias[PackedIndex];
 				FVector2D LightmapUVBias = FVector2D(Packed.X, Packed.Y);
@@ -676,7 +928,7 @@ void FISMCInstanceDataSceneProxyNoGPUScene::ApplyDataChanges(FISMInstanceUpdateC
 		{
 			int32 PackedIndex = It.GetItemIndex();
 			int32 InstanceIndex = It.GetIndex();
-			if (IndexRemap.RemapIndex(InstanceIndex))
+			if (IndexRemap.RemapDestIndex(InstanceIndex))
 			{
 				FColor HitProxyColor;
 				bool bSelected;
@@ -729,7 +981,7 @@ void FISMCInstanceDataSceneProxyNoGPUScene::Update(FISMInstanceUpdateChangeSet&&
 		{
 			// This is somewhat nonintuitive, but the current instance->index map is where we retain knowledge of where the instance used to be placed (in the component address space at last update)
 			int32 InstanceIndex = InstanceIdIndexMap.IdToIndex(FPrimitiveInstanceId{It.GetIndex()});
-			if (IndexRemapOld.RemapIndex(InstanceIndex))
+			if (IndexRemapOld.RemapDestIndex(InstanceIndex))
 			{
 				LOG_INST_DATA(TEXT("Update/HideInstance, ID: %d, IDX: %d"), It.GetIndex(), InstanceIndex);
 				LegacyInstanceBuffer->InstanceData->NullifyInstance(InstanceIndex);
@@ -737,7 +989,7 @@ void FISMCInstanceDataSceneProxyNoGPUScene::Update(FISMInstanceUpdateChangeSet&&
 		}
 	}
 
-	UpdateIdMapping(ChangeSet);
+	UpdateIdMapping(ChangeSet, FIdentityIndexRemap());
 	check(bLegacyReordered || ChangeSet.PostUpdateNumInstances == InstanceIdIndexMap.GetMaxInstanceIndex());
 
 	LegacyInstanceReorderTable = MoveTemp(ChangeSet.LegacyInstanceReorderTable);
@@ -765,7 +1017,7 @@ void FISMCInstanceDataSceneProxyNoGPUScene::Build(FISMInstanceUpdateChangeSet&&
 
 	check(bLegacyReordered || ChangeSet.LegacyInstanceReorderTable.IsEmpty());
 
-	UpdateIdMapping(ChangeSet);
+	UpdateIdMapping(ChangeSet, FIdentityIndexRemap());
 	check(bLegacyReordered || ChangeSet.PostUpdateNumInstances == InstanceIdIndexMap.GetMaxInstanceIndex());
 
 	LegacyInstanceReorderTable = MoveTemp(ChangeSet.LegacyInstanceReorderTable);
diff --git a/Engine/Source/Runtime/Engine/Private/InstancedStaticMesh/ISMScatterGatherUtil.h b/Engine/Source/Runtime/Engine/Private/InstancedStaticMesh/ISMScatterGatherUtil.h
--- a/Engine/Source/Runtime/Engine/Private/InstancedStaticMesh/ISMScatterGatherUtil.h
+++ b/Engine/Source/Runtime/Engine/Private/InstancedStaticMesh/ISMScatterGatherUtil.h
@@ -37,18 +37,19 @@ void Scatter(const DeltaType &Delta, TArray<ValueType> &OutDest, int32 DestNumEl
  * using the same delta information. 
  * Never takes move shortcut as there is an index remap.
  */
-template <typename DeltaType, typename ValueType, typename DestIndexRemapType>
-void Scatter(const DeltaType &Delta, TArray<ValueType> &OutDest, int32 DestNumElements, const TArray<ValueType> &InSource, const DestIndexRemapType &DestIndexRemap, int32 ElementStride = 1)
+template <typename DeltaType, typename ValueType, typename IndexRemapType>
+void Scatter(const DeltaType &Delta, TArray<ValueType> &OutDest, int32 DestNumElements, const TArray<ValueType> &InSource, const IndexRemapType &IndexRemap, int32 ElementStride = 1)
 {
 	check(InSource.Num() == Delta.GetNumItems() * ElementStride);
 	OutDest.SetNumUninitialized(DestNumElements * ElementStride);
 	for (auto It = Delta.GetIterator(); It; ++It)
 	{
-		int32 ItemIndex = It.GetItemIndex();
-		int32 DestIndex = DestIndexRemap[It.GetIndex()];
-		if (DestIndex != INDEX_NONE)
+		int32 SrcIndex = It.GetItemIndex();
+		int32 DestIndex = It.GetIndex();
+
+		if (IndexRemap.Remap(SrcIndex, DestIndex))
 		{
-			FMemory::Memcpy(&OutDest[DestIndex * ElementStride], &InSource[ItemIndex * ElementStride], ElementStride * sizeof(ValueType));
+			FMemory::Memcpy(&OutDest[DestIndex * ElementStride], &InSource[SrcIndex * ElementStride], ElementStride * sizeof(ValueType));
 		}
 	}
 }
diff --git a/Engine/Source/Runtime/Engine/Public/InstanceDataSceneProxy.h b/Engine/Source/Runtime/Engine/Public/InstanceDataSceneProxy.h
--- a/Engine/Source/Runtime/Engine/Public/InstanceDataSceneProxy.h
+++ b/Engine/Source/Runtime/Engine/Public/InstanceDataSceneProxy.h
@@ -6,6 +6,7 @@
 #include "RenderTransform.h"
 #include "Tasks/Task.h"
 #include "Engine/EngineTypes.h"
+#include "Rendering/RenderingSpatialHash.h"
 
 class FPrimitiveDrawInterface;
 
@@ -111,6 +112,7 @@ struct FInstanceEditorData
 
 #endif
 
+class FInstanceSceneDataImmutable;
 
 /**
  */
@@ -187,8 +189,16 @@ public:
 
 	void ValidateData() const;
 
+	struct FCompressedSpatialHashItem
+	{
+		// TODO: compress the location - store delta to primitive loc for example
+		RenderingSpatialHash::FLocation64 Location;
+		int32 NumInstances;
+	};
+
 	struct FWriteView
 	{
+		FInstanceSceneDataBuffers::FAccessTag AccessTag;
 		FRenderTransform &PrimitiveToRelativeWorld;
 		FVector &PrimitiveWorldSpaceOffset;
 		TArray<FRenderBounds, TInlineAllocator<1>> &InstanceLocalBounds;
@@ -216,6 +226,7 @@ public:
 		check(CurrentWriterTag.compare_exchange_strong(PrevTagValue, AccessTag.WriterTag));
 		return 	FWriteView
 		{
+			AccessTag,
 			PrimitiveToRelativeWorld,
 			PrimitiveWorldSpaceOffset,
 			InstanceLocalBounds,
@@ -291,6 +302,9 @@ public:
 		};
 	}
 
+	ENGINE_API void SetImmutable(FInstanceSceneDataImmutable &&ImmutableData, FAccessTag AccessTag);
+	inline TSharedPtr<FInstanceSceneDataImmutable, ESPMode::ThreadSafe> GetImmutable() const { return Immutable; }
+
 protected:
 	FRenderTransform PrimitiveToRelativeWorld;
 	FVector PrimitiveWorldSpaceOffset;
@@ -322,6 +336,21 @@ protected:
 	FORCEINLINE void ValidateAccess(const FAccessTag& AccessTag) const {}
 #endif
 
+	TSharedPtr<FInstanceSceneDataImmutable, ESPMode::ThreadSafe> Immutable;
+};
+
+/**
+ * Stores instance data that is immutable - i.e., it is never changed after construction, it can be held onto by the renderer through a shared pointer.
+ */
+class FInstanceSceneDataImmutable
+{
+public:
+	inline FInstanceSceneDataImmutable(const TArray<FInstanceSceneDataBuffers::FCompressedSpatialHashItem> &InCompressedInstanceSpatialHashes) : CompressedInstanceSpatialHashes(InCompressedInstanceSpatialHashes) {}
+
+	inline const TArray<FInstanceSceneDataBuffers::FCompressedSpatialHashItem> &GetCompressedInstanceSpatialHashes() const { return CompressedInstanceSpatialHashes; }
+private:
+	// Encodes consecutive ranges of instances that share the same spatial hash location
+	TArray<FInstanceSceneDataBuffers::FCompressedSpatialHashItem> CompressedInstanceSpatialHashes;
 };
 
 /**
diff --git a/Engine/Source/Runtime/Engine/Public/InstanceDataTypes.h b/Engine/Source/Runtime/Engine/Public/InstanceDataTypes.h
--- a/Engine/Source/Runtime/Engine/Public/InstanceDataTypes.h
+++ b/Engine/Source/Runtime/Engine/Public/InstanceDataTypes.h
@@ -37,6 +37,7 @@ union FInstanceDataFlags
 		uint16 bHasPerInstancePayloadExtension : 1;
 		uint16 bHasPerInstanceVisible : 1;
 		uint16 bHasPerInstanceEditorData : 1;
+		uint16 bHasCompressedSpatialHash : 1;
 	};
 
 	inline bool operator == (FInstanceDataFlags Other) const { return Packed == Other.Packed; }
diff --git a/Engine/Source/Runtime/Engine/Public/InstancedStaticMesh/ISMInstanceDataManager.h b/Engine/Source/Runtime/Engine/Public/InstancedStaticMesh/ISMInstanceDataManager.h
--- a/Engine/Source/Runtime/Engine/Public/InstancedStaticMesh/ISMInstanceDataManager.h
+++ b/Engine/Source/Runtime/Engine/Public/InstancedStaticMesh/ISMInstanceDataManager.h
@@ -16,6 +16,7 @@ class HHitProxy;
 struct FInstanceUpdateComponentDesc
 {
 	FMatrix PrimitiveLocalToWorld;
+	EComponentMobility::Type ComponentMobility = EComponentMobility::Movable;
 	FRenderBounds StaticMeshBounds;
 	FInstanceDataFlags Flags;
 	FPrimitiveMaterialPropertyDescriptor PrimitiveMaterialDesc;
@@ -53,6 +54,7 @@ public:
 		Initial, // In the initial state, there is no proxy and therefore changes do not need to be tracked, e.g., during initial setup of an ISM component.
 		Tracked,
 		Disabled,
+		Optimized, // In the optimized state there's no need to track any delta changes, but if anything changes at all we must rebuild.
 	};
 
 	enum class EMode
@@ -102,13 +104,10 @@ public:
 	 * and the tracking state is not Disabled.
 	 */
 	inline bool HasAnyChanges() const { return GetState() != ETrackingState::Disabled && (GetState() != ETrackingState::Tracked || HasAnyInstanceChanges());}
-
-	void SerializeRenderData(FArchive& Ar, bool bCooked);
-
 	bool FlushChanges(FInstanceUpdateComponentDesc &&ComponentData, bool bNewPrimitiveProxy);
 
 	// 
-	void PostLoad(int32 InNumInstances, TUniquePtr<FStaticMeshInstanceData> &&InStaticMeshInstanceData);
+	void PostLoad(int32 InNumInstances, FInstanceUpdateComponentDesc &&ComponentData);
 
 	/**
 	 * Clear the ID/Index association and reset the mapping to identity & number of instances to the given number.
@@ -158,6 +157,13 @@ public:
 	 */
 	void OnRegister(int32 InNumInstances);
 
+#if WITH_EDITOR
+	void BeginCacheForCookedPlatformData(const ITargetPlatform* TargetPlatform, FInstanceUpdateComponentDesc &&ComponentData, TStridedView<FMatrix> InstanceTransforms);
+
+	void WriteCookedRenderData(FArchive& Ar, FInstanceUpdateComponentDesc &&ComponentData, TStridedView<FMatrix> InstanceTransforms);
+#endif
+	void ReadCookedRenderData(FArchive& Ar);
+
 private:
 	template <typename TaskLambdaType>
 	static void BeginUpdateTask(FInstanceDataUpdateTaskInfo &InstanceDataUpdateTaskInfo, TaskLambdaType &&TaskLambda, const FInstanceDataBufferHeader &InInstanceDataBufferHeader);
@@ -182,6 +188,15 @@ private:
 
 	void InitChangeSet(const union FChangeDesc &ChangeDesc, const FInstanceUpdateComponentDesc &ComponentData, FISMInstanceUpdateChangeSet &ChangeSet);
 
+#if WITH_EDITOR
+
+	bool ShouldWriteCookedData(const ITargetPlatform* TargetPlatform, int32 NumInstancesToBuildFor);
+
+	/**
+	 * Build precomputed data from the input.
+	 */
+	static FISMPrecomputedSpatialHashData PrecomputeOptimizationData(FInstanceUpdateComponentDesc &&ComponentData, TStridedView<FMatrix> InstanceTransforms);
+#endif
 	EMode Mode = EMode::Default;
 	ETrackingState TrackingState = ETrackingState::Initial;
 
@@ -199,6 +214,7 @@ private:
 	bool bAnyEditorDataChanged = false;
 #endif	
 	bool bPrimitiveTransformChanged = false;
+	bool bAnyInstanceChange = false;
 
 	TSharedPtr<FISMCInstanceDataSceneProxy, ESPMode::ThreadSafe> Proxy;
 	TWeakObjectPtr<UPrimitiveComponent> PrimitiveComponent = nullptr;
@@ -208,8 +224,9 @@ private:
 
 	TPimplPtr<struct FLegacyBuildData> LegacyBuildData;
 
-	// Used for serialized legacy data to drive the non-GPU scene rendering path.
-	TUniquePtr<FStaticMeshInstanceData> LegacyStaticMeshInstanceData;
+	using PrecomputedOptimizationDataPtr = FISMCInstanceDataSceneProxy::FISMPrecomputedSpatialHashDataPtr;
+
+	PrecomputedOptimizationDataPtr PrecomputedOptimizationData;
 
 	// Must track this to detect changes 
 	// TODO: make event driven and save the storage?
diff --git a/Engine/Source/Runtime/Engine/Public/InstancedStaticMesh/ISMInstanceDataSceneProxy.h b/Engine/Source/Runtime/Engine/Public/InstancedStaticMesh/ISMInstanceDataSceneProxy.h
--- a/Engine/Source/Runtime/Engine/Public/InstancedStaticMesh/ISMInstanceDataSceneProxy.h
+++ b/Engine/Source/Runtime/Engine/Public/InstancedStaticMesh/ISMInstanceDataSceneProxy.h
@@ -15,6 +15,15 @@ class HHitProxy;
 
 DECLARE_LOG_CATEGORY_EXTERN(LogInstanceProxy, Log, All)
 
+/**
+ * Precomputed optimization data that descrives the spatial hashes and reordering needed.
+ */
+struct FISMPrecomputedSpatialHashData
+{
+	TArray<FInstanceSceneDataBuffers::FCompressedSpatialHashItem> Hashes;
+	TArray<int32> ProxyIndexToComponentIndexRemap;
+};
+
 /**
  * Proxy class that represents scene instance data to the renderer.
  * Responsible for preparing data for render use per-platform & serializing such data for cooked builds.
@@ -64,14 +73,22 @@ public:
 
 	ENGINE_API static FVector3f GetLocalBoundsPadExtent(const FRenderTransform& LocalToWorld, float PadAmount);
 
+	using FISMPrecomputedSpatialHashDataPtr = TSharedPtr<const FISMPrecomputedSpatialHashData, ESPMode::ThreadSafe>;
 protected:
+	/**
+	 * Build an optimized instance buffer, where the order is sorted such that instances with the same spatial hash are consecutive.
+	 */
+	void BuildFromOptimizedDataBuffers(FISMInstanceUpdateChangeSet& ChangeSet, FInstanceIdIndexMap &OutInstanceIdIndexMap, FInstanceSceneDataBuffers::FWriteView &OutData);
+
 	// Update the InstanceIdIndexMap given the change set.
-	void UpdateIdMapping(FISMInstanceUpdateChangeSet& ChangeSet);
+	template <typename IndexRemapType>
+	void UpdateIdMapping(FISMInstanceUpdateChangeSet& ChangeSet, const IndexRemapType &IndexRemap);
 
 	template <typename IndexRemapType>
 	void ApplyDataChanges(FISMInstanceUpdateChangeSet &ChangeSet, const IndexRemapType &IndexRemap, int32 PostUpdateNumInstances, FInstanceSceneDataBuffers::FWriteView &ProxyData);
 	
-	friend class FPrimitiveInstanceDataManager;
+	template <typename IndexRemapType>
+	void ApplyAttributeChanges(FISMInstanceUpdateChangeSet &ChangeSet, const IndexRemapType &IndexRemap, FInstanceSceneDataBuffers::FWriteView &ProxyData);
 
 	FStaticShaderPlatform ShaderPlatform;
 	ERHIFeatureLevel::Type FeatureLevel;
@@ -85,11 +102,16 @@ protected:
 	FInstanceDataUpdateTaskInfo InstanceDataUpdateTaskInfo;
 	// True when it has never been updated before.
 	bool bIsNew = true;
-
+	// This should be set when constructing for a static primitive that wants pre-built instances.
+	bool bBuildOptimized = false;
 #if WITH_EDITOR
 	/**Container for hitproxies that are used by the instances, uses the FDeferredCleanupInterface machinery to delete itself back on the game thread when replaced. */
 	TPimplPtr<FOpaqueHitProxyContainer> HitProxyContainer;
 #endif
+
+	friend class FPrimitiveInstanceDataManager;
+
+	FISMPrecomputedSpatialHashDataPtr PrecomputedOptimizationData;
 };
 
 /**
diff --git a/Engine/Source/Runtime/Renderer/Private/RendererScene.cpp b/Engine/Source/Runtime/Renderer/Private/RendererScene.cpp
--- a/Engine/Source/Runtime/Renderer/Private/RendererScene.cpp
+++ b/Engine/Source/Runtime/Renderer/Private/RendererScene.cpp
@@ -6610,10 +6610,6 @@ void FScene::UpdateAllPrimitiveSceneInfos(FRDGBuilder& GraphBuilder, EUpdateAllP
 
 	SceneExtensionsUpdaters.PostGPUSceneUpdate(GraphBuilder, SceneUB);
 
-	// Need to do this here since we delete the proxy next 
-	// TODO: should refactor to add this as a dependency for a proxy-deletion task instead.
-	SceneCullingUpdater.GetAsyncProxyUseTaskHandle().Wait();
-
 	GraphBuilder.AddSetupTask([DeletedPrimitiveSceneInfos = MoveTemp(DeletedPrimitiveSceneInfos)]
 	{
 		SCOPED_NAMED_EVENT(FScene_DeletePrimitiveSceneInfo, FColor::Red);
diff --git a/Engine/Source/Runtime/Renderer/Private/SceneCulling/SceneCulling.cpp b/Engine/Source/Runtime/Renderer/Private/SceneCulling/SceneCulling.cpp
--- a/Engine/Source/Runtime/Renderer/Private/SceneCulling/SceneCulling.cpp
+++ b/Engine/Source/Runtime/Renderer/Private/SceneCulling/SceneCulling.cpp
@@ -111,6 +111,12 @@ static TAutoConsoleVariable<int32> CVarSceneCulling(
 	}),
 	ECVF_RenderThreadSafe);
 
+static TAutoConsoleVariable<int32> CVarSceneCullingPrecomputed(
+	TEXT("r.SceneCulling.Precomputed"), 
+	0, 
+	TEXT("Enable/Disable precomputed spatial hashes for scene culling."),
+	ECVF_RenderThreadSafe | ECVF_ReadOnly);
+
 static TAutoConsoleVariable<int32> CVarSceneCullingAsyncUpdate(
 	TEXT("r.SceneCulling.Async.Update"), 
 	1, 
@@ -277,11 +283,9 @@ const FString &FSceneCulling::FPrimitiveState::ToString() const
 	case SinglePrim:
 		Result.Append(TEXT("SinglePrim"));
 		break;
-#if SCENE_CULLING_USE_PRECOMPUTED
 	case Precomputed:
 		Result.Append(TEXT("Precomputed"));
 		break;
-#endif
 	case Dynamic:
 		Result.Append(TEXT("Dynamic"));
 		break;
@@ -296,12 +300,10 @@ const FString &FSceneCulling::FPrimitiveState::ToString() const
 	return Result;
 }
 
-#if SCENE_CULLING_USE_PRECOMPUTED
 inline bool operator==(const FInstanceSceneDataBuffers::FCompressedSpatialHashItem A, const FInstanceSceneDataBuffers::FCompressedSpatialHashItem B)
 {
 	return A.Location == B.Location && A.NumInstances == B.NumInstances;
 }
-#endif
 
 #if OLA_TODO
 
@@ -872,6 +874,8 @@ public:
 		// TODO: this is not needed when we also call update for added primitives correctly, remove and replace with a check!
 		SceneCulling.PrimitiveStates.SetNum(SceneCulling.Scene.GetMaxPersistentPrimitiveIndex());
 
+		bUsePrecomputed = CVarSceneCullingPrecomputed.GetValueOnAnyThread() != 0;
+
 		if (CVarTreatDynamicInstancedAsUncullable.GetValueOnRenderThread() != 0)
 		{
 			// Flip dynamic stuff into the uncullabe bucket. 
@@ -1165,7 +1169,7 @@ public:
 		int32 RemovedInstanceCount[EUpdateFrequencyCategory::Num] = { 0, 0 };
 		FCellHeader PrevCellHeader = FCellHeader { 0u, InvalidCellFlag };
 
-		FORCEINLINE int32 GetRemovedInstanceCount() const
+		SC_FORCEINLINE int32 GetRemovedInstanceCount() const
 		{
 			return RemovedInstanceCount[EUpdateFrequencyCategory::Static] + RemovedInstanceCount[EUpdateFrequencyCategory::Dynamic];
 		}
@@ -1620,7 +1624,7 @@ public:
 		SceneCulling.FreeChunk(ChunkId);
 	}
 
-	FORCEINLINE FPrimitiveState ComputePrimitiveState(const FPrimitiveBounds& Bounds, FPrimitiveSceneInfo* PrimitiveSceneInfo, int32 NumInstances, int32 InstanceDataOffset, FPrimitiveSceneProxy* SceneProxy, FInstanceDataFlags InstanceDataFlags, const FPrimitiveState &PrevState)
+	SC_FORCEINLINE FPrimitiveState ComputePrimitiveState(const FPrimitiveBounds& Bounds, FPrimitiveSceneInfo* PrimitiveSceneInfo, int32 NumInstances, int32 InstanceDataOffset, FPrimitiveSceneProxy* SceneProxy, FInstanceDataFlags InstanceDataFlags, const FPrimitiveState &PrevState)
 	{
 		FPrimitiveState NewState;
 		NewState.bDynamic = PrevState.bDynamic || SceneProxy->IsOftenMoving();
@@ -1639,13 +1643,11 @@ public:
 			}
 			else
 			{
-#if SCENE_CULLING_USE_PRECOMPUTED
 				if (InstanceDataFlags.bHasCompressedSpatialHash && bUsePrecomputed)
 				{
 					NewState.State = FPrimitiveState::Precomputed;
 				}
 				else 
-#endif // SCENE_CULLING_USE_PRECOMPUTED
 				{
 					NewState.State = NewState.bDynamic ? DynamicInstancedPrimitiveState : FPrimitiveState::Cached;
 				}		
@@ -1658,7 +1660,7 @@ public:
 
 
 	template <EUpdateFrequencyCategory::EType UpdateFrequencyCategory>
-	FORCEINLINE int32 AddCachedOrDynamic(const FInstanceSceneDataBuffers *InstanceSceneDataBuffers, int32 CacheIndex, const bool bHasPerInstanceLocalBounds, int32 InstanceDataOffset, int32 NumInstances)
+	SC_FORCEINLINE int32 AddCachedOrDynamic(const FInstanceSceneDataBuffers *InstanceSceneDataBuffers, int32 CacheIndex, const bool bHasPerInstanceLocalBounds, int32 InstanceDataOffset, int32 NumInstances)
 	{
 		check(InstanceSceneDataBuffers != nullptr);
 
@@ -1694,7 +1696,7 @@ public:
 		return PrimitiveSceneInfo->GetInstanceSceneDataBuffers();
 	}
 
-	FORCEINLINE void AddInstances(FPersistentPrimitiveIndex PersistentId, FPrimitiveSceneInfo* PrimitiveSceneInfo)
+	SC_FORCEINLINE void AddInstances(FPersistentPrimitiveIndex PersistentId, FPrimitiveSceneInfo* PrimitiveSceneInfo)
 	{
 		FScene &Scene = SceneCulling.Scene;
 		TArray<FPrimitiveState> &PrimitiveStates = SceneCulling.PrimitiveStates;
@@ -1734,9 +1736,9 @@ public:
 				NewPrimitiveState.Payload = CellIndex;
 			}
 			break;
-#if SCENE_CULLING_USE_PRECOMPUTED
 			case FPrimitiveState::Precomputed:
 			{
+				check(InstanceSceneDataBuffers);
 				SC_SCOPED_NAMED_EVENT_DETAIL(SceneCulling_Post_AddInstances_Precomputed, FColor::Emerald);
 
 				BUILDER_LOG_LIST("Add CompressedInstanceSpatialHashes");
@@ -1755,7 +1757,6 @@ public:
 				NewPrimitiveState.InstanceSceneDataImmutable = InstanceSceneDataImmutable;
 			}
 			break;
-#endif // SCENE_CULLING_USE_PRECOMPUTED
 			case FPrimitiveState::UnCullable:
 			{
 				SceneCulling.UnCullablePrimitives.Add(PersistentId);
@@ -1763,12 +1764,16 @@ public:
 			break;
 			case FPrimitiveState::Dynamic:
 			{
+				check(InstanceSceneDataBuffers);
+
 				SC_SCOPED_NAMED_EVENT_DETAIL(SceneCulling_Post_AddInstances_Dynamic, FColor::Emerald);
 				NewPrimitiveState.Payload = AddCachedOrDynamic<EUpdateFrequencyCategory::Dynamic>(InstanceSceneDataBuffers, AllocateCacheEntry(), InstanceDataFlags.bHasPerInstanceLocalBounds, InstanceDataOffset, NumInstances);
 			}
 			break;
 			case FPrimitiveState::Cached:
 			{
+				check(InstanceSceneDataBuffers);
+
 				SC_SCOPED_NAMED_EVENT_DETAIL(SceneCulling_Post_AddInstances_Cached, FColor::Emerald);
 				NewPrimitiveState.Payload = AddCachedOrDynamic<EUpdateFrequencyCategory::Static>(InstanceSceneDataBuffers, AllocateCacheEntry(), InstanceDataFlags.bHasPerInstanceLocalBounds, InstanceDataOffset, NumInstances);
 			}
@@ -1806,7 +1811,6 @@ public:
 		MarkCellForRemove(CellIndex, NumInstances, UpdateFrequencyCategory);
 	}
 
-#if SCENE_CULLING_USE_PRECOMPUTED
 	inline void RemovePrecomputed(FSceneCulling::FPrimitiveState &PrimitiveState)
 	{
 		BUILDER_LOG("FPrimitiveState::Precomputed");
@@ -1820,7 +1824,6 @@ public:
 			MarkCellForRemove(CellIndex, Item.NumInstances, EUpdateFrequencyCategory::Static);
 		}
 	}
-#endif
 
 	inline void MarkInstancesForRemoval(FPersistentPrimitiveIndex PersistentPrimitiveIndex, FPrimitiveSceneInfo *PrimitiveSceneInfo)
 	{
@@ -1856,12 +1859,10 @@ public:
 			int32 CellIndex = PrimitiveState.Payload;
 			MarkCellForRemove(CellIndex, PrimitiveState.NumInstances, PrimitiveState.bDynamic ? EUpdateFrequencyCategory::Dynamic : EUpdateFrequencyCategory::Static);
 		}
-#if SCENE_CULLING_USE_PRECOMPUTED
 		else if (PrimitiveState.State == FPrimitiveState::Precomputed)
 		{
 			RemovePrecomputed(PrimitiveState);
 		}
-#endif
 		else if (PrimitiveState.State == FPrimitiveState::Cached || PrimitiveState.State == FPrimitiveState::Dynamic)
 		{
 			int32 CacheIndex = PrimitiveState.Payload;
@@ -1912,7 +1913,7 @@ public:
 
 	// Mark those that need for remove, queue others for add
 
-	FORCEINLINE void UpdateInstances(FPersistentPrimitiveIndex PersistentPrimitiveIndex, FPrimitiveSceneInfo* PrimitiveSceneInfo)
+	SC_FORCEINLINE void UpdateInstances(FPersistentPrimitiveIndex PersistentPrimitiveIndex, FPrimitiveSceneInfo* PrimitiveSceneInfo)
 	{
 		FSceneCulling::FPrimitiveState PrevPrimitiveState = SceneCulling.PrimitiveStates[PersistentPrimitiveIndex.Index];
 		BUILDER_LOG_SCOPE("UpdateInstances: %d [%s]", PersistentPrimitiveIndex.Index, *PrevPrimitiveState.ToString());
@@ -2018,7 +2019,6 @@ public:
 				UpdateProcessDynamicInstances(HashLocationComputer, InstanceDataOffset, NumInstances, PrevPrimitiveState.NumInstances, CellIndexCacheEntry);
 			}
 		}
-#if SCENE_CULLING_USE_PRECOMPUTED
 		else if (PrevPrimitiveState.State == FPrimitiveState::Precomputed)
 		{
 			SC_SCOPED_NAMED_EVENT_DETAIL(SceneCulling_Post_UpdateInstances_Precomputed, FColor::Emerald);
@@ -2033,7 +2033,6 @@ public:
 
 			bNeedsAdd = true;
 		}
-#endif
 		// In all other cases we have something that must be removed and is in either Cached or Dynamic state which are both removed in the same way.
 		else if (PrevPrimitiveState.State != FPrimitiveState::Unknown)
 		{
@@ -2081,14 +2080,12 @@ public:
 					NewPrimitiveState.Payload = CellIndex;
 				}
 				break;
-#if SCENE_CULLING_USE_PRECOMPUTED
 				case FPrimitiveState::Precomputed:
 				{
 					// this is wrong, post-update a precomputed item should be transitioned to dynamic.
 					check(false);
 				}
 				break;
-#endif // SCENE_CULLING_USE_PRECOMPUTED
 				case FPrimitiveState::UnCullable:
 				{
 					SceneCulling.UnCullablePrimitives.Add(PersistentPrimitiveIndex);
@@ -2323,6 +2320,7 @@ public:
 	int32 TotalUpdatedInstances = 0;
 	int32 TotalAddedInstances = 0;
 	int32 TotalRemovedInstances = 0;
+	bool bUsePrecomputed = false;
 
 #if SC_ENABLE_DETAILED_LOGGING
 	bool bIsLoggingEnabled = false;
@@ -2462,11 +2460,6 @@ void FSceneCulling::FUpdater::OnPreSceneUpdate(FRDGBuilder& GraphBuilder, const
 	}
 	SC_DETAILED_LOGGING_SCOPE(Implementation);
 
-	// This can be run async, but we need to either
-	//  1. [if SCENE_CULLING_USE_PRECOMPUTED] retain a reference to the compressed data (could ref count this), or 
-	//  2. keep the proxy alive and store the pointer to the proxy somewhere before kicking off the async task.
-	// The post-callback then needs to queue its work to happen after this task.
-
 	// Handle all removed primitives 
 	// Step 1. Mark all removed instances
 #if SC_ALLOW_ASYNC_TASKS
@@ -2492,16 +2485,6 @@ void FSceneCulling::FUpdater::OnPreSceneUpdate(FRDGBuilder& GraphBuilder, const
 	// Updated primitives are not handled here (state-caching allows deferring those until the post update)
 }
 
-UE::Tasks::FTask FSceneCulling::FUpdater::GetAsyncProxyUseTaskHandle()
-{
-#if SCENE_CULLING_USE_PRECOMPUTED
-	// Sync is only requiured if using precomputed
-	return PreUpdateTaskHandle;
-#else 
-	return UE::Tasks::FTask();
-#endif
-}
-
 void FSceneCulling::FUpdater::OnPostSceneUpdate(FRDGBuilder& GraphBuilder, const FScenePostUpdateChangeSet& ScenePostUpdateData)
 {
 	LLM_SCOPE_BYTAG(SceneCulling);
diff --git a/Engine/Source/Runtime/Renderer/Private/SceneCulling/SceneCulling.h b/Engine/Source/Runtime/Renderer/Private/SceneCulling/SceneCulling.h
--- a/Engine/Source/Runtime/Renderer/Private/SceneCulling/SceneCulling.h
+++ b/Engine/Source/Runtime/Renderer/Private/SceneCulling/SceneCulling.h
@@ -20,9 +20,6 @@ class FScenePostUpdateChangeSet;
 class FScene;
 struct FPrimitiveBounds;
 
-// Note: the needed precomputation is not yet implemented in the ISM proxy / etc.
-#define SCENE_CULLING_USE_PRECOMPUTED 0
-
 /**
  * Represents either a set of planes, or a sphere,
  */
@@ -51,8 +48,7 @@ public:
 	{
 	public:
 		void OnPreSceneUpdate(FRDGBuilder& GraphBuilder, const FScenePreUpdateChangeSet& ScenePreUpdateData);
-		// Call to get a task to wait for the pre-update task to finish using scene proxies that were deleted - before actually deleting them
-		UE::Tasks::FTask GetAsyncProxyUseTaskHandle();
+
 		void OnPostSceneUpdate(FRDGBuilder& GraphBuilder, const FScenePostUpdateChangeSet& ScenePostUpdateData);
 		void FinalizeAndClear(FRDGBuilder& GraphBuilder, bool bPublishStats);
 
@@ -224,9 +220,7 @@ private:
 		{
 			Unknown,
 			SinglePrim,
-#if SCENE_CULLING_USE_PRECOMPUTED
 			Precomputed,
-#endif
 			UnCullable,
 			Dynamic,
 			Cached,
@@ -244,9 +238,7 @@ private:
 
 		const FString &ToString() const;
 
-#if SCENE_CULLING_USE_PRECOMPUTED
 		TSharedPtr<FInstanceSceneDataImmutable, ESPMode::ThreadSafe> InstanceSceneDataImmutable;
-#endif
 	};
 
 	TArray<FPrimitiveState> PrimitiveStates;
