WorldPartition - HLOD: Changed NetDormancy from DORM_DormantAll to DORM_Initial
* DORM_Initial should be used for map placed actors
* Will prevent HLOD actors from being considered in the initial replication
* Also fixes an issue where unbuilt HLOD actors with no components are causing warnings to be written to the logs during the initial replication phase

diff --git a/Engine/Source/Runtime/Engine/Private/WorldPartition/HLOD/HLODActor.cpp b/Engine/Source/Runtime/Engine/Private/WorldPartition/HLOD/HLODActor.cpp
--- a/Engine/Source/Runtime/Engine/Private/WorldPartition/HLOD/HLODActor.cpp
+++ b/Engine/Source/Runtime/Engine/Private/WorldPartition/HLOD/HLODActor.cpp
@@ -54,7 +54,7 @@ AWorldPartitionHLOD::AWorldPartitionHLOD(const FObjectInitializer& ObjectInitial
 	// needs to be replicated.
 	bReplicates = true;
 
-	NetDormancy = DORM_DormantAll;
+	NetDormancy = DORM_Initial;
 	NetUpdateFrequency = 1.f;
 
 #if WITH_EDITORONLY_DATA
