Expose virtual texture physical pool settings in the project engine settings UI.
Enable live edit of the settings, with physical pools being recreated on changes.
Add config option to auto grow physical pools whenever a pool is oversubscribed.
The config option applies in editor only. There is an additional CVar "r.VT.PoolAutoGrow" that does the same thing in cooked builds.
When pools autogrow the new settings are stored in a transient setting, visible in the UI. These can be copy/pasted to the serialized settings to provide a way to tune in editor.

diff --git a/Engine/Config/BaseEngine.ini b/Engine/Config/BaseEngine.ini
--- a/Engine/Config/BaseEngine.ini
+++ b/Engine/Config/BaseEngine.ini
@@ -3394,17 +3394,6 @@ bAllowDelayedMessaging=false
 [/Script/ChaosSolverEngine.ChaosSolverSettings]
 DefaultChaosSolverActorClass=/Script/ChaosSolverEngine.ChaosSolverActor
 
-[/Script/Engine.VirtualTexturePoolConfig]
-; Configure VT physical memory pools. Useful commands to set these:
-; "stat virtualtexturing" show dynamic use of the VT system including the cache loads
-; "r.VT.ListPhysicalPool" list in-depth details on the allocated physical pools
-; "r.VT.DumpPoolUsage" list in-depth details on asset usage of physical pools
-; The default pool size and behavior is set here. It can be overridden in project config files.
-+Pools=(SizeInMegabyte=64, bAllowSizeScale=True, bEnableResidencyMipMapBias=True)
-; Specific pools can be configured per tile size or format in project config files. Examples commented below:
-;+Pools=(Formats=(PF_DXT1), SizeInMegabyte=36, bAllowSizeScale=True, ScalabilityGroup=1, bEnableResidencyMipMapBias=False)
-;+Pools=(Formats=(PF_DXT5, PF_DXT5), SizeInMegabyte=100)
-
 [PlatformPaths]
 UnrealBuildTool=Engine/Build/BatchFiles/Build.bat
 
diff --git a/Engine/Source/Runtime/Engine/Classes/VT/VirtualTexturePoolConfig.h b/Engine/Source/Runtime/Engine/Classes/VT/VirtualTexturePoolConfig.h
--- a/Engine/Source/Runtime/Engine/Classes/VT/VirtualTexturePoolConfig.h
+++ b/Engine/Source/Runtime/Engine/Classes/VT/VirtualTexturePoolConfig.h
@@ -2,83 +2,114 @@
 
 #pragma once
 
-#include "CoreMinimal.h"
-#include "UObject/ObjectMacros.h"
+#include "Engine/DeveloperSettings.h"
 #include "RenderResource.h"
-#include "RenderCommandFence.h"
-#include "Engine/Texture.h"
-#include "VirtualTexturing.h"
+
 #include "VirtualTexturePoolConfig.generated.h"
 
-/**
-* Settings of a single pool
-*/
+/** Settings for a single virtual texture physical pool. */
 USTRUCT()
 struct FVirtualTextureSpacePoolConfig
 {
 	GENERATED_USTRUCT_BODY()
 
-	FVirtualTextureSpacePoolConfig() 
-		: MinTileSize(0)
-		, MaxTileSize(0)
-		, SizeInMegabyte(0)
-		, bEnableResidencyMipMapBias(false)
-		, bAllowSizeScale(false)
-		, ScalabilityGroup(0)
-		, MinScaledSizeInMegabyte(0)
-		, MaxScaledSizeInMegabyte(0)
-	{}
+	/** Formats of the layers in the physical pool. Leave empty to match any format. */
+	UPROPERTY(EditAnywhere, Category = "PoolConfig|Filter")
+	TArray< TEnumAsByte<EPixelFormat> > Formats;
 
 	/** Minimum tile size to match (including tile border). */
-	UPROPERTY()
-	int32 MinTileSize;
+	UPROPERTY(EditAnywhere, Category = "PoolConfig|Filter")
+	int32 MinTileSize = 0;
 
-	/** Maximum tile size to match (including tile border). */
-	UPROPERTY()
-	int32 MaxTileSize;
-
-	/** Format set to match. One pool can contain multiple layers with synchronized page table mappings. */
-	UPROPERTY()
-	TArray< TEnumAsByte<EPixelFormat> > Formats;
+	/** Maximum tile size to match (including tile border). Set to 0 to match any tile size. */
+	UPROPERTY(EditAnywhere, Category = "PoolConfig|Filter")
+	int32 MaxTileSize = 0;
 
-	/** Size in megabytes to allocate for the pool. The allocator will allocate as close as possible below this limit. */
-	UPROPERTY()
-	int32 SizeInMegabyte;
+	/** Upper limit size in megabytes to allocate for the pool. The allocator will allocate as close as possible below this limit. */
+	UPROPERTY(EditAnywhere, Category = PoolConfig)
+	int32 SizeInMegabyte = 0;
 
 	/** Enable MipMapBias based on pool residency tracking. */
-	UPROPERTY()
-	bool bEnableResidencyMipMapBias;
+	UPROPERTY(EditAnywhere, Category = PoolConfig)
+	bool bEnableResidencyMipMapBias = false;
 
-	/** Allow the size to allocate for the pool to be scaled by some factor. */
-	UPROPERTY()
-	bool bAllowSizeScale;
-
-	/** Scalability group index that gives the size scale. */
-	UPROPERTY()
-	uint32 ScalabilityGroup;
+	/** Allow the size to allocate for the pool to be scaled by scalability settings. */
+	UPROPERTY(EditAnywhere, Category = "PoolConfig|Scalabiity")
+	bool bAllowSizeScale = false;
 
 	/** Lower limit of size in megabytes to allocate for the pool after size scaling. */
-	UPROPERTY()
-	int32 MinScaledSizeInMegabyte;
+	UPROPERTY(EditAnywhere, Category = "PoolConfig|Scalability")
+	int32 MinScaledSizeInMegabyte = 0;
 
-	/** Upper limit of size in megabytes to allocate for the pool after size scaling. */
-	UPROPERTY()
-	int32 MaxScaledSizeInMegabyte;
+	/** Upper limit of size in megabytes to allocate for the pool after size scaling. Set to 0 to ignore. */
+	UPROPERTY(EditAnywhere, Category = "PoolConfig|Scalability")
+	int32 MaxScaledSizeInMegabyte = 0;
 
 	/** Is this the default config? Use this setting when we can't find any other match. */
 	bool IsDefault() const { return Formats.Num() == 0 && SizeInMegabyte > 0; }
 };
 
-UCLASS(config = Engine, transient, MinimalAPI)
-class UVirtualTexturePoolConfig : public UObject
+/** Configuration for virtual texture physical pool sizes. */
+UCLASS(config = Engine, defaultconfig, meta = (DisplayName = "Virtual Texture Pool"), MinimalAPI)
+class UVirtualTexturePoolConfig : public UDeveloperSettings
 {
 	GENERATED_UCLASS_BODY()
-public:
-	UPROPERTY(config)
-	int32 DefaultSizeInMegabyte; // Size in tiles of any pools not explicitly specified in the config
-
-	UPROPERTY(config)
-	TArray<FVirtualTextureSpacePoolConfig> Pools; // All the VT pools specified in the config
+	
+	/** Upper size limit in megabytes for any pools not explicitly matched by a config entry in the Pools array. */
+	UPROPERTY(config, EditAnywhere, Category = PoolConfig)
+	int32 DefaultSizeInMegabyte = 64;
+
+	/** 
+	 * Enable physical pools growing on oversubscription.
+	 * Each physical pool will grow to the maximum size so far requested.
+	 * This setting applies to the editor only. To have similar behavior in a cooked build use r.VT.PoolAutoGrow.
+	 */
+	UPROPERTY(config, EditAnywhere, Category = PoolConfig)
+	bool bPoolAutoGrowInEditor = true;
+
+	/** 
+	 * Serialized array of configs. 
+	 * A virtual texture physical pool iterates these from last to first and uses the first matching config that it finds. 
+	 */
+	UPROPERTY(config, EditAnywhere, Category = PoolConfig, meta=(DisplayName="Fixed Pools", TitleProperty = "Formats"))
+	TArray<FVirtualTextureSpacePoolConfig> Pools;
+
+	/** 
+	 * Transient array of runtime detected configs used by the PoolAutoGrow system.
+	 * A virtual texture physical pool searches these to find a match before searching the configs in Pools.
+	 * These tracked configs can be copied to the serialized Pools as a good estimation of the fixed pool sizes that a cooked project needs.
+	 */
+	UPROPERTY(Transient, EditAnywhere, Category = PoolConfig, meta = (TitleProperty = "Formats"))
+	TArray<FVirtualTextureSpacePoolConfig> TransientPools;
+
+	/** Version number used to help track configuration changes. */
+	UPROPERTY(Transient)
+	uint32 Version = 0;
+
+	/** Find a matching config by first searching the TransientPools and then the Pools. */
+	void FindPoolConfig(TEnumAsByte<EPixelFormat> const* InFormats, int32 InNumLayers, int32 InTileSize, FVirtualTextureSpacePoolConfig& OutConfig) const;
+
+	/** Add a transient config. If there is already a matching config, then this will just copy the SizeInMegabyte setting to the existing config. Returns false if no change is made. */
+	bool AddOrModifyTransientPoolConfig(FVirtualTextureSpacePoolConfig const& InConfig);
+
+#if WITH_EDITOR
+	void PostEditChangeProperty(struct FPropertyChangedEvent& PropertyChangedEvent) override;
+#endif
+};
 
-	ENGINE_API void FindPoolConfig(TEnumAsByte<EPixelFormat> const* InFormats, int32 InNumLayers, int32 InTileSize, FVirtualTextureSpacePoolConfig& OutConfig) const;
+/** Namespace container for static helper functions. */
+struct VirtualTexturePool
+{
+	/** Find a matching pool config. */
+	ENGINE_API static void FindPoolConfig(TEnumAsByte<EPixelFormat> const* InFormats, int32 InNumLayers, int32 InTileSize, FVirtualTextureSpacePoolConfig& OutConfig);
+	/** Add pool configs. If one of the configs already has a matching config, then this will just copy the SizeInMegabyte setting to the existing config. */
+	ENGINE_API static void AddOrModifyTransientPoolConfigs_RenderThread(TArray<FVirtualTextureSpacePoolConfig>& InConfigs);
+	/** Get the scalability setting for scaling virtual texture physical pool sizes. */
+	ENGINE_API static float GetPoolSizeScale();
+	/** Get if the virtual texture physical pools should auto grow on oversubscription. */
+	ENGINE_API static bool GetPoolAutoGrow();
+	/** Get the scalability setting whether we can create multiple physical pools for the same format. */
+	ENGINE_API static int32 GetSplitPhysicalPoolSize();
+	/** Get a hash of all settings that affect virtual texture physical pool creation. */
+	ENGINE_API static uint32 GetConfigHash();
 };
diff --git a/Engine/Source/Runtime/Engine/Private/VT/VirtualTexturePoolConfig.cpp b/Engine/Source/Runtime/Engine/Private/VT/VirtualTexturePoolConfig.cpp
--- a/Engine/Source/Runtime/Engine/Private/VT/VirtualTexturePoolConfig.cpp
+++ b/Engine/Source/Runtime/Engine/Private/VT/VirtualTexturePoolConfig.cpp
@@ -2,8 +2,70 @@
 
 #include "VT/VirtualTexturePoolConfig.h"
 
+#include "Async/TaskGraphInterfaces.h"
+#include "HAL/IConsoleManager.h"
+#include "VT/VirtualTextureRecreate.h"
+
+#if WITH_EDITOR
+#include "Framework/Docking/TabManager.h"
+#include "Framework/Notifications/NotificationManager.h"
+#include "Widgets/Notifications/SNotificationList.h"
+#endif
+
 #include UE_INLINE_GENERATED_CPP_BY_NAME(VirtualTexturePoolConfig)
 
+#define LOCTEXT_NAMESPACE "VirtualTexturePool"
+
+void OnVirtualTexturePoolConfigUpdate();
+
+static TAutoConsoleVariable<float> CVarVTPoolSizeScale(
+	TEXT("r.VT.PoolSizeScale"),
+	1.0f,
+	TEXT("Scale factor for virtual texture physical pool size.\n"),
+	FConsoleVariableDelegate::CreateLambda([](IConsoleVariable* InVariable) { OnVirtualTexturePoolConfigUpdate(); }),
+	ECVF_RenderThreadSafe | ECVF_Scalability | ECVF_ExcludeFromPreview
+);
+
+static TAutoConsoleVariable<bool> CVarVTPoolAutoGrow(
+	TEXT("r.VT.PoolAutoGrow"),
+	false,
+	TEXT("Enable physical pool growing on oversubscription."),
+	FConsoleVariableDelegate::CreateLambda([](IConsoleVariable* InVariable) { OnVirtualTexturePoolConfigUpdate(); }),
+	ECVF_RenderThreadSafe
+);
+
+static TAutoConsoleVariable<int32> CVarVTSplitPhysicalPoolSize(
+	TEXT("r.VT.SplitPhysicalPoolSize"),
+	0,
+	TEXT("Create multiple physical pools per format to keep pools at this maximum size in tiles.\n")
+	TEXT("A value of 64 tiles will force 16bit page tables. This can be a page table memory optimization for large physical pools.\n")
+	TEXT("Defaults to 0 (off)."),
+	FConsoleVariableDelegate::CreateLambda([](IConsoleVariable* InVariable) { OnVirtualTexturePoolConfigUpdate(); }),
+	ECVF_RenderThreadSafe
+);
+
+/** Track changes and apply to relevant systems. This allows us to dynamically change the scalability settings. */
+void OnVirtualTexturePoolConfigUpdate()
+{
+	// CVar updates can come early in start up before UVirtualTexturePoolConfig is serialized.
+	// That leads to UVirtualTexturePoolConfig not loading correctly.
+	// We early out here to avoid that case (we will update all later on serialization anyway).
+	if (!IsClassLoaded<UVirtualTexturePoolConfig>())
+	{
+		return;
+	}
+
+	const uint32 ConfigHash = VirtualTexturePool::GetConfigHash();
+	static float LastConfigHash = ConfigHash;
+
+	if (LastConfigHash != ConfigHash)
+	{
+		LastConfigHash = ConfigHash;
+		VirtualTexture::Recreate();
+	}
+}
+
+
 UVirtualTexturePoolConfig::UVirtualTexturePoolConfig(const FObjectInitializer& ObjectInitializer)
 	: Super(ObjectInitializer)
 {
@@ -11,12 +73,40 @@ UVirtualTexturePoolConfig::UVirtualTexturePoolConfig(const FObjectInitializer& O
 
 void UVirtualTexturePoolConfig::FindPoolConfig(TEnumAsByte<EPixelFormat> const* InFormats, int32 InNumLayers, int32 InTileSize, FVirtualTextureSpacePoolConfig& OutConfig) const
 {
-	bool bFoundDefaultConfig = false;
+	// Create a default config that will be used if no other default config is found.
 	FVirtualTextureSpacePoolConfig DefaultConfig;
 	DefaultConfig.SizeInMegabyte = DefaultSizeInMegabyte;
+	DefaultConfig.bEnableResidencyMipMapBias = true;
+	DefaultConfig.bAllowSizeScale = true;
 
-	// Reverse iterate so that project config can override base config
-	for (int32 Id = Pools.Num() - 1; Id >= 0 ; Id--)
+	bool bFoundDefaultConfig = false;
+
+	// First look in transient pool descriptions.
+	for (int32 Id = TransientPools.Num() - 1; Id >= 0; Id--)
+	{
+		const FVirtualTextureSpacePoolConfig& Config = TransientPools[Id];
+		if (Config.MinTileSize <= InTileSize && (Config.MaxTileSize == 0 || Config.MaxTileSize >= InTileSize) && InNumLayers == Config.Formats.Num())
+		{
+			bool bAllFormatsMatch = true;
+			for (int Layer = 0; Layer < InNumLayers && bAllFormatsMatch; ++Layer)
+			{
+				if (InFormats[Layer] != Config.Formats[Layer])
+				{
+					bAllFormatsMatch = false;
+				}
+			}
+
+			if (bAllFormatsMatch)
+			{
+				OutConfig = Config;
+				return;
+			}
+		}
+	}
+
+	// Now look in serialized pool descriptions.
+	// Note that we reverse iterate so that project config can override base config.
+	for (int32 Id = Pools.Num() - 1; Id >= 0; Id--)
 	{
 		const FVirtualTextureSpacePoolConfig& Config = Pools[Id];
 		if (Config.MinTileSize <= InTileSize && (Config.MaxTileSize == 0 || Config.MaxTileSize >= InTileSize) && InNumLayers == Config.Formats.Num())
@@ -44,6 +134,160 @@ void UVirtualTexturePoolConfig::FindPoolConfig(TEnumAsByte<EPixelFormat> const*
 		}
 	}
 
+	// Didn't find an exact match so return whatever default config that we first found.
 	OutConfig = DefaultConfig;
 }
 
+bool UVirtualTexturePoolConfig::AddOrModifyTransientPoolConfig(FVirtualTextureSpacePoolConfig const& InConfig)
+{
+	for (int32 Index = TransientPools.Num() - 1; Index >= 0; Index--)
+	{
+		FVirtualTextureSpacePoolConfig& Config = TransientPools[Index];
+		if (InConfig.Formats == Config.Formats && InConfig.MaxTileSize == Config.MaxTileSize && InConfig.MinTileSize == Config.MinTileSize)
+		{
+			if (Config.SizeInMegabyte != InConfig.SizeInMegabyte)
+			{
+				Config.SizeInMegabyte = InConfig.SizeInMegabyte;
+				Version++;
+				return true;
+			}
+
+			// No change.
+			return false;
+		}
+	}
+
+	Version++;
+	TransientPools.Add(InConfig);
+	return true;
+}
+
+#if WITH_EDITOR
+
+void UVirtualTexturePoolConfig::PostEditChangeProperty(struct FPropertyChangedEvent& PropertyChangedEvent)
+{
+	Super::PostEditChangeProperty(PropertyChangedEvent);
+	Version++;
+	VirtualTexture::Recreate();
+}
+
+#endif
+
+
+/**
+ * A task to modify the virtual texture physical pool configs.
+ * This needs queuing up on the game thread so that the virtual texture UObjects can be reinitialized correctly.
+ */
+class FAddOrModifyTransientPoolConfigsTask
+{
+	TArray<FVirtualTextureSpacePoolConfig> Configs;
+
+public:
+	FAddOrModifyTransientPoolConfigsTask(TArray<FVirtualTextureSpacePoolConfig>& InConfigs)
+		: Configs(MoveTemp(InConfigs))
+	{
+	}
+
+	FORCEINLINE TStatId GetStatId() const
+	{
+		RETURN_QUICK_DECLARE_CYCLE_STAT(FAddOrModifyTransientPoolConfigsTask, STATGROUP_TaskGraphTasks);
+	}
+
+	ENamedThreads::Type GetDesiredThread()
+	{
+		return ENamedThreads::GameThread;
+	}
+
+	static ESubsequentsMode::Type GetSubsequentsMode()
+	{
+		return ESubsequentsMode::FireAndForget;
+	}
+
+	void DoTask(ENamedThreads::Type CurrentThread, const FGraphEventRef& MyCompletionGraphEvent)
+	{
+		check(IsInGameThread());
+
+		TArray<int32, TInlineAllocator<8>> ChangedConfigs;
+
+		UVirtualTexturePoolConfig* PoolConfig = GetMutableDefault<UVirtualTexturePoolConfig>();
+		for (int32 ConfigIndex = 0; ConfigIndex < Configs.Num(); ++ConfigIndex)
+		{
+			if (PoolConfig->AddOrModifyTransientPoolConfig(Configs[ConfigIndex]))
+			{
+				ChangedConfigs.Add(ConfigIndex);
+			}
+		}
+
+		// Defer virtual texture recreation until all changes to config are done.
+		for (int32 ConfigIndex : ChangedConfigs)
+		{
+			VirtualTexture::Recreate(Configs[ConfigIndex].Formats);
+		}
+
+#if WITH_EDITOR
+		// Notify user about changes.
+		if (ChangedConfigs.Num())
+		{
+			static TWeakPtr<class SNotificationItem> NotificationHandle;
+			if (!NotificationHandle.IsValid())
+			{
+				FNotificationInfo Info(LOCTEXT("PoolResizeNotify", "Resizing Virtual Texture Pools."));
+				Info.SubText = LOCTEXT("PoolResizeNotifySubtext", "Size changes are not saved by default. To keep changes, copy from 'Transient Pools' to 'Fixed Pools' in the Virtual Texture Pool settings.");
+				Info.ExpireDuration = 8.0f;
+				Info.HyperlinkText = LOCTEXT("OpenSettings", "Open Project Settings");
+				Info.Hyperlink = FSimpleDelegate::CreateStatic([]() { FGlobalTabmanager::Get()->TryInvokeTab(FName("ProjectSettings")); });
+			
+				NotificationHandle = FSlateNotificationManager::Get().AddNotification(Info);
+			}
+		}
+#endif
+	}
+};
+
+void VirtualTexturePool::FindPoolConfig(TEnumAsByte<EPixelFormat> const* InFormats, int32 InNumLayers, int32 InTileSize, FVirtualTextureSpacePoolConfig& OutConfig)
+{
+	UVirtualTexturePoolConfig const* PoolConfig = GetDefault<UVirtualTexturePoolConfig>();
+	PoolConfig->FindPoolConfig(InFormats, InNumLayers, InTileSize, OutConfig);
+}
+
+void VirtualTexturePool::AddOrModifyTransientPoolConfigs_RenderThread(TArray<FVirtualTextureSpacePoolConfig>& InConfigs)
+{
+	check(IsInRenderingThread());
+	TGraphTask<FAddOrModifyTransientPoolConfigsTask>::CreateTask(nullptr).ConstructAndDispatchWhenReady(InConfigs);
+}
+
+float VirtualTexturePool::GetPoolSizeScale()
+{
+	return CVarVTPoolSizeScale.GetValueOnAnyThread();
+}
+
+bool VirtualTexturePool::GetPoolAutoGrow()
+{
+#if WITH_EDITOR
+	UVirtualTexturePoolConfig const* PoolConfig = GetDefault<UVirtualTexturePoolConfig>();
+	if (PoolConfig->bPoolAutoGrowInEditor)
+	{
+		return true;
+	}
+#endif
+	return CVarVTPoolAutoGrow.GetValueOnAnyThread();
+}
+
+int32 VirtualTexturePool::GetSplitPhysicalPoolSize()
+{
+	return FMath::Max(CVarVTSplitPhysicalPoolSize.GetValueOnAnyThread(), 0);
+}
+
+uint32 VirtualTexturePool::GetConfigHash()
+{
+	uint32 Hash = GetTypeHash(GetPoolSizeScale());
+	Hash = HashCombine(Hash, GetTypeHash(GetPoolAutoGrow()));
+	Hash = HashCombine(Hash, GetTypeHash(GetSplitPhysicalPoolSize()));
+
+	UVirtualTexturePoolConfig const* PoolConfig = GetDefault<UVirtualTexturePoolConfig>();
+	Hash = HashCombine(Hash, GetTypeHash(PoolConfig->Version));
+
+	return Hash;
+}
+
+#undef LOCTEXT_NAMESPACE
diff --git a/Engine/Source/Runtime/Engine/Private/VT/VirtualTextureRecreate.cpp b/Engine/Source/Runtime/Engine/Private/VT/VirtualTextureRecreate.cpp
new file mode 100644
--- /dev/null
+++ b/Engine/Source/Runtime/Engine/Private/VT/VirtualTextureRecreate.cpp
@@ -0,0 +1,86 @@
+// Copyright Epic Games, Inc. All Rights Reserved.
+
+#include "VT/VirtualTextureRecreate.h"
+
+#include "Components/RuntimeVirtualTextureComponent.h"
+#include "Engine/Texture2D.h"
+#include "TextureResource.h"
+#include "VT/RuntimeVirtualTexture.h"
+#include "UObject/UObjectIterator.h"
+
+namespace VirtualTexture
+{
+	// Release all VT render resources here.
+	// Assuming all virtual textures are released, then virtual texture pools will reach a zero ref count and release, which is needed for any pool size scale to be effective.
+	// Note that for pool size scale changes, there will be a transition period (with high memory watermark) where new pools are created before old pools are released.
+	void Recreate()
+	{
+		// Reinit streaming virtual textures.
+		for (TObjectIterator<UTexture2D> It(RF_ClassDefaultObject, true, EInternalObjectFlags::Garbage); It; ++It)
+		{
+			if (It->IsCurrentlyVirtualTextured())
+			{
+				if (FVirtualTexture2DResource* Resource = static_cast<FVirtualTexture2DResource*>(It->GetResource()))
+				{
+					BeginUpdateResourceRHI(Resource);
+				}
+			}
+		}
+
+		// Reinit runtime virtual textures.
+		for (TObjectIterator<URuntimeVirtualTextureComponent> It(RF_ClassDefaultObject, true, EInternalObjectFlags::Garbage); It; ++It)
+		{
+			It->MarkRenderStateDirty();
+		}
+	}
+
+	void Recreate(TConstArrayView < TEnumAsByte<EPixelFormat> > InFormat)
+	{
+		// Reinit streaming virtual textures that match one of our passed in format arrays.
+		for (TObjectIterator<UTexture2D> It(RF_ClassDefaultObject, true, EInternalObjectFlags::Garbage); It; ++It)
+		{
+			if (It->IsCurrentlyVirtualTextured())
+			{
+				if (FVirtualTexture2DResource* Resource = static_cast<FVirtualTexture2DResource*>(It->GetResource()))
+				{
+					bool bIsMatchingFormat = Resource->GetNumLayers() == InFormat.Num();
+					for (int32 LayerIndex = 0; bIsMatchingFormat && LayerIndex < InFormat.Num(); LayerIndex++)
+					{
+						if (Resource->GetFormat(LayerIndex) != InFormat[LayerIndex])
+						{
+							bIsMatchingFormat = false;
+							break;
+						}
+					}
+
+					if (bIsMatchingFormat)
+					{
+						BeginUpdateResourceRHI(Resource);
+					}
+				}
+			}
+		}
+
+		// Reinit runtime virtual textures that match one of our passed in format arrays.
+		for (TObjectIterator<URuntimeVirtualTextureComponent> It(RF_ClassDefaultObject, true, EInternalObjectFlags::Garbage); It; ++It)
+		{
+			if (URuntimeVirtualTexture* VirtualTexture = It->GetVirtualTexture())
+			{
+				bool bIsMatchingFormat = VirtualTexture->GetLayerCount() == InFormat.Num();
+				for (int32 LayerIndex = 0; bIsMatchingFormat && LayerIndex < InFormat.Num(); LayerIndex++)
+				{
+					if (VirtualTexture->GetLayerFormat(LayerIndex) != InFormat[LayerIndex])
+					{
+						bIsMatchingFormat = false;
+						break;
+					}
+				}
+
+				if (bIsMatchingFormat)
+				{
+					It->MarkRenderStateDirty();
+				}
+			}
+		}
+	}
+}
diff --git a/Engine/Source/Runtime/Engine/Private/VT/VirtualTextureRecreate.h b/Engine/Source/Runtime/Engine/Private/VT/VirtualTextureRecreate.h
new file mode 100644
--- /dev/null
+++ b/Engine/Source/Runtime/Engine/Private/VT/VirtualTextureRecreate.h
@@ -0,0 +1,14 @@
+// Copyright Epic Games, Inc. All Rights Reserved.
+
+#pragma once
+
+#include "RenderResource.h"
+
+namespace VirtualTexture
+{
+	/** Reinitialize all virtual texture producers. */
+	void Recreate();
+
+	/** Reinitialize all virtual texture that match the passed in format. */
+	void Recreate(TConstArrayView < TEnumAsByte<EPixelFormat> > InFormat);
+}
diff --git a/Engine/Source/Runtime/Engine/Private/VT/VirtualTextureScalability.cpp b/Engine/Source/Runtime/Engine/Private/VT/VirtualTextureScalability.cpp
--- a/Engine/Source/Runtime/Engine/Private/VT/VirtualTextureScalability.cpp
+++ b/Engine/Source/Runtime/Engine/Private/VT/VirtualTextureScalability.cpp
@@ -2,14 +2,9 @@
 
 #include "VT/VirtualTextureScalability.h"
 
-#include "Components/RuntimeVirtualTextureComponent.h"
-#include "EngineModule.h"
-#include "Engine/Texture2D.h"
-#include "RendererInterface.h"
-#include "RenderingThread.h"
-#include "TextureResource.h"
-#include "UObject/UObjectIterator.h"
-#include "VT/RuntimeVirtualTexture.h"
+#include "HAL/IConsoleManager.h"
+#include "VT/VirtualTexturePoolConfig.h"
+#include "VT/VirtualTextureRecreate.h"
 
 namespace VirtualTextureScalability
 {
@@ -45,13 +40,6 @@ namespace VirtualTextureScalability
 		ECVF_RenderThreadSafe | ECVF_Scalability
 	);
 
-	static TAutoConsoleVariable<int32> CVarVTMaxReleasedPerFrame(
-		TEXT("r.VT.MaxReleasedPerFrame"),
-		0,
-		TEXT("Max number of allocated virtual textures to release per frame"),
-		ECVF_RenderThreadSafe | ECVF_Scalability
-	);
-
 #if WITH_EDITOR
 	static TAutoConsoleVariable<int32> CVarVTMaxContinuousUpdatesPerFrameInEditor(
 		TEXT("r.VT.MaxContinuousUpdatesPerFrameInEditor"),
@@ -68,44 +56,22 @@ namespace VirtualTextureScalability
 		ECVF_RenderThreadSafe | ECVF_Scalability
 	);
 
-	static TAutoConsoleVariable<int32> CVarVTMaxAnisotropy(
-		TEXT("r.VT.MaxAnisotropy"),
-		8,
-		TEXT("MaxAnisotropy setting for Virtual Texture sampling."),
+	static TAutoConsoleVariable<int32> CVarVTMaxReleasedPerFrame(
+		TEXT("r.VT.MaxReleasedPerFrame"),
+		0,
+		TEXT("Max number of allocated virtual textures to release per frame"),
 		ECVF_RenderThreadSafe | ECVF_Scalability
 	);
 
-	static const int NumScalabilityGroups = 3;
+	static TAutoConsoleVariable<int32> CVarVTPageFreeThreshold(
+		TEXT("r.VT.PageFreeThreshold"),
+		60,
+		TEXT("Number of frames since the last time a VT page was used, before it's considered free.\n")
+		TEXT("VT pages are not necesarily marked as used on the CPU every time they're accessed by the GPU.\n")
+		TEXT("Increasing this threshold reduces the chances that an in-use frame is considered free."),
+		ECVF_RenderThreadSafe);
 
-	static float GPoolSizeScales[NumScalabilityGroups] = { 1.f, 1.f, 1.f };
-	static FAutoConsoleVariableRef CVarVTPoolSizeScale_ForBackwardsCompat(
-		TEXT("r.VT.PoolSizeScale"),
-		GPoolSizeScales[0],
-		TEXT("Scale factor for virtual texture physical pool size.\n")
-		TEXT(" Group 0"),
-		ECVF_Scalability | ECVF_ExcludeFromPreview
-	);
-	static FAutoConsoleVariableRef CVarVTPoolSizeScale0(
-		TEXT("r.VT.PoolSizeScale.Group0"),
-		GPoolSizeScales[0],
-		TEXT("Scale factor for virtual texture physical pool size.\n")
-		TEXT(" Group 0"),
-		ECVF_Scalability | ECVF_ExcludeFromPreview
-	);
-	static FAutoConsoleVariableRef CVarVTPoolSizeScale1(
-		TEXT("r.VT.PoolSizeScale.Group1"),
-		GPoolSizeScales[1],
-		TEXT("Scale factor for virtual texture physical pool sizes.\n")
-		TEXT(" Group 1"),
-		ECVF_Scalability | ECVF_ExcludeFromPreview
-	);
-	static FAutoConsoleVariableRef CVarVTPoolSizeScale2(
-		TEXT("r.VT.PoolSizeScale.Group2"),
-		GPoolSizeScales[2],
-		TEXT("Scale factor for virtual texture physical pool sizes.\n")
-		TEXT(" Group 2"),
-		ECVF_Scalability | ECVF_ExcludeFromPreview
-	);
+	static const int NumScalabilityGroups = 3;
 
 	static float GTileCountBiases[NumScalabilityGroups] = { 0 };
 	static FAutoConsoleVariableRef CVarVTTileCountBias_ForBackwardsCompat(
@@ -137,14 +103,12 @@ namespace VirtualTextureScalability
 		ECVF_Scalability
 	);
 
-	static TAutoConsoleVariable<int32> CVarVTSplitPhysicalPoolSize(
-		TEXT("r.VT.SplitPhysicalPoolSize"),
-		0,
-		TEXT("Create multiple physical pools per format to keep pools at this maximum size in tiles.")
-		TEXT("A value of 64 tiles will force 16bit page tables. This can be a page table memory optimization for large physical pools.")
-		TEXT("Defaults to 0 (off)."),
-		ECVF_RenderThreadSafe
-	);
+	static TAutoConsoleVariable<int32> CVarVTMobileEnableManualTrilinearFiltering(
+		TEXT("r.VT.Mobile.ManualTrilinearFiltering"),
+		1,
+		TEXT("Whether to use a manual trilinear filtering for VTs on mobile platforms.\n")
+		TEXT("This more expensive filtering is used on mobile platforms that do not support Temporal Anti-Aliasing.\n"),
+		ECVF_RenderThreadSafe | ECVF_ReadOnly);
 
 	static TAutoConsoleVariable<int32> CVarVTEnableAnisotropy(
 		TEXT("r.VT.AnisotropicFiltering"),
@@ -152,28 +116,20 @@ namespace VirtualTextureScalability
 		TEXT("Is anisotropic filtering for VTs enabled?"),
 		ECVF_RenderThreadSafe | ECVF_ReadOnly);
 
-	static TAutoConsoleVariable<int32> CVarVTMobileEnableManualTrilinearFiltering(
-		TEXT("r.VT.Mobile.ManualTrilinearFiltering"),
-		1,
-		TEXT("Whether to use a manual trilinear filtering for VTs on mobile platforms.\n")
-		TEXT("This more expensive filtering is used on mobile platforms that do not support Temporal Anti-Aliasing.\n"),
-		ECVF_RenderThreadSafe | ECVF_ReadOnly);
+	static TAutoConsoleVariable<int32> CVarVTMaxAnisotropy(
+		TEXT("r.VT.MaxAnisotropy"),
+		8,
+		TEXT("MaxAnisotropy setting for Virtual Texture sampling."),
+		ECVF_RenderThreadSafe | ECVF_Scalability
+	);
+
 
-	static TAutoConsoleVariable<int32> CVarVTPageFreeThreshold(
-		TEXT("r.VT.PageFreeThreshold"),
-		60,
-		TEXT("Number of frames since the last time a VT page was used, before it's considered free.\n")
-		TEXT("VT pages are not necesarily marked as used on the CPU every time they're accessed by the GPU.\n")
-		TEXT("Increasing this threshold reduces the chances that an in-use frame is considered free."),
-		ECVF_RenderThreadSafe);
-	
 	/** Track changes and apply to relevant systems. This allows us to dynamically change the scalability settings. */
 	static void OnUpdate()
 	{
 		const float MaxAnisotropy = CVarVTMaxAnisotropy.GetValueOnGameThread();
 
 		static float LastMaxAnisotropy = MaxAnisotropy;
-		static float LastPoolSizeScales[3] = { GPoolSizeScales[0], GPoolSizeScales[1], GPoolSizeScales[2] };
 		static float LastTileCountBiases[3] = { GTileCountBiases[0], GTileCountBiases[1], GTileCountBiases[2] };
 
 		bool bUpdate = false;
@@ -182,13 +138,6 @@ namespace VirtualTextureScalability
 			LastMaxAnisotropy = MaxAnisotropy;
 			bUpdate = true;
 		}
-		if (LastPoolSizeScales[0] != GPoolSizeScales[0] || LastPoolSizeScales[1] != GPoolSizeScales[1] || LastPoolSizeScales[2] != GPoolSizeScales[2])
-		{
-			LastPoolSizeScales[0] = GPoolSizeScales[0];
-			LastPoolSizeScales[1] = GPoolSizeScales[1];
-			LastPoolSizeScales[2] = GPoolSizeScales[2];
-			bUpdate = true;
-		}
 		if (LastTileCountBiases[0] != GTileCountBiases[0] || LastTileCountBiases[1] != GTileCountBiases[1] || LastTileCountBiases[2] != GTileCountBiases[2])
 		{
 			LastTileCountBiases[0] = GTileCountBiases[0];
@@ -199,24 +148,7 @@ namespace VirtualTextureScalability
 
 		if (bUpdate)
 		{
-			// Release all VT render resources here.
-			// Assuming all virtual textures are released, then virtual texture pools will reach a zero ref count and release, which is needed for any pool size scale to be effective.
-			// Note that for pool size scale changes, there will be a transition period (with high memory watermark) where new pools are created before old pools are released.
-
-			// Reinit streaming virtual textures.
-			for (TObjectIterator<UTexture2D> It(RF_ClassDefaultObject, false, EInternalObjectFlags::Garbage); It; ++It)
-			{
-				if (It->IsCurrentlyVirtualTextured() && It->GetResource() != nullptr)
-				{
-					BeginUpdateResourceRHI(It->GetResource());
-				}
-			}
-
-			// Reinit runtime virtual textures.
-			for (TObjectIterator<URuntimeVirtualTextureComponent> It(RF_ClassDefaultObject, false, EInternalObjectFlags::Garbage); It; ++It)
-			{
-				It->MarkRenderStateDirty();
-			}
+			VirtualTexture::Recreate();
 		}
 	}
 
@@ -254,11 +186,6 @@ namespace VirtualTextureScalability
 		return CVarMaxPagesProducedPerFrame.GetValueOnAnyThread();
 	}
 
-	int32 GetMaxAllocatedVTReleasedPerFrame()
-	{
-		return CVarVTMaxReleasedPerFrame.GetValueOnAnyThread();
-	}
-
 	int32 GetMaxContinuousUpdatesPerFrame()
 	{
 #if WITH_EDITOR
@@ -269,42 +196,52 @@ namespace VirtualTextureScalability
 #endif
 	}
 
-	bool IsAnisotropicFilteringEnabled()
+	int32 GetMaxAllocatedVTReleasedPerFrame()
 	{
-		return CVarVTEnableAnisotropy.GetValueOnAnyThread() != 0;
+#if WITH_EDITOR
+		return 0;
+#else
+		return CVarVTMaxReleasedPerFrame.GetValueOnAnyThread();
+#endif
 	}
 
-	int32 GetMaxAnisotropy()
+	uint32 GetPageFreeThreshold()
 	{
-		return CVarVTMaxAnisotropy.GetValueOnAnyThread();
+		return FMath::Max(CVarVTPageFreeThreshold.GetValueOnRenderThread(), 0);
 	}
 
-	float GetPoolSizeScale(uint32 GroupIndex)
+	int32 GetRuntimeVirtualTextureSizeBias(uint32 GroupIndex)
 	{
-		// This is called on render thread but uses non render thread cvar. However it should be safe enough due to the calling pattern.
-		// Using ECVF_RenderThreadSafe would mean that OnUpdate() logic can fail to detect a change due to the cvar ref pointing at the render thread value.
-		return GroupIndex < NumScalabilityGroups ? GPoolSizeScales[GroupIndex] : 1.f;
+		return GroupIndex < NumScalabilityGroups ? GTileCountBiases[GroupIndex] : 0;
 	}
 
-	int32 GetRuntimeVirtualTextureSizeBias(uint32 GroupIndex)
+	bool IsAnisotropicFilteringEnabled()
 	{
-		return GroupIndex < NumScalabilityGroups ? GTileCountBiases[GroupIndex] : 0;
+		return CVarVTEnableAnisotropy.GetValueOnAnyThread() != 0;
 	}
 
-	int32 GetSplitPhysicalPoolSize()
+	int32 GetMaxAnisotropy()
 	{
-		return FMath::Max(CVarVTSplitPhysicalPoolSize.GetValueOnRenderThread(), 0);
+		return CVarVTMaxAnisotropy.GetValueOnAnyThread();
 	}
 
-	uint32 GetPageFreeThreshold()
+	// Begin deprecated functions.
+	// Can remove include of VirtualTexturePoolConfig.h when these are removed.
+	float GetPoolSizeScale() 
 	{
-		return FMath::Max(CVarVTPageFreeThreshold.GetValueOnRenderThread(), 0);
+		return VirtualTexturePool::GetPoolSizeScale(); 
+	}
+	float GetPoolSizeScale(uint32 GroupIndex) 
+	{
+		return VirtualTexturePool::GetPoolSizeScale(); 
+	}
+	int32 GetSplitPhysicalPoolSize() 
+	{
+		return VirtualTexturePool::GetSplitPhysicalPoolSize(); 
 	}
-	
-	uint32 GetPhysicalPoolSettingsHash()
+	uint32 GetPhysicalPoolSettingsHash() 
 	{
-		uint32 Hash = GetArrayHash(GPoolSizeScales, UE_ARRAY_COUNT(GPoolSizeScales));
-		Hash = HashCombine(Hash, GetTypeHash(GetSplitPhysicalPoolSize()));
-		return Hash;
+		return VirtualTexturePool::GetConfigHash(); 
 	}
+	// End deprecated functions.
 }
diff --git a/Engine/Source/Runtime/Engine/Public/VT/VirtualTextureScalability.h b/Engine/Source/Runtime/Engine/Public/VT/VirtualTextureScalability.h
--- a/Engine/Source/Runtime/Engine/Public/VT/VirtualTextureScalability.h
+++ b/Engine/Source/Runtime/Engine/Public/VT/VirtualTextureScalability.h
@@ -16,26 +16,23 @@ namespace VirtualTextureScalability
 	ENGINE_API int32 GetMaxContinuousUpdatesPerFrame();
 	/** Get max allocated virtual textures to release per frame. */
 	ENGINE_API int32 GetMaxAllocatedVTReleasedPerFrame();
-	/** Get scale factor for virtual texture physical pool sizes. */
-	ENGINE_API float GetPoolSizeScale();
+	/** Get the number of frames a page must be unused, before it's considered free */
+	ENGINE_API uint32 GetPageFreeThreshold();
 	/** Get resolution bias for runtime virtual textures. */
 	ENGINE_API int32 GetRuntimeVirtualTextureSizeBias();
+	/** Get resolution bias for runtime virtual textures. */
+	ENGINE_API int32 GetRuntimeVirtualTextureSizeBias(uint32 GroupIndex);
 	/** Is HW Anisotropic filtering enabled for VT */
 	ENGINE_API bool IsAnisotropicFilteringEnabled();
-	/**
-	 * Get maximum anisotropy when virtual texture sampling. 
-	 * This is also clamped per virtual texture according to the tile border size.
-	 */
+	/** Get maximum anisotropy when virtual texture sampling. This is also clamped per virtual texture according to the tile border size. */
 	ENGINE_API int32 GetMaxAnisotropy();
-	/** Get scale factor for virtual texture physical pool sizes. */
+
+	UE_DEPRECATED(5.4, "Use VirtualTexturePool::GetPoolSizeScale() instead.")
+	ENGINE_API float GetPoolSizeScale();
+	UE_DEPRECATED(5.4, "Use VirtualTexturePool::GetPoolSizeScale() instead.")
 	ENGINE_API float GetPoolSizeScale(uint32 GroupIndex);
-	/** Get resolution bias for runtime virtual textures. */
-	ENGINE_API int32 GetRuntimeVirtualTextureSizeBias(uint32 GroupIndex);
-	/** Get maximum size in tiles for physical pools before we split them. */
+	UE_DEPRECATED(5.4, "Use VirtualTexturePool::GetSplitPhysicalPoolSize() instead.")
 	ENGINE_API int32 GetSplitPhysicalPoolSize();
-	/** Get the number of frames a page must be unused, before it's considered free */
-	ENGINE_API uint32 GetPageFreeThreshold();
-
-	/** Get a unique hash of all state affecting physical pool allocation.*/
+	UE_DEPRECATED(5.4, "Use VirtualTexturePool::GetConfigHash() instead.")
 	ENGINE_API uint32 GetPhysicalPoolSettingsHash();
 }
diff --git a/Engine/Source/Runtime/Renderer/Private/VT/VirtualTexturePhysicalSpace.cpp b/Engine/Source/Runtime/Renderer/Private/VT/VirtualTexturePhysicalSpace.cpp
--- a/Engine/Source/Runtime/Renderer/Private/VT/VirtualTexturePhysicalSpace.cpp
+++ b/Engine/Source/Runtime/Renderer/Private/VT/VirtualTexturePhysicalSpace.cpp
@@ -53,12 +53,11 @@ static TAutoConsoleVariable<float> CVarVTResidencyAdjustmentRate(
 	ECVF_RenderThreadSafe);
 
 
-FVirtualTexturePhysicalSpace::FVirtualTexturePhysicalSpace(const FVTPhysicalSpaceDescription& InDesc, uint16 InID, int32 InTileWidthHeight, bool bInEnableResidencyMipMapBias)
+FVirtualTexturePhysicalSpace::FVirtualTexturePhysicalSpace(uint16 InID, const FVTPhysicalSpaceDescription& InDesc, FVTPhysicalSpaceDescriptionExt& InDescExt)
 	: Description(InDesc)
-	, TextureSizeInTiles(InTileWidthHeight)
-	, NumRefs(0u)
+	, DescriptionExt(InDescExt)
 	, ID(InID)
-	, bEnableResidencyMipMapBias(bInEnableResidencyMipMapBias)
+	, NumRefs(0u)
 	, ResidencyMipMapBias(0.0f)
 	, LastFrameOversubscribed(0)
 #if !UE_BUILD_SHIPPING
@@ -186,14 +185,26 @@ void FVirtualTexturePhysicalSpace::FinalizeTextures(FRDGBuilder& GraphBuilder)
 	}
 }
 
-uint32 FVirtualTexturePhysicalSpace::GetSizeInBytes() const
+uint32 FVirtualTexturePhysicalSpace::GetTileSizeInBytes() const
 {
 	SIZE_T TileSizeBytes = 0;
 	for (int32 Layer = 0; Layer < Description.NumLayers; ++Layer)
 	{
 		TileSizeBytes += CalculateImageBytes(Description.TileSize, Description.TileSize, 0, Description.Format[Layer]);
 	}
-	return GetNumTiles() * TileSizeBytes;
+	return TileSizeBytes;
+
+}
+
+uint32 FVirtualTexturePhysicalSpace::GetSizeInBytes() const
+{
+	SIZE_T TextureSizeBytes = 0;
+	const uint32 TextureSize = GetTextureSize();
+	for (int32 Layer = 0; Layer < Description.NumLayers; ++Layer)
+	{
+		TextureSizeBytes += CalculateImageBytes(TextureSize, TextureSize, 0, Description.Format[Layer]);
+	}
+	return TextureSizeBytes;
 }
 
 void FVirtualTexturePhysicalSpace::UpdateResidencyTracking(uint32 Frame)
@@ -229,7 +240,7 @@ void FVirtualTexturePhysicalSpace::UpdateResidencyTracking(uint32 Frame)
 
 	ResidencyMipMapBias = FMath::Clamp(ResidencyMipMapBias, 0.f, MaxMipMapBias);
 
-	if (!bEnableResidencyMipMapBias || LockedPageResidency > LockedUpperBound)
+	if (!DescriptionExt.bEnableResidencyMipMapBias || LockedPageResidency > LockedUpperBound)
 	{
 		ResidencyMipMapBias = 0.f;
 	}
diff --git a/Engine/Source/Runtime/Renderer/Private/VT/VirtualTexturePhysicalSpace.h b/Engine/Source/Runtime/Renderer/Private/VT/VirtualTexturePhysicalSpace.h
--- a/Engine/Source/Runtime/Renderer/Private/VT/VirtualTexturePhysicalSpace.h
+++ b/Engine/Source/Runtime/Renderer/Private/VT/VirtualTexturePhysicalSpace.h
@@ -10,9 +10,9 @@
 #include "VirtualTextureShared.h"
 #include "VirtualTexturing.h"
 
+/** Description that determines a unique physical space. This is filled out by a producer. */
 struct FVTPhysicalSpaceDescription
 {
-	uint32 PoolSettingsHash;
 	uint32 TileSize;
 	uint8 Dimensions;
 	uint8 NumLayers;
@@ -23,8 +23,7 @@ struct FVTPhysicalSpaceDescription
 
 inline bool operator==(const FVTPhysicalSpaceDescription& Lhs, const FVTPhysicalSpaceDescription& Rhs)
 {
-	if (Lhs.PoolSettingsHash != Rhs.PoolSettingsHash ||
-		Lhs.TileSize != Rhs.TileSize ||
+	if (Lhs.TileSize != Rhs.TileSize ||
 		Lhs.NumLayers != Rhs.NumLayers || 
 		Lhs.Dimensions != Rhs.Dimensions || 
 		Lhs.bContinuousUpdate != Rhs.bContinuousUpdate)
@@ -53,8 +52,7 @@ inline bool operator!=(const FVTPhysicalSpaceDescription& Lhs, const FVTPhysical
 
 inline uint32 GetTypeHash(const FVTPhysicalSpaceDescription& Desc)
 {
-	uint32 Hash = Desc.PoolSettingsHash;
-	Hash = HashCombine(Hash, GetTypeHash(Desc.TileSize));
+	uint32 Hash = Desc.TileSize;
 	Hash = HashCombine(Hash, GetTypeHash(Desc.Dimensions));
 	Hash = HashCombine(Hash, GetTypeHash(Desc.NumLayers));
 	Hash = HashCombine(Hash, GetTypeHash(Desc.bContinuousUpdate));
@@ -66,24 +64,40 @@ inline uint32 GetTypeHash(const FVTPhysicalSpaceDescription& Desc)
 	return Hash;
 }
 
+/** Additional part of the description for a unique physical space that is derived from the virtual texture pool config. */
+struct FVTPhysicalSpaceDescriptionExt
+{
+	int32 TileWidthHeight = 0;
+	int32 PoolCount = 1;
+	bool bEnableResidencyMipMapBias = false;
+};
+
+inline bool operator==(const FVTPhysicalSpaceDescriptionExt& Lhs, const FVTPhysicalSpaceDescriptionExt& Rhs)
+{
+	return Lhs.TileWidthHeight == Rhs.TileWidthHeight && Lhs.PoolCount == Rhs.PoolCount && Lhs.bEnableResidencyMipMapBias == Rhs.bEnableResidencyMipMapBias;
+}
+
+/** Virtual texture physical space that contains the physical texture and the page pool that tracks residency. */
 class FVirtualTexturePhysicalSpace final : public FRenderResource
 {
 public:
-	FVirtualTexturePhysicalSpace(const FVTPhysicalSpaceDescription& InDesc, uint16 InID, int32 InTileWidthHeight, bool bInEnableResidencyMipMapBias);
+	FVirtualTexturePhysicalSpace(uint16 InID, const FVTPhysicalSpaceDescription& InDesc, FVTPhysicalSpaceDescriptionExt& InDescExt);
 	virtual ~FVirtualTexturePhysicalSpace();
 
 	inline const FVTPhysicalSpaceDescription& GetDescription() const { return Description; }
+	inline const FVTPhysicalSpaceDescriptionExt& GetDescriptionExt() const { return DescriptionExt; }
 	inline const FString& GetFormatString() const { return FormatString; }
 	inline EPixelFormat GetFormat(int32 Layer) const { return Description.Format[Layer]; }
 	inline uint16 GetID() const { return ID; }
-	inline uint32 GetNumTiles() const { return TextureSizeInTiles * TextureSizeInTiles; }
-	inline uint32 GetSizeInTiles() const { return TextureSizeInTiles; }
-	inline uint32 GetTextureSize() const { return TextureSizeInTiles * Description.TileSize; }
-	inline FIntVector GetPhysicalLocation(uint16 pAddress) const { return FIntVector(pAddress % TextureSizeInTiles, pAddress / TextureSizeInTiles, 0); }
+	inline uint32 GetSizeInTiles() const { return DescriptionExt.TileWidthHeight; }
+	inline uint32 GetNumTiles() const { return GetSizeInTiles() * GetSizeInTiles(); }
+	inline uint32 GetTextureSize() const { return GetSizeInTiles() * Description.TileSize; }
+	inline FIntVector GetPhysicalLocation(uint16 pAddress) const { return FIntVector(pAddress % GetSizeInTiles(), pAddress / GetSizeInTiles(), 0); }
 
 	// 16bit page tables allocate 6bits to address TileX/Y, so can only address tiles from 0-63
-	inline bool DoesSupport16BitPageTable() const { return TextureSizeInTiles <= 64u; }
+	inline bool DoesSupport16BitPageTable() const { return GetSizeInTiles() <= 64u; }
 
+	uint32 GetTileSizeInBytes() const;
 	uint32 GetSizeInBytes() const;
 
 	inline const FTexturePagePool& GetPagePool() const { return Pool; }
@@ -129,17 +143,16 @@ public:
 
 private:
 	FVTPhysicalSpaceDescription Description;
+	FVTPhysicalSpaceDescriptionExt DescriptionExt;
 	FString FormatString;
 	FTexturePagePool Pool;
 	TRefCountPtr<IPooledRenderTarget> PooledRenderTarget[VIRTUALTEXTURE_SPACE_MAXLAYERS];
 	FShaderResourceViewRHIRef TextureSRV[VIRTUALTEXTURE_SPACE_MAXLAYERS];
 	FShaderResourceViewRHIRef TextureSRV_SRGB[VIRTUALTEXTURE_SPACE_MAXLAYERS];
 
-	uint32 TextureSizeInTiles;
-	uint32 NumRefs;
 	uint16 ID;
+	uint32 NumRefs;
 	
-	bool bEnableResidencyMipMapBias;
 	float ResidencyMipMapBias;
 	uint32 LastFrameOversubscribed;
 
diff --git a/Engine/Source/Runtime/Renderer/Private/VT/VirtualTextureProducer.cpp b/Engine/Source/Runtime/Renderer/Private/VT/VirtualTextureProducer.cpp
--- a/Engine/Source/Runtime/Renderer/Private/VT/VirtualTextureProducer.cpp
+++ b/Engine/Source/Runtime/Renderer/Private/VT/VirtualTextureProducer.cpp
@@ -3,7 +3,6 @@
 #include "VirtualTextureProducer.h"
 
 #include "VT/VirtualTexturePhysicalSpace.h"
-#include "VT/VirtualTextureScalability.h"
 #include "VT/VirtualTextureSystem.h"
 
 FVirtualTextureProducer::~FVirtualTextureProducer()
@@ -73,7 +72,6 @@ FVirtualTextureProducerHandle FVirtualTextureProducerCollection::RegisterProduce
 	Entry.DestroyedCallbacksIndex = AcquireCallback();
 
 	FVTPhysicalSpaceDescription PhysicalSpaceDesc;
-	PhysicalSpaceDesc.PoolSettingsHash = VirtualTextureScalability::GetPhysicalPoolSettingsHash();
 	PhysicalSpaceDesc.Dimensions = InDesc.Dimensions;
 	PhysicalSpaceDesc.TileSize = InDesc.TileSize + InDesc.TileBorderSize * 2u;
 	PhysicalSpaceDesc.bContinuousUpdate = InDesc.bContinuousUpdate;
diff --git a/Engine/Source/Runtime/Renderer/Private/VT/VirtualTextureSystem.cpp b/Engine/Source/Runtime/Renderer/Private/VT/VirtualTextureSystem.cpp
--- a/Engine/Source/Runtime/Renderer/Private/VT/VirtualTextureSystem.cpp
+++ b/Engine/Source/Runtime/Renderer/Private/VT/VirtualTextureSystem.cpp
@@ -909,27 +909,16 @@ void FVirtualTextureSystem::ReleaseSpace(FVirtualTextureSpace* Space)
 	}
 }
 
-/** 
- * Description of additional settings to use when initializing a physical space. 
- * This is filed by GetPoolInitDescription() based on the requested space description and the current config settings.
- */
-struct FVTPhysicalSpaceInitDescription
+/** Get the extra physical space description that depends on the virtual texture pool config. */
+void GetPhysicalSpaceExtraDescription(FVTPhysicalSpaceDescription const& InDesc, FVTPhysicalSpaceDescriptionExt& OutDescExt)
 {
-	int32 TileWidthHeight = 0;
-	int32 PoolCount = 1;
-	bool bEnableResidencyMipBias = false;
-};
-
-void GetPoolInitDescription(FVTPhysicalSpaceDescription const& InDesc, FVTPhysicalSpaceInitDescription& OutInitDescription)
-{
-	// Find matching config from ini file
+	// Find matching config from pool settings.
 	FVirtualTextureSpacePoolConfig Config;
-	UVirtualTexturePoolConfig const* PoolConfig = GetDefault<UVirtualTexturePoolConfig>();
-	PoolConfig->FindPoolConfig(InDesc.Format, InDesc.NumLayers, InDesc.TileSize, Config);
+	VirtualTexturePool::FindPoolConfig(InDesc.Format, InDesc.NumLayers, InDesc.TileSize, Config);
 	int32 SizeInMegabyte = Config.SizeInMegabyte;
 
 	// Adjust found config for scaling.
-	const float Scale = Config.bAllowSizeScale ? VirtualTextureScalability::GetPoolSizeScale(Config.ScalabilityGroup) : 1.f;
+	const float Scale = Config.bAllowSizeScale ? VirtualTexturePool::GetPoolSizeScale() : 1.f;
 	SizeInMegabyte = (int32)(Scale * (float)SizeInMegabyte);
 	if (Scale < 1.f && Config.MinScaledSizeInMegabyte > 0)
 	{
@@ -969,7 +958,7 @@ void GetPoolInitDescription(FVTPhysicalSpaceDescription const& InDesc, FVTPhysic
 			break;
 		}
 
-		const int32 SplitPhysicalPoolSize = VirtualTextureScalability::GetSplitPhysicalPoolSize();
+		const int32 SplitPhysicalPoolSize = VirtualTexturePool::GetSplitPhysicalPoolSize();
 		if (SplitPhysicalPoolSize <= 0 || TileWidthHeight <= SplitPhysicalPoolSize)
 		{
 			break;
@@ -978,32 +967,34 @@ void GetPoolInitDescription(FVTPhysicalSpaceDescription const& InDesc, FVTPhysic
 		PoolCount++;
 	}
 
-	OutInitDescription.TileWidthHeight = TileWidthHeight;
-	OutInitDescription.PoolCount = PoolCount;
-	OutInitDescription.bEnableResidencyMipBias = Config.bEnableResidencyMipMapBias;
+	OutDescExt.TileWidthHeight = TileWidthHeight;
+	OutDescExt.PoolCount = PoolCount;
+	OutDescExt.bEnableResidencyMipMapBias = Config.bEnableResidencyMipMapBias;
 }
 
-/** Cached version of GetPoolInitDescription() to avoid regularly repeating the heavy work in that function. */
-void GetPoolInitDescription_Cached(FVTPhysicalSpaceDescription const& InDesc, FVTPhysicalSpaceInitDescription& OutInitDescription)
+/** Cached version of GetPhysicalSpaceExtraDescription() to avoid regularly repeating the heavy work in that function. */
+void GetPhysicalSpaceExtraDescription_Cached(FVTPhysicalSpaceDescription const& InDesc, FVTPhysicalSpaceDescriptionExt& OutDescExt)
 {
-	static TMap<FVTPhysicalSpaceDescription, FVTPhysicalSpaceInitDescription> Map;
+	static TMap<FVTPhysicalSpaceDescription, FVTPhysicalSpaceDescriptionExt> Map;
 
-	// Invalidate the cache if any relevant CVar settings change.
-	static uint32 PhysicalPoolSettingsHash = VirtualTextureScalability::GetPhysicalPoolSettingsHash();
-	if (PhysicalPoolSettingsHash != VirtualTextureScalability::GetPhysicalPoolSettingsHash())
+	// Invalidate the cache if any config settings change.
+	uint32 PhysicalPoolSettingsHash = VirtualTexturePool::GetConfigHash();
+	static uint32 LastPhysicalPoolSettingsHash = PhysicalPoolSettingsHash;
+	if (LastPhysicalPoolSettingsHash != PhysicalPoolSettingsHash)
 	{
+		LastPhysicalPoolSettingsHash = PhysicalPoolSettingsHash;
 		Map.Reset();
 	}
 
-	FVTPhysicalSpaceInitDescription* InitDescriptionPtr = Map.Find(InDesc);
+	FVTPhysicalSpaceDescriptionExt* InitDescriptionPtr = Map.Find(InDesc);
 	if (InitDescriptionPtr == nullptr)
 	{
-		GetPoolInitDescription(InDesc, OutInitDescription);
-		Map.Add(InDesc, OutInitDescription);
+		GetPhysicalSpaceExtraDescription(InDesc, OutDescExt);
+		Map.Add(InDesc, OutDescExt);
 	}
 	else
 	{
-		OutInitDescription = *InitDescriptionPtr;
+		OutDescExt = *InitDescriptionPtr;
 	}
 }
 
@@ -1011,22 +1002,23 @@ FVirtualTexturePhysicalSpace* FVirtualTextureSystem::AcquirePhysicalSpace(FRHICo
 {
 	LLM_SCOPE(ELLMTag::VirtualTextureSystem);
 
+	// Get extra setup information from the virtual pool configs.
+	FVTPhysicalSpaceDescriptionExt DescExt;
+	GetPhysicalSpaceExtraDescription_Cached(InDesc, DescExt);
+
 	// Find matching pools.
 	// We support multiple matching pools to allow for 16bit page table memory optimization.
 	TArray<int32, TInlineAllocator<8>> Matching;
 	for (int32 i = 0; i < PhysicalSpaces.Num(); ++i)
 	{
 		FVirtualTexturePhysicalSpace* PhysicalSpace = PhysicalSpaces[i];
-		if (PhysicalSpace && PhysicalSpace->GetDescription() == InDesc)
+		if (PhysicalSpace && PhysicalSpace->GetDescription() == InDesc && PhysicalSpace->GetDescriptionExt() == DescExt)
 		{
 			Matching.Add(i);
 		}
 	}
 
-	FVTPhysicalSpaceInitDescription InitDescription;
-	GetPoolInitDescription_Cached(InDesc, InitDescription);
-
-	if (InitDescription.PoolCount <= Matching.Num())
+	if (DescExt.PoolCount <= Matching.Num())
 	{
 		// Randomly select from any pools that exist.
 		int32 RandomIndex = FMath::RandHelper(Matching.Num());
@@ -1051,7 +1043,7 @@ FVirtualTexturePhysicalSpace* FVirtualTextureSystem::AcquirePhysicalSpace(FRHICo
 		PhysicalSpaces.AddZeroed();
 	}
 
-	FVirtualTexturePhysicalSpace* PhysicalSpace = new FVirtualTexturePhysicalSpace(InDesc, ID, InitDescription.TileWidthHeight, InitDescription.bEnableResidencyMipBias);
+	FVirtualTexturePhysicalSpace* PhysicalSpace = new FVirtualTexturePhysicalSpace(ID, InDesc, DescExt);
 	PhysicalSpaces[ID] = PhysicalSpace;
 
 	INC_MEMORY_STAT_BY(STAT_TotalPhysicalMemory, PhysicalSpace->GetSizeInBytes());
@@ -1934,6 +1926,41 @@ void FVirtualTextureSystem::UpdateResidencyTracking() const
 	}
 }
 
+void FVirtualTextureSystem::GrowPhysicalPools() const
+{
+	if (!VirtualTexturePool::GetPoolAutoGrow())
+	{
+		return;
+	}
+
+	TArray<FVirtualTextureSpacePoolConfig> Configs;
+	for (int32 i = 0; i < PhysicalSpaces.Num(); ++i)
+	{
+		FVirtualTexturePhysicalSpace* PhysicalSpace = PhysicalSpaces[i];
+		if (PhysicalSpace && PhysicalSpace->GetLastFrameOversubscribed() == Frame)
+		{
+			FVirtualTextureSpacePoolConfig& Config = Configs.AddDefaulted_GetRef();
+			const FVTPhysicalSpaceDescription& Desc = PhysicalSpace->GetDescription();
+			Config.Formats.Append(Desc.Format, Desc.NumLayers);
+			Config.MaxTileSize = Config.MinTileSize = Desc.TileSize;
+
+			// Increase pool by 1 tile or 4MB, whichever is greater.
+			const int32 TileCount = PhysicalSpace->GetSizeInTiles();
+			const int32 TileSizeInBytes = PhysicalSpace->GetTileSizeInBytes();
+			const int32 CurrentSizeInBytes = TileCount * TileCount * TileSizeInBytes;
+			const int32 NextSizeInBytes = (TileCount + 1) * (TileCount + 1) * TileSizeInBytes;
+			const int32 MinIncreaseInBytes = 4 * 1024 * 1024;
+			const int32 ClampedNextSizeInBytes = FMath::Max(NextSizeInBytes, CurrentSizeInBytes + MinIncreaseInBytes);
+			Config.SizeInMegabyte = FMath::DivideAndRoundUp(ClampedNextSizeInBytes, 1024 * 1024);
+		}
+	}
+
+	if (Configs.Num())
+	{
+		VirtualTexturePool::AddOrModifyTransientPoolConfigs_RenderThread(Configs);
+	}
+}
+
 void FVirtualTextureSystem::SubmitRequestsFromLocalTileList(FRHICommandList& RHICmdList, TArray<FVirtualTextureLocalTile>& OutDeferredTiles, const TSet<FVirtualTextureLocalTile>& LocalTileList, EVTProducePageFlags Flags, ERHIFeatureLevel::Type FeatureLevel, uint32 MaxRequestsToProduce)
 {
 	LLM_SCOPE(ELLMTag::VirtualTextureSystem);
@@ -2832,6 +2859,8 @@ void FVirtualTextureSystem::EndUpdate(FRDGBuilder& GraphBuilder, TUniquePtr<FVir
 	Producers.NotifyRequestsCompleted();
 
 	UpdateResidencyTracking();
+	GrowPhysicalPools();
+
 
 #if !UE_BUILD_SHIPPING
 	UpdateCsvStats();
diff --git a/Engine/Source/Runtime/Renderer/Private/VT/VirtualTextureSystem.h b/Engine/Source/Runtime/Renderer/Private/VT/VirtualTextureSystem.h
--- a/Engine/Source/Runtime/Renderer/Private/VT/VirtualTextureSystem.h
+++ b/Engine/Source/Runtime/Renderer/Private/VT/VirtualTextureSystem.h
@@ -227,6 +227,7 @@ private:
 	void GetContinuousUpdatesToProduce(FUniqueRequestList const* RequestList, int32 MaxTilesToProduce, int32 MaxContinuousUpdates);
 
 	void UpdateResidencyTracking() const;
+	void GrowPhysicalPools() const;
 
 #if WITH_EDITOR
 	void RecordPageRequests(FUniquePageList const* UniquePageList, TSet<uint64>& OutPages);
