Fixed errors when advancing sync marker phases

Addresses the following:
- Not keeping track of previous or next marker when advancing a marker forwards or backwards respectively. This is only an issue when two or more markers are passed in a single tick.
- When advancing current time to match next or previous marker their TimeToMarker should be 0 since the sequence's current position is that of the matched marker's time.

[RN] minor
#10969

diff --git a/Engine/Source/Runtime/Engine/Private/Animation/AnimSequence.cpp b/Engine/Source/Runtime/Engine/Private/Animation/AnimSequence.cpp
--- a/Engine/Source/Runtime/Engine/Private/Animation/AnimSequence.cpp
+++ b/Engine/Source/Runtime/Engine/Private/Animation/AnimSequence.cpp
@@ -2882,7 +2882,7 @@ void UAnimSequence::AdvanceMarkerPhaseAsLeader(bool bLooping, float MoveDelta, c
 		// Repeat until there is no more move delta to handle.
 		while (true)
 		{
-			// Our next marker is the end boundary.
+			// Our next marker is the end boundary. (Only possible if sequence is not looping)
 			if (NextMarker.MarkerIndex == MarkerIndexSpecialValues::AnimationBoundary)
 			{
 				const float PrevCurrentTime = CurrentTime;
@@ -2909,7 +2909,7 @@ void UAnimSequence::AdvanceMarkerPhaseAsLeader(bool bLooping, float MoveDelta, c
 
 				// Make our new previous marker be the marker we just passed.
 				PrevMarker.MarkerIndex = NextMarker.MarkerIndex; 
-				PrevMarker.TimeToMarker = -CurrentMoveDelta;
+				PrevMarker.TimeToMarker = 0.0f;
 
 				// Record that we just passed a marker.
 				const int32 PassedMarker = MarkersPassed.Add(FPassedMarker());
@@ -2996,7 +2996,7 @@ void UAnimSequence::AdvanceMarkerPhaseAsLeader(bool bLooping, float MoveDelta, c
 
 				// Make our new next marker be the marker we just passed.
 				NextMarker.MarkerIndex = PrevMarker.MarkerIndex;
-				NextMarker.TimeToMarker = -CurrentMoveDelta;
+				NextMarker.TimeToMarker = 0.0f;
 
 				// Record that we just passed a marker.
 				const int32 PassedMarker = MarkersPassed.Add(FPassedMarker());
@@ -3591,6 +3591,8 @@ void UAnimSequence::AdvanceMarkerPhaseAsFollower(const FMarkerTickContext& Conte
 		// Ensure next marker matches leader's next marker after tick.
 		if (NextMarker.MarkerIndex != MarkerIndexSpecialValues::AnimationBoundary && Context.MarkersPassedThisTick.Num() > 0)
 		{
+			PreviousMarker.MarkerIndex = NextMarker.MarkerIndex;
+			
 			AdvanceMarkerForwards(NextMarker.MarkerIndex, LeaderEndPosition.NextMarkerName, bLooping, AuthoredSyncMarkers, MirrorTable);
 		}
 
@@ -3657,6 +3659,8 @@ void UAnimSequence::AdvanceMarkerPhaseAsFollower(const FMarkerTickContext& Conte
 		// Ensure previous marker match leader's previous marker after tick
 		if (PreviousMarker.MarkerIndex != MarkerIndexSpecialValues::AnimationBoundary && Context.MarkersPassedThisTick.Num() > 0)
 		{
+			NextMarker.MarkerIndex = PreviousMarker.MarkerIndex;
+			
 			AdvanceMarkerBackwards(PreviousMarker.MarkerIndex, LeaderEndPosition.PreviousMarkerName, bLooping, AuthoredSyncMarkers, MirrorTable);
 		}
 
