Disabled MinKeysForResampling when resampling is not selected in anim compression settings

diff --git a/Engine/Source/Runtime/Engine/Classes/Animation/AnimCompress_PerTrackCompression.h b/Engine/Source/Runtime/Engine/Classes/Animation/AnimCompress_PerTrackCompression.h
--- a/Engine/Source/Runtime/Engine/Classes/Animation/AnimCompress_PerTrackCompression.h
+++ b/Engine/Source/Runtime/Engine/Classes/Animation/AnimCompress_PerTrackCompression.h
@@ -57,7 +57,7 @@ class UAnimCompress_PerTrackCompression : public UAnimCompress_RemoveLinearKeys
 	float ResampledFramerate;
 
 	/** Animations with fewer keys than MinKeysForResampling will not be resampled. */
-	UPROPERTY(EditAnywhere, Category=Resampling)
+	UPROPERTY(EditAnywhere, Category=Resampling, meta=(editcondition = "bResampleAnimation"))
 	int32 MinKeysForResampling;
 
 	/** If true, adjust the error thresholds based on the 'height' within the skeleton */
