Added clamping to the UI for "Remove Every Second Key" anim compression's Min Keys property

diff --git a/Engine/Source/Runtime/Engine/Classes/Animation/AnimCompress_RemoveEverySecondKey.h b/Engine/Source/Runtime/Engine/Classes/Animation/AnimCompress_RemoveEverySecondKey.h
--- a/Engine/Source/Runtime/Engine/Classes/Animation/AnimCompress_RemoveEverySecondKey.h
+++ b/Engine/Source/Runtime/Engine/Classes/Animation/AnimCompress_RemoveEverySecondKey.h
@@ -18,7 +18,7 @@ class UAnimCompress_RemoveEverySecondKey : public UAnimCompress
 	GENERATED_UCLASS_BODY()
 
 	/** Animations with fewer than MinKeys will not lose any keys. */
-	UPROPERTY(EditAnywhere, Category=AnimationCompressionAlgorithm_RemoveEverySecondKey)
+	UPROPERTY(EditAnywhere, Category=AnimationCompressionAlgorithm_RemoveEverySecondKey, meta=(UIMin=1, ClampMin=1))
 	int32 MinKeys;
 
 	/**
