Exposed AddAnimNotifies publicly in FAnimInstanceProxy

Intended to be used by custom anim nodes

diff --git a/Engine/Source/Runtime/Engine/Public/Animation/AnimInstanceProxy.h b/Engine/Source/Runtime/Engine/Public/Animation/AnimInstanceProxy.h
--- a/Engine/Source/Runtime/Engine/Public/Animation/AnimInstanceProxy.h
+++ b/Engine/Source/Runtime/Engine/Public/Animation/AnimInstanceProxy.h
@@ -549,6 +549,13 @@ public:
 	/** Get the debug data for this instance's anim bp */
 	ENGINE_API FAnimBlueprintDebugData* GetAnimBlueprintDebugData() const;
 
+	/**
+	 * Add anim notifies to the proxies notify queue
+	 * @param NewNotifies		The notifies to add
+	 * @param InstanceWeight	The effective weight of the notifies (used for trigger filtering)
+	 **/
+	ENGINE_API void AddAnimNotifies(const TArray<FAnimNotifyEventReference>& NewNotifies, const float InstanceWeight);
+
 	/** Only restricted classes can access the protected interface */
 	friend class UAnimInstance;
 	friend class UAnimSingleNodeInstance;
@@ -693,9 +700,6 @@ protected:
 		return BufferWriteIndex; 
 	}
 
-	/** Add anim notifier **/
-	ENGINE_API void AddAnimNotifies(const TArray<FAnimNotifyEventReference>& NewNotifies, const float InstanceWeight);
-
 	/** Returns the baked sync group index from the compile step */
 	ENGINE_API int32 GetSyncGroupIndexFromName(FName SyncGroupName) const;
 
