Correctly use the key reduction thresholds in "Remove Trivial Keys" animation compression

#5933

diff --git a/Engine/Source/Runtime/Engine/Private/Animation/AnimCompress_RemoveTrivialKeys.cpp b/Engine/Source/Runtime/Engine/Private/Animation/AnimCompress_RemoveTrivialKeys.cpp
--- a/Engine/Source/Runtime/Engine/Private/Animation/AnimCompress_RemoveTrivialKeys.cpp
+++ b/Engine/Source/Runtime/Engine/Private/Animation/AnimCompress_RemoveTrivialKeys.cpp
@@ -13,9 +13,9 @@ UAnimCompress_RemoveTrivialKeys::UAnimCompress_RemoveTrivialKeys(const FObjectIn
 	: Super(ObjectInitializer)
 {
 	Description = TEXT("Remove Trivial Keys");
-	MaxPosDiff = 0.0001f;
-	MaxAngleDiff = 0.0003f;
-	MaxScaleDiff = 0.00001f;
+	MaxPosDiff = TRANSLATION_ZEROING_THRESHOLD;
+	MaxAngleDiff = QUATERNION_ZEROING_THRESHOLD;
+	MaxScaleDiff = SCALE_ZEROING_THRESHOLD;
 }
 
 #if WITH_EDITOR
@@ -29,7 +29,7 @@ bool UAnimCompress_RemoveTrivialKeys::DoReduction(const FCompressibleAnimData& C
 	SeparateRawDataIntoTracks( CompressibleAnimData.RawAnimationData, CompressibleAnimData.SequenceLength, TranslationData, RotationData, ScaleData );
 	
 	// remove obviously redundant keys from the source data
-	FilterTrivialKeys(TranslationData, RotationData, ScaleData, TRANSLATION_ZEROING_THRESHOLD, QUATERNION_ZEROING_THRESHOLD, SCALE_ZEROING_THRESHOLD);
+	FilterTrivialKeys(TranslationData, RotationData, ScaleData, MaxPosDiff, MaxAngleDiff, MaxScaleDiff);
 
 	// record the proper runtime decompressor to use
 	FUECompressedAnimDataMutable& AnimData = static_cast<FUECompressedAnimDataMutable&>(*OutResult.AnimData);
