Fix warning V547: Expression '* Buffer == 0' is always false.

diff --git a/Engine/Source/Runtime/Engine/Private/EdGraph/EdGraphPin.cpp b/Engine/Source/Runtime/Engine/Private/EdGraph/EdGraphPin.cpp
--- a/Engine/Source/Runtime/Engine/Private/EdGraph/EdGraphPin.cpp
+++ b/Engine/Source/Runtime/Engine/Private/EdGraph/EdGraphPin.cpp
@@ -1146,14 +1146,14 @@ bool UEdGraphPin::ImportTextItem(const TCHAR*& Buffer, int32 PortFlags, class UO
 		while (*Buffer != '=')
 		{
 			Buffer++;
-		}
-
-		if (*Buffer == 0)
-		{
-			ErrorText->Logf(ELogVerbosity::Warning, TEXT("%s: Unexpected end-of-stream while importing property values."), *GetName());
 
-			// Parse error
-			return false;
+			if (*Buffer == 0)
+			{
+				ErrorText->Logf(ELogVerbosity::Warning, TEXT("%s: Unexpected end-of-stream while importing property values."), *GetName());
+	
+				// Parse error
+				return false;
+			}
 		}
 
 		int32 NumCharsInToken = Buffer - StartBuffer;
