Fix custom instance data on mobile with a temporary rollback of 29980807

diff --git a/Engine/Source/Runtime/Engine/Private/Materials/HLSLMaterialTranslator.cpp b/Engine/Source/Runtime/Engine/Private/Materials/HLSLMaterialTranslator.cpp
--- a/Engine/Source/Runtime/Engine/Private/Materials/HLSLMaterialTranslator.cpp
+++ b/Engine/Source/Runtime/Engine/Private/Materials/HLSLMaterialTranslator.cpp
@@ -250,7 +250,6 @@ struct FHLSLMaterialTranslator::FEnvironmentDefines
 	bool bUsesEyeAdaptation;
 	bool bVirtualTextureOutput;
 	bool bUsesPerInstanceCustomData;
-	bool bUsesPerInstanceCustomDataPS;
 	bool bUsesPerInstanceRandomPS;
 	bool bUsesPerInstanceFadeAmount;
 	bool bUsesVertexInterpolator;
@@ -2488,8 +2487,6 @@ void FHLSLMaterialTranslator::GetMaterialEnvironment(EShaderPlatform InPlatform,
 	}
 
 	OutEnvironment.SetDefine(TEXT("USES_PER_INSTANCE_CUSTOM_DATA"), EnvironmentDefines->bUsesPerInstanceCustomData);
-	// Whether we need to pass custom data paramaters to PS through interpolators
-	OutEnvironment.SetDefine(TEXT("USES_PER_INSTANCE_CUSTOM_DATA_PS"), EnvironmentDefines->bUsesPerInstanceCustomDataPS && Material->IsUsedWithInstancedStaticMeshes());
 
 	OutEnvironment.SetDefine(TEXT("USES_PER_INSTANCE_FADE_AMOUNT"), EnvironmentDefines->bUsesPerInstanceFadeAmount);
 	OutEnvironment.SetDefine(TEXT("USES_VERTEX_INTERPOLATOR"), EnvironmentDefines->bUsesVertexInterpolator);
@@ -14215,7 +14212,6 @@ void FHLSLMaterialTranslator::PrepareEnvironmentDefines()
 	EnvironmentDefines->bUsesEyeAdaptation = MaterialCompilationOutput.bUsesEyeAdaptation;
 	EnvironmentDefines->bVirtualTextureOutput = MaterialCompilationOutput.bHasRuntimeVirtualTextureOutputNode;
 	EnvironmentDefines->bUsesPerInstanceCustomData = MaterialCompilationOutput.bUsesPerInstanceCustomData && Material->IsUsedWithInstancedStaticMeshes();
-	EnvironmentDefines->bUsesPerInstanceCustomDataPS = bUsesPerInstanceCustomDataPS && Material->IsUsedWithInstancedStaticMeshes();
 	EnvironmentDefines->bUsesPerInstanceRandomPS = bUsesPerInstanceRandomPS && Material->IsUsedWithInstancedStaticMeshes();
 	EnvironmentDefines->bUsesPerInstanceFadeAmount = bUsesPerInstanceFadeAmount && Material->IsUsedWithInstancedStaticMeshes();
 	EnvironmentDefines->bUsesVertexInterpolator = MaterialCompilationOutput.bUsesVertexInterpolator;
