Don't detach skeletal mesh component from its actor when bodies are set to simulate. Implement IsSimulatingPhysics for SkeletalMeshComponent so that it distinguishes based on whether or not the SKM transform is driven by the simulation of the root body.
Fixes problems with the skeletal mesh detaching from the actor, both in standalone, and also when networked. Adds some additional info/warning to the weld option when attaching components/actors.

diff --git a/Engine/Source/Runtime/Engine/Classes/Components/SceneComponent.h b/Engine/Source/Runtime/Engine/Classes/Components/SceneComponent.h
--- a/Engine/Source/Runtime/Engine/Classes/Components/SceneComponent.h
+++ b/Engine/Source/Runtime/Engine/Classes/Components/SceneComponent.h
@@ -729,7 +729,7 @@ public:
 	* @param  LocationRule				How to handle translation when attaching.
 	* @param  RotationRule				How to handle rotation when attaching.
 	* @param  ScaleRule					How to handle scale when attaching.
-	* @param  bWeldSimulatedBodies		Whether to weld together simulated physics bodies.
+	* @param  bWeldSimulatedBodies		Whether to weld together simulated physics bodies. This transfers the shapes in the welded object into the parent (if simulated), which can result in permanent changes that persist even after subsequently detaching.
 	* @return True if attachment is successful (or already attached to requested parent/socket), false if attachment is rejected and there is no change in AttachParent.
 	*/
 	UFUNCTION(BlueprintCallable, Category = "Transformation", meta = (DisplayName = "Attach Component To Component", ScriptName = "AttachToComponent", bWeldSimulatedBodies=true))
@@ -741,7 +741,7 @@ public:
 	ENGINE_API virtual void DetachFromParent(bool bMaintainWorldPosition = false, bool bCallModify = true);
 
 	/** 
-	 * Detach this component from whatever it is attached to. Automatically unwelds components that are welded together (See WeldTo)
+	 * Detach this component from whatever it is attached to. Automatically unwelds components that are welded together (see AttachToComponent), though note that some effects of welding may not be undone.
 	 * @param LocationRule				How to handle translations when detaching.
 	 * @param RotationRule				How to handle rotation when detaching.
 	 * @param ScaleRule					How to handle scales when detaching.
@@ -751,7 +751,7 @@ public:
 	ENGINE_API void K2_DetachFromComponent(EDetachmentRule LocationRule = EDetachmentRule::KeepRelative, EDetachmentRule RotationRule = EDetachmentRule::KeepRelative, EDetachmentRule ScaleRule = EDetachmentRule::KeepRelative, bool bCallModify = true);
 
 	/** 
-	 * Detach this component from whatever it is attached to. Automatically unwelds components that are welded together (See WeldTo)
+	 * Detach this component from whatever it is attached to. Automatically unwelds components that are welded together (See AttachToComponent), though note that some effects of welding may not be undone.
 	 * @param DetachmentRules			How to handle transforms & modification when detaching.
 	 */
 	ENGINE_API virtual void DetachFromComponent(const FDetachmentTransformRules& DetachmentRules);
diff --git a/Engine/Source/Runtime/Engine/Classes/Components/SkeletalMeshComponent.h b/Engine/Source/Runtime/Engine/Classes/Components/SkeletalMeshComponent.h
--- a/Engine/Source/Runtime/Engine/Classes/Components/SkeletalMeshComponent.h
+++ b/Engine/Source/Runtime/Engine/Classes/Components/SkeletalMeshComponent.h
@@ -1753,6 +1753,7 @@ public:
 
 	ENGINE_API virtual class UBodySetup* GetBodySetup() override;
 	ENGINE_API virtual bool CanEditSimulatePhysics() override;
+	ENGINE_API virtual bool IsSimulatingPhysics(FName BoneName = NAME_None) const override;
 	ENGINE_API virtual FBodyInstance* GetBodyInstance(FName BoneName = NAME_None, bool bGetWelded = true, int32 Index = INDEX_NONE) const override;
 	ENGINE_API virtual void UpdatePhysicsToRBChannels() override;
 	ENGINE_API virtual void SetAllPhysicsAngularVelocityInRadians(FVector const& NewVel, bool bAddToCurrent = false) override;
@@ -2098,6 +2099,10 @@ public:
 	UFUNCTION(BlueprintCallable, Category="Physics")
 	ENGINE_API void SetAllBodiesBelowSimulatePhysics(const FName& InBoneName, bool bNewSimulate, bool bIncludeSelf = true );
 
+	/** Set a single bone to be simulated (or not) */
+	UFUNCTION(BlueprintCallable, Category="Physics")
+	ENGINE_API void SetBodySimulatePhysics(const FName& InBoneName, bool bSimulate);
+
 	/** Allows you to reset bodies Simulate state based on where bUsePhysics is set to true in the BodySetup. */
 	UFUNCTION(BlueprintCallable, Category="Physics")
 	ENGINE_API void ResetAllBodiesSimulatePhysics();
diff --git a/Engine/Source/Runtime/Engine/Classes/GameFramework/Actor.h b/Engine/Source/Runtime/Engine/Classes/GameFramework/Actor.h
--- a/Engine/Source/Runtime/Engine/Classes/GameFramework/Actor.h
+++ b/Engine/Source/Runtime/Engine/Classes/GameFramework/Actor.h
@@ -1809,7 +1809,7 @@ public:
 	 * @param LocationRule				How to handle translation when attaching.
 	 * @param RotationRule				How to handle rotation when attaching.
 	 * @param ScaleRule					How to handle scale when attaching.
-	 * @param bWeldSimulatedBodies		Whether to weld together simulated physics bodies.
+	 * @param bWeldSimulatedBodies		Whether to weld together simulated physics bodies. This transfers the shapes in the welded object into the parent (if simulated), which can result in permanent changes that persist even after subsequently detaching.
 	 * @return							Whether the attachment was successful or not
 	 */
 	UFUNCTION(BlueprintCallable, meta = (DisplayName = "Attach Actor To Component", ScriptName = "AttachToComponent", bWeldSimulatedBodies = true), Category = "Transformation")
@@ -1845,7 +1845,7 @@ public:
 	 * @param LocationRule				How to handle translation when attaching.
 	 * @param RotationRule				How to handle rotation when attaching.
 	 * @param ScaleRule					How to handle scale when attaching.
-	 * @param bWeldSimulatedBodies		Whether to weld together simulated physics bodies.
+	 * @param bWeldSimulatedBodies		Whether to weld together simulated physics bodies.This transfers the shapes in the welded object into the parent (if simulated), which can result in permanent changes that persist even after subsequently detaching.
 	 * @return							Whether the attachment was successful or not
 	 */
 	UFUNCTION(BlueprintCallable, meta = (DisplayName = "Attach Actor To Actor", ScriptName = "AttachToActor", bWeldSimulatedBodies=true), Category = "Transformation")
diff --git a/Engine/Source/Runtime/Engine/Private/SkeletalMeshComponentPhysics.cpp b/Engine/Source/Runtime/Engine/Private/SkeletalMeshComponentPhysics.cpp
--- a/Engine/Source/Runtime/Engine/Private/SkeletalMeshComponentPhysics.cpp
+++ b/Engine/Source/Runtime/Engine/Private/SkeletalMeshComponentPhysics.cpp
@@ -223,7 +223,7 @@ void USkeletalMeshComponent::SetSimulatePhysics(bool bSimulate)
 				{
 					if (PhysAssetBodySetup->PhysicsType == EPhysicsType::PhysType_Default)
 					{
-						BodyInst->SetInstanceSimulatePhysics(bSimulate);
+						BodyInst->SetInstanceSimulatePhysics(bSimulate, false, true);
 					}
 				}
 			}
@@ -1091,7 +1091,7 @@ void USkeletalMeshComponent::SetAllBodiesSimulatePhysics(bool bNewSimulate)
 {
 	for(int32 i=0; i<Bodies.Num(); i++)
 	{
-		Bodies[i]->SetInstanceSimulatePhysics(bNewSimulate);
+		Bodies[i]->SetInstanceSimulatePhysics(bNewSimulate, false, true);
 	}
 
 	SetRootBodyIndex(RootBodyData.BodyIndex);	//Update the root body data cache in case animation has moved root body relative to root joint
@@ -1167,7 +1167,7 @@ void USkeletalMeshComponent::SetAllBodiesBelowSimulatePhysics( const FName& InBo
 {
 	int32 NumBodiesFound = ForEachBodyBelow(InBoneName, bIncludeSelf, /*bSkipCustomPhysicsType=*/ false, [bNewSimulate](FBodyInstance* BI)
 	{
-		BI->SetInstanceSimulatePhysics(bNewSimulate);
+		BI->SetInstanceSimulatePhysics(bNewSimulate, false, true);
 	});
 
 	if (NumBodiesFound)
@@ -1182,6 +1182,22 @@ void USkeletalMeshComponent::SetAllBodiesBelowSimulatePhysics( const FName& InBo
 	}
 }
 
+void USkeletalMeshComponent::SetBodySimulatePhysics(const FName& InBoneName, bool bSimulate)
+{
+	FBodyInstance* BI = GetBodyInstance(InBoneName);
+	if (BI)
+	{
+		BI->SetInstanceSimulatePhysics(bSimulate, false, true);
+
+		if (IsSimulatingPhysics())
+		{
+			SetRootBodyIndex(RootBodyData.BodyIndex);	//Update the root body data cache in case animation has moved root body relative to root joint
+		}
+
+		UpdateEndPhysicsTickRegisteredState();
+		UpdateClothTickRegisteredState();
+	}
+}
 
 void USkeletalMeshComponent::SetAllMotorsAngularPositionDrive(bool bEnableSwingDrive, bool bEnableTwistDrive, bool bSkipCustomPhysicsType)
 {
@@ -1375,11 +1391,11 @@ void USkeletalMeshComponent::ResetAllBodiesSimulatePhysics()
 			{
 				if (BodyInstSetup->PhysicsType == PhysType_Simulated)
 				{
-					BodyInst->SetInstanceSimulatePhysics(true);
+					BodyInst->SetInstanceSimulatePhysics(true, false, true);
 				}
 				else
 				{
-					BodyInst->SetInstanceSimulatePhysics(false);
+					BodyInst->SetInstanceSimulatePhysics(false, false, true);
 				}
 			}
 		}
@@ -1715,7 +1731,7 @@ void USkeletalMeshComponent::UpdateMeshForBrokenConstraints()
 						if( !ChildBodyInst->IsInstanceSimulatingPhysics() )
 						{
 							DEBUGBROKENCONSTRAINTUPDATE(UE_LOG(LogSkeletalMesh, Log, TEXT("      Unfixing body."));)
-							ChildBodyInst->SetInstanceSimulatePhysics(true);
+							ChildBodyInst->SetInstanceSimulatePhysics(true, false, true);
 						}
 					}
 
@@ -1763,6 +1779,25 @@ FName USkeletalMeshComponent::FindConstraintBoneName( int32 ConstraintIndex )
 	return PhysicsAsset ? PhysicsAsset->FindConstraintBoneName(ConstraintIndex) : NAME_None;
 }
 
+bool USkeletalMeshComponent::IsSimulatingPhysics(FName BoneName) const
+{
+	// If no bone name is specified, then we respond referring to the component.
+	// If the component is not set to follow physics, then the component is not controlled by simulation.
+	if (BoneName == NAME_None && 
+		PhysicsTransformUpdateMode == EPhysicsTransformUpdateMode::ComponentTransformIsKinematic)
+	{
+		return false;
+	}
+
+	// We respond based on either the body (if a bone is specified), or the root body (if no bone is
+	// specified, and the component is controlled by simulation).
+	FBodyInstance* BI = GetBodyInstance(BoneName);
+	if (BI)
+	{
+		return BI->IsInstanceSimulatingPhysics();
+	}
+	return false;
+}
 
 FBodyInstance* USkeletalMeshComponent::GetBodyInstance(FName BoneName, bool, int32) const
 {
@@ -1774,7 +1809,7 @@ FBodyInstance* USkeletalMeshComponent::GetBodyInstance(FName BoneName, bool, int
 		// A name of NAME_None indicates 'root body'
 		if(BoneName == NAME_None)
 		{
-			if(Bodies.IsValidIndex(RootBodyData.BodyIndex))
+			if (Bodies.IsValidIndex(RootBodyData.BodyIndex))
 			{
 				BodyInst = Bodies[RootBodyData.BodyIndex];
 			}
@@ -2005,7 +2040,7 @@ void USkeletalMeshComponent::BreakConstraint(FVector Impulse, FVector HitLocatio
 	if( Body != NULL && !Body->IsInstanceSimulatingPhysics() )
 	{
 		// Unfix body so it can be broken.
-		Body->SetInstanceSimulatePhysics(true);
+		Body->SetInstanceSimulatePhysics(true, false, true);
 	}
 
 	// Break Constraint
@@ -2040,7 +2075,7 @@ void USkeletalMeshComponent::SetAngularLimits(FName InBoneName, float Swing1Limi
 	if (Body != NULL && Body->IsInstanceSimulatingPhysics())
 	{
 		// Unfix body so it can be broken.
-		Body->SetInstanceSimulatePhysics(true);
+		Body->SetInstanceSimulatePhysics(true, false, true);
 	}
 
 	// update limits
