Only precache PSOs when already precached on mesh changes on UStaticMeshComponent otherwise there is still a pending PSO precache which will happen during PostLoad

diff --git a/Engine/Source/Runtime/Engine/Private/Components/StaticMeshComponent.cpp b/Engine/Source/Runtime/Engine/Private/Components/StaticMeshComponent.cpp
--- a/Engine/Source/Runtime/Engine/Private/Components/StaticMeshComponent.cpp
+++ b/Engine/Source/Runtime/Engine/Private/Components/StaticMeshComponent.cpp
@@ -2325,8 +2325,13 @@ bool UStaticMeshComponent::SetStaticMesh(UStaticMesh* NewMesh)
 		);
 	}
 
-	// Precache the PSOs
-	PrecachePSOs();
+#if UE_WITH_PSO_PRECACHING
+	// Precache the PSOs if already precache otherwise will happen later on during PostLoad
+	if (bPSOPrecacheCalled)
+	{ 
+		PrecachePSOs();
+	}
+#endif // UE_WITH_PSO_PRECACHING
 
 	// Need to send this to render thread at some point
 	if (IsRenderStateCreated())
