Support welding to kinematic root bodies.

This can significantly reduce the cost of moving a kinematic actor that has a large number of attached components.

NOTE: this is still not supported by the AttachTo methods. Welding to kinematics can only be done via C++.

diff --git a/Engine/Source/Runtime/Engine/Classes/Components/PrimitiveComponent.h b/Engine/Source/Runtime/Engine/Classes/Components/PrimitiveComponent.h
--- a/Engine/Source/Runtime/Engine/Classes/Components/PrimitiveComponent.h
+++ b/Engine/Source/Runtime/Engine/Classes/Components/PrimitiveComponent.h
@@ -2241,14 +2241,23 @@ public:
 	*	Welding allows the child physics object to become physically connected to its parent. This is useful for creating compound rigid bodies with correct mass distribution.
 	*   @param InParent the component to be physically attached to
 	*   @param InSocketName optional socket to attach component to
+	*	@param bWeldToKinematicParent if true, children will be welded onto the parent even if the parent is kinematic (default false)
+	*
+	* By default if the root is kinematic then the welded bodies are set to kinematic rather than actually welded. This is beneficial if
+	* the actor is not moved very often or only contains a few welded shapes. However it can be expensive to move a kinematic actor
+	* that contains a large number of (unwelded) kinematic children, in which case you can setting bWeldToKinematicParent to true to generate
+	* a welded kinematic actor. There is no real benefit to welding if the actor does not move and, since the initial weld cost is fairly high,
+	* you generally would not enable welding on all kinematics in the world if you have a lot of them.
 	*/
-	ENGINE_API virtual void WeldTo(class USceneComponent* InParent, FName InSocketName = NAME_None);
+	ENGINE_API virtual void WeldTo(class USceneComponent* InParent, FName InSocketName = NAME_None, bool bWeldToKinematicParent = false);
 
 	/**
 	*	Does the actual work for welding.
+	*	@param bWeldSimulatedChild if true, simulated children will be welded onto the parent (default true)
+	*	@param bWeldToKinematicParent if true, children will be welded onto the parent even if the parent is kinematic (default false)
 	*	@return true if did a true weld of shapes, meaning body initialization is not needed
 	*/
-	ENGINE_API virtual bool WeldToImplementation(USceneComponent * InParent, FName ParentSocketName = NAME_None, bool bWeldSimulatedChild = true);
+	ENGINE_API virtual bool WeldToImplementation(USceneComponent * InParent, FName ParentSocketName = NAME_None, bool bWeldSimulatedChild = true, bool bWeldToKinematicParent = false);
 
 	/**
 	*   UnWelds this component from its parent component. Attachment is maintained (DetachFromParent automatically unwelds)
diff --git a/Engine/Source/Runtime/Engine/Private/PrimitiveComponentPhysics.cpp b/Engine/Source/Runtime/Engine/Private/PrimitiveComponentPhysics.cpp
--- a/Engine/Source/Runtime/Engine/Private/PrimitiveComponentPhysics.cpp
+++ b/Engine/Source/Runtime/Engine/Private/PrimitiveComponentPhysics.cpp
@@ -780,7 +780,7 @@ void UPrimitiveComponent::GetWeldedBodies(TArray<FBodyInstance*> & OutWeldedBodi
 	}
 }
 
-bool UPrimitiveComponent::WeldToImplementation(USceneComponent * InParent, FName ParentSocketName /* = Name_None */, bool bWeldSimulatedChild /* = false */)
+bool UPrimitiveComponent::WeldToImplementation(USceneComponent * InParent, FName ParentSocketName /* = Name_None */, bool bWeldSimulatedChild /* = true */, bool bWeldToKinematicParent /* = false */)
 {
 	SCOPE_CYCLE_COUNTER(STAT_WeldPhysics);
 
@@ -829,7 +829,7 @@ bool UPrimitiveComponent::WeldToImplementation(USceneComponent * InParent, FName
 			//Child always inherits from root
 
 			//if root is kinematic simply set child to be kinematic and we're done
-			if (RootComponent->IsSimulatingPhysics(SocketName) == false)
+			if ((RootComponent->IsSimulatingPhysics(SocketName) == false) && (bWeldToKinematicParent == false))
 			{
 				FPlatformAtomics::InterlockedExchangePtr((void**)&BI->WeldParent, nullptr);
 				SetSimulatePhysics(false);
@@ -847,7 +847,7 @@ bool UPrimitiveComponent::WeldToImplementation(USceneComponent * InParent, FName
 	return false;
 }
 
-void UPrimitiveComponent::WeldTo(USceneComponent* InParent, FName InSocketName /* = NAME_None */)
+void UPrimitiveComponent::WeldTo(USceneComponent* InParent, FName InSocketName /* = NAME_None */, bool bWeldToKinematicParent /* = false */)
 {
 	//automatically attach if needed
 	if (GetAttachParent() != InParent || GetAttachSocketName() != InSocketName)
@@ -855,7 +855,8 @@ void UPrimitiveComponent::WeldTo(USceneComponent* InParent, FName InSocketName /
 		AttachToComponent(InParent, FAttachmentTransformRules::KeepWorldTransform, InSocketName);
 	}
 
-	WeldToImplementation(InParent, InSocketName);
+	const bool bWeldSimulatedChild = true;
+	WeldToImplementation(InParent, InSocketName, bWeldSimulatedChild, bWeldToKinematicParent);
 }
 
 void UPrimitiveComponent::UnWeldFromParent()
