Fix Incremental UnrealEditor MacArm64 compilation.

diff --git a/Engine/Source/Runtime/Engine/Public/ExternalDirtyActorsTracker.h b/Engine/Source/Runtime/Engine/Public/ExternalDirtyActorsTracker.h
--- a/Engine/Source/Runtime/Engine/Public/ExternalDirtyActorsTracker.h
+++ b/Engine/Source/Runtime/Engine/Public/ExternalDirtyActorsTracker.h
@@ -32,8 +32,8 @@ public:
 
 	const MapType& GetDirtyActors() const { return DirtyActors; }
 
-	virtual void OnRemoveInvalidDirtyActor(const TWeakObjectPtr<AActor> InActor, StoreType::Type& InValue) {}
-	virtual void OnRemoveNonDirtyActor(const TWeakObjectPtr<AActor> InActor, StoreType::Type& InValue) {}
+	virtual void OnRemoveInvalidDirtyActor(const TWeakObjectPtr<AActor> InActor, typename StoreType::Type& InValue) {}
+	virtual void OnRemoveNonDirtyActor(const TWeakObjectPtr<AActor> InActor, typename StoreType::Type& InValue) {}
 
 protected:
 	void OnPackageDirtyStateChanged(UPackage* InPackage)
