Fix branching point notifies anim montages on load that were saved with a null LinkedMontage ptr

Effectively reorders calls to FAnimNotifyEvent::Link and RefreshCachedData in PostLoad, ensuring that LinkedMontage is always valid.
The issue was that prior to 5.1 (and the fix in 22268778) for some workflows LinkSequence() was being called, resulting in a null LinkedMontage ptr being serialized.

diff --git a/Engine/Source/Runtime/Engine/Private/Animation/AnimMontage.cpp b/Engine/Source/Runtime/Engine/Private/Animation/AnimMontage.cpp
--- a/Engine/Source/Runtime/Engine/Private/Animation/AnimMontage.cpp
+++ b/Engine/Source/Runtime/Engine/Private/Animation/AnimMontage.cpp
@@ -400,6 +400,29 @@ FFrameRate UAnimMontage::GetSamplingFrameRate() const
 
 void UAnimMontage::PostLoad()
 {
+	// Link notifies before we call the Super::PostLoad (and eventually RefreshCacheData). This is to ensure that branching points get correctly
+	// picked up as FAnimNotifyEvent::IsBranchingPoint() relies on the LinkedMontage being valid, and due to an issue where (since deprecated)
+	// LinkSequence() was called instead of LinkMontage() there exists content for which LinkedMontage is saved as a null reference.
+	for(FAnimNotifyEvent& Notify : Notifies)
+	{
+#if WITH_EDITORONLY_DATA
+		if(Notify.DisplayTime_DEPRECATED != 0.0f)
+		{
+			Notify.Clear();
+			Notify.Link(this, Notify.DisplayTime_DEPRECATED);
+		}
+		else
+#endif
+		{
+			Notify.Link(this, Notify.GetTime());
+		}
+
+		if(Notify.Duration != 0.0f)
+		{
+			Notify.EndLink.Link(this, Notify.GetTime() + Notify.Duration);
+		}
+	}
+	
 	Super::PostLoad();
 
 	// copy deprecated variable to new one, temporary code to keep data copied. Am deleting it right after this
@@ -508,26 +531,6 @@ void UAnimMontage::PostLoad()
 		}
 	}
 
-	for(FAnimNotifyEvent& Notify : Notifies)
-	{
-#if WITH_EDITORONLY_DATA
-		if(Notify.DisplayTime_DEPRECATED != 0.0f)
-		{
-			Notify.Clear();
-			Notify.Link(this, Notify.DisplayTime_DEPRECATED);
-		}
-		else
-#endif
-		{
-			Notify.Link(this, Notify.GetTime());
-		}
-
-		if(Notify.Duration != 0.0f)
-		{
-			Notify.EndLink.Link(this, Notify.GetTime() + Notify.Duration);
-		}
-	}
-
 	// Convert BranchingPoints to AnimNotifies.
 	if (GetLinker() && (GetLinker()->UEVer() < VER_UE4_MONTAGE_BRANCHING_POINT_REMOVAL) )
 	{
