[NetPhysics] Resim improvements
- Allow applying closest state on rewind if there is no state for the rewind frame.
- Cache incoming inputs for local player.
- Null check input and state histories weak pointers in RewindData.
- Set a Prediction Interpolation object to resimulation while it interacts with a resim object.

diff --git a/Engine/Source/Runtime/Engine/Private/PhysicsEngine/NetworkPhysicsComponent.cpp b/Engine/Source/Runtime/Engine/Private/PhysicsEngine/NetworkPhysicsComponent.cpp
--- a/Engine/Source/Runtime/Engine/Private/PhysicsEngine/NetworkPhysicsComponent.cpp
+++ b/Engine/Source/Runtime/Engine/Private/PhysicsEngine/NetworkPhysicsComponent.cpp
@@ -18,6 +18,15 @@
 
 #include UE_INLINE_GENERATED_CPP_BY_NAME(NetworkPhysicsComponent)
 
+namespace PhysicsReplicationCVars
+{
+	namespace ResimulationCVars
+	{
+		static bool bAllowRewindToClosestState = true;
+		static FAutoConsoleVariableRef CVarResimAllowRewindToClosestState(TEXT("np2.Resim.AllowRewindToClosestState"), bAllowRewindToClosestState, TEXT("When rewinding to a specific frame, if the client doens't have state data for that frame, use closest data available. Only affects the first rewind frame, when FPBDRigidsEvolution is set to Reset."));
+	}
+}
+
 /** These CVars are deprecated from UE 5.4, physics frame offset for networked physics prediction is now handled via PlayerController with automatic time dilation
 * p.net.CmdOffsetEnabled = 0 is recommended to disable the deprecated flow */
 namespace InputCmdCVars
@@ -639,8 +648,7 @@ void UNetworkPhysicsComponent::OnRep_SetReplicatedInputs()
 		PlayerController = GetWorld()->GetFirstPlayerController();
 	}
 
-	// For local controller we should already have correct replicated inputs
-	if (PlayerController && !IsLocallyControlled() && !HasServerWorld() && InputsHistory)
+	if (PlayerController && !HasServerWorld() && InputsHistory)
 	{
 		const int32 LocalOffset = PlayerController->GetNetworkPhysicsTickOffset();
 
@@ -741,7 +749,7 @@ void UNetworkPhysicsComponent::OnPreProcessInputsInternal(const int32 PhysicsSte
 			}
 		}
 
-		// for the inputs client local ones are ground truth otherwise use the replicated ones coming from the server
+		// Apply replicated inputs on server and simulated proxies (and on local player if we are resimulating)
 		if (!IsLocallyControlled() || bIsSolverResim)
 		{
 			FNetworkPhysicsDatas* PhysicsDatas = InputsDatas.Get();
@@ -772,6 +780,7 @@ void UNetworkPhysicsComponent::OnPreProcessInputsInternal(const int32 PhysicsSte
 			}
 		}
 
+		// Apply replicated state on clients if we are resimulating
 		if (!HasServerWorld() && bIsSolverResim)
 		{
 			FNetworkPhysicsDatas* PhysicsDatas = StatesDatas.Get();
@@ -779,7 +788,8 @@ void UNetworkPhysicsComponent::OnPreProcessInputsInternal(const int32 PhysicsSte
 	#if DEBUG_NETWORK_PHYSICS
 			UE_LOG(LogChaos, Log, TEXT("		Extracting history states at frame %d | Component = %s"), PhysicsStep, *GetFullName());
 	#endif
-			if (StatesHistory->ExtractDatas(PhysicsStep, bIsSolverReset, PhysicsDatas, true))
+			const bool bExactFrame = PhysicsReplicationCVars::ResimulationCVars::bAllowRewindToClosestState ? !bIsSolverReset : true;
+			if (StatesHistory->ExtractDatas(PhysicsStep, bIsSolverReset, PhysicsDatas, bExactFrame))
 			{
 				PhysicsDatas->ApplyDatas(ActorComponent);
 			}
diff --git a/Engine/Source/Runtime/Engine/Private/PhysicsEngine/PhysicsReplication.cpp b/Engine/Source/Runtime/Engine/Private/PhysicsEngine/PhysicsReplication.cpp
--- a/Engine/Source/Runtime/Engine/Private/PhysicsEngine/PhysicsReplication.cpp
+++ b/Engine/Source/Runtime/Engine/Private/PhysicsEngine/PhysicsReplication.cpp
@@ -778,7 +778,7 @@ void FPhysicsReplicationAsync::OnPreSimulate_Internal()
 					// If final resim frame, mark interpolated targets as waiting for up to date data from the server.
 					if (Target.RepMode == EPhysicsReplicationMode::PredictiveInterpolation)
 					{
-						Target.SetWaiting(RigidsSolver->GetCurrentFrame() + Target.FrameOffset);
+						Target.SetWaiting(RigidsSolver->GetCurrentFrame() + Target.FrameOffset, Target.RepModeOverride);
 					}
 				}
 			}
@@ -866,6 +866,7 @@ void FPhysicsReplicationAsync::UpdateAsyncTarget(const FPhysicsRepAsyncInputData
 		Target->PrevPosTarget = Input.TargetState.Position;
 		Target->PrevRotTarget = Input.TargetState.Quaternion;
 		Target->PrevLinVel = Input.TargetState.LinVel;
+		Target->RepModeOverride = Input.RepMode;
 	}
 
 	/** Target Update Description
@@ -1059,7 +1060,7 @@ void FPhysicsReplicationAsync::ApplyTargetStatesAsync(const float DeltaSeconds,
 
 			if (FPBDRigidParticleHandle* RigidHandle = Handle->CastToRigidParticle())
 			{
-				const EPhysicsReplicationMode RepMode = Target.RepMode;
+				const EPhysicsReplicationMode RepMode = Target.IsWaiting() ? Target.RepModeOverride : Target.RepMode;
 				switch (RepMode)
 				{
 					case EPhysicsReplicationMode::Default:
@@ -1405,7 +1406,7 @@ bool FPhysicsReplicationAsync::PredictiveInterpolation(Chaos::FPBDRigidParticleH
 	if (PhysicsReplicationCVars::ResimulationCVars::bDisableReplicationOnInteraction && ParticlesInResimIslands.Contains(Handle->GetHandleIdx()))
 	{
 		// If particle is in an island with a resim object, don't run replication and wait for an up to date target (after leaving the island)
-		Target.SetWaiting(RigidsSolver->GetCurrentFrame() + Target.FrameOffset);
+		Target.SetWaiting(RigidsSolver->GetCurrentFrame() + Target.FrameOffset, EPhysicsReplicationMode::Resimulation);
 		return false;
 	}
 
diff --git a/Engine/Source/Runtime/Engine/Public/PhysicsReplication.h b/Engine/Source/Runtime/Engine/Public/PhysicsReplication.h
--- a/Engine/Source/Runtime/Engine/Public/PhysicsReplication.h
+++ b/Engine/Source/Runtime/Engine/Public/PhysicsReplication.h
@@ -137,6 +137,7 @@ struct FReplicatedPhysicsTargetAsync
 
 	/** The replication mode this PhysicsObject should use */
 	EPhysicsReplicationMode RepMode;
+	EPhysicsReplicationMode RepModeOverride;
 
 	/** Correction values from previous update */
 	FVector PrevPosTarget;
@@ -167,10 +168,18 @@ public:
 	/** Is this target waiting for up to date data? */
 	const bool IsWaiting() { return WaitForServerFrame > INDEX_NONE; } const
 
+	/** Set target to wait for data newer than @param InWaitForServerFrame and while waiting replicate via @param InRepModeOverride */
+	void SetWaiting(int32 InWaitForServerFrame, EPhysicsReplicationMode InRepModeOverride)
+	{
+		SetWaiting(InWaitForServerFrame);
+		RepModeOverride = InRepModeOverride;
+	}
+
 	/** Set target to wait for data newer than @param InWaitForServerFrame */
 	void SetWaiting(int32 InWaitForServerFrame)
 	{
-		WaitForServerFrame = InWaitForServerFrame; 
+		RepModeOverride = RepMode;
+		WaitForServerFrame = InWaitForServerFrame;
 		PRAGMA_DISABLE_DEPRECATION_WARNINGS
 		bWaiting = IsWaiting();
 		PRAGMA_ENABLE_DEPRECATION_WARNINGS
diff --git a/Engine/Source/Runtime/Experimental/Chaos/Private/RewindData.cpp b/Engine/Source/Runtime/Experimental/Chaos/Private/RewindData.cpp
--- a/Engine/Source/Runtime/Experimental/Chaos/Private/RewindData.cpp
+++ b/Engine/Source/Runtime/Experimental/Chaos/Private/RewindData.cpp
@@ -992,18 +992,19 @@ int32 FRewindData::FindValidResimFrame(const int32 RequestedFrame)
 		{
 			for (auto& InputsHistory : InputsHistories)
 			{
-				if (!InputsHistory.Pin().Get()->HasValidDatas(ValidFrame))
+				if (InputsHistory.IsValid() && !InputsHistory.Pin().Get()->HasValidDatas(ValidFrame))
 				{
 					bHasTargetHistory = false;
 					break;
 				}
 			}
 		}
+
 		if (bHasTargetHistory)
 		{
 			for (auto& StatesHistory : StatesHistories)
 			{
-				if (!StatesHistory.Pin().Get()->HasValidDatas(ValidFrame))
+				if (StatesHistory.IsValid() && !StatesHistory.Pin().Get()->HasValidDatas(ValidFrame))
 				{
 					bHasTargetHistory = false;
 					break;
