Added edit condition to disable anim per-track compression adaptive error settings when the features are not enabled

diff --git a/Engine/Source/Runtime/Engine/Classes/Animation/AnimCompress_PerTrackCompression.h b/Engine/Source/Runtime/Engine/Classes/Animation/AnimCompress_PerTrackCompression.h
--- a/Engine/Source/Runtime/Engine/Classes/Animation/AnimCompress_PerTrackCompression.h
+++ b/Engine/Source/Runtime/Engine/Classes/Animation/AnimCompress_PerTrackCompression.h
@@ -64,12 +64,12 @@ class UAnimCompress_PerTrackCompression : public UAnimCompress_RemoveLinearKeys
 	UPROPERTY(EditAnywhere, Category=AdaptiveError)
 	uint32 bUseAdaptiveError:1;
 
-	/** If true, uses MinEffectorDiff as the threhsold for end effectors */
-	UPROPERTY(EditAnywhere, Category=AdaptiveError)
+	/** If true, uses MinEffectorDiff as the threshold for end effectors */
+	UPROPERTY(EditAnywhere, Category=AdaptiveError, meta=(editcondition = "bUseAdaptiveError"))
 	uint32 bUseOverrideForEndEffectors:1;
 
 	/** A bias added to the track height before using it to calculate the adaptive error */
-	UPROPERTY(EditAnywhere, Category=AdaptiveError)
+	UPROPERTY(EditAnywhere, Category=AdaptiveError, meta=(editcondition = "bUseAdaptiveError"))
 	int32 TrackHeightBias;
 
 	/**
@@ -77,7 +77,7 @@ class UAnimCompress_PerTrackCompression : public UAnimCompress_RemoveLinearKeys
 	 * EffectiveErrorTolerance = Max(BaseErrorTolerance / Power(ParentingDivisor, Max(Height+Bias,0) * ParentingDivisorExponent), ZeroingThreshold)
 	 * Only has an effect bUseAdaptiveError is true
 	 */
-	UPROPERTY(EditAnywhere, Category=AdaptiveError, meta=(ClampMin = "1.0"))
+	UPROPERTY(EditAnywhere, Category=AdaptiveError, meta=(ClampMin = "1.0", editcondition = "bUseAdaptiveError"))
 	float ParentingDivisor;
 
 	/**
@@ -85,7 +85,7 @@ class UAnimCompress_PerTrackCompression : public UAnimCompress_RemoveLinearKeys
 	 * EffectiveErrorTolerance = Max(BaseErrorTolerance / Power(ParentingDivisor, Max(Height+Bias,0) * ParentingDivisorExponent), ZeroingThreshold)
 	 * Only has an effect bUseAdaptiveError is true
 	 */
-	UPROPERTY(EditAnywhere, Category=AdaptiveError, meta=(ClampMin = "0.1"))
+	UPROPERTY(EditAnywhere, Category=AdaptiveError, meta=(ClampMin = "0.1", editcondition = "bUseAdaptiveError"))
 	float ParentingDivisorExponent;
 
 	/**
@@ -99,27 +99,27 @@ class UAnimCompress_PerTrackCompression : public UAnimCompress_RemoveLinearKeys
 	 * This ratio determines how much error in end effector rotation can come from a given track's rotation error or translation error.
 	 * If 1, all of it must come from rotation error, if 0.5, half can come from each, and if 0.0, all must come from translation error.
 	 */
-	UPROPERTY(EditAnywhere, Category=AdaptiveError2, meta=(ClampMin = "0.0", ClampMax = "1.0"))
+	UPROPERTY(EditAnywhere, Category=AdaptiveError2, meta=(ClampMin = "0.0", ClampMax = "1.0", editcondition = "bUseAdaptiveError2"))
 	float RotationErrorSourceRatio;
 
 	/**
 	 * This ratio determines how much error in end effector translation can come from a given track's rotation error or translation error.
 	 * If 1, all of it must come from rotation error, if 0.5, half can come from each, and if 0.0, all must come from translation error.
 	 */
-	UPROPERTY(EditAnywhere, Category=AdaptiveError2, meta=(ClampMin = "0.0", ClampMax = "1.0"))
+	UPROPERTY(EditAnywhere, Category=AdaptiveError2, meta=(ClampMin = "0.0", ClampMax = "1.0", editcondition = "bUseAdaptiveError2"))
 	float TranslationErrorSourceRatio;
 
 	/**
 	 * This ratio determines how much error in end effector scale can come from a given track's rotation error or scale error.
 	 * If 1, all of it must come from rotation error, if 0.5, half can come from each, and if 0.0, all must come from scale error.
 	 */
-	UPROPERTY(EditAnywhere, Category=AdaptiveError2, meta=(ClampMin = "0.0", ClampMax = "1.0"))
+	UPROPERTY(EditAnywhere, Category=AdaptiveError2, meta=(ClampMin = "0.0", ClampMax = "1.0", editcondition = "bUseAdaptiveError2"))
 	float ScaleErrorSourceRatio;
 
 	/**
 	 * A fraction that determines how much of the total error budget can be introduced by any particular track
 	 */
-	UPROPERTY(EditAnywhere, Category=AdaptiveError2, meta=(ClampMin = "0.0", ClampMax = "1.0"))
+	UPROPERTY(EditAnywhere, Category=AdaptiveError2, meta=(ClampMin = "0.0", ClampMax = "1.0", editcondition = "bUseAdaptiveError2"))
 	float MaxErrorPerTrackRatio;
 
 	/**
