World Partition: fix FWorldPartitionActorDesc::ShouldResave to properly handle the HLOD relevant case and also include actor/component properties as they can impact other world builders, such as UWorldPartitionRuntimeVirtualTextureBuilder.

diff --git a/Engine/Source/Runtime/Engine/Private/WorldPartition/WorldPartitionActorDesc.cpp b/Engine/Source/Runtime/Engine/Private/WorldPartition/WorldPartitionActorDesc.cpp
--- a/Engine/Source/Runtime/Engine/Private/WorldPartition/WorldPartitionActorDesc.cpp
+++ b/Engine/Source/Runtime/Engine/Private/WorldPartition/WorldPartitionActorDesc.cpp
@@ -310,7 +310,8 @@ bool FWorldPartitionActorDesc::ShouldResave(const FWorldPartitionActorDesc* Othe
 		ContentBundleGuid != Other->ContentBundleGuid ||
 		!CompareUnsortedArrays(DataLayers, Other->DataLayers) ||
 		!CompareUnsortedArrays(References, Other->References) ||
-		!CompareUnsortedArrays(EditorOnlyReferences, Other->EditorOnlyReferences))
+		!CompareUnsortedArrays(EditorOnlyReferences, Other->EditorOnlyReferences) ||
+		Properties != Other->Properties)
 	{
 		return true;
 	}
@@ -331,9 +332,9 @@ bool FWorldPartitionActorDesc::ShouldResave(const FWorldPartitionActorDesc* Othe
 		}
 	}
 
-	// If the actor descriptor says the actor is HLOD relebant but in reality it's not, this will incur a loading time penalty during HLOD generation
+	// If the actor descriptor says the actor is HLOD relevant but in reality it's not, this will incur a loading time penalty during HLOD generation
 	// but will not affect the final result, as the value from the loaded actor will be used instead, so don't consider this as affecting streaming generation.
-	return (bActorIsHLODRelevant == Other->bActorIsHLODRelevant) || (bActorIsHLODRelevant && !Other->bActorIsHLODRelevant);
+	return !bActorIsHLODRelevant && Other->bActorIsHLODRelevant;
 }
 
 void FWorldPartitionActorDesc::SerializeTo(TArray<uint8>& OutData) const
