Fix potential race condition during animation compression

diff --git a/Engine/Source/Runtime/Engine/Private/Animation/AnimSequence.cpp b/Engine/Source/Runtime/Engine/Private/Animation/AnimSequence.cpp
--- a/Engine/Source/Runtime/Engine/Private/Animation/AnimSequence.cpp
+++ b/Engine/Source/Runtime/Engine/Private/Animation/AnimSequence.cpp
@@ -4577,6 +4577,10 @@ FIoHash UAnimSequence::BeginCacheDerivedData(const ITargetPlatform* TargetPlatfo
 		return FIoHash::Zero;
 	}
 
+	// Wait for any in-flight requests to finish
+	// Once a compression request finishes, it might modify the compression settings below and the key hash as well
+	UE::Anim::FAnimSequenceCompilingManager::Get().FinishCompilation({ this });
+
 	// Ensure that there are valid compression settings
 	if (BoneCompressionSettings == nullptr || !BoneCompressionSettings->AreSettingsValid())
 	{
@@ -4592,7 +4596,7 @@ FIoHash UAnimSequence::BeginCacheDerivedData(const ITargetPlatform* TargetPlatfo
 		VariableFrameStrippingSettings = FAnimationUtils::GetDefaultVariableFrameStrippingSettings();
 	}
 
-	// Make sure all our required dependencies are loaded
+	// Make sure all our required dependencies are loaded, we need them to compute the KeyHash
 	FAnimationUtils::EnsureAnimSequenceLoaded(*this);
 		
 	const FIoHash KeyHash = CreateDerivedDataKeyHash(TargetPlatform);
@@ -4602,10 +4606,6 @@ FIoHash UAnimSequence::BeginCacheDerivedData(const ITargetPlatform* TargetPlatfo
 		return KeyHash;
 	}
 
-	// Wait for any in-flight requests
-	UE::Anim::FAnimSequenceCompilingManager::Get().FinishCompilation({this});
-
-	// We should wait here for any previous compilations?
 	FCompressedAnimSequence* TargetData = nullptr;
 	if (TargetPlatform->IsRunningPlatform())
 	{
