AInstancedFoliageActor - Fix map checks errors during PostLoad() when building HLODs
* Assign a dummy runtime cell to streamed levels created with UWorldPartitionLevelStreamingDynamic::LoadInEditor()

diff --git a/Engine/Source/Runtime/Engine/Private/WorldPartition/WorldPartitionLevelStreamingDynamic.cpp b/Engine/Source/Runtime/Engine/Private/WorldPartition/WorldPartitionLevelStreamingDynamic.cpp
--- a/Engine/Source/Runtime/Engine/Private/WorldPartition/WorldPartitionLevelStreamingDynamic.cpp
+++ b/Engine/Source/Runtime/Engine/Private/WorldPartition/WorldPartitionLevelStreamingDynamic.cpp
@@ -1,6 +1,8 @@
 // Copyright Epic Games, Inc. All Rights Reserved.
 
 #include "WorldPartition/WorldPartitionLevelStreamingDynamic.h"
+#include "WorldPartition/WorldPartitionRuntimeLevelStreamingCell.h"
+
 #include "Engine/LevelStreaming.h"
 #include "WorldPartition/WorldPartition.h"
 #include "UObject/PropertyPortFlags.h"
@@ -90,6 +92,10 @@ UWorldPartitionLevelStreamingDynamic* UWorldPartitionLevelStreamingDynamic::Load
 	FString PackageName = FString::Printf(TEXT("Temp/%s"), *LevelStreamingName.ToString());
 	TSoftObjectPtr<UWorld> WorldAsset(FSoftObjectPath(FString::Printf(TEXT("/%.*s/%s.%s"), WorldMountPointName.Len(), WorldMountPointName.GetData(), *PackageName, *World->GetName())));
 	LevelStreaming->SetWorldAsset(WorldAsset);
+
+	// Assign a dummy runtime cell to ensure that code that is testing for IsWorldPartitionRuntimeCell() behaves as expected
+	UWorldPartitionRuntimeLevelStreamingCell* StreamingCell = NewObject<UWorldPartitionRuntimeLevelStreamingCell>(LevelStreaming, NAME_None, RF_Transient);
+	LevelStreaming->StreamingCell = StreamingCell;
 	
 	LevelStreaming->LevelTransform = FTransform::Identity;
 	LevelStreaming->Initialize(World, InPackages);
@@ -161,8 +167,7 @@ void UWorldPartitionLevelStreamingDynamic::CreateRuntimeLevel()
 	check(World && (World->IsGameWorld() || GetShouldBeVisibleInEditor()));
 
 	// Make sure we are creating a runtime level for a cell.
-	// Or, we created a transient WPLevelStreamingDynamic via UWorldPartitionLevelStreamingDynamic::LoadInEditor, to load a pack of actors (Hlods generation).
-	check(StreamingCell != nullptr || GetShouldBeVisibleInEditor());
+	check(StreamingCell != nullptr);
 
 	// Create streaming cell Level package
 	RuntimeLevel = FWorldPartitionLevelHelper::CreateEmptyLevelForRuntimeCell(StreamingCell.Get(), World, GetWorldAsset().ToString());
