Removed changes to FPrimitiveInstanceDataManager::PostLoad that were not needed.

#rnx

diff --git a/Engine/Source/Runtime/Engine/Private/InstancedStaticMesh.cpp b/Engine/Source/Runtime/Engine/Private/InstancedStaticMesh.cpp
--- a/Engine/Source/Runtime/Engine/Private/InstancedStaticMesh.cpp
+++ b/Engine/Source/Runtime/Engine/Private/InstancedStaticMesh.cpp
@@ -4838,13 +4838,7 @@ void UInstancedStaticMeshComponent::PostLoad()
 
 	if (!HasAnyFlags(RF_ClassDefaultObject|RF_ArchetypeObject))
 	{
-		FInstanceUpdateComponentDesc ComponentData;
-		if (GetStaticMesh() != nullptr)
-		{
-			UpdateComponentToWorld();
-			BuildComponentInstanceData(GMaxRHIFeatureLevel, ComponentData);
-		}
-		PrimitiveInstanceDataManager.PostLoad(PerInstanceSMData.Num(), MoveTemp(ComponentData));
+		PrimitiveInstanceDataManager.PostLoad(PerInstanceSMData.Num());
 	}
 
 	// Has different implementation in HISMC
diff --git a/Engine/Source/Runtime/Engine/Private/InstancedStaticMesh/ISMInstanceDataManager.cpp b/Engine/Source/Runtime/Engine/Private/InstancedStaticMesh/ISMInstanceDataManager.cpp
--- a/Engine/Source/Runtime/Engine/Private/InstancedStaticMesh/ISMInstanceDataManager.cpp
+++ b/Engine/Source/Runtime/Engine/Private/InstancedStaticMesh/ISMInstanceDataManager.cpp
@@ -788,7 +788,7 @@ bool FPrimitiveInstanceDataManager::FlushChanges(FInstanceUpdateComponentDesc &&
 	return true;
 }
 
-void FPrimitiveInstanceDataManager::PostLoad(int32 InNumInstances, FInstanceUpdateComponentDesc &&ComponentData)
+void FPrimitiveInstanceDataManager::PostLoad(int32 InNumInstances)
 {
 	if (GetState() == ETrackingState::Disabled)
 	{
diff --git a/Engine/Source/Runtime/Engine/Public/InstancedStaticMesh/ISMInstanceDataManager.h b/Engine/Source/Runtime/Engine/Public/InstancedStaticMesh/ISMInstanceDataManager.h
--- a/Engine/Source/Runtime/Engine/Public/InstancedStaticMesh/ISMInstanceDataManager.h
+++ b/Engine/Source/Runtime/Engine/Public/InstancedStaticMesh/ISMInstanceDataManager.h
@@ -107,7 +107,7 @@ public:
 	bool FlushChanges(FInstanceUpdateComponentDesc &&ComponentData, bool bNewPrimitiveProxy);
 
 	// 
-	void PostLoad(int32 InNumInstances, FInstanceUpdateComponentDesc &&ComponentData);
+	void PostLoad(int32 InNumInstances);
 
 	/**
 	 * Clear the ID/Index association and reset the mapping to identity & number of instances to the given number.
