Add an option to place actors under a subfolder instead of the root when breaking a Level Instance.

If this option is enabled, the actors will be placed inside the folder the LI is inside of, under a subfolder with the name of the Level Instance, and also keeping their original folder structure.
So if i.e. the Level Instance Actor is called "Desert/LI_House2", and an actor inside is named "Lights/Light_Sun", the actor will be moved to "Desert/LI_House2/Lights/Light_Sun" in the outer level.

diff --git a/Engine/Source/Editor/LevelInstanceEditor/Private/LevelInstanceEditorModule.cpp b/Engine/Source/Editor/LevelInstanceEditor/Private/LevelInstanceEditorModule.cpp
--- a/Engine/Source/Editor/LevelInstanceEditor/Private/LevelInstanceEditorModule.cpp
+++ b/Engine/Source/Editor/LevelInstanceEditor/Private/LevelInstanceEditorModule.cpp
@@ -402,54 +402,90 @@ namespace LevelInstanceMenuUtils
 	void CreateBreakSubMenu(UToolMenu* Menu, TArray<ILevelInstanceInterface*> BreakableLevelInstances)
 	{
 		static int32 BreakLevels = 1;
+		ULevelInstanceEditorPerProjectUserSettings* Settings = GetMutableDefault<ULevelInstanceEditorPerProjectUserSettings>();
 
 		if (ULevelInstanceSubsystem* LevelInstanceSubsystem = GEditor->GetEditorWorldContext().World()->GetSubsystem<ULevelInstanceSubsystem>())
 		{
-			FToolMenuSection& Section = Menu->AddSection(NAME_None, LOCTEXT("LevelInstanceBreakSection", "Break Level Instance"));
-			TSharedRef<SWidget> MenuWidget =
-				SNew(SVerticalBox)
+			FToolMenuSection& Section = Menu->AddSection("Options", LOCTEXT("LevelInstanceBreakOptionsSection", "Options"));
+
+			FToolMenuEntry OrganizeInFoldersEntry = FToolMenuEntry::InitMenuEntry(
+				"OrganizeInFolders",
+				LOCTEXT("OrganizeActorsInFolders", "Keep Folders"),
+				LOCTEXT(
+					"OrganizeActorsInFoldersTooltip",
+					"Should the actors be placed in the folder the Level Instance is in, "
+					"and keep the folder structure they had inside the Level Instance?"
+				),
+				FSlateIcon(),
+				FUIAction(
+					FExecuteAction::CreateLambda([Settings]()
+					{
+						Settings->bKeepFoldersDuringBreak = !Settings->bKeepFoldersDuringBreak;
+					}),
+					FCanExecuteAction(),
+					FGetActionCheckState::CreateLambda([Settings]
+					{
+						return Settings->bKeepFoldersDuringBreak ? ECheckBoxState::Checked : ECheckBoxState::Unchecked;
+					})
+				),
+				EUserInterfaceActionType::ToggleButton
+			);
+			OrganizeInFoldersEntry.bShouldCloseWindowAfterMenuSelection = false;
+			Section.AddEntry(OrganizeInFoldersEntry);
 
-				+SVerticalBox::Slot()
+			TSharedRef<SWidget> MenuWidget =
+				SNew(SBox)
+				.Padding(FMargin(5, 2, 5, 0))
 				[
-					SNew(SHorizontalBox)
-					+ SHorizontalBox::Slot()
+					SNew(SNumericEntryBox<int32>)
+					.MinValue(1)
+					.Value_Lambda([]() { return BreakLevels; })
+					.OnValueChanged_Lambda([](int32 InValue) { BreakLevels = InValue; })
+					.Label()
 					[
-						SNew(SNumericEntryBox<int32>)
-						.MinValue(1)
-						.Value_Lambda([]() { return BreakLevels; })
-						.OnValueChanged_Lambda([](int32 InValue) { BreakLevels = InValue; })
-						.LabelPadding(0)
-						.Label()
-						[
-							SNumericEntryBox<int32>::BuildLabel(LOCTEXT("BreakLevelsLabel", "Levels"), FLinearColor::White, SNumericEntryBox<int32>::BlueLabelBackgroundColor)
-						]
+						SNumericEntryBox<int32>::BuildLabel(LOCTEXT("BreakLevelsLabel", "Levels"), FLinearColor::White, FLinearColor::Transparent)
 					]
-				]
+				];
 
-				+SVerticalBox::Slot()
-				.VAlign(VAlign_Center)
-				.HAlign(HAlign_Center)
-				.Padding(0, 5, 0, 0)
-				[
-					SNew(SButton)
-					.HAlign(HAlign_Center)
-					.ContentPadding(FAppStyle::GetMargin("StandardDialog.ContentPadding"))
-					.OnClicked_Lambda([BreakableLevelInstances, LevelInstanceSubsystem]() 
+			Section.AddEntry(FToolMenuEntry::InitWidget("SetBreakLevels", MenuWidget, FText::GetEmpty(), false));
+
+			Section.AddSeparator(NAME_None);
+
+			FToolMenuEntry ExecuteEntry = FToolMenuEntry::InitMenuEntry(
+				"ExecuteBreak",
+				LOCTEXT("BreakLevelInstances_BreakLevelInstanceButton", "Break Level Instance(s)"),
+				FText(),
+				FSlateIcon(),
+				FUIAction(
+					FExecuteAction::CreateLambda([BreakableLevelInstances, LevelInstanceSubsystem, Settings]()
 					{
-						const FText LevelInstanceBreakWarning = FText::Format(LOCTEXT("BreakingLevelInstance", "You are about to break {0} level instance(s). This action cannot be undone. Are you sure ?"),  FText::FromString(FString::FromInt(BreakableLevelInstances.Num())));
+						const FText LevelInstanceBreakWarning = FText::Format(
+							LOCTEXT(
+								"BreakingLevelInstance",
+								"You are about to break {0} level instance(s). This action cannot be undone. Are you sure ?"
+							),
+							FText::AsNumber(BreakableLevelInstances.Num())
+						);
+
 						if (FMessageDialog::Open(EAppMsgType::YesNo, LevelInstanceBreakWarning) == EAppReturnType::Yes)
 						{
+							ELevelInstanceBreakFlags Flags = ELevelInstanceBreakFlags::None;
+							if (Settings->bKeepFoldersDuringBreak)
+							{
+								Flags |= ELevelInstanceBreakFlags::KeepFolders;
+							}
+
 							for (ILevelInstanceInterface* LevelInstance : BreakableLevelInstances)
 							{
-								LevelInstanceSubsystem->BreakLevelInstance(LevelInstance, BreakLevels);
+								LevelInstanceSubsystem->BreakLevelInstance(LevelInstance, BreakLevels, nullptr, Flags);
 							}
 						}
-						return FReply::Handled();
 					})
-					.Text(LOCTEXT("BreakLevelInstances_BreakLevelInstanceButton", "Break Level Instance(s)"))
-				];
+				),
+				EUserInterfaceActionType::Button
+			);
 
-			Section.AddEntry(FToolMenuEntry::InitWidget("SetBreakLevels", MenuWidget, FText::GetEmpty(), false));
+			Section.AddEntry(ExecuteEntry);
 		}
 	}
 
diff --git a/Engine/Source/Editor/LevelInstanceEditor/Public/LevelInstanceEditorSettings.h b/Engine/Source/Editor/LevelInstanceEditor/Public/LevelInstanceEditorSettings.h
--- a/Engine/Source/Editor/LevelInstanceEditor/Public/LevelInstanceEditorSettings.h
+++ b/Engine/Source/Editor/LevelInstanceEditor/Public/LevelInstanceEditorSettings.h
@@ -51,4 +51,16 @@ public:
 
 	UPROPERTY(config, EditAnywhere, Category = Pivot)
 	ELevelInstancePivotType PivotType;
+
+	/**
+	 * When the Level Instance is broken via "Level->Break..", its actors will be placed inside the folder the LI is
+	 * inside of, under a subfolder with the name of the Level Instance, and also keeping their original folder structure.
+	 * So if i.e. the Level Instance Actor is called "Desert/LI_House2", and an actor inside is named "Lights/Light_Sun",
+	 * the actor will be moved to "Desert/LI_House2/Lights/Light_Sun" in the outer level.
+	 *
+	 * If this flag is not set, actors will be placed either in the root folder of the outer level (but their original
+	 * folders from the LI kept), or, if "Current Folder" is set, they'll be moved there without any subfolders.
+	 */
+	UPROPERTY(config, EditAnywhere, Category = Break)
+	bool bKeepFoldersDuringBreak = false;
 };
diff --git a/Engine/Source/Runtime/Engine/Private/LevelInstance/LevelInstanceSubsystem.cpp b/Engine/Source/Runtime/Engine/Private/LevelInstance/LevelInstanceSubsystem.cpp
--- a/Engine/Source/Runtime/Engine/Private/LevelInstance/LevelInstanceSubsystem.cpp
+++ b/Engine/Source/Runtime/Engine/Private/LevelInstance/LevelInstanceSubsystem.cpp
@@ -1359,7 +1359,8 @@ ILevelInstanceInterface* ULevelInstanceSubsystem::CreateLevelInstanceFrom(const
 	return GetLevelInstance(NewLevelInstanceID);
 }
 
-bool ULevelInstanceSubsystem::BreakLevelInstance(ILevelInstanceInterface* LevelInstance, uint32 Levels /* = 1 */, TArray<AActor*>* OutMovedActors /* = nullptr */)
+bool ULevelInstanceSubsystem::BreakLevelInstance(ILevelInstanceInterface* LevelInstance, uint32 Levels /* = 1 */,
+	TArray<AActor*>* OutMovedActors /* = nullptr */, ELevelInstanceBreakFlags Flags /* = None */)
 {
 	const double StartTime = FPlatformTime::Seconds();
 
@@ -1368,7 +1369,7 @@ bool ULevelInstanceSubsystem::BreakLevelInstance(ILevelInstanceInterface* LevelI
 	GetMutableDefault<ULevelEditorMiscSettings>()->bAvoidRelabelOnPasteSelected = 1;
 
 	TArray<AActor*> MovedActors;
-	BreakLevelInstance_Impl(LevelInstance, Levels, MovedActors);
+	BreakLevelInstance_Impl(LevelInstance, Levels, MovedActors, Flags);
 
 	USelection* ActorSelection = GEditor->GetSelectedActors();
 	ActorSelection->BeginBatchSelectOperation();
@@ -1391,7 +1392,8 @@ bool ULevelInstanceSubsystem::BreakLevelInstance(ILevelInstanceInterface* LevelI
 	return bStatus;
 }
 
-void ULevelInstanceSubsystem::BreakLevelInstance_Impl(ILevelInstanceInterface* LevelInstance, uint32 Levels, TArray<AActor*>& OutMovedActors)
+void ULevelInstanceSubsystem::BreakLevelInstance_Impl(ILevelInstanceInterface* LevelInstance, uint32 Levels,
+	TArray<AActor*>& OutMovedActors, ELevelInstanceBreakFlags Flags)
 {
 	if (Levels > 0)
 	{
@@ -1487,60 +1489,83 @@ void ULevelInstanceSubsystem::BreakLevelInstance_Impl(ILevelInstanceInterface* L
 		const bool bWarnAboutReferences = true;
 		const bool bWarnAboutRenaming = false;
 		const bool bMoveAllOrFail = true;
-		if (!EditorLevelUtils::CopyActorsToLevel(ActorsToMove.Array(), DestinationLevel, bWarnAboutReferences, bWarnAboutRenaming, bMoveAllOrFail))
+
+		TArray<AActor*> ActorsMovedThisStage;
+		if (!EditorLevelUtils::CopyActorsToLevel(ActorsToMove.Array(), DestinationLevel, bWarnAboutReferences, bWarnAboutRenaming, bMoveAllOrFail, &ActorsMovedThisStage))
 		{
 			UE_LOG(LogLevelInstance, Warning, TEXT("Failed to break Level Instance because not all actors could be moved"));
 			return;
 		}
 
-		// Clear undo buffer here because operation of Breaking a level instance is not undoable.
-		// - Do it before Unload and Destroy calls because those calls will try and unload the level.
-		//	 The unloading of the level might cause a call to UEngine::FindAndPrintStaleReferencesToObject
-		//	 which will slow down the unloading because the Trans buffer has some references to the moved actors.
-		if (GEditor->Trans)
-		{
-			GEditor->Trans->Reset(LOCTEXT("BreakLevelInstance", "Break Level Instance"));
-		}
+		OutMovedActors.Append(ActorsMovedThisStage);
 
-		if (LevelInstanceActor->IsA<APackedLevelActor>())
+		const bool bKeepFolders = EnumHasAnyFlags(Flags, ELevelInstanceBreakFlags::KeepFolders);
+		FString LevelInstanceFolder;
+		if (bKeepFolders)
 		{
-			BlockUnloadLevelInstance(LevelInstance);
+			// Build the folder name to move the actors into
+			const FName LevelInstancePath = LevelInstanceActor->GetFolderPath();
+			if (!LevelInstancePath.IsNone())
+			{
+				LevelInstanceFolder = LevelInstancePath.ToString();
+				LevelInstanceFolder += TEXT('/');
+			}
+			LevelInstanceFolder += LevelInstanceActor->GetActorNameOrLabel();
 		}
 
-		// Destroy the old LevelInstance instance actor
-		GetWorld()->DestroyActor(LevelInstanceActor);
-	
-		const bool bContinueBreak = Levels > 1;
-		TArray<ILevelInstanceInterface*> Children;
-
-		for (FSelectionIterator It(GEditor->GetSelectedActorIterator()); It; ++It)
+		TArray<ILevelInstanceInterface*> ChildLevelInstances;
+		for (AActor* Actor : ActorsMovedThisStage)
 		{
-			AActor* Actor = Cast<AActor>(*It);
-
-			if (Actor)
+			if (bKeepFolders)
 			{
-				OutMovedActors.Add(Actor);
+				// Update the folder path of the moved actor, combining LI's path + LI's name + actor's path
+				TStringBuilder<128> NewActorPath;
+				NewActorPath += LevelInstanceFolder;
+
+				const FName OldActorPath = Actor->GetFolderPath();
+				if (!OldActorPath.IsNone())
+				{
+					NewActorPath += TEXT('/');
+					NewActorPath += Actor->GetFolderPath().ToString();
+				}
+
+				Actor->SetFolderPath(FName(NewActorPath));
 			}
 
 			// Break up any sub LevelInstances if more levels are requested
-			if (bContinueBreak)
+			if (Levels > 1)
 			{
 				if (ILevelInstanceInterface* ChildLevelInstance = Cast<ILevelInstanceInterface>(Actor))
 				{
-					OutMovedActors.Remove(Actor);
+					OutMovedActors.RemoveSingleSwap(Actor);
 
-					Children.Add(ChildLevelInstance);
+					ChildLevelInstances.Add(ChildLevelInstance);
 				}
 			}
 		}
 
-		for (auto& Child : Children)
+		// Clear undo buffer here because operation of Breaking a level instance is not undoable.
+		// - Do it before Unload and Destroy calls because those calls will try and unload the level.
+		//	 The unloading of the level might cause a call to UEngine::FindAndPrintStaleReferencesToObject
+		//	 which will slow down the unloading because the Trans buffer has some references to the moved actors.
+		if (GEditor->Trans)
 		{
-			BreakLevelInstance_Impl(Child, Levels - 1, OutMovedActors);
+			GEditor->Trans->Reset(LOCTEXT("BreakLevelInstance", "Break Level Instance"));
 		}
-	}
 
-	return;
+		if (LevelInstanceActor->IsA<APackedLevelActor>())
+		{
+			BlockUnloadLevelInstance(LevelInstance);
+		}
+
+		// Destroy the old LevelInstance instance actor
+		GetWorld()->DestroyActor(LevelInstanceActor);
+
+		for (auto& Child : ChildLevelInstances)
+		{
+			BreakLevelInstance_Impl(Child, Levels - 1, OutMovedActors, Flags);
+		}
+	}
 }
 
 bool ULevelInstanceSubsystem::LevelInstanceHasLevelScriptBlueprint(const ILevelInstanceInterface* LevelInstance) const
diff --git a/Engine/Source/Runtime/Engine/Public/LevelInstance/LevelInstanceSubsystem.h b/Engine/Source/Runtime/Engine/Public/LevelInstance/LevelInstanceSubsystem.h
--- a/Engine/Source/Runtime/Engine/Public/LevelInstance/LevelInstanceSubsystem.h
+++ b/Engine/Source/Runtime/Engine/Public/LevelInstance/LevelInstanceSubsystem.h
@@ -25,6 +25,23 @@ class ULevelStreamingLevelInstanceEditor;
 class UWorldPartitionSubsystem;
 class UBlueprint;
 
+enum class ELevelInstanceBreakFlags : uint8
+{
+	None = 0,
+
+	/**
+	 * The actors will be placed inside the folder the LI is inside of, under a subfolder with the name of the
+	 * Level Instance, and also keeping their original folder structure.
+	 * So if i.e. the Level Instance Actor is called "Desert/LI_House2", and an actor inside is named "Lights/Light_Sun",
+	 * the actor will be moved to "Desert/LI_House2/Lights/Light_Sun" in the outer level.
+	 *
+	 * If this flag is not set, actors will be placed either in the root folder of the outer level (but their original
+	 * folders from the LI kept), or, if context folder is set, they'll be moved there without any subfolders.
+	 */
+	KeepFolders = 1 << 0,
+};
+ENUM_CLASS_FLAGS(ELevelInstanceBreakFlags);
+
 /**
  * ULevelInstanceSubsystem
  */
@@ -109,7 +126,7 @@ public:
 	ENGINE_API ILevelInstanceInterface* CreateLevelInstanceFrom(const TArray<AActor*>& ActorsToMove, const FNewLevelInstanceParams& CreationParams);
 	ENGINE_API bool MoveActorsToLevel(const TArray<AActor*>& ActorsToRemove, ULevel* DestinationLevel, TArray<AActor*>* OutActors = nullptr) const;
 	ENGINE_API bool MoveActorsTo(ILevelInstanceInterface* LevelInstance, const TArray<AActor*>& ActorsToMove, TArray<AActor*>* OutActors = nullptr);
-	ENGINE_API bool BreakLevelInstance(ILevelInstanceInterface* LevelInstance, uint32 Levels = 1, TArray<AActor*>* OutMovedActors = nullptr);
+	ENGINE_API bool BreakLevelInstance(ILevelInstanceInterface* LevelInstance, uint32 Levels = 1, TArray<AActor*>* OutMovedActors = nullptr, ELevelInstanceBreakFlags Flags = ELevelInstanceBreakFlags::None);
 
 	ENGINE_API bool CanMoveActorToLevel(const AActor* Actor, FText* OutReason = nullptr) const;
 	ENGINE_API void OnActorDeleted(AActor* Actor);
@@ -172,7 +189,7 @@ private:
 	ENGINE_API bool ForEachLevelInstanceChildImpl(const ILevelInstanceInterface* LevelInstance, bool bRecursive, TFunctionRef<bool(const ILevelInstanceInterface*)> Operation) const;
 	ENGINE_API bool ForEachLevelInstanceChildImpl(ILevelInstanceInterface* LevelInstance, bool bRecursive, TFunctionRef<bool(ILevelInstanceInterface*)> Operation) const;
 
-	ENGINE_API void BreakLevelInstance_Impl(ILevelInstanceInterface* LevelInstance, uint32 Levels, TArray<AActor*>& OutMovedActors);
+	ENGINE_API void BreakLevelInstance_Impl(ILevelInstanceInterface* LevelInstance, uint32 Levels, TArray<AActor*>& OutMovedActors, ELevelInstanceBreakFlags Flags);
 
 	static ENGINE_API bool ShouldIgnoreDirtyPackage(UPackage* DirtyPackage, const UWorld* EditingWorld);
 
