World Partition: remove IWorldPartitionActorDescInstance base interface.
- The main reason is to remove ambiguity for functions taking actor descriptors (immutable) and actor descriptor views, which would result in ambigous APIs.
- This was discussed in the original review (an intermediate class was made to fix this), but there's no need for a base class and this needs to be fixed before the 5.4 release.

diff --git a/Engine/Source/Runtime/Engine/Public/WorldPartition/WorldPartitionActorDescInstance.h b/Engine/Source/Runtime/Engine/Public/WorldPartition/WorldPartitionActorDescInstance.h
--- a/Engine/Source/Runtime/Engine/Public/WorldPartition/WorldPartitionActorDescInstance.h
+++ b/Engine/Source/Runtime/Engine/Public/WorldPartition/WorldPartitionActorDescInstance.h
@@ -3,20 +3,17 @@
 #pragma once
 
 #include "CoreMinimal.h"
-#include "WorldPartition/WorldPartitionActorDescInstanceInterface.h"
 
 #if WITH_EDITOR
-
 #include "WorldPartition/WorldPartitionActorDesc.h"
 #include "UObject/WeakObjectPtr.h"
 
 class AActor;
 class UWorldPartition;
 class UActorDescContainerInstance;
-
 #endif // WITH_EDITOR
 
-class FWorldPartitionActorDescInstance : public IWorldPartitionActorDescInstance
+class FWorldPartitionActorDescInstance
 {
 #if WITH_EDITOR
 	friend struct FWorldPartitionHandleImpl;
@@ -36,69 +33,69 @@ public:
 	ENGINE_API FWorldPartitionActorDescInstance(UActorDescContainerInstance* InContainerInstance, FWorldPartitionActorDesc* InActorDesc);
 	virtual ~FWorldPartitionActorDescInstance() {}
 
-	//~ Begin IWorldPartitionActorDescInstanceViewInterface
-	virtual const FGuid& GetGuid() const override { return GetActorDesc()->GetGuid(); }
-	virtual FTopLevelAssetPath GetBaseClass() const override { return GetActorDesc()->GetBaseClass(); }
-	virtual FTopLevelAssetPath GetNativeClass() const override { return GetActorDesc()->GetNativeClass(); }
-	virtual UClass* GetActorNativeClass() const override { return GetActorDesc()->GetActorNativeClass(); }
+	virtual const FGuid& GetGuid() const { return GetActorDesc()->GetGuid(); }
+	virtual FTopLevelAssetPath GetBaseClass() const { return GetActorDesc()->GetBaseClass(); }
+	virtual FTopLevelAssetPath GetNativeClass() const { return GetActorDesc()->GetNativeClass(); }
+	virtual UClass* GetActorNativeClass() const { return GetActorDesc()->GetActorNativeClass(); }
 	
-	virtual FName GetRuntimeGrid() const override { return GetActorDesc()->GetRuntimeGrid(); }
-	virtual bool GetIsSpatiallyLoaded() const override { return !GetForceNonSpatiallyLoaded() && GetActorDesc()->GetIsSpatiallyLoaded(); }
-	virtual bool GetActorIsEditorOnly() const override { return GetActorDesc()->GetActorIsEditorOnly(); }
-	virtual bool GetActorIsRuntimeOnly() const override { return GetActorDesc()->GetActorIsRuntimeOnly(); }
-	ENGINE_API virtual bool IsRuntimeRelevant() const override;
-	ENGINE_API virtual bool IsEditorRelevant() const override;
-
-	virtual bool IsUsingDataLayerAsset() const override { return GetActorDesc()->IsUsingDataLayerAsset(); }
-	virtual const TArray<FName>& GetDataLayers() const override { return GetActorDesc()->GetDataLayers(); }
-
-	virtual bool GetActorIsHLODRelevant() const override { return GetActorDesc()->GetActorIsHLODRelevant(); }
-	virtual FSoftObjectPath GetHLODLayer() const override { return GetActorDesc()->GetHLODLayer(); }
-
-	virtual const TArray<FName>& GetTags() const override { return GetActorDesc()->GetTags(); }
-	virtual FName GetActorPackage() const override { return GetActorDesc()->GetActorPackage(); }
-	ENGINE_API virtual FSoftObjectPath GetActorSoftPath() const override;
-	virtual FName GetActorLabel() const override { return GetActorDesc()->GetActorLabel(); }
-	ENGINE_API virtual FName GetActorName() const override;
-	virtual FName GetFolderPath() const override { return GetActorDesc()->GetFolderPath(); }
-	virtual const FGuid& GetFolderGuid() const override { return GetActorDesc()->GetFolderGuid(); }
-
-	virtual FBox GetEditorBounds() const override { return GetActorDesc()->GetEditorBounds(); }
-	virtual FBox GetRuntimeBounds() const override { return GetActorDesc()->GetRuntimeBounds(); }
-
-	virtual bool GetProperty(FName PropertyName, FName* PropertyValue) const override { return GetActorDesc()->GetProperty(PropertyName, PropertyValue); }
-	virtual bool HasProperty(FName PropertyName) const override { return GetActorDesc()->HasProperty(PropertyName); }
-
-	virtual const TArray<FGuid>& GetReferences() const override { return GetActorDesc()->GetReferences(); }
-	virtual const TArray<FGuid>& GetEditorOnlyReferences() const override { return GetActorDesc()->GetEditorOnlyReferences(); }
-	virtual bool IsEditorOnlyReference(const FGuid& ReferenceGuid) const override { return GetActorDesc()->IsEditorOnlyReference(ReferenceGuid); }
+	virtual FName GetRuntimeGrid() const { return GetActorDesc()->GetRuntimeGrid(); }
+	virtual bool GetIsSpatiallyLoaded() const { return !GetForceNonSpatiallyLoaded() && GetActorDesc()->GetIsSpatiallyLoaded(); }
+	virtual bool GetActorIsEditorOnly() const { return GetActorDesc()->GetActorIsEditorOnly(); }
+	virtual bool GetActorIsRuntimeOnly() const { return GetActorDesc()->GetActorIsRuntimeOnly(); }
+	ENGINE_API virtual bool IsRuntimeRelevant() const;
+	ENGINE_API virtual bool IsEditorRelevant() const;
+
+	virtual bool IsUsingDataLayerAsset() const { return GetActorDesc()->IsUsingDataLayerAsset(); }
+	virtual const TArray<FName>& GetDataLayers() const { return GetActorDesc()->GetDataLayers(); }
+
+	virtual bool GetActorIsHLODRelevant() const { return GetActorDesc()->GetActorIsHLODRelevant(); }
+	virtual FSoftObjectPath GetHLODLayer() const { return GetActorDesc()->GetHLODLayer(); }
+
+	virtual const TArray<FName>& GetTags() const { return GetActorDesc()->GetTags(); }
+	virtual FName GetActorPackage() const { return GetActorDesc()->GetActorPackage(); }
+	ENGINE_API virtual FSoftObjectPath GetActorSoftPath() const;
+	virtual FName GetActorLabel() const { return GetActorDesc()->GetActorLabel(); }
+	ENGINE_API virtual FName GetActorName() const;
+	virtual FName GetFolderPath() const { return GetActorDesc()->GetFolderPath(); }
+	virtual const FGuid& GetFolderGuid() const { return GetActorDesc()->GetFolderGuid(); }
+
+	virtual FBox GetEditorBounds() const { return GetActorDesc()->GetEditorBounds(); }
+	virtual FBox GetRuntimeBounds() const { return GetActorDesc()->GetRuntimeBounds(); }
+
+	virtual bool GetProperty(FName PropertyName, FName* PropertyValue) const { return GetActorDesc()->GetProperty(PropertyName, PropertyValue); }
+	virtual bool HasProperty(FName PropertyName) const { return GetActorDesc()->HasProperty(PropertyName); }
+
+	virtual const TArray<FGuid>& GetReferences() const { return GetActorDesc()->GetReferences(); }
+	virtual const TArray<FGuid>& GetEditorOnlyReferences() const { return GetActorDesc()->GetEditorOnlyReferences(); }
+	virtual bool IsEditorOnlyReference(const FGuid& ReferenceGuid) const { return GetActorDesc()->IsEditorOnlyReference(ReferenceGuid); }
 	
-	virtual const FGuid& GetParentActor() const override { return GetActorDesc()->GetParentActor(); }
+	virtual const FGuid& GetParentActor() const { return GetActorDesc()->GetParentActor(); }
 		
-	virtual FGuid GetContentBundleGuid() const override { return GetActorDesc()->GetContentBundleGuid(); }
+	virtual FGuid GetContentBundleGuid() const { return GetActorDesc()->GetContentBundleGuid(); }
 
-	virtual bool IsChildContainerInstance() const override { return ChildContainerInstance || GetActorDesc()->IsChildContainerInstance(); }
-	virtual FName GetChildContainerPackage() const override { return GetActorDesc()->GetChildContainerPackage(); }
-	virtual EWorldPartitionActorFilterType GetChildContainerFilterType() const override { return GetActorDesc()->GetChildContainerFilterType(); }
-	virtual const FWorldPartitionActorFilter* GetChildContainerFilter() const override { return GetActorDesc()->GetChildContainerFilter(); }
-	virtual bool GetChildContainerInstance(FWorldPartitionActorDesc::FContainerInstance& OutContainerInstance) const override { return GetActorDesc()->GetChildContainerInstance(this, OutContainerInstance); }
+	virtual bool IsChildContainerInstance() const { return ChildContainerInstance || GetActorDesc()->IsChildContainerInstance(); }
+	virtual FName GetChildContainerPackage() const { return GetActorDesc()->GetChildContainerPackage(); }
+	virtual EWorldPartitionActorFilterType GetChildContainerFilterType() const { return GetActorDesc()->GetChildContainerFilterType(); }
+	virtual const FWorldPartitionActorFilter* GetChildContainerFilter() const { return GetActorDesc()->GetChildContainerFilter(); }
+	virtual bool GetChildContainerInstance(FWorldPartitionActorDesc::FContainerInstance& OutContainerInstance) const { return GetActorDesc()->GetChildContainerInstance(this, OutContainerInstance); }
 			
-	virtual bool IsMainWorldOnly() const override { return GetActorDesc()->IsMainWorldOnly(); }
-	virtual bool IsListedInSceneOutliner() const override { return GetActorDesc()->IsListedInSceneOutliner(); }
-	virtual const FGuid& GetSceneOutlinerParent() const override { return GetActorDesc()->GetSceneOutlinerParent(); }
+	virtual bool IsMainWorldOnly() const { return GetActorDesc()->IsMainWorldOnly(); }
+	virtual bool IsListedInSceneOutliner() const { return GetActorDesc()->IsListedInSceneOutliner(); }
+	virtual const FGuid& GetSceneOutlinerParent() const { return GetActorDesc()->GetSceneOutlinerParent(); }
 			
-	ENGINE_API virtual const FWorldPartitionActorDesc* GetActorDesc() const override { check(ActorDesc); return ActorDesc; }
+	ENGINE_API virtual const FWorldPartitionActorDesc* GetActorDesc() const { check(ActorDesc); return ActorDesc; }
 	
-	virtual bool HasResolvedDataLayerInstanceNames() const override { return ResolvedDataLayerInstanceNames.IsSet(); }
-	ENGINE_API virtual const TArray<FName>& GetDataLayerInstanceNames() const override;
+	virtual bool HasResolvedDataLayerInstanceNames() const { return ResolvedDataLayerInstanceNames.IsSet(); }
+	ENGINE_API virtual const TArray<FName>& GetDataLayerInstanceNames() const;
 
-	ENGINE_API virtual bool IsLoaded(bool bEvenIfPendingKill = false) const override;
-	ENGINE_API virtual AActor* GetActor(bool bEvenIfPendingKill = true, bool bEvenIfUnreachable = false) const override;
+	ENGINE_API virtual bool IsLoaded(bool bEvenIfPendingKill = false) const;
+	ENGINE_API virtual AActor* GetActor(bool bEvenIfPendingKill = true, bool bEvenIfUnreachable = false) const;
 	
-	ENGINE_API virtual FString ToString(FWorldPartitionActorDesc::EToStringMode Mode = FWorldPartitionActorDesc::EToStringMode::Compact) const override;
+	ENGINE_API virtual FString ToString(FWorldPartitionActorDesc::EToStringMode Mode = FWorldPartitionActorDesc::EToStringMode::Compact) const;
 	
-	virtual UActorDescContainerInstance* GetContainerInstance() const override { return ContainerInstance; }
-	//~ End
+	virtual UActorDescContainerInstance* GetContainerInstance() const { return ContainerInstance; }
+
+	inline FName GetActorLabelOrName() const { return GetActorLabel().IsNone() ? GetActorName() : GetActorLabel(); }
 
 	inline UActorDescContainerInstance* GetChildContainerInstance() const { return ChildContainerInstance; }
 
@@ -116,6 +113,7 @@ public:
 		
 	ENGINE_API AActor* Load();
 	ENGINE_API void Unload();
+
 protected:
 	UWorldPartition* GetLoadedChildWorldPartition() const { return GetActorDesc()->GetLoadedChildWorldPartition(this); }
 	ENGINE_API void UpdateActorDesc(FWorldPartitionActorDesc* InActorDesc);
diff --git a/Engine/Source/Runtime/Engine/Public/WorldPartition/WorldPartitionActorDescInstanceInterface.h b/Engine/Source/Runtime/Engine/Public/WorldPartition/WorldPartitionActorDescInstanceInterface.h
deleted file mode 100644
diff --git a/Engine/Source/Runtime/Engine/Public/WorldPartition/WorldPartitionActorDescInstanceViewInterface.h b/Engine/Source/Runtime/Engine/Public/WorldPartition/WorldPartitionActorDescInstanceViewInterface.h
--- a/Engine/Source/Runtime/Engine/Public/WorldPartition/WorldPartitionActorDescInstanceViewInterface.h
+++ b/Engine/Source/Runtime/Engine/Public/WorldPartition/WorldPartitionActorDescInstanceViewInterface.h
@@ -2,25 +2,89 @@
 #pragma once
 
 #if WITH_EDITOR
-
 #include "CoreMinimal.h"
-#include "WorldPartition/WorldPartitionActorDescInstanceInterface.h"
+#include "WorldPartition/WorldPartitionActorDesc.h"
 
 class AActor;
 class IStreamingGenerationErrorHandler;
 class UWorldPartition;
 struct FWorldPartitionActorFilter;
 enum class EWorldPartitionActorFilterType : uint8;
-
 #endif // WITH_EDITOR
 
 /**
  * Interface for a view on top of an actor desc, used to cache information that can be (potentially) different than the actor desc
  * itself due to streaming generation logic, etc.
  */
-class IWorldPartitionActorDescInstanceView : public IWorldPartitionActorDescInstance
+class IWorldPartitionActorDescInstanceView
 {
+#if WITH_EDITOR
 public:
+	virtual const FGuid& GetGuid() const = 0;
+
+	virtual FTopLevelAssetPath GetBaseClass() const = 0;
+	virtual FTopLevelAssetPath GetNativeClass() const = 0;
+	virtual UClass* GetActorNativeClass() const = 0;
+	
+	virtual FName GetRuntimeGrid() const = 0;
+	virtual bool GetIsSpatiallyLoaded() const = 0;
+	virtual bool GetActorIsEditorOnly() const = 0;
+	virtual bool GetActorIsRuntimeOnly() const = 0;
+	virtual bool IsRuntimeRelevant() const = 0;
+	virtual bool IsEditorRelevant() const = 0;
+
+	virtual bool IsUsingDataLayerAsset() const = 0;
+	virtual const TArray<FName>& GetDataLayers() const = 0;
+
+	virtual bool GetActorIsHLODRelevant() const = 0;
+	virtual FSoftObjectPath GetHLODLayer() const = 0;
+	
+	virtual const TArray<FName>& GetTags() const = 0;
+	virtual FName GetActorPackage() const = 0;
+	virtual FSoftObjectPath GetActorSoftPath() const = 0;
+	virtual FName GetActorLabel() const = 0;
+	virtual FName GetActorName() const = 0;
+	virtual FName GetFolderPath() const = 0;
+	virtual const FGuid& GetFolderGuid() const = 0;
+		
+	virtual FBox GetEditorBounds() const = 0;
+	virtual FBox GetRuntimeBounds() const = 0;
+		
+	virtual bool GetProperty(FName PropertyName, FName* PropertyValue) const = 0;
+	virtual bool HasProperty(FName PropertyName) const = 0;
+
+	virtual const TArray<FGuid>& GetReferences() const = 0;
+	virtual const TArray<FGuid>& GetEditorOnlyReferences() const = 0;
+	virtual bool IsEditorOnlyReference(const FGuid& ReferenceGuid) const = 0;
+		
+	virtual const FGuid& GetParentActor() const = 0;
+	
+	virtual FGuid GetContentBundleGuid() const = 0;
+	
+	virtual bool IsChildContainerInstance() const = 0;
+	virtual FName GetChildContainerPackage() const = 0;
+	virtual EWorldPartitionActorFilterType GetChildContainerFilterType() const = 0;
+	virtual const FWorldPartitionActorFilter* GetChildContainerFilter() const = 0;
+	virtual bool GetChildContainerInstance(FWorldPartitionActorDesc::FContainerInstance& OutContainerInstance) const = 0;
+			
+	virtual bool IsMainWorldOnly() const = 0;
+	virtual bool IsListedInSceneOutliner() const = 0;
+	virtual const FGuid& GetSceneOutlinerParent() const = 0;
+		
+	virtual FString ToString(FWorldPartitionActorDesc::EToStringMode Mode) const = 0;
+	virtual const FWorldPartitionActorDesc* GetActorDesc() const = 0;
+
+	virtual FName GetActorLabelOrName() const { return GetActorLabel().IsNone() ? GetActorName() : GetActorLabel(); }
+
+	virtual bool HasResolvedDataLayerInstanceNames() const = 0;
+	virtual const TArray<FName>& GetDataLayerInstanceNames() const = 0;
+
+	virtual AActor* GetActor(bool bEvenIfPendingKill = true, bool bEvenIfUnreachable = false) const = 0;
+	virtual bool IsLoaded(bool bEvenIfPendingKill = false) const = 0;
+
+	virtual UActorDescContainerInstance* GetContainerInstance() const = 0;
+
 	virtual void CheckForErrors(IStreamingGenerationErrorHandler* ErrorHandler) const = 0;
+#endif // WITH_EDITOR
 };
 
