Fix to possible Signed/unsigned `int32` conversion in `Texture2DDynamic.cpp` could lead to out-of-bounds read

diff --git a/Engine/Source/Runtime/Engine/Private/Texture2DDynamic.cpp b/Engine/Source/Runtime/Engine/Private/Texture2DDynamic.cpp
--- a/Engine/Source/Runtime/Engine/Private/Texture2DDynamic.cpp
+++ b/Engine/Source/Runtime/Engine/Private/Texture2DDynamic.cpp
@@ -95,18 +95,27 @@ void FTexture2DDynamicResource::WriteRawToTexture_RenderThread(TArrayView64<cons
 {
 	check(IsInRenderingThread());
 
-	const int32 Width = Texture2DRHI->GetSizeX();
-	const int32 Height = Texture2DRHI->GetSizeY();
+	const uint32 Width = Texture2DRHI->GetSizeX();
+	const uint32 Height = Texture2DRHI->GetSizeY();
+
+	// Prevent from locking texture if the source is empty or of size 0 or if source is is too small.
+	const uint64 SourceSize = RawData.Num();
+	if (!ensure(Width * Height != 0 && SourceSize >= Width * Height && Texture2DRHI->GetDesc().Format == EPixelFormat::PF_B8G8R8A8))
+	{
+		return;
+	}
 
 	uint32 DestStride = 0;
 	uint8* DestData = reinterpret_cast<uint8*>(RHILockTexture2D(Texture2DRHI, 0, RLM_WriteOnly, DestStride, false, false));
 
-	for (int32 y = 0; y < Height; y++)
+
+	for (uint32 y = 0; y < Height; y++)
 	{
-		uint8* DestPtr = &DestData[((int64)Height - 1 - y) * DestStride];
+		const uint64 CurrentLine = ((uint64)Height - 1 - y);
+		uint8* DestPtr = &DestData[CurrentLine * DestStride];
 
-		const FColor* SrcPtr = &((FColor*)(RawData.GetData()))[((int64)Height - 1 - y) * Width];
-		for (int32 x = 0; x < Width; x++)
+		const FColor* SrcPtr = &((FColor*)(RawData.GetData()))[CurrentLine * Width];
+		for (uint32 x = 0; x < Width; x++)
 		{
 			*DestPtr++ = SrcPtr->B;
 			*DestPtr++ = SrcPtr->G;
