Various RVT tidy up fixes.
Added support for writing fixed color when baking SVT. This replaces the fixed "debug" function and allows us to use the feature outside of debug usage.
Fixed path that creates an SVT for displacement map. Here the fixed color can be used to effectively set a zero displacement for the low resolution mips.
Added r.VT.RVT.MipColors cvar which uses the fixed color support to generate RVT using debug mip colors. This is useful for debugging issues with materials sampling higher resolution mips than expected, putting pressure on the physical pools.
Removed r.VT.RVT.EnableVolumes cvar which is no longer required now that volumes can be enabled/disabled per platform without triggering page table or physical texture allocations.
Added support for custom expanding of bounds for RVT volume.
Tidied naming of categories and sections for the RVT components.

diff --git a/Engine/Plugins/Experimental/VirtualHeightfieldMesh/Source/VirtualHeightfieldMeshEditor/Private/HeightfieldMinMaxTextureBuild.cpp b/Engine/Plugins/Experimental/VirtualHeightfieldMesh/Source/VirtualHeightfieldMeshEditor/Private/HeightfieldMinMaxTextureBuild.cpp
--- a/Engine/Plugins/Experimental/VirtualHeightfieldMesh/Source/VirtualHeightfieldMeshEditor/Private/HeightfieldMinMaxTextureBuild.cpp
+++ b/Engine/Plugins/Experimental/VirtualHeightfieldMesh/Source/VirtualHeightfieldMeshEditor/Private/HeightfieldMinMaxTextureBuild.cpp
@@ -206,7 +206,7 @@ namespace VirtualHeightfieldMesh
 					Desc.MaxLevel = MaxLevel;
 					Desc.bClearTextures = true;
 					Desc.bIsThumbnails = false;
-					Desc.DebugType = ERuntimeVirtualTextureDebugType::None;
+					Desc.FixedColor = FLinearColor::Transparent;
 					Desc.NumPageDescs = 1;
 					Desc.Targets[0].Texture = RenderTileResources.GetTileRenderTarget()->GetRHI();
 					Desc.Targets[1].Texture = nullptr;
diff --git a/Engine/Shaders/Private/VirtualTextureMaterial.usf b/Engine/Shaders/Private/VirtualTextureMaterial.usf
--- a/Engine/Shaders/Private/VirtualTextureMaterial.usf
+++ b/Engine/Shaders/Private/VirtualTextureMaterial.usf
@@ -110,8 +110,9 @@ void FPixelShaderInOut_MainPS(
 	float Displacement = 0.f;
 #endif
 
-	// Apply debug shading
-	BaseColor = lerp(BaseColor, half3(In.SvPosition.xy * ResolvedView.ViewSizeAndInvSize.zw, 0), ResolvedView.RuntimeVirtualTextureDebugParams.x * 0.25);
+	// Apply any fixed color.
+	BaseColor = lerp(BaseColor, ResolvedView.RuntimeVirtualTextureDebugParams.xyz, ResolvedView.RuntimeVirtualTextureDebugParams.w);
+	Displacement = lerp(Displacement, ResolvedView.RuntimeVirtualTextureDebugParams.x, ResolvedView.RuntimeVirtualTextureDebugParams.w);
 
 #if defined(OUT_BASECOLOR)
 	Out.MRT[0] = float4(BaseColor, 1.f) * Opacity;
diff --git a/Engine/Source/Editor/DetailCustomizations/Private/DetailCustomizations.cpp b/Engine/Source/Editor/DetailCustomizations/Private/DetailCustomizations.cpp
--- a/Engine/Source/Editor/DetailCustomizations/Private/DetailCustomizations.cpp
+++ b/Engine/Source/Editor/DetailCustomizations/Private/DetailCustomizations.cpp
@@ -820,6 +820,20 @@ void FDetailCustomizationsModule::RegisterSectionMappings()
 			Section->RemoveCategory("Brush Settings");
 		}
 	}
+
+	// RuntimeVirtualTextureComponent
+	{
+		{
+			TSharedRef<FPropertySection> Section = PropertyModule.FindOrCreateSection("RuntimeVirtualTextureComponent", "Rendering", LOCTEXT("Rendering", "Rendering"));
+			Section->AddCategory("Rendering");
+			Section->AddCategory("Runtime Virtual Texture");
+			Section->AddCategory("Streaming Virtual Texture");
+		}
+		{
+			TSharedRef<FPropertySection> Section = PropertyModule.FindOrCreateSection("RuntimeVirtualTextureComponent", "General", LOCTEXT("General", "General"));
+			Section->AddCategory("Volume Bounds");
+		}
+	}
 }
 
 #undef LOCTEXT_NAMESPACE
diff --git a/Engine/Source/Editor/VirtualTexturingEditor/Private/RuntimeVirtualTextureBuildStreamingMips.cpp b/Engine/Source/Editor/VirtualTexturingEditor/Private/RuntimeVirtualTextureBuildStreamingMips.cpp
--- a/Engine/Source/Editor/VirtualTexturingEditor/Private/RuntimeVirtualTextureBuildStreamingMips.cpp
+++ b/Engine/Source/Editor/VirtualTexturingEditor/Private/RuntimeVirtualTextureBuildStreamingMips.cpp
@@ -31,9 +31,9 @@ namespace
 
 			for (int32 Layer = 0; Layer < NumLayers; ++Layer)
 			{
-				check(InLayerFormats[Layer] == PF_G16 || InLayerFormats[Layer] == PF_B8G8R8A8 || InLayerFormats[Layer] == PF_DXT1 || InLayerFormats[Layer] == PF_DXT5 || InLayerFormats[Layer] == PF_BC5
-					|| InLayerFormats[Layer] == PF_R5G6B5_UNORM || InLayerFormats[Layer] == PF_B5G5R5A1_UNORM);
-				LayerFormats[Layer] = InLayerFormats[Layer] == PF_G16 ? PF_G16 : PF_B8G8R8A8;
+				check(InLayerFormats[Layer] == PF_G16 || InLayerFormats[Layer] == PF_B8G8R8A8 || InLayerFormats[Layer] == PF_DXT1 || InLayerFormats[Layer] == PF_DXT5 || InLayerFormats[Layer] == PF_BC4
+					|| InLayerFormats[Layer] == PF_BC5 || InLayerFormats[Layer] == PF_R5G6B5_UNORM || InLayerFormats[Layer] == PF_B5G5R5A1_UNORM);
+				LayerFormats[Layer] = InLayerFormats[Layer] == PF_G16 || InLayerFormats[Layer] == PF_BC4 ? PF_G16 : PF_B8G8R8A8;
 				LayerOffsets[Layer] = TotalSizeBytes;
 				TotalSizeBytes += CalculateImageBytes(InTileSize, InTileSize, 0, LayerFormats[Layer]) * InNumTilesX * InNumTilesY;
 			}
@@ -153,13 +153,13 @@ namespace RuntimeVirtualTexture
 		return true;
 	}
 	
-	bool BuildStreamedMips(URuntimeVirtualTextureComponent* InComponent, ERuntimeVirtualTextureDebugType DebugType)
+	bool BuildStreamedMips(URuntimeVirtualTextureComponent* InComponent, FLinearColor const& FixedColor)
 	{
 		EShadingPath ShadingPath = (InComponent && InComponent->GetScene()) ? InComponent->GetScene()->GetShadingPath() : EShadingPath::Deferred;
-		return BuildStreamedMips(ShadingPath, InComponent, DebugType);
+		return BuildStreamedMips(ShadingPath, InComponent, FixedColor);
 	}
 
-	bool BuildStreamedMips(EShadingPath ShadingPath, URuntimeVirtualTextureComponent* InComponent, ERuntimeVirtualTextureDebugType DebugType)
+	bool BuildStreamedMips(EShadingPath ShadingPath, URuntimeVirtualTextureComponent* InComponent, FLinearColor const& FixedColor)
 	{
 		if (!HasStreamedMips(ShadingPath, InComponent))
 		{
@@ -243,7 +243,7 @@ namespace RuntimeVirtualTexture
 					TileX, TileY,
 					TileSize, ImageSizeX, ImageSizeY, 
 					&FinalPixels,
-					DebugType](FRHICommandListImmediate& RHICmdList)
+					FixedColor](FRHICommandListImmediate& RHICmdList)
 				{
 					const FBox2D TileBox(FVector2D(0, 0), FVector2D(TileSize, TileSize));
 					const FIntRect TileRect(0, 0, TileSize, TileSize);
@@ -266,7 +266,7 @@ namespace RuntimeVirtualTexture
 						Desc.MaxLevel = MaxLevel;
 						Desc.bClearTextures = true;
 						Desc.bIsThumbnails = false;
-						Desc.DebugType = DebugType;
+						Desc.FixedColor = FixedColor;
 						Desc.NumPageDescs = 1;
 						Desc.Targets[0].Texture = RenderTileResources.GetRenderTarget(0);
 						Desc.Targets[1].Texture = RenderTileResources.GetRenderTarget(1);
diff --git a/Engine/Source/Editor/VirtualTexturingEditor/Private/RuntimeVirtualTextureBuildStreamingMips.h b/Engine/Source/Editor/VirtualTexturingEditor/Private/RuntimeVirtualTextureBuildStreamingMips.h
--- a/Engine/Source/Editor/VirtualTexturingEditor/Private/RuntimeVirtualTextureBuildStreamingMips.h
+++ b/Engine/Source/Editor/VirtualTexturingEditor/Private/RuntimeVirtualTextureBuildStreamingMips.h
@@ -4,7 +4,6 @@
 
 #include "CoreMinimal.h"
 
-enum class ERuntimeVirtualTextureDebugType;
 class URuntimeVirtualTextureComponent;
 enum class EShadingPath;
 
@@ -15,6 +14,6 @@ namespace RuntimeVirtualTexture
 	bool HasStreamedMips(EShadingPath ShadingPath, URuntimeVirtualTextureComponent* InComponent);
 
 	/** Build the streaming mips texture. */
-	bool BuildStreamedMips(URuntimeVirtualTextureComponent* InComponent, ERuntimeVirtualTextureDebugType DebugType);
-	bool BuildStreamedMips(EShadingPath ShadingPath, URuntimeVirtualTextureComponent* InComponent, ERuntimeVirtualTextureDebugType DebugType);
+	bool BuildStreamedMips(URuntimeVirtualTextureComponent* InComponent, FLinearColor const& FixedColor);
+	bool BuildStreamedMips(EShadingPath ShadingPath, URuntimeVirtualTextureComponent* InComponent, FLinearColor const& FixedColor);
 };
diff --git a/Engine/Source/Editor/VirtualTexturingEditor/Private/RuntimeVirtualTextureDetailsCustomization.cpp b/Engine/Source/Editor/VirtualTexturingEditor/Private/RuntimeVirtualTextureDetailsCustomization.cpp
--- a/Engine/Source/Editor/VirtualTexturingEditor/Private/RuntimeVirtualTextureDetailsCustomization.cpp
+++ b/Engine/Source/Editor/VirtualTexturingEditor/Private/RuntimeVirtualTextureDetailsCustomization.cpp
@@ -336,9 +336,8 @@ FReply FRuntimeVirtualTextureComponentDetailsCustomization::BuildStreamedMips()
 
 		RuntimeVirtualTextureComponent->GetStreamingTexture()->Modify();
 
-		const bool bDebug = RuntimeVirtualTextureComponent->IsBuildDebugStreamingMips();
-		const ERuntimeVirtualTextureDebugType DebugType = bDebug ? ERuntimeVirtualTextureDebugType::Debug : ERuntimeVirtualTextureDebugType::None;
-		if (RuntimeVirtualTexture::BuildStreamedMips(RuntimeVirtualTextureComponent, DebugType))
+		const FLinearColor FixedColor = RuntimeVirtualTextureComponent->GetStreamingMipsFixedColor();
+		if (RuntimeVirtualTexture::BuildStreamedMips(RuntimeVirtualTextureComponent, FixedColor))
 		{
 			bOK = true;
 		}
diff --git a/Engine/Source/Editor/VirtualTexturingEditor/Private/RuntimeVirtualTextureSetBounds.cpp b/Engine/Source/Editor/VirtualTexturingEditor/Private/RuntimeVirtualTextureSetBounds.cpp
--- a/Engine/Source/Editor/VirtualTexturingEditor/Private/RuntimeVirtualTextureSetBounds.cpp
+++ b/Engine/Source/Editor/VirtualTexturingEditor/Private/RuntimeVirtualTextureSetBounds.cpp
@@ -72,6 +72,13 @@ namespace RuntimeVirtualTexture
 			}
 		}
 
+		// Expand bounds.
+		const float ExpandBounds = InComponent->GetExpandBounds();
+		if (Bounds.IsValid && ExpandBounds > 0)
+		{
+			Bounds = Bounds.ExpandBy(ExpandBounds);
+		}
+
 		// Calculate the transform to fit the bounds.
 		FTransform Transform;
 		const FVector LocalPosition = Bounds.Min;
diff --git a/Engine/Source/Editor/VirtualTexturingEditor/Private/RuntimeVirtualTextureThumbnailRenderer.cpp b/Engine/Source/Editor/VirtualTexturingEditor/Private/RuntimeVirtualTextureThumbnailRenderer.cpp
--- a/Engine/Source/Editor/VirtualTexturingEditor/Private/RuntimeVirtualTextureThumbnailRenderer.cpp
+++ b/Engine/Source/Editor/VirtualTexturingEditor/Private/RuntimeVirtualTextureThumbnailRenderer.cpp
@@ -113,7 +113,7 @@ void URuntimeVirtualTextureThumbnailRenderer::Draw(UObject* Object, int32 X, int
 		Desc.MaxLevel = MaxLevel;
 		Desc.bClearTextures = true;
 		Desc.bIsThumbnails = true;
-		Desc.DebugType = ERuntimeVirtualTextureDebugType::None;
+		Desc.FixedColor = FLinearColor::Transparent;
 		Desc.NumPageDescs = 1;
 		Desc.Targets[0].Texture = RenderTarget->GetRenderTargetTexture();
 		Desc.PageDescs[0].DestBox[0] = DestBox;
diff --git a/Engine/Source/Editor/VirtualTexturingEditor/Private/VirtualTexturingEditorModule.cpp b/Engine/Source/Editor/VirtualTexturingEditor/Private/VirtualTexturingEditorModule.cpp
--- a/Engine/Source/Editor/VirtualTexturingEditor/Private/VirtualTexturingEditorModule.cpp
+++ b/Engine/Source/Editor/VirtualTexturingEditor/Private/VirtualTexturingEditorModule.cpp
@@ -91,7 +91,7 @@ bool FVirtualTexturingEditorModule::HasStreamedMips(URuntimeVirtualTextureCompon
 
 bool FVirtualTexturingEditorModule::BuildStreamedMips(URuntimeVirtualTextureComponent* InComponent) const
 {
-	return RuntimeVirtualTexture::BuildStreamedMips(EShadingPath::Deferred, InComponent, ERuntimeVirtualTextureDebugType::None);
+	return RuntimeVirtualTexture::BuildStreamedMips(EShadingPath::Deferred, InComponent, FLinearColor::Transparent);
 }
 
 bool FVirtualTexturingEditorModule::HasStreamedMips(EShadingPath ShadingPath, URuntimeVirtualTextureComponent* InComponent) const
@@ -101,7 +101,7 @@ bool FVirtualTexturingEditorModule::HasStreamedMips(EShadingPath ShadingPath, UR
 
 bool FVirtualTexturingEditorModule::BuildStreamedMips(EShadingPath ShadingPath, URuntimeVirtualTextureComponent* InComponent) const
 {
-	return RuntimeVirtualTexture::BuildStreamedMips(ShadingPath, InComponent, ERuntimeVirtualTextureDebugType::None);
+	return RuntimeVirtualTexture::BuildStreamedMips(ShadingPath, InComponent, FLinearColor::Transparent);
 }
 
 void FVirtualTexturingEditorModule::ConvertVirtualTextures(const TArray<UTexture2D *>& Textures, bool bConvertBackToNonVirtual, const TArray<UMaterial *>* RelatedMaterials /* = nullptr */) const
diff --git a/Engine/Source/Runtime/Engine/Classes/Components/RuntimeVirtualTextureComponent.h b/Engine/Source/Runtime/Engine/Classes/Components/RuntimeVirtualTextureComponent.h
--- a/Engine/Source/Runtime/Engine/Classes/Components/RuntimeVirtualTextureComponent.h
+++ b/Engine/Source/Runtime/Engine/Classes/Components/RuntimeVirtualTextureComponent.h
@@ -15,58 +15,62 @@ class UVirtualTextureBuilder;
 enum class EShadingPath;
 
 /** Component used to place a URuntimeVirtualTexture in the world. */
-UCLASS(Blueprintable, ClassGroup = Rendering, HideCategories = (Activation, Collision, Cooking, Mobility, LOD, Object, Physics, Rendering), MinimalAPI)
+UCLASS(Blueprintable, ClassGroup = Rendering, HideCategories = (Activation, Collision, Cooking, HLOD, Mobility, LOD, Navigation, Object, Physics), MinimalAPI)
 class URuntimeVirtualTextureComponent : public USceneComponent
 {
 	GENERATED_UCLASS_BODY()
 
 protected:
 	/** Actor to align rotation to. If set this actor is always included in the bounds calculation. */
-	UPROPERTY(EditAnywhere, Category = TransformFromBounds)
+	UPROPERTY(EditAnywhere, Category = VolumeBounds)
 	TSoftObjectPtr<AActor> BoundsAlignActor = nullptr;
 
 	/** Placeholder for details customization button. */
-	UPROPERTY(VisibleAnywhere, Transient, Category = TransformFromBounds)
+	UPROPERTY(VisibleAnywhere, Transient, Category = VolumeBounds)
 	bool bSetBoundsButton;
 
 	/** If the Bounds Align Actor is a Landscape then this will snap the bounds so that virtual texture texels align with landscape vertex positions. */
-	UPROPERTY(EditAnywhere, AdvancedDisplay, Category = TransformFromBounds, meta = (DisplayName = "Snap To Landscape"))
+	UPROPERTY(EditAnywhere, AdvancedDisplay, Category = VolumeBounds, meta = (DisplayName = "Snap To Landscape"))
 	bool bSnapBoundsToLandscape;
 
+	/** Amount to expand the Bounds during calculation. */
+	UPROPERTY(EditAnywhere, AdvancedDisplay, Category = VolumeBounds, meta = (UIMin = "0", ClampMin = "0"))
+	float ExpandBounds = 0;
+
 	/** The virtual texture object to use. */
-	UPROPERTY(EditAnywhere, BlueprintReadOnly, TextExportTransient, Category = VirtualTexture)
+	UPROPERTY(EditAnywhere, BlueprintReadOnly, TextExportTransient, Category = RuntimeVirtualTexture)
 	TObjectPtr<URuntimeVirtualTexture> VirtualTexture = nullptr;
 
 	/** Per platform overrides for enabling the virtual texture. Only affects In-Game and PIE. */
-	UPROPERTY(EditAnywhere, AdvancedDisplay, Category = VirtualTexture)
+	UPROPERTY(EditAnywhere, AdvancedDisplay, Category = RuntimeVirtualTexture)
 	FPerPlatformBool EnableInGamePerPlatform;
 
 	/** Enable the virtual texture only when Nanite is enabled. Can be used for a Displacement virtual texture with Nanite tessellation. */
-	UPROPERTY(EditAnywhere, AdvancedDisplay, Category = VirtualTexture)
+	UPROPERTY(EditAnywhere, AdvancedDisplay, Category = RuntimeVirtualTexture)
 	bool bEnableForNaniteOnly = false;
 
 	/** Set to true to enable scalability settings for the virtual texture. */
-	UPROPERTY(EditAnywhere, AdvancedDisplay, Category = VirtualTexture, meta = (InlineEditConditionToggle))
+	UPROPERTY(EditAnywhere, AdvancedDisplay, Category = RuntimeVirtualTexture, meta = (InlineEditConditionToggle))
 	bool bEnableScalability = false;
 
 	/** Group index of the scalability settings to use for the virtual texture. */
-	UPROPERTY(EditAnywhere, AdvancedDisplay, Category = VirtualTexture, meta = (UIMin = "0", UIMax = "2", EditCondition = bEnableScalability))
+	UPROPERTY(EditAnywhere, AdvancedDisplay, Category = RuntimeVirtualTexture, meta = (UIMin = "0", UIMax = "2", EditCondition = bEnableScalability))
 	uint32 ScalabilityGroup = 0;
 
 	/** Hide primitives in the main pass. Hidden primitives will be those that draw to this virtual texture with 'Draw in Main Pass' set to 'From Virtual Texture'. */
-	UPROPERTY(EditAnywhere, AdvancedDisplay, Category = VirtualTexture)
+	UPROPERTY(EditAnywhere, AdvancedDisplay, Category = RuntimeVirtualTexture)
 	bool bHidePrimitives = false;
 
 	/** Texture object containing streamed low mips. This can reduce rendering update cost. */
-	UPROPERTY(EditAnywhere, BlueprintReadOnly, Category = VirtualTextureBuild)
+	UPROPERTY(EditAnywhere, BlueprintReadOnly, Category = StreamingVirtualTexture)
 	TObjectPtr<UVirtualTextureBuilder> StreamingTexture = nullptr;
 
 	/** Number of streaming low mips to build for the virtual texture. */
-	UPROPERTY(EditAnywhere, Category = VirtualTextureBuild, meta = (UIMin = "0", UIMax = "12", DisplayName = "Build Levels"))
+	UPROPERTY(EditAnywhere, Category = StreamingVirtualTexture, meta = (UIMin = "0", UIMax = "12", DisplayName = "Build Levels"))
 	int32 StreamLowMips = 0;
 
 	/** Placeholder for details customization button. */
-	UPROPERTY(VisibleAnywhere, Transient, Category = VirtualTextureBuild)
+	UPROPERTY(VisibleAnywhere, Transient, Category = StreamingVirtualTexture)
 	bool bBuildStreamingMipsButton;
 
 	/** 
@@ -74,17 +78,21 @@ protected:
 	 * For compressors that support EncodeSpeed (i.e. Oodle), this is only applied if enabled (see Project Settings -> Texture Encoding). 
 	 * Note that this is in addition to any unavoidable loss due to the target format. Selecting "No Lossy Compression" will not result in zero distortion for BCn formats.
 	*/
-	UPROPERTY(EditAnywhere, BlueprintReadWrite, AdvancedDisplay, Category = VirtualTextureBuild)
+	UPROPERTY(EditAnywhere, BlueprintReadWrite, AdvancedDisplay, Category = StreamingVirtualTexture)
 	TEnumAsByte<ETextureLossyCompressionAmount> LossyCompressionAmount = TLCA_Default;
 
+	/** Build the streaming low mips using a fixed color. */
+	UPROPERTY(EditAnywhere, BlueprintReadOnly, AdvancedDisplay, Category = StreamingVirtualTexture, meta = (InlineEditConditionToggle))
+	bool bUseStreamingMipsFixedColor = false;
+
+	/** Fixed color to use when building the streaming low mips. This only affects BaseColor and Displacement attributes. The Red channel is used for fixed Displacement. */
+	UPROPERTY(EditAnywhere, BlueprintReadOnly, AdvancedDisplay, Category = StreamingVirtualTexture, meta = (DisplayName = "Fixed Color", HideAlphaChannel, EditCondition = bUseStreamingMipsFixedColor))
+	FLinearColor StreamingMipsFixedColor;
+
 	/** Use streaming low mips when rendering in editor. Set true to view and debug the baked streaming low mips. */
-	UPROPERTY(EditAnywhere, AdvancedDisplay, Category = VirtualTextureBuild, meta = (DisplayName = "View in Editor"))
+	UPROPERTY(EditAnywhere, AdvancedDisplay, Category = StreamingVirtualTexture, meta = (DisplayName = "View in Editor"))
 	bool bUseStreamingLowMipsInEditor = false;
 
-	/** Build the streaming low mips using debug coloring. This can help show where streaming mips are being used. */
-	UPROPERTY(EditAnywhere, BlueprintReadOnly, Transient, AdvancedDisplay, Category = VirtualTextureBuild, meta = (DisplayName = "Build Debug"))
-	bool bBuildDebugStreamingMips = false;
-
 #if WITH_EDITOR
 	/** Delegate handle for our function called on PIE end. */
 	FDelegateHandle PieEndDelegateHandle;
@@ -102,7 +110,7 @@ public:
 	 * This function marks an area of the runtime virtual texture as dirty.
 	 * @param WorldBounds : The world space bounds of the pages to invalidate.
 	 */
-	UFUNCTION(BlueprintCallable, Category = "VirtualTexture")
+	UFUNCTION(BlueprintCallable, Category = VirtualTexture)
 	ENGINE_API void Invalidate(FBoxSphereBounds const& WorldBounds);
 
 	/** Set the runtime virtual texture object on this component. */
@@ -135,8 +143,8 @@ public:
 	/** Get if we want to use any streaming low mips on this component. */
 	ENGINE_API bool IsStreamingLowMips(EShadingPath ShadingPath) const;
 
-	/** Public getter for debug streaming mips flag. */
-	bool IsBuildDebugStreamingMips() { return bBuildDebugStreamingMips; }
+	/** Public getter for streaming mips fixed color. */
+	ENGINE_API FLinearColor GetStreamingMipsFixedColor() const;
 
 	/** Public getter for lossy compression setting. */
 	TEnumAsByte<ETextureLossyCompressionAmount> GetLossyCompressionAmount() const { return LossyCompressionAmount; }
@@ -160,6 +168,8 @@ public:
 	TSoftObjectPtr<AActor>& GetBoundsAlignActor() { return BoundsAlignActor; }
 	/** Get if SnapBoundsToLandscape is set on this component. */
 	bool GetSnapBoundsToLandscape() const { return bSnapBoundsToLandscape; }
+	/** Get amount to expand the calculated bounds on this component. */
+	float GetExpandBounds() const { return ExpandBounds; }
 #endif
 	/** Get a translation to account for any vertex sample offset from the use of bSnapBoundsToLandscape. */
 	ENGINE_API FTransform GetTexelSnapTransform() const;
diff --git a/Engine/Source/Runtime/Engine/Classes/VT/RuntimeVirtualTextureVolume.h b/Engine/Source/Runtime/Engine/Classes/VT/RuntimeVirtualTextureVolume.h
--- a/Engine/Source/Runtime/Engine/Classes/VT/RuntimeVirtualTextureVolume.h
+++ b/Engine/Source/Runtime/Engine/Classes/VT/RuntimeVirtualTextureVolume.h
@@ -7,7 +7,7 @@
 #include "RuntimeVirtualTextureVolume.generated.h"
 
 /** Actor used to place a URuntimeVirtualTexture in the world. */
-UCLASS(Blueprintable, HideCategories = (Actor, Collision, Cooking, Input, LOD, Physics, Replication, Rendering), MinimalAPI)
+UCLASS(Blueprintable, HideCategories = (Actor, Collision, Cooking, HLOD, Input, LOD, Networking, Physics, Replication), MinimalAPI)
 class ARuntimeVirtualTextureVolume : public AActor
 {
 	GENERATED_UCLASS_BODY()
diff --git a/Engine/Source/Runtime/Engine/Private/Components/RuntimeVirtualTextureComponent.cpp b/Engine/Source/Runtime/Engine/Private/Components/RuntimeVirtualTextureComponent.cpp
--- a/Engine/Source/Runtime/Engine/Private/Components/RuntimeVirtualTextureComponent.cpp
+++ b/Engine/Source/Runtime/Engine/Private/Components/RuntimeVirtualTextureComponent.cpp
@@ -23,12 +23,6 @@
 
 #define LOCTEXT_NAMESPACE "URuntimeVirtualTextureComponent"
 
-static TAutoConsoleVariable<int32> CVarRVTEnableVolumes(
-	TEXT("r.VT.RVT.EnableVolumes"),
-	1,
-	TEXT("Enable RVT volumes."),
-	ECVF_Default);
-
 URuntimeVirtualTextureComponent::URuntimeVirtualTextureComponent(const FObjectInitializer& ObjectInitializer)
 	: Super(ObjectInitializer)
 	, EnableInGamePerPlatform(true)
@@ -143,11 +137,6 @@ void URuntimeVirtualTextureComponent::DestroyRenderState_Concurrent()
 
 bool URuntimeVirtualTextureComponent::IsEnabledInScene() const
 {
-	if (CVarRVTEnableVolumes.GetValueOnGameThread() == 0)
-	{
-		return false;
-	}
-
 	const bool bUseNanite = UseNanite(GetScene()->GetShaderPlatform());
 	if (bEnableForNaniteOnly && !bUseNanite)
 	{
@@ -275,6 +264,18 @@ bool URuntimeVirtualTextureComponent::IsStreamingTextureInvalid() const
 		StreamingTexture->BuildHash != CalculateStreamingTextureSettingsHash();
 }
 
+FLinearColor URuntimeVirtualTextureComponent::GetStreamingMipsFixedColor() const 
+{
+	if (!bUseStreamingMipsFixedColor)
+	{
+		return FLinearColor::Transparent;
+	} 
+	
+	FLinearColor Color(StreamingMipsFixedColor);
+	Color.A = 1.f;
+	return Color;
+}
+
 // RAII class to release and recreate runtime virtual texture producers associated with a UVirtualTextureBuilder.
 // Required around modifications of a UVirtualTextureBuilder because virtual producers hold pointers to the internal data.
 class FScopedRuntimeVirtualTextureRecreate
@@ -311,7 +312,7 @@ private:
 
 static void GetLayerFormatSettings(FTextureFormatSettings& OutFormatSettings, EPixelFormat LayerFormat, bool IsLayerYCoCg, bool IsLayerSRGB, bool IsLayerLQCompression)
 {
-	OutFormatSettings.CompressionSettings = IsLayerLQCompression? TC_LQ : (LayerFormat == PF_BC5 ? TC_Normalmap : TC_Default);
+	OutFormatSettings.CompressionSettings = IsLayerLQCompression ? TC_LQ : (LayerFormat == PF_BC5 ? TC_Normalmap : (LayerFormat == PF_BC4 ? TC_Alpha : TC_Default));
 	OutFormatSettings.CompressionNone = LayerFormat == PF_B8G8R8A8 || LayerFormat == PF_G16;
 	OutFormatSettings.CompressionNoAlpha = LayerFormat == PF_DXT1 || LayerFormat == PF_BC5 || LayerFormat == PF_R5G6B5_UNORM;
 	OutFormatSettings.CompressionForceAlpha = LayerFormat == PF_DXT5;
@@ -343,7 +344,7 @@ void URuntimeVirtualTextureComponent::InitializeStreamingTexture(EShadingPath Sh
 		for (int32 Layer = 0; Layer < BuildDesc.LayerCount; Layer++) 
 		{
 			const EPixelFormat LayerFormat = VirtualTexture->GetLayerFormat(Layer);
-			BuildDesc.LayerFormats[Layer] = LayerFormat == PF_G16 ? TSF_G16 : TSF_BGRA8;
+			BuildDesc.LayerFormats[Layer] = LayerFormat == PF_G16 || LayerFormat == PF_BC4 ? TSF_G16 : TSF_BGRA8;
 			bool IsLayerLQCompression = (VirtualTexture->GetMaterialType() == ERuntimeVirtualTextureMaterialType::BaseColor_Normal_Roughness && VirtualTexture->GetLQCompression() && LayerFormat != PF_B8G8R8A8);
 			GetLayerFormatSettings(BuildDesc.LayerFormatSettings[Layer], LayerFormat, VirtualTexture->IsLayerYCoCg(Layer), VirtualTexture->IsLayerSRGB(Layer), IsLayerLQCompression);
 		}
@@ -355,7 +356,6 @@ void URuntimeVirtualTextureComponent::InitializeStreamingTexture(EShadingPath Sh
 		BuildDesc.InData = InData;
 
 		StreamingTexture->BuildTexture(ShadingPath, BuildDesc);
-		StreamingTexture->EnableCookPerPlatform = EnableInGamePerPlatform;
 		StreamingTexture->Modify();
 	}
 }
@@ -367,12 +367,6 @@ bool URuntimeVirtualTextureComponent::CanEditChange(const FProperty* InProperty)
 	{
 		bCanEdit &= GetVirtualTexture() != nullptr && GetStreamingTexture() != nullptr;
 	}
-	else if (InProperty->GetFName() == TEXT("bBuildDebugStreamingMips"))
-	{
-		bCanEdit = GetVirtualTexture() != nullptr && 
-			GetVirtualTexture()->GetMaterialType() != ERuntimeVirtualTextureMaterialType::WorldHeight && 
-			GetVirtualTexture()->GetMaterialType() != ERuntimeVirtualTextureMaterialType::Displacement;
-	}
 	return bCanEdit;
 }
 
diff --git a/Engine/Source/Runtime/Engine/Public/VT/RuntimeVirtualTextureEnum.h b/Engine/Source/Runtime/Engine/Public/VT/RuntimeVirtualTextureEnum.h
--- a/Engine/Source/Runtime/Engine/Public/VT/RuntimeVirtualTextureEnum.h
+++ b/Engine/Source/Runtime/Engine/Public/VT/RuntimeVirtualTextureEnum.h
@@ -79,10 +79,3 @@ enum ERuntimeVirtualTextureShaderUniform
 	ERuntimeVirtualTextureShaderUniform_WorldHeightUnpack,
 	ERuntimeVirtualTextureShaderUniform_Count,
 };
-
-/** Enumeration of runtime virtual texture debug modes. */
-enum class ERuntimeVirtualTextureDebugType
-{
-	None,
-	Debug
-};
diff --git a/Engine/Source/Runtime/Renderer/Private/VT/RuntimeVirtualTextureProducer.cpp b/Engine/Source/Runtime/Renderer/Private/VT/RuntimeVirtualTextureProducer.cpp
--- a/Engine/Source/Runtime/Renderer/Private/VT/RuntimeVirtualTextureProducer.cpp
+++ b/Engine/Source/Runtime/Renderer/Private/VT/RuntimeVirtualTextureProducer.cpp
@@ -70,7 +70,7 @@ void FRuntimeVirtualTextureFinalizer::Finalize(FRDGBuilder& GraphBuilder)
 	RenderPageBatchDesc.MaxLevel = Desc.MaxLevel;
 	RenderPageBatchDesc.bClearTextures = bClearTextures;
 	RenderPageBatchDesc.bIsThumbnails = false;
-	RenderPageBatchDesc.DebugType = ERuntimeVirtualTextureDebugType::None;
+	RenderPageBatchDesc.FixedColor = FLinearColor::Transparent;
 	
 	for (int LayerIndex = 0; LayerIndex < RuntimeVirtualTexture::MaxTextureLayers; ++LayerIndex)
 	{
diff --git a/Engine/Source/Runtime/Renderer/Private/VT/RuntimeVirtualTextureRender.cpp b/Engine/Source/Runtime/Renderer/Private/VT/RuntimeVirtualTextureRender.cpp
--- a/Engine/Source/Runtime/Renderer/Private/VT/RuntimeVirtualTextureRender.cpp
+++ b/Engine/Source/Runtime/Renderer/Private/VT/RuntimeVirtualTextureRender.cpp
@@ -26,6 +26,12 @@
 
 namespace RuntimeVirtualTexture
 {
+	static TAutoConsoleVariable<bool> CVarVTMipColors(
+		TEXT("r.VT.RVT.MipColors"),
+		false,
+		TEXT("Render mip colors to RVT BaseColor."),
+		ECVF_RenderThreadSafe);
+
 	static TAutoConsoleVariable<int32> CVarVTHighQualityPerPixelHeight(
 		TEXT("r.VT.RVT.HighQualityPerPixelHeight"),
 		1,
@@ -1351,6 +1357,16 @@ namespace RuntimeVirtualTexture
 		return Scene != nullptr && Scene->GetRenderScene() != nullptr && Scene->GetRenderScene()->GPUScene.IsRendering();
 	}
 
+	FLinearColor GetMipLevelColor(uint32 InLevel)
+	{
+		static const uint32 MipColors[] = {
+			0xC0FFFFFF, 0xC0FFFF00, 0xC000FFFF, 0xC000FF00, 0xC0FF00FF, 0xC0FF0000, 0xC00000FF,
+			0xC0808080, 0xC0808000, 0xC0008080, 0xC0008000, 0xC0800080, 0xC0800000, 0xC0000080 };
+
+		InLevel = FMath::Min<uint32>(InLevel, sizeof(MipColors) / sizeof(MipColors[0]) - 1);
+		return FLinearColor(FColor(MipColors[InLevel]));
+	}
+
 	void RenderPage(
 		FRDGBuilder& GraphBuilder,
 		FScene* Scene,
@@ -1373,7 +1389,7 @@ namespace RuntimeVirtualTexture
 		FBox2D const& UVRange,
 		uint8 vLevel,
 		uint8 MaxLevel,
-		ERuntimeVirtualTextureDebugType DebugType)
+		FLinearColor const& FixedColor)
 	{
 		RDG_EVENT_SCOPE(GraphBuilder, "VirtualTextureDynamicCache");
 
@@ -1430,7 +1446,7 @@ namespace RuntimeVirtualTexture
 		View->SetupUniformBufferParameters(nullptr, 0, *View->CachedViewUniformShaderParameters);
 		View->CachedViewUniformShaderParameters->RuntimeVirtualTextureMipLevel = MipLevelParameter;
 		View->CachedViewUniformShaderParameters->RuntimeVirtualTexturePackHeight = FVector2f(WorldHeightPackParameter);	// LWC_TODO: Precision loss
-		View->CachedViewUniformShaderParameters->RuntimeVirtualTextureDebugParams = FVector4f(DebugType == ERuntimeVirtualTextureDebugType::Debug ? 1.f : 0.f, 0.f, 0.f, 0.f);
+		View->CachedViewUniformShaderParameters->RuntimeVirtualTextureDebugParams = CVarVTMipColors.GetValueOnRenderThread() ? GetMipLevelColor(vLevel) : FixedColor;
 		View->ViewUniformBuffer = TUniformBufferRef<FViewUniformShaderParameters>::CreateUniformBufferImmediate(*View->CachedViewUniformShaderParameters, UniformBuffer_SingleFrame);
 
 		// Build graph
@@ -1555,7 +1571,7 @@ namespace RuntimeVirtualTexture
 					PageDesc.UVRange,
 					PageDesc.vLevel,
 					InDesc.MaxLevel,
-					InDesc.DebugType);
+					InDesc.FixedColor);
 			}
 		}
 	}
diff --git a/Engine/Source/Runtime/Renderer/Public/VT/RuntimeVirtualTextureRender.h b/Engine/Source/Runtime/Renderer/Public/VT/RuntimeVirtualTextureRender.h
--- a/Engine/Source/Runtime/Renderer/Public/VT/RuntimeVirtualTextureRender.h
+++ b/Engine/Source/Runtime/Renderer/Public/VT/RuntimeVirtualTextureRender.h
@@ -72,8 +72,8 @@ namespace RuntimeVirtualTexture
 		bool bClearTextures;
 		/** Set to true for thumbnail rendering. */
 		bool bIsThumbnails;
-		/** Debug visualization to render with. */
-		ERuntimeVirtualTextureDebugType DebugType;
+		/** Fixed BaseColor to apply. Uses alpha channel to blend with material output. */
+		FLinearColor FixedColor;
 
 		/** Physical texture targets to render to. */
 		FRenderPageTarget Targets[RuntimeVirtualTexture::MaxTextureLayers];
