Add LocalPlayer async step to server travel requests.
#rnx

diff --git a/Engine/Source/Runtime/Engine/Classes/Engine/LocalPlayer.h b/Engine/Source/Runtime/Engine/Classes/Engine/LocalPlayer.h
--- a/Engine/Source/Runtime/Engine/Classes/Engine/LocalPlayer.h
+++ b/Engine/Source/Runtime/Engine/Classes/Engine/LocalPlayer.h
@@ -420,7 +420,14 @@ public:
 	 * @return False if an error occurred, true if the play actor was successfully spawned.	 
 	 */
 	ENGINE_API virtual bool SpawnPlayActor(const FString& URL,FString& OutError, UWorld* InWorld);
-	
+
+	DECLARE_DELEGATE(FOnPreBeginHandshakeCompleteDelegate);
+
+	/**
+	 * Allow local player to run any async tasks needed before starting travel to a server.
+	 */
+	virtual void PreBeginHandshake(const FOnPreBeginHandshakeCompleteDelegate& OnComplete) { OnComplete.ExecuteIfBound(); }
+
 	/** Send a splitscreen join command to the server to allow a splitscreen player to connect to the game
 	 * the client must already be connected to a server for this function to work
 	 * @note this happens automatically for all viewports that exist during the initial server connect
diff --git a/Engine/Source/Runtime/Engine/Classes/Engine/PendingNetGame.h b/Engine/Source/Runtime/Engine/Classes/Engine/PendingNetGame.h
--- a/Engine/Source/Runtime/Engine/Classes/Engine/PendingNetGame.h
+++ b/Engine/Source/Runtime/Engine/Classes/Engine/PendingNetGame.h
@@ -93,6 +93,11 @@ public:
 
 	void	InitNetDriver();
 
+	/**
+	 * Begin initial handshake if needed, or call SendInitialJoin.
+	 */
+	void BeginHandshake();
+
 	/**
 	 * Send the packet for triggering the initial join
 	 */
diff --git a/Engine/Source/Runtime/Engine/Private/PendingNetGame.cpp b/Engine/Source/Runtime/Engine/Private/PendingNetGame.cpp
--- a/Engine/Source/Runtime/Engine/Private/PendingNetGame.cpp
+++ b/Engine/Source/Runtime/Engine/Private/PendingNetGame.cpp
@@ -54,19 +54,20 @@ void UPendingNetGame::InitNetDriver()
 
 		if( NetDriver->InitConnect( this, URL, ConnectionError ) )
 		{
-			UNetConnection* ServerConn = NetDriver->ServerConnection;
-
 			FNetDelegates::OnPendingNetGameConnectionCreated.Broadcast(this);
 
-			// Kick off the connection handshake
-			if (ServerConn->Handler.IsValid())
+			ULocalPlayer* LocalPlayer = GEngine->GetFirstGamePlayer(this);
+			if (LocalPlayer)
 			{
-				ServerConn->Handler->BeginHandshaking(
-					FPacketHandlerHandshakeComplete::CreateUObject(this, &UPendingNetGame::SendInitialJoin));
+				LocalPlayer->PreBeginHandshake(ULocalPlayer::FOnPreBeginHandshakeCompleteDelegate::CreateWeakLambda(this,
+					[this]()
+					{
+						BeginHandshake();
+					}));
 			}
 			else
 			{
-				SendInitialJoin();
+				BeginHandshake();
 			}
 		}
 		else
@@ -90,6 +91,21 @@ void UPendingNetGame::InitNetDriver()
 	}
 }
 
+void UPendingNetGame::BeginHandshake()
+{
+	// Kick off the connection handshake
+	UNetConnection* ServerConn = NetDriver->ServerConnection;
+	if (ServerConn->Handler.IsValid())
+	{
+		ServerConn->Handler->BeginHandshaking(
+			FPacketHandlerHandshakeComplete::CreateUObject(this, &UPendingNetGame::SendInitialJoin));
+	}
+	else
+	{
+		SendInitialJoin();
+	}
+}
+
 void UPendingNetGame::SendInitialJoin()
 {
 	if (NetDriver != nullptr)
