Newly-spawned anim instances now call BeginPlay overide points when they are spawned 'late'

I.e. the instance is created after the component has had BeginPlay() called on it

diff --git a/Engine/Source/Runtime/Engine/Private/Animation/AnimInstance.cpp b/Engine/Source/Runtime/Engine/Private/Animation/AnimInstance.cpp
--- a/Engine/Source/Runtime/Engine/Private/Animation/AnimInstance.cpp
+++ b/Engine/Source/Runtime/Engine/Private/Animation/AnimInstance.cpp
@@ -3252,6 +3252,12 @@ void UAnimInstance::PerformLinkedLayerOverlayOperation(TSubclassOf<UAnimInstance
 								NewLinkedInstance->bReceiveNotifiesFromLinkedInstances = LayerNode->bReceiveNotifiesFromLinkedInstances;
 								NewLinkedInstance->InitializeAnimation();
 
+								if(MeshComp->HasBegunPlay())
+								{
+									NewLinkedInstance->NativeBeginPlay();
+									NewLinkedInstance->BlueprintBeginPlay();
+								}
+
 								// Unlink any layer nodes in the new linked instance, as they may have been hooked up to self in InitializeAnimation above.
 								UnlinkLayerNodesInInstance(NewLinkedInstance, LayerPair.Value);
 
@@ -3308,6 +3314,12 @@ void UAnimInstance::PerformLinkedLayerOverlayOperation(TSubclassOf<UAnimInstance
 							NewLinkedInstance->bCreatedByLinkedAnimGraph = true;
 							NewLinkedInstance->InitializeAnimation();
 
+							if(MeshComp->HasBegunPlay())
+							{
+								NewLinkedInstance->NativeBeginPlay();
+								NewLinkedInstance->BlueprintBeginPlay();
+							}
+
 							// Unlink any layer nodes in the new linked instance, as they may have been hooked up to self in InitializeAnimation above.
 							UnlinkLayerNodesInInstance(NewLinkedInstance, LayerPair.Value);
 
diff --git a/Engine/Source/Runtime/Engine/Private/Animation/AnimNode_LinkedAnimGraph.cpp b/Engine/Source/Runtime/Engine/Private/Animation/AnimNode_LinkedAnimGraph.cpp
--- a/Engine/Source/Runtime/Engine/Private/Animation/AnimNode_LinkedAnimGraph.cpp
+++ b/Engine/Source/Runtime/Engine/Private/Animation/AnimNode_LinkedAnimGraph.cpp
@@ -312,6 +312,12 @@ void FAnimNode_LinkedAnimGraph::ReinitializeLinkedAnimInstance(const UAnimInstan
 			// Initialize the new instance
 			InstanceToRun->InitializeAnimation();
 
+			if(MeshComp->HasBegunPlay())
+			{
+				InstanceToRun->NativeBeginPlay();
+				InstanceToRun->BlueprintBeginPlay();
+			}
+
 			MeshComp->GetLinkedAnimInstances().Add(InstanceToRun);
 		}
 
diff --git a/Engine/Source/Runtime/Engine/Private/Animation/AnimSubsystem_SharedLinkedAnimLayers.cpp b/Engine/Source/Runtime/Engine/Private/Animation/AnimSubsystem_SharedLinkedAnimLayers.cpp
--- a/Engine/Source/Runtime/Engine/Private/Animation/AnimSubsystem_SharedLinkedAnimLayers.cpp
+++ b/Engine/Source/Runtime/Engine/Private/Animation/AnimSubsystem_SharedLinkedAnimLayers.cpp
@@ -87,6 +87,8 @@ FLinkedAnimLayerInstanceData* FLinkedAnimLayerClassData::FindInstanceData(const
 
 UAnimInstance* FLinkedAnimLayerClassData::FindOrAddInstanceForLinking(UAnimInstance* OwningInstance, FName Function, bool& bIsNewInstance)
 {
+	USkeletalMeshComponent* Mesh = OwningInstance->GetSkelMeshComponent();
+
 	for (FLinkedAnimLayerInstanceData& LayerInstanceData : InstancesData)
 	{
 		// Check if function is already linked
@@ -95,7 +97,6 @@ UAnimInstance* FLinkedAnimLayerClassData::FindOrAddInstanceForLinking(UAnimInsta
 			// Re-add persistent instance of first function re-bind
 			if (LayerInstanceData.IsPersistent() && LayerInstanceData.GetLinkedFunctions().Num() == 0)
 			{
-				USkeletalMeshComponent* Mesh = OwningInstance->GetSkelMeshComponent();
 				// Make sure the bones to update are up to date with LOD changes / bone visibility / cosmetics, etc.
 				LayerInstanceData.Instance->RecalcRequiredBones();
 				check(!Mesh->GetLinkedAnimInstances().Contains(LayerInstanceData.Instance));
@@ -111,9 +112,16 @@ UAnimInstance* FLinkedAnimLayerClassData::FindOrAddInstanceForLinking(UAnimInsta
 
 	// Create a new object
 	bIsNewInstance = true;
-	UAnimInstance* NewAnimInstance = NewObject<UAnimInstance>(OwningInstance->GetSkelMeshComponent(), Class);
+	UAnimInstance* NewAnimInstance = NewObject<UAnimInstance>(Mesh, Class);
 	NewAnimInstance->bCreatedByLinkedAnimGraph = true;
 	NewAnimInstance->InitializeAnimation();
+
+	if(Mesh->HasBegunPlay())
+	{
+		NewAnimInstance->NativeBeginPlay();
+		NewAnimInstance->BlueprintBeginPlay();
+	}
+	
 	FLinkedAnimLayerInstanceData& NewInstanceData = AddInstance(NewAnimInstance);
 	NewInstanceData.AddLinkedFunction(Function, NewAnimInstance);
 	OwningInstance->GetSkelMeshComponent()->GetLinkedAnimInstances().Add(NewAnimInstance);
diff --git a/Engine/Source/Runtime/Engine/Private/Components/SkeletalMeshComponent.cpp b/Engine/Source/Runtime/Engine/Private/Components/SkeletalMeshComponent.cpp
--- a/Engine/Source/Runtime/Engine/Private/Components/SkeletalMeshComponent.cpp
+++ b/Engine/Source/Runtime/Engine/Private/Components/SkeletalMeshComponent.cpp
@@ -1011,6 +1011,14 @@ bool USkeletalMeshComponent::InitializeAnimScriptInstance(bool bForceReinit, boo
 				ResetLinkedAnimInstances();
 
 				AnimScriptInstance->InitializeAnimation(bInDeferRootNodeInitialization);
+
+				// Call BeginPlay on anim instances created at runtime.
+				if (HasBegunPlay())
+				{
+					AnimScriptInstance->NativeBeginPlay();
+					AnimScriptInstance->BlueprintBeginPlay();
+				}
+
 				bInitializedMainInstance = true;
 			}
 		}
@@ -1032,8 +1040,15 @@ bool USkeletalMeshComponent::InitializeAnimScriptInstance(bool bForceReinit, boo
 				if (AnimScriptInstance)
 				{
 					ResetLinkedAnimInstances();
-				
+
 					AnimScriptInstance->InitializeAnimation(bInDeferRootNodeInitialization);
+
+					if(HasBegunPlay())
+					{
+						AnimScriptInstance->NativeBeginPlay();
+						AnimScriptInstance->BlueprintBeginPlay();
+					}
+
 					bInitializedMainInstance = true;
 				}
 
@@ -1072,6 +1087,13 @@ bool USkeletalMeshComponent::InitializeAnimScriptInstance(bool bForceReinit, boo
 			{
 				PostProcessAnimInstance->InitializeAnimation();
 
+				// Call BeginPlay on anim instances created at runtime.
+				if (HasBegunPlay())
+				{
+					PostProcessAnimInstance->NativeBeginPlay();
+					PostProcessAnimInstance->BlueprintBeginPlay();
+				}
+				
 				if(FAnimNode_LinkedInputPose* InputNode = PostProcessAnimInstance->GetLinkedInputPoseNode())
 				{
 					InputNode->CachedInputPose.SetBoneContainer(&PostProcessAnimInstance->GetRequiredBones());
@@ -1195,6 +1217,12 @@ void USkeletalMeshComponent::PostEditChangeProperty(FPropertyChangedEvent& Prope
 						SCOPE_CYCLE_COUNTER(STAT_AnimSpawnTime);
 						AnimScriptInstance = NewObject<UAnimInstance>(this, AnimClass);
 						AnimScriptInstance->InitializeAnimation();
+
+						if(HasBegunPlay())
+						{
+							AnimScriptInstance->NativeBeginPlay();
+							AnimScriptInstance->BlueprintBeginPlay();
+						}
 					}
 				}
 			}
@@ -1223,6 +1251,12 @@ void USkeletalMeshComponent::PostEditChangeProperty(FPropertyChangedEvent& Prope
 					{
 						PostProcessAnimInstance = NewObject<UAnimInstance>(this, *SkelMesh->GetPostProcessAnimBlueprint());
 						PostProcessAnimInstance->InitializeAnimation();
+
+						if(HasBegunPlay())
+						{
+							PostProcessAnimInstance->NativeBeginPlay();
+							PostProcessAnimInstance->BlueprintBeginPlay();
+						}
 					}
 					else
 					{
