World Partition Runtime HashSet: fix streaming cells GUIDs to not include the world name, so they remain valid when creating a map from a template with HLODs.

diff --git a/Engine/Source/Runtime/Engine/Private/WorldPartition/RuntimeHashSet/RuntimePartitionLHGrid.cpp b/Engine/Source/Runtime/Engine/Private/WorldPartition/RuntimeHashSet/RuntimePartitionLHGrid.cpp
--- a/Engine/Source/Runtime/Engine/Private/WorldPartition/RuntimeHashSet/RuntimePartitionLHGrid.cpp
+++ b/Engine/Source/Runtime/Engine/Private/WorldPartition/RuntimeHashSet/RuntimePartitionLHGrid.cpp
@@ -142,4 +142,9 @@ bool URuntimePartitionLHGrid::GenerateStreaming(const FGenerateStreamingParams&
 
 	return true;
 }
+
+FArchive& URuntimePartitionLHGrid::AppendCellGuid(FArchive& InAr)
+{
+	return Super::AppendCellGuid(InAr) << CellSize;
+}
 #endif
\ No newline at end of file
diff --git a/Engine/Source/Runtime/Engine/Private/WorldPartition/RuntimeHashSet/WorldPartitionRuntimeHashSet.cpp b/Engine/Source/Runtime/Engine/Private/WorldPartition/RuntimeHashSet/WorldPartitionRuntimeHashSet.cpp
--- a/Engine/Source/Runtime/Engine/Private/WorldPartition/RuntimeHashSet/WorldPartitionRuntimeHashSet.cpp
+++ b/Engine/Source/Runtime/Engine/Private/WorldPartition/RuntimeHashSet/WorldPartitionRuntimeHashSet.cpp
@@ -610,55 +610,65 @@ void UWorldPartitionRuntimeHashSet::PostEditChangeChainProperty(FPropertyChanged
 
 UWorldPartitionRuntimeHashSet::FCellUniqueId UWorldPartitionRuntimeHashSet::GetCellUniqueId(const URuntimePartition::FCellDescInstance& InCellDescInstance) const
 {
-	TStringBuilder<128> CellNameBuilder;
+	FCellUniqueId CellUniqueId;
+	FName CellNameID = InCellDescInstance.Name;
+	FDataLayersID DataLayersID(InCellDescInstance.DataLayerInstances);
+	FGuid ContentBundleID(InCellDescInstance.ContentBundleID);
 
-	UWorld* OuterWorld = GetTypedOuter<UWorld>();
-	check(OuterWorld);
+	// Build cell unique name
+	{
+		TStringBuilder<128> CellNameBuilder;
 
-	FString WorldName = FPackageName::GetShortName(OuterWorld->GetPackage());
+		UWorld* OuterWorld = GetTypedOuter<UWorld>();
+		check(OuterWorld);
 
-	CellNameBuilder.Appendf(TEXT("%s_%s"), *WorldName, *InCellDescInstance.Name.ToString());
+		FString WorldName = FPackageName::GetShortName(OuterWorld->GetPackage());
 
-	const FDataLayersID DataLayersID(InCellDescInstance.DataLayerInstances);
-	if (DataLayersID.GetHash())
-	{
-		CellNameBuilder.Appendf(TEXT("_d%X"), DataLayersID.GetHash());
-	}
+		CellNameBuilder.Appendf(TEXT("%s_%s"), *WorldName, *CellNameID.ToString());
 
-	if (InCellDescInstance.ContentBundleID.IsValid())
-	{
-		CellNameBuilder.Appendf(TEXT("_c%s"), *UContentBundleDescriptor::GetContentBundleCompactString(InCellDescInstance.ContentBundleID));
-	}
+		if (DataLayersID.GetHash())
+		{
+			CellNameBuilder.Appendf(TEXT("_d%X"), DataLayersID.GetHash());
+		}
 
-	if (!IsRunningCookCommandlet() && OuterWorld->IsGameWorld())
-	{
-		FString SourceWorldPath;
-		FString InstancedWorldPath;
-		if (OuterWorld->GetSoftObjectPathMapping(SourceWorldPath, InstancedWorldPath))
+		if (ContentBundleID.IsValid())
+		{
+			CellNameBuilder.Appendf(TEXT("_c%s"), *UContentBundleDescriptor::GetContentBundleCompactString(ContentBundleID));
+		}
+
+		if (!IsRunningCookCommandlet() && OuterWorld->IsGameWorld())
 		{
-			const FTopLevelAssetPath SourceAssetPath(SourceWorldPath);
-			WorldName = FPackageName::GetShortName(SourceAssetPath.GetPackageName());
+			FString SourceWorldPath;
+			FString InstancedWorldPath;
+			if (OuterWorld->GetSoftObjectPathMapping(SourceWorldPath, InstancedWorldPath))
+			{
+				const FTopLevelAssetPath SourceAssetPath(SourceWorldPath);
+				WorldName = FPackageName::GetShortName(SourceAssetPath.GetPackageName());
 						
-			InstancedWorldPath = UWorld::RemovePIEPrefix(InstancedWorldPath);
+				InstancedWorldPath = UWorld::RemovePIEPrefix(InstancedWorldPath);
 
-			const FString SourcePackageName = SourceAssetPath.GetPackageName().ToString();
-			const FTopLevelAssetPath InstanceAssetPath(InstancedWorldPath);
-			const FString InstancePackageName = InstanceAssetPath.GetPackageName().ToString();
+				const FString SourcePackageName = SourceAssetPath.GetPackageName().ToString();
+				const FTopLevelAssetPath InstanceAssetPath(InstancedWorldPath);
+				const FString InstancePackageName = InstanceAssetPath.GetPackageName().ToString();
 
-			if (int32 Index = InstancePackageName.Find(SourcePackageName); Index != INDEX_NONE)
-			{
-				CellNameBuilder.Appendf(TEXT("_i%s"), *InstancePackageName.Mid(Index + SourcePackageName.Len()));
+				if (int32 Index = InstancePackageName.Find(SourcePackageName); Index != INDEX_NONE)
+				{
+					CellNameBuilder.Appendf(TEXT("_i%s"), *InstancePackageName.Mid(Index + SourcePackageName.Len()));
+				}
 			}
 		}
+	
+		CellUniqueId.Name = CellNameBuilder.ToString();
 	}
 
-	FCellUniqueId CellUniqueId;
-	CellUniqueId.Name = CellNameBuilder.ToString();
-
-	FArchiveMD5 ArMD5;
-	ArMD5 << CellUniqueId.Name;
-	CellUniqueId.Guid = ArMD5.GetGuidFromHash();
-	check(CellUniqueId.Guid.IsValid());
+	// Build cell guid
+	{
+		FArchiveMD5 ArMD5;
+		ArMD5 << CellNameID << DataLayersID << ContentBundleID;
+		InCellDescInstance.SourcePartition->AppendCellGuid(ArMD5);
+		CellUniqueId.Guid = ArMD5.GetGuidFromHash();
+		check(CellUniqueId.Guid.IsValid());
+	}
 
 	return CellUniqueId;
 }
diff --git a/Engine/Source/Runtime/Engine/Public/WorldPartition/DataLayer/DataLayersID.h b/Engine/Source/Runtime/Engine/Public/WorldPartition/DataLayer/DataLayersID.h
--- a/Engine/Source/Runtime/Engine/Public/WorldPartition/DataLayer/DataLayersID.h
+++ b/Engine/Source/Runtime/Engine/Public/WorldPartition/DataLayer/DataLayersID.h
@@ -28,6 +28,11 @@ public:
 		return Hash;
 	}
 
+	friend FArchive& operator<<(FArchive& InAr, FDataLayersID& InDataLayersID)
+	{
+		return InAr << InDataLayersID.Hash;
+	}
+
 	friend uint32 GetTypeHash(const FDataLayersID& InDataLayersID)
 	{
 		return InDataLayersID.Hash;
diff --git a/Engine/Source/Runtime/Engine/Public/WorldPartition/RuntimeHashSet/RuntimePartition.h b/Engine/Source/Runtime/Engine/Public/WorldPartition/RuntimeHashSet/RuntimePartition.h
--- a/Engine/Source/Runtime/Engine/Public/WorldPartition/RuntimeHashSet/RuntimePartition.h
+++ b/Engine/Source/Runtime/Engine/Public/WorldPartition/RuntimeHashSet/RuntimePartition.h
@@ -68,6 +68,7 @@ public:
 	virtual void InitHLODRuntimePartitionFrom(const URuntimePartition* InRuntimePartition, int32 InHLODIndex);
 	virtual bool IsValidPartitionTokens(const TArray<FName>& InPartitionTokens) const PURE_VIRTUAL(URuntimePartition::IsValidPartitionTokens, return false;);
 	virtual bool GenerateStreaming(const FGenerateStreamingParams& InParams, FGenerateStreamingResult& OutResult) PURE_VIRTUAL(URuntimePartition::GenerateStreaming, return false;);
+	virtual FArchive& AppendCellGuid(FArchive& InAr) { return InAr << Name << HLODIndex; }
 #endif
 
 #if WITH_EDITORONLY_DATA
diff --git a/Engine/Source/Runtime/Engine/Public/WorldPartition/RuntimeHashSet/RuntimePartitionLHGrid.h b/Engine/Source/Runtime/Engine/Public/WorldPartition/RuntimeHashSet/RuntimePartitionLHGrid.h
--- a/Engine/Source/Runtime/Engine/Public/WorldPartition/RuntimeHashSet/RuntimePartitionLHGrid.h
+++ b/Engine/Source/Runtime/Engine/Public/WorldPartition/RuntimeHashSet/RuntimePartitionLHGrid.h
@@ -26,6 +26,7 @@ public:
 	virtual void SetDefaultValues() override;
 	virtual bool IsValidPartitionTokens(const TArray<FName>& InPartitionTokens) const override;
 	virtual bool GenerateStreaming(const FGenerateStreamingParams& InParams, FGenerateStreamingResult& OutResult) override;
+	virtual FArchive& AppendCellGuid(FArchive& InAr) override;
 	//~ End URuntimePartition interface
 #endif
 
