New HLSL translator support for UMaterialExpressionLandscapeLayerSwitch
ENGINE_API FExpressionSwitchBase so it and derived classes can be reached from editor dlls
Implement FExpressionLandscapeLayerSwitch
Genericised code for StaticTerrainLayerWeightIndices to reduce duplicated code

#tests materials, editor, win64

diff --git a/Engine/Source/Runtime/Engine/Private/Materials/MaterialHLSLTree.cpp b/Engine/Source/Runtime/Engine/Private/Materials/MaterialHLSLTree.cpp
--- a/Engine/Source/Runtime/Engine/Private/Materials/MaterialHLSLTree.cpp
+++ b/Engine/Source/Runtime/Engine/Private/Materials/MaterialHLSLTree.cpp
@@ -1668,53 +1668,41 @@ void FExpressionTextureSample::EmitValueShader(FEmitContext& Context, FEmitScope
 		bAutomaticViewMipBias);
 }
 
+FName FExpressionStaticTerrainLayerWeight::BuildWeightmapName(const TCHAR* Weightmap, int32 Index, bool bUseIndex) const
+{
+	FName Name;
+	if (bUseIndex)
+	{
+		Name = *FString::Printf(TEXT("%s%d"), Weightmap, Index);
+	}
+	else
+	{
+		Name = *FString::Printf(TEXT("%sArray"), Weightmap);
+	}
+	return Name;
+}
+
 bool FExpressionStaticTerrainLayerWeight::PrepareValue(FEmitContext& Context, FEmitScope& Scope, const FRequestedType& RequestedType, FPrepareValueResult& OutResult) const
 {
-	bool bFoundMatchingParameter = false;
 	FEmitData& EmitData = Context.FindData<FEmitData>();
 	// TODO: revist whether we need to add the parameters to CachedExpressionData
 	const bool bTextureArrayEnabled = UseTextureArraySample(Context);
-	if (EmitData.StaticParameters)
-	{
-		for (int32 ParameterIndex = 0; ParameterIndex < EmitData.StaticParameters->EditorOnly.TerrainLayerWeightParameters.Num(); ++ParameterIndex)
-		{
-			const FStaticTerrainLayerWeightParameter& Parameter = EmitData.StaticParameters->EditorOnly.TerrainLayerWeightParameters[ParameterIndex];
-
-			// If there are multiple weight maps with the same name, they should be numbered to allow for unique masks
-			FName LayerNameTest = Parameter.LayerName;
-			LayerNameTest.SetNumber(0);
-
-			if (LayerNameTest != BaseParameterInfo.Name)
-			{
-				continue;
-			}
-
-			const int32 WeightmapIndex = Parameter.WeightmapIndex;
-
-			if (WeightmapIndex == INDEX_NONE)
-			{
-				continue;
-			}
 
-			bFoundMatchingParameter = true;
+	TArray<int32> Indices;
+	EmitData.GatherStaticTerrainLayerParamIndices(BaseParameterInfo.Name, Indices);
+	bool bFoundMatchingParameter = !Indices.IsEmpty();
 
-			if (!Context.bMarkLiveValues || !EmitData.CachedExpressionData)
-			{
-				break;
-			}
+	if (Context.bMarkLiveValues && EmitData.CachedExpressionData)
+	{
+		for (int32 Index : Indices)
+		{
+			const FStaticTerrainLayerWeightParameter& Parameter = EmitData.StaticParameters->EditorOnly.TerrainLayerWeightParameters[Index];
 
 			FMaterialParameterInfo WeightmapParameterInfo = BaseParameterInfo;
+			WeightmapParameterInfo.Name = BuildWeightmapName(TEXT("Weightmap"), Parameter.WeightmapIndex, !bTextureArrayEnabled);
+
 			FMaterialParameterMetadata WeightmapParameterMeta;
-			if (bTextureArrayEnabled)
-			{
-				WeightmapParameterInfo.Name = FName(TEXT("WeightmapArray"));
-				WeightmapParameterMeta.Value = GEngine->WeightMapPlaceholderTexture;
-			}
-			else
-			{
-				WeightmapParameterInfo.Name = *FString::Printf(TEXT("Weightmap%d"), WeightmapIndex);
-				WeightmapParameterMeta.Value = GEngine->WeightMapArrayPlaceholderTexture;
-			}
+			WeightmapParameterMeta.Value = bTextureArrayEnabled ? GEngine->WeightMapArrayPlaceholderTexture : GEngine->WeightMapPlaceholderTexture;
 
 			UObject* UnusedReferencedTexture;
 			EmitData.CachedExpressionData->AddParameter(WeightmapParameterInfo, WeightmapParameterMeta, UnusedReferencedTexture);
@@ -1727,12 +1715,12 @@ bool FExpressionStaticTerrainLayerWeight::PrepareValue(FEmitContext& Context, FE
 
 			EmitData.CachedExpressionData->AddParameter(LayerMaskParameterInfo, LayerMaskParameterMeta, UnusedReferencedTexture);
 		}
+	}
 
-		if (bFoundMatchingParameter && EmitData.CachedExpressionData)
-		{
-			EmitData.CachedExpressionData->ReferencedTextures.AddUnique(GEngine->WeightMapArrayPlaceholderTexture);
-			EmitData.CachedExpressionData->ReferencedTextures.AddUnique(GEngine->WeightMapPlaceholderTexture);
-		}
+	if (bFoundMatchingParameter && EmitData.CachedExpressionData)
+	{
+		EmitData.CachedExpressionData->ReferencedTextures.AddUnique(GEngine->WeightMapArrayPlaceholderTexture);
+		EmitData.CachedExpressionData->ReferencedTextures.AddUnique(GEngine->WeightMapPlaceholderTexture);
 	}
 
 	if (!bFoundMatchingParameter)
@@ -1786,46 +1774,19 @@ void FExpressionStaticTerrainLayerWeight::EmitValueShader(FEmitContext& Context,
 	FEmitShaderExpression* EmitResult = nullptr;
 	int32 NumWeightmapParameters = 0;
 
-	for (int32 ParameterIndex = 0; ParameterIndex < EmitData.StaticParameters->EditorOnly.TerrainLayerWeightParameters.Num(); ++ParameterIndex)
-	{
-		const FStaticTerrainLayerWeightParameter& Parameter = EmitData.StaticParameters->EditorOnly.TerrainLayerWeightParameters[ParameterIndex];
-
-		// If there are multiple weight maps with the same name, they should be numbered to allow for unique masks
-		FName LayerNameTest = Parameter.LayerName;
-		LayerNameTest.SetNumber(0);
-
-		if (LayerNameTest != BaseParameterInfo.Name)
-		{
-			continue;
-		}
+	TArray<int32> Indices;
+	EmitData.GatherStaticTerrainLayerParamIndices(BaseParameterInfo.Name, Indices);
 
+	for (int32 Index : Indices)
+	{
+		const FStaticTerrainLayerWeightParameter& Parameter = EmitData.StaticParameters->EditorOnly.TerrainLayerWeightParameters[Index];
 		const int32 WeightmapIndex = Parameter.WeightmapIndex;
 
-		if (WeightmapIndex == INDEX_NONE)
-		{
-			continue;
-		}
-
 		FMaterialTextureValue TextureValue;
-		if (bTextureArrayEnabled)
-		{
-			TextureValue.Texture = GEngine->WeightMapArrayPlaceholderTexture;
-			TextureValue.SamplerType = SAMPLERTYPE_Masks;
-			TextureValue.ParameterInfo = BaseParameterInfo;
-			TextureValue.ParameterInfo.Name = FName(TEXT("WeightmapArray"));
-		}
-		else
-		{
-			TextureValue.Texture = GEngine->WeightMapPlaceholderTexture;
-			TextureValue.SamplerType = SAMPLERTYPE_Masks;
-			TextureValue.ParameterInfo = BaseParameterInfo;
-			TextureValue.ParameterInfo.Name = *FString::Printf(TEXT("Weightmap%d"), WeightmapIndex);	
-		}
-		
-		if (!EmitTexCoordValue)
-		{
-			EmitTexCoordValue = TexCoordExpression->GetValueShader(Context, Scope, TexCoordType);
-		}
+		TextureValue.SamplerType = SAMPLERTYPE_Masks;
+		TextureValue.ParameterInfo = BaseParameterInfo;
+		TextureValue.Texture = bTextureArrayEnabled ? GEngine->WeightMapArrayPlaceholderTexture : GEngine->WeightMapPlaceholderTexture;
+		TextureValue.ParameterInfo.Name = BuildWeightmapName(TEXT("Weightmap"), WeightmapIndex, !bTextureArrayEnabled);
 
 		if (bTextureArrayEnabled)
 		{
@@ -3265,6 +3226,22 @@ bool FExpressionFinalShadingModelSwitch::IsInputActive(const FEmitContext& Conte
 	}
 }
 
+bool FExpressionLandscapeLayerSwitch::IsInputActive(const FEmitContext& Context, int32 Index) const
+{
+	bool bFoundMatchingParameter = false;
+	const FEmitData& EmitData = Context.FindData<FEmitData>();
+	if (!bPreviewUsed)
+	{
+		TArray<int32> Indices;
+		EmitData.GatherStaticTerrainLayerParamIndices(ParameterName, Indices);
+		bFoundMatchingParameter = !Indices.IsEmpty();
+	}
+
+	// 0 is LayerNotUsed, 1 is LayerUsed
+	int32 DesiredIndex = bFoundMatchingParameter || bPreviewUsed ? 1 : 0;
+	return Index == DesiredIndex;
+}
+
 bool FExpressionAtmosphericFogColorFunction::PrepareValue(FEmitContext& Context, FEmitScope& Scope, const FRequestedType& RequestedType, FPrepareValueResult& OutResult) const
 {
 	const FPreparedType& PositionType = Context.PrepareExpression(PositionExpression, Scope, Shader::EValueType::Double3);
@@ -3568,6 +3545,36 @@ int32 FEmitData::FindOrAddParameterCollection(const class UMaterialParameterColl
 	return CollectionIndex;
 }
 
+void FEmitData::GatherStaticTerrainLayerParamIndices(FName LayerName, TArray<int32>& ParamIndices) const
+{
+	if (StaticParameters)
+	{
+		for (int32 ParameterIndex = 0; ParameterIndex < StaticParameters->EditorOnly.TerrainLayerWeightParameters.Num(); ++ParameterIndex)
+		{
+			const FStaticTerrainLayerWeightParameter& Parameter = StaticParameters->EditorOnly.TerrainLayerWeightParameters[ParameterIndex];
+
+			// If there are multiple weight maps with the same name, they should be numbered to allow for unique masks
+			FName LayerNameTest = Parameter.LayerName;
+			if (Parameter.bIsRepeatedLayer)
+			{
+				LayerNameTest.SetNumber(0);
+			}
+
+			if (LayerNameTest != LayerName)
+			{
+				continue;
+			}
+
+			if (Parameter.WeightmapIndex == INDEX_NONE)
+			{
+				continue;
+			}
+
+			ParamIndices.Add(ParameterIndex);
+		}
+	}
+}
+
 } // namespace UE::HLSLTree::Material
 
 #endif // WITH_EDITOR
diff --git a/Engine/Source/Runtime/Engine/Public/HLSLTree/HLSLTreeCommon.h b/Engine/Source/Runtime/Engine/Public/HLSLTree/HLSLTreeCommon.h
--- a/Engine/Source/Runtime/Engine/Public/HLSLTree/HLSLTreeCommon.h
+++ b/Engine/Source/Runtime/Engine/Public/HLSLTree/HLSLTreeCommon.h
@@ -243,7 +243,7 @@ public:
 	ENGINE_API virtual void EmitValuePreshader(FEmitContext& Context, FEmitScope& Scope, const FRequestedType& RequestedType, FEmitValuePreshaderResult& OutResult) const override;
 };
 
-class FExpressionSwitchBase : public FExpression
+class ENGINE_API FExpressionSwitchBase : public FExpression
 {
 public:
 	static constexpr int8 MaxInputs = 8;
diff --git a/Engine/Source/Runtime/Engine/Public/MaterialHLSLTree.h b/Engine/Source/Runtime/Engine/Public/MaterialHLSLTree.h
--- a/Engine/Source/Runtime/Engine/Public/MaterialHLSLTree.h
+++ b/Engine/Source/Runtime/Engine/Public/MaterialHLSLTree.h
@@ -526,6 +526,28 @@ public:
 	virtual bool IsInputActive(const FEmitContext& Context, int32 Index) const override;
 };
 
+class ENGINE_API FExpressionLandscapeLayerSwitch : public FExpressionSwitchBase
+{
+public:
+	FName ParameterName;
+	bool bPreviewUsed;
+
+	FExpressionLandscapeLayerSwitch(TConstArrayView<const FExpression*> InInputs, FName InParameterName, bool bInPreviewUsed)
+		: FExpressionSwitchBase(InInputs)
+		, ParameterName(InParameterName)
+		, bPreviewUsed(bInPreviewUsed)
+	{
+		check(InInputs.Num() == 2);
+	}
+
+	virtual const FExpression* NewSwitch(UE::HLSLTree::FTree& Tree, TConstArrayView<const FExpression*> InInputs) const override
+	{
+		return Tree.NewExpression<FExpressionLandscapeLayerSwitch>(InInputs, ParameterName, bPreviewUsed);
+	}
+
+	virtual bool IsInputActive(const FEmitContext& Context, int32 Index) const override;
+};
+
 class FExpressionNaniteReplaceFunction : public FExpression
 {
 public:
@@ -655,6 +677,7 @@ public:
 	virtual void EmitValueShader(FEmitContext& Context, FEmitScope& Scope, const FRequestedType& RequestedType, FEmitValueShaderResult& OutResult) const override;
 private:
 	bool UseTextureArraySample(const FEmitContext& Context) const;
+	FName BuildWeightmapName(const TCHAR* Weightmap, int32 Index, bool bUseIndex) const;
 };
 
 class FExpressionTextureProperty : public FExpression
@@ -1216,6 +1239,7 @@ public:
 	void EmitInterpolatorShader(FEmitContext& Context, FStringBuilderBase& OutCode);
 
 	int32 FindOrAddParameterCollection(const class UMaterialParameterCollection* ParameterCollection);
+	void GatherStaticTerrainLayerParamIndices(FName LayerName, TArray<int32>& WeightIndices) const;
 };
 
 } // namespace UE::HLSLTree::Material
diff --git a/Engine/Source/Runtime/Landscape/Classes/Materials/MaterialExpressionLandscapeLayerSwitch.h b/Engine/Source/Runtime/Landscape/Classes/Materials/MaterialExpressionLandscapeLayerSwitch.h
--- a/Engine/Source/Runtime/Landscape/Classes/Materials/MaterialExpressionLandscapeLayerSwitch.h
+++ b/Engine/Source/Runtime/Landscape/Classes/Materials/MaterialExpressionLandscapeLayerSwitch.h
@@ -46,6 +46,7 @@ public:
 	LANDSCAPE_API virtual void SetEditableName(const FString& NewName) override;
 	LANDSCAPE_API virtual void GetCaption(TArray<FString>& OutCaptions) const override;
 	LANDSCAPE_API virtual bool MatchesSearchQuery(const TCHAR* SearchQuery) override;
+	virtual bool GenerateHLSLExpression(FMaterialHLSLGenerator& Generator, UE::HLSLTree::FScope& Scope, int32 OutputIndex, UE::HLSLTree::FExpression const*& OutExpression) const override;
 #endif
 	LANDSCAPE_API virtual UObject* GetReferencedTexture() const override;
 	LANDSCAPE_API virtual ReferencedTextureArray GetReferencedTextures() const override;
diff --git a/Engine/Source/Runtime/Landscape/Private/Materials/MaterialExpressionLandscapeLayerSwitch.cpp b/Engine/Source/Runtime/Landscape/Private/Materials/MaterialExpressionLandscapeLayerSwitch.cpp
--- a/Engine/Source/Runtime/Landscape/Private/Materials/MaterialExpressionLandscapeLayerSwitch.cpp
+++ b/Engine/Source/Runtime/Landscape/Private/Materials/MaterialExpressionLandscapeLayerSwitch.cpp
@@ -5,6 +5,8 @@
 #include "Engine/Texture.h"
 #include "EngineGlobals.h"
 #include "MaterialCompiler.h"
+#include "MaterialHLSLGenerator.h"
+#include "MaterialHLSLTree.h"
 #include "Materials/Material.h"
 #include "LandscapeUtils.h"
 
@@ -76,6 +78,20 @@ int32 UMaterialExpressionLandscapeLayerSwitch::Compile(class FMaterialCompiler*
 
 	return ReturnCode;
 }
+
+bool UMaterialExpressionLandscapeLayerSwitch::GenerateHLSLExpression(FMaterialHLSLGenerator& Generator, UE::HLSLTree::FScope& Scope, int32 OutputIndex, UE::HLSLTree::FExpression const*& OutExpression) const
+{
+	using namespace UE::HLSLTree;
+
+	const FExpression* Inputs[] = {
+		LayerNotUsed.TryAcquireHLSLExpression(Generator, Scope),
+		LayerUsed.TryAcquireHLSLExpression(Generator, Scope)
+	};
+
+	OutExpression = Generator.GetTree().NewExpression<Material::FExpressionLandscapeLayerSwitch>(Inputs, ParameterName, PreviewUsed!=0);
+	return OutExpression != nullptr;
+}
+
 #endif // WITH_EDITOR
 
 UObject* UMaterialExpressionLandscapeLayerSwitch::GetReferencedTexture() const
