World Partition: fix wp.Editor.DumpStreamingGenerationLog console command for worlds with Content Bundles.

diff --git a/Engine/Source/Runtime/Engine/Private/WorldPartition/ContentBundle/ContentBundleContainer.cpp b/Engine/Source/Runtime/Engine/Private/WorldPartition/ContentBundle/ContentBundleContainer.cpp
--- a/Engine/Source/Runtime/Engine/Private/WorldPartition/ContentBundle/ContentBundleContainer.cpp
+++ b/Engine/Source/Runtime/Engine/Private/WorldPartition/ContentBundle/ContentBundleContainer.cpp
@@ -423,10 +423,10 @@ void FContentBundleContainer::OnPreGenerateStreaming(TArray<FString>* OutPackage
 
 	GetInjectedWorld()->ContentBundleManager->GetPIEDuplicateHelper()->Clear();
 
-	const bool bIsPIE = true;
+	UWorldPartition* WorldPartition = GetInjectedWorld()->GetWorldPartition();
 	for (TSharedPtr<FContentBundleEditor>& ContentBundle : GetEditorContentBundles())
 	{
-		ContentBundle->GenerateStreaming(OutPackageToGenerate, bIsPIE);
+		ContentBundle->GenerateStreaming(OutPackageToGenerate, WorldPartition->bIsPIE);
 	}
 }
 
diff --git a/Engine/Source/Runtime/Engine/Private/WorldPartition/ContentBundle/ContentBundleEditor.cpp b/Engine/Source/Runtime/Engine/Private/WorldPartition/ContentBundle/ContentBundleEditor.cpp
--- a/Engine/Source/Runtime/Engine/Private/WorldPartition/ContentBundle/ContentBundleEditor.cpp
+++ b/Engine/Source/Runtime/Engine/Private/WorldPartition/ContentBundle/ContentBundleEditor.cpp
@@ -365,10 +365,10 @@ void FContentBundleEditor::GenerateStreaming(TArray<FString>* OutPackageToGenera
 		{
 			UE_LOG(LogContentBundle, Error, TEXT("%s Failed to store streaming object. PIE duplication will not work."), *ContentBundle::Log::MakeDebugInfoString(*this));
 		}
-
-		// Clear streaming object. It will be kept alive & duplicated by UContentBundleDuplicateForPIEHelper.
-		ExternalStreamingObject = nullptr;
 	}
+
+	// Clear streaming object, as it will be kept alive and duplicated by UContentBundleDuplicateForPIEHelper (or simply GC'd for manual streaming generation).
+	ExternalStreamingObject = nullptr;
 }
 
 void FContentBundleEditor::OnBeginCook(IWorldPartitionCookPackageContext& CookContext)
diff --git a/Engine/Source/Runtime/Engine/Private/WorldPartition/WorldPartitionStreamingGeneration.cpp b/Engine/Source/Runtime/Engine/Private/WorldPartition/WorldPartitionStreamingGeneration.cpp
--- a/Engine/Source/Runtime/Engine/Private/WorldPartition/WorldPartitionStreamingGeneration.cpp
+++ b/Engine/Source/Runtime/Engine/Private/WorldPartition/WorldPartitionStreamingGeneration.cpp
@@ -48,6 +48,7 @@ static FAutoConsoleCommand DumpStreamingGenerationLog(
 				UWorldPartition::FGenerateStreamingContext Context;
 				WorldPartition->GenerateStreaming(Params, Context);
 				WorldPartition->FlushStreaming();
+				CollectGarbage(GARBAGE_COLLECTION_KEEPFLAGS);
 			}
 		}
 	})
diff --git a/Engine/Source/Runtime/Engine/Public/WorldPartition/WorldPartition.h b/Engine/Source/Runtime/Engine/Public/WorldPartition/WorldPartition.h
--- a/Engine/Source/Runtime/Engine/Public/WorldPartition/WorldPartition.h
+++ b/Engine/Source/Runtime/Engine/Public/WorldPartition/WorldPartition.h
@@ -123,6 +123,7 @@ class UWorldPartition final : public UObject, public FActorDescContainerInstance
 	friend class FUnrealEdMisc;
 	friend class UActorDescContainer;
 	friend class UActorDescContainerInstance;
+	friend class FContentBundleContainer;
 
 public:
 #if WITH_EDITOR
