Switch default curve compression codec to ACL and add a fallback codec if the plugin is disabled

diff --git a/Engine/Config/BaseEngine.ini b/Engine/Config/BaseEngine.ini
--- a/Engine/Config/BaseEngine.ini
+++ b/Engine/Config/BaseEngine.ini
@@ -3347,8 +3347,10 @@ BoneCompressionSettingsFallback="/Engine/Animation/DefaultAnimBoneCompressionSet
 BoneCompressionSettings="/ACLPlugin/ACLAnimBoneCompressionSettings"
 ; Default codec for the animation recorder
 AnimationRecorderBoneCompressionSettings="/Engine/Animation/DefaultRecorderBoneCompression"
+; Fallback is used only if the default codec isn't found (e.g. if the plugin isn't enabled)
+CurveCompressionSettingsFallback="/Engine/Animation/DefaultAnimCurveCompressionSettings"
 ; Default codec assigned when new sequences are imported/created
-CurveCompressionSettings="/Engine/Animation/DefaultAnimCurveCompressionSettings"
+CurveCompressionSettings="/ACLPlugin/ACLAnimCurveCompressionSettings"
 VariableFrameStrippingSettings="/Engine/Animation/DefaultVariableFrameStrippingSettings"
 
 [/Script/Engine.MeshSimplificationSettings]
diff --git a/Engine/Source/Runtime/Engine/Private/Animation/AnimationUtils.cpp b/Engine/Source/Runtime/Engine/Private/Animation/AnimationUtils.cpp
--- a/Engine/Source/Runtime/Engine/Private/Animation/AnimationUtils.cpp
+++ b/Engine/Source/Runtime/Engine/Private/Animation/AnimationUtils.cpp
@@ -283,7 +283,7 @@ void FAnimationUtils::ComputeCompressionError(const FCompressibleAnimData& Compr
 		const FTransform EndEffectorDummyBone(FQuat::Identity, FVector(END_EFFECTOR_DUMMY_BONE_LENGTH));
 
 		FAnimSequenceDecompressionContext DecompContext(CompressibleAnimData.SampledFrameRate, CompressibleAnimData.GetNumberOfFrames(), CompressibleAnimData.Interpolation, CompressibleAnimData.AnimFName, *CompressedData.AnimData, RefPose,
-			CompressibleAnimData.TrackToSkeletonMapTable, nullptr, CompressibleAnimData.bIsValidAdditive);
+			CompressibleAnimData.TrackToSkeletonMapTable, nullptr, CompressibleAnimData.bIsValidAdditive, CompressibleAnimData.AdditiveType);
 
 		const TArray<FBoneData>& BoneData = CompressibleAnimData.BoneData;
 
@@ -968,7 +968,12 @@ UAnimCurveCompressionSettings* FAnimationUtils::GetDefaultAnimationCurveCompress
 {
 	if (DefaultCurveCompressionSettings == nullptr)
 	{
-		DefaultCurveCompressionSettings = Cast<UAnimCurveCompressionSettings>(GetDefaultAnimationCompressionSettings(TEXT("CurveCompressionSettings"), true));
+		DefaultCurveCompressionSettings = Cast<UAnimCurveCompressionSettings>(GetDefaultAnimationCompressionSettings(TEXT("CurveCompressionSettings"), false));
+
+		if (DefaultCurveCompressionSettings == nullptr)
+		{
+			DefaultCurveCompressionSettings = Cast<UAnimCurveCompressionSettings>(GetDefaultAnimationCompressionSettings(TEXT("CurveCompressionSettingsFallback"), true));
+		}
 	}
 
 	return DefaultCurveCompressionSettings;
@@ -1111,7 +1116,7 @@ void FAnimationUtils::ExtractTransformFromCompressionData(const FCompressibleAni
 		// Build our read-only version from the mutable source
 		FUECompressedAnimData AnimData(AnimDataMutable);
 
-		FAnimSequenceDecompressionContext DecompContext(CompressibleAnimData.SampledFrameRate, CompressibleAnimData.GetNumberOfFrames(), CompressibleAnimData.Interpolation, CompressibleAnimData.AnimFName, AnimData, CompressibleAnimData.RefLocalPoses, CompressibleAnimData.TrackToSkeletonMapTable, nullptr, CompressibleAnimData.bIsValidAdditive);
+		FAnimSequenceDecompressionContext DecompContext(CompressibleAnimData.SampledFrameRate, CompressibleAnimData.GetNumberOfFrames(), CompressibleAnimData.Interpolation, CompressibleAnimData.AnimFName, AnimData, CompressibleAnimData.RefLocalPoses, CompressibleAnimData.TrackToSkeletonMapTable, nullptr, CompressibleAnimData.bIsValidAdditive, CompressibleAnimData.AdditiveType);
 		DecompContext.Seek(Time);
 		CompressedAnimData.Codec->DecompressBone(DecompContext, TrackIndex, OutBoneTransform);
 		return;
