World Partition Runtime HashSet:
- Added project settings to allow changing default editor and runtime hash default classes for empty templates.
- New maps from the open world and empty templates now defaults to the new runtime hashset, with a basic LHGrid setup.
- Converted all UE5 templates to use the new runtime hashset, with a basic LHGrid setup.

#virtualized

diff --git a/Engine/Config/BaseEngine.ini b/Engine/Config/BaseEngine.ini
--- a/Engine/Config/BaseEngine.ini
+++ b/Engine/Config/BaseEngine.ini
@@ -3537,4 +3537,8 @@ ProviderModule=AnalyticsLog
 UsageType=Editor
 ; Override the default output location from Saved/Telemetry/Telemetry.json
 ; FileName="somefile.json"
-; FolderPath="somefolder"
\ No newline at end of file
+; FolderPath="somefolder"
+
+[/Script/Engine.WorldPartitionSettings]
+EditorHashDefaultClass=/Script/CoreUObject.Class'/Script/Engine.WorldPartitionEditorSpatialHash'
+RuntimeHashDefaultClass=/Script/CoreUObject.Class'/Script/Engine.WorldPartitionRuntimeHashSet'
\ No newline at end of file
diff --git a/Engine/Source/Editor/UnrealEd/Private/Commandlets/WorldPartitionConvertCommandlet.cpp b/Engine/Source/Editor/UnrealEd/Private/Commandlets/WorldPartitionConvertCommandlet.cpp
--- a/Engine/Source/Editor/UnrealEd/Private/Commandlets/WorldPartitionConvertCommandlet.cpp
+++ b/Engine/Source/Editor/UnrealEd/Private/Commandlets/WorldPartitionConvertCommandlet.cpp
@@ -19,6 +19,7 @@
 #include "Engine/MapBuildDataRegistry.h"
 #include "ActorReferencesUtils.h"
 #include "WorldPartition/WorldPartition.h"
+#include "WorldPartition/WorldPartitionSettings.h"
 #include "WorldPartition/WorldPartitionSubsystem.h"
 #include "WorldPartition/HLOD/HLODActor.h"
 #include "WorldPartition/WorldPartitionMiniMap.h"
@@ -715,8 +716,8 @@ int32 UWorldPartitionConvertCommandlet::Main(const FString& Params)
 		}
 		else
 		{
-			EditorHashClass = FindObject<UClass>(nullptr, TEXT("/Script/Engine.WorldPartitionEditorSpatialHash"));
-			RuntimeHashClass = FindObject<UClass>(nullptr, TEXT("/Script/Engine.WorldPartitionRuntimeSpatialHash"));
+			EditorHashClass = UWorldPartitionSettings::Get()->GetEditorHashDefaultClass();
+			RuntimeHashClass = UWorldPartitionSettings::Get()->GetRuntimeHashDefaultClass();
 		}
 	}
 
diff --git a/Engine/Source/Runtime/Engine/Private/WorldPartition/RuntimeHashSet/WorldPartitionRuntimeHashSet.cpp b/Engine/Source/Runtime/Engine/Private/WorldPartition/RuntimeHashSet/WorldPartitionRuntimeHashSet.cpp
--- a/Engine/Source/Runtime/Engine/Private/WorldPartition/RuntimeHashSet/WorldPartitionRuntimeHashSet.cpp
+++ b/Engine/Source/Runtime/Engine/Private/WorldPartition/RuntimeHashSet/WorldPartitionRuntimeHashSet.cpp
@@ -2,6 +2,7 @@
 
 #include "WorldPartition/RuntimeHashSet/WorldPartitionRuntimeHashSet.h"
 #include "WorldPartition/RuntimeHashSet/RuntimePartition.h"
+#include "WorldPartition/RuntimeHashSet/RuntimePartitionLHGrid.h"
 #include "WorldPartition/RuntimeHashSet/RuntimePartitionPersistent.h"
 #include "WorldPartition/HLOD/HLODLayer.h"
 #include "WorldPartition/ContentBundle/ContentBundleDescriptor.h"
@@ -92,7 +93,17 @@ void UWorldPartitionRuntimeHashSet::PostLoad()
 
 #if WITH_EDITOR
 void UWorldPartitionRuntimeHashSet::SetDefaultValues()
-{}
+{
+	check(RuntimePartitions.IsEmpty());
+
+	FRuntimePartitionDesc& RuntimePartitionDesc = RuntimePartitions.AddDefaulted_GetRef();
+	RuntimePartitionDesc.Class = URuntimePartitionLHGrid::StaticClass();
+	RuntimePartitionDesc.Name = TEXT("MainPartition");
+
+	RuntimePartitionDesc.MainLayer = NewObject<URuntimePartitionLHGrid>(this, NAME_None);
+	RuntimePartitionDesc.MainLayer->Name = RuntimePartitionDesc.Name;
+	RuntimePartitionDesc.MainLayer->SetDefaultValues();
+}
 
 void UWorldPartitionRuntimeHashSet::FlushStreaming()
 {
diff --git a/Engine/Source/Runtime/Engine/Private/WorldPartition/WorldPartition.cpp b/Engine/Source/Runtime/Engine/Private/WorldPartition/WorldPartition.cpp
--- a/Engine/Source/Runtime/Engine/Private/WorldPartition/WorldPartition.cpp
+++ b/Engine/Source/Runtime/Engine/Private/WorldPartition/WorldPartition.cpp
@@ -991,9 +991,10 @@ UWorldPartition* UWorldPartition::CreateOrRepairWorldPartition(AWorldSettings* W
 	{
 		if (!EditorHashClass)
 		{
-			EditorHashClass = FindObject<UClass>(nullptr, TEXT("/Script/Engine.WorldPartitionEditorSpatialHash"));
+			EditorHashClass = UWorldPartitionSettings::Get()->GetEditorHashDefaultClass();
 		}
 
+		check(EditorHashClass);
 		WorldPartition->EditorHash = NewObject<UWorldPartitionEditorHash>(WorldPartition, EditorHashClass);
 		WorldPartition->EditorHash->SetDefaultValues();
 	}
@@ -1002,9 +1003,10 @@ UWorldPartition* UWorldPartition::CreateOrRepairWorldPartition(AWorldSettings* W
 	{
 		if (!RuntimeHashClass)
 		{
-			RuntimeHashClass = FindObject<UClass>(nullptr, TEXT("/Script/Engine.WorldPartitionRuntimeSpatialHash"));
+			RuntimeHashClass = UWorldPartitionSettings::Get()->GetRuntimeHashDefaultClass();
 		}
 
+		check(RuntimeHashClass);
 		WorldPartition->RuntimeHash = NewObject<UWorldPartitionRuntimeHash>(WorldPartition, RuntimeHashClass, NAME_None, RF_Transactional);
 		WorldPartition->RuntimeHash->SetDefaultValues();
 	}
diff --git a/Engine/Source/Runtime/Engine/Public/WorldPartition/WorldPartitionSettings.h b/Engine/Source/Runtime/Engine/Public/WorldPartition/WorldPartitionSettings.h
--- a/Engine/Source/Runtime/Engine/Public/WorldPartition/WorldPartitionSettings.h
+++ b/Engine/Source/Runtime/Engine/Public/WorldPartition/WorldPartitionSettings.h
@@ -16,9 +16,20 @@ public:
 	static UWorldPartitionSettings* Get() { return CastChecked<UWorldPartitionSettings>(UWorldPartitionSettings::StaticClass()->GetDefaultObject()); }
 
 	EWorldPartitionDataLayersLogicOperator GetNewMapsDataLayersLogicOperator() const { return NewMapsDataLayersLogicOperator; }
+	
+	TSubclassOf<UWorldPartitionEditorHash> GetEditorHashDefaultClass() const { return EditorHashDefaultClass; }
+	TSubclassOf<UWorldPartitionRuntimeHash> GetRuntimeHashDefaultClass() const { return RuntimeHashDefaultClass; }
 
 protected:
 	/** Set the default logical operator for actor data layers activation for new maps */
-	UPROPERTY(EditAnywhere, config, Category = WorldPartition)
+	UPROPERTY(EditAnywhere, Config, Category = WorldPartition)
 	EWorldPartitionDataLayersLogicOperator NewMapsDataLayersLogicOperator = EWorldPartitionDataLayersLogicOperator::Or;
+
+	/** Set the default editor hash class to use for the editor */
+	UPROPERTY(Config)
+	TSubclassOf<UWorldPartitionEditorHash> EditorHashDefaultClass;
+
+	/** Set the default runtime hash class to use for new maps  */
+	UPROPERTY(EditAnywhere, Config, NoClear, Category = WorldPartition)
+	TSubclassOf<UWorldPartitionRuntimeHash> RuntimeHashDefaultClass;
 };
\ No newline at end of file
