Switch default curve compression codec to ACL and add a fallback codec if the plugin is disabled

diff --git a/Engine/Config/BaseEngine.ini b/Engine/Config/BaseEngine.ini
--- a/Engine/Config/BaseEngine.ini
+++ b/Engine/Config/BaseEngine.ini
@@ -3347,8 +3347,10 @@ BoneCompressionSettingsFallback="/Engine/Animation/DefaultAnimBoneCompressionSet
 BoneCompressionSettings="/ACLPlugin/ACLAnimBoneCompressionSettings"
 ; Default codec for the animation recorder
 AnimationRecorderBoneCompressionSettings="/Engine/Animation/DefaultRecorderBoneCompression"
+; Fallback is used only if the default codec isn't found (e.g. if the plugin isn't enabled)
+CurveCompressionSettingsFallback="/Engine/Animation/DefaultAnimCurveCompressionSettings"
 ; Default codec assigned when new sequences are imported/created
-CurveCompressionSettings="/Engine/Animation/DefaultAnimCurveCompressionSettings"
+CurveCompressionSettings="/ACLPlugin/ACLAnimCurveCompressionSettings"
 VariableFrameStrippingSettings="/Engine/Animation/DefaultVariableFrameStrippingSettings"
 
 [/Script/Engine.MeshSimplificationSettings]
diff --git a/Engine/Source/Runtime/Engine/Private/Animation/AnimationUtils.cpp b/Engine/Source/Runtime/Engine/Private/Animation/AnimationUtils.cpp
--- a/Engine/Source/Runtime/Engine/Private/Animation/AnimationUtils.cpp
+++ b/Engine/Source/Runtime/Engine/Private/Animation/AnimationUtils.cpp
@@ -968,7 +968,12 @@ UAnimCurveCompressionSettings* FAnimationUtils::GetDefaultAnimationCurveCompress
 {
 	if (DefaultCurveCompressionSettings == nullptr)
 	{
-		DefaultCurveCompressionSettings = Cast<UAnimCurveCompressionSettings>(GetDefaultAnimationCompressionSettings(TEXT("CurveCompressionSettings"), true));
+		DefaultCurveCompressionSettings = Cast<UAnimCurveCompressionSettings>(GetDefaultAnimationCompressionSettings(TEXT("CurveCompressionSettings"), false));
+
+		if (DefaultCurveCompressionSettings == nullptr)
+		{
+			DefaultCurveCompressionSettings = Cast<UAnimCurveCompressionSettings>(GetDefaultAnimationCompressionSettings(TEXT("CurveCompressionSettingsFallback"), true));
+		}
 	}
 
 	return DefaultCurveCompressionSettings;
