TextureLODSettings make TMGS_Angular fall back to SimpleAverage rather than NoMipMaps when not a cube map

#rnx

diff --git a/Engine/Source/Runtime/Engine/Private/TextureLODSettings.cpp b/Engine/Source/Runtime/Engine/Private/TextureLODSettings.cpp
--- a/Engine/Source/Runtime/Engine/Private/TextureLODSettings.cpp
+++ b/Engine/Source/Runtime/Engine/Private/TextureLODSettings.cpp
@@ -143,6 +143,7 @@ int32 UTextureLODSettings::CalculateLODBias(const UTexture* Texture, bool bIncCi
 	TextureMaxSize = Texture->MaxTextureSize;
 #endif // #if WITH_EDITORONLY_DATA
 
+	// note: these are the Platform texture size and can return 0 during build!
 	float Width = Texture->GetSurfaceWidth();
 	float Height = Texture->GetSurfaceHeight();
 
@@ -298,9 +299,9 @@ void UTextureLODSettings::GetMipGenSettings(const UTexture& Texture, TextureMipG
 
 	// angular filtering only applies to cubemaps
 	// note: currently intentionally NOT allowed on cubearrays
-	if (Setting == TMGS_Angular && !Texture.IsA(UTextureCube::StaticClass()))
+	if (Setting == TMGS_Angular && Texture.GetTextureClass() != ETextureClass::Cube )
 	{
-		Setting = TMGS_NoMipmaps;
+		Setting = TMGS_SimpleAverage;
 	}
 
 	OutMipGenSettings = Setting;
