Materials
- Added MinimalAPi specifier to UMaterialExpressionSceneTexture so that it can be used with Cast.

diff --git a/Engine/Source/Runtime/Engine/Classes/Materials/MaterialExpressionSceneTexture.h b/Engine/Source/Runtime/Engine/Classes/Materials/MaterialExpressionSceneTexture.h
--- a/Engine/Source/Runtime/Engine/Classes/Materials/MaterialExpressionSceneTexture.h
+++ b/Engine/Source/Runtime/Engine/Classes/Materials/MaterialExpressionSceneTexture.h
@@ -10,7 +10,7 @@
 #include "Materials/MaterialExpression.h"
 #include "MaterialExpressionSceneTexture.generated.h"
 
-UCLASS(collapsecategories, hidecategories=Object)
+UCLASS(MinimalAPI, collapsecategories, hidecategories=Object)
 class UMaterialExpressionSceneTexture : public UMaterialExpression
 {
 	GENERATED_UCLASS_BODY()
