SkeletalMesh : Fixed crash in cases where the rendermesh has not been created.

diff --git a/Engine/Source/Runtime/Engine/Private/Components/SkinnedMeshComponent.cpp b/Engine/Source/Runtime/Engine/Private/Components/SkinnedMeshComponent.cpp
--- a/Engine/Source/Runtime/Engine/Private/Components/SkinnedMeshComponent.cpp
+++ b/Engine/Source/Runtime/Engine/Private/Components/SkinnedMeshComponent.cpp
@@ -932,57 +932,59 @@ void USkinnedMeshComponent::CreateRenderState_Concurrent(FRegisterComponentConte
 		if ( FApp::CanEverRender() && ShouldComponentAddToScene() )
 		{
 			ERHIFeatureLevel::Type SceneFeatureLevel = GetWorld()->GetFeatureLevel();
-			FSkeletalMeshRenderData* SkelMeshRenderData = GetSkinnedAsset()->GetResourceForRendering();
+			if( FSkeletalMeshRenderData* SkelMeshRenderData = GetSkinnedAsset()->GetResourceForRendering() )
+			{
 			int32 MinLODIndex = ComputeMinLOD();
 			
 #if DO_CHECK
-			for (int LODIndex = MinLODIndex; LODIndex < SkelMeshRenderData->LODRenderData.Num(); LODIndex++)
-			{
-				FSkeletalMeshLODRenderData& LODData = SkelMeshRenderData->LODRenderData[LODIndex];
-				const FPositionVertexBuffer* PositionVertexBufferPtr = &LODData.StaticVertexBuffers.PositionVertexBuffer;
-				if (!PositionVertexBufferPtr || (PositionVertexBufferPtr->GetNumVertices() <= 0))
+				for (int LODIndex = MinLODIndex; LODIndex < SkelMeshRenderData->LODRenderData.Num(); LODIndex++)
 				{
-					UE_LOG(LogSkinnedMeshComp, Warning, TEXT("Invalid Lod %i for Rendering Asset: %s"), LODIndex, *GetSkinnedAsset()->GetFullName());
+					FSkeletalMeshLODRenderData& LODData = SkelMeshRenderData->LODRenderData[LODIndex];
+					const FPositionVertexBuffer* PositionVertexBufferPtr = &LODData.StaticVertexBuffers.PositionVertexBuffer;
+					if (!PositionVertexBufferPtr || (PositionVertexBufferPtr->GetNumVertices() <= 0))
+					{
+						UE_LOG(LogSkinnedMeshComp, Warning, TEXT("Invalid Lod %i for Rendering Asset: %s"), LODIndex, *GetSkinnedAsset()->GetFullName());
+					}
 				}
-			}
 #endif
 	
-			// Also check if skeletal mesh has too many bones/chunk for GPU skinning.
-			if (MeshObjectFactory)
-			{
-				MeshObject = MeshObjectFactory(MeshObjectFactoryUserData, this, SkelMeshRenderData, SceneFeatureLevel);
-			}
-			if (!MeshObject)
-			{
 				// Also check if skeletal mesh has too many bones/chunk for GPU skinning.
-				if (bRenderStatic)
-				{
-					// GPU skin vertex buffer + LocalVertexFactory
-					MeshObject = ::new FSkeletalMeshObjectStatic(this, SkelMeshRenderData, SceneFeatureLevel);
-				}
-				else if (ShouldCPUSkin())
-				{
-					MeshObject = ::new FSkeletalMeshObjectCPUSkin(this, SkelMeshRenderData, SceneFeatureLevel);
-				}
-				// don't silently enable CPU skinning for unsupported meshes, just do not render them, so their absence can be noticed and fixed
-				else if (!SkelMeshRenderData->RequiresCPUSkinning(SceneFeatureLevel, MinLODIndex))
+				if (MeshObjectFactory)
 				{
-					MeshObject = ::new FSkeletalMeshObjectGPUSkin(this, SkelMeshRenderData, SceneFeatureLevel);
+					MeshObject = MeshObjectFactory(MeshObjectFactoryUserData, this, SkelMeshRenderData, SceneFeatureLevel);
 				}
-				else
+				if (!MeshObject)
 				{
-					int32 MaxBonesPerChunk = SkelMeshRenderData->GetMaxBonesPerSection(MinLODIndex);
-					int32 MaxSupportedGPUSkinBones = FGPUBaseSkinVertexFactory::GetMaxGPUSkinBones();
-					int32 NumBoneInfluences = SkelMeshRenderData->GetNumBoneInfluences(MinLODIndex);
-					FString FeatureLevelName; GetFeatureLevelName(SceneFeatureLevel, FeatureLevelName);
+					// Also check if skeletal mesh has too many bones/chunk for GPU skinning.
+					if (bRenderStatic)
+					{
+						// GPU skin vertex buffer + LocalVertexFactory
+						MeshObject = ::new FSkeletalMeshObjectStatic(this, SkelMeshRenderData, SceneFeatureLevel);
+					}
+					else if (ShouldCPUSkin())
+					{
+						MeshObject = ::new FSkeletalMeshObjectCPUSkin(this, SkelMeshRenderData, SceneFeatureLevel);
+					}
+					// don't silently enable CPU skinning for unsupported meshes, just do not render them, so their absence can be noticed and fixed
+					else if (!SkelMeshRenderData->RequiresCPUSkinning(SceneFeatureLevel, MinLODIndex))
+					{
+						MeshObject = ::new FSkeletalMeshObjectGPUSkin(this, SkelMeshRenderData, SceneFeatureLevel);
+					}
+					else
+					{
+						int32 MaxBonesPerChunk = SkelMeshRenderData->GetMaxBonesPerSection(MinLODIndex);
+						int32 MaxSupportedGPUSkinBones = FGPUBaseSkinVertexFactory::GetMaxGPUSkinBones();
+						int32 NumBoneInfluences = SkelMeshRenderData->GetNumBoneInfluences(MinLODIndex);
+						FString FeatureLevelName; GetFeatureLevelName(SceneFeatureLevel, FeatureLevelName);
 
-					UE_LOG(LogSkinnedMeshComp, Warning, TEXT("SkeletalMesh %s, is not supported for current feature level (%s) and will not be rendered. MinLOD %d, NumBones %d (supported %d), NumBoneInfluences: %d"),
-						*GetNameSafe(GetSkinnedAsset()), *FeatureLevelName, MinLODIndex, MaxBonesPerChunk, MaxSupportedGPUSkinBones, NumBoneInfluences);
+						UE_LOG(LogSkinnedMeshComp, Warning, TEXT("SkeletalMesh %s, is not supported for current feature level (%s) and will not be rendered. MinLOD %d, NumBones %d (supported %d), NumBoneInfluences: %d"),
+							*GetNameSafe(GetSkinnedAsset()), *FeatureLevelName, MinLODIndex, MaxBonesPerChunk, MaxSupportedGPUSkinBones, NumBoneInfluences);
+					}
 				}
-			}
 
-			//Allow the editor a chance to manipulate it before its added to the scene
-			PostInitMeshObject(MeshObject);
+				//Allow the editor a chance to manipulate it before its added to the scene
+				PostInitMeshObject(MeshObject);
+			}
 		}
 	}
 
