[NetPhysics] Correct physics offset if it goes too out of sync.
Can happen if a client or server freezes up for a while or when debugging and putting breakpoints on client or server.
Normally any small offset will be instantly adjusted through time dilation.

diff --git a/Engine/Source/Runtime/Engine/Private/PlayerController.cpp b/Engine/Source/Runtime/Engine/Private/PlayerController.cpp
--- a/Engine/Source/Runtime/Engine/Private/PlayerController.cpp
+++ b/Engine/Source/Runtime/Engine/Private/PlayerController.cpp
@@ -131,6 +131,9 @@ namespace NetworkPhysicsCvars
 	int32 TickOffsetUpdateInterval = 10;
 	FAutoConsoleVariableRef CVarTickOffsetUpdateInterval(TEXT("np2.TickOffsetUpdateInterval"), TickOffsetUpdateInterval, TEXT("How many physics ticks to wait between each tick offset update. Lowest viable value = 1, which means update each tick. Deactivate physics offset updates by setting to 0 or negative value."));
 	
+	int32 TickOffsetCorrectionLimit = 10;
+	FAutoConsoleVariableRef CVarTickOffsetCorrectionLimit(TEXT("np2.TickOffsetCorrectionLimit"), TickOffsetCorrectionLimit, TEXT("If the client gets out of sync with physics ticks more than this limit, cut the losses and reset the offset."));
+	
 	float TimeDilationAmount = 0.01f;
 	FAutoConsoleVariableRef CVarTimeDilationAmount(TEXT("np2.TimeDilationAmount"), TimeDilationAmount, TEXT("Server-side CVar, Disable TimeDilation by setting to 0 | Default: 0.01 | Value is in percent where 0.01 = 1% dilation. Example: 1.0/0.01 = 100, meaning that over the time it usually takes to tick 100 physics steps we will tick 99 or 101 depending on if we dilate up or down."));
 
@@ -6232,9 +6235,9 @@ void APlayerController::ServerSendLatestAsyncPhysicsTimestamp_Implementation(FAs
 	const int32 PredictedServerFrame = Timestamp.ServerFrame;
 	bNetworkPhysicsTickOffsetAssigned |= PredictedServerFrame != INDEX_NONE;
 
-	// Send update to client if offset is not assigned
+	// Send update to client if offset is not assigned or over correction limit
 	// Note that we are sending the current ServerFrame along with the frame buffer added, to the client.
-	if (!bNetworkPhysicsTickOffsetAssigned)
+	if (!bNetworkPhysicsTickOffsetAssigned || FMath::Abs(PredictedServerFrame - ActualTimestamp.ServerFrame) > NetworkPhysicsCvars::TickOffsetCorrectionLimit)
 	{
 		Timestamp.ServerFrame = ActualTimestamp.ServerFrame;
 		NetworkPhysicsTickOffset = Timestamp.ServerFrame - Timestamp.LocalFrame;
@@ -6309,11 +6312,8 @@ void APlayerController::ClientSetupNetworkPhysicsTimestamp_Implementation(FAsync
 	ensure(UPhysicsSettings::Get()->PhysicsPrediction.bEnablePhysicsPrediction);
 
 	// Assign async physics tick offset
-	if (!bNetworkPhysicsTickOffsetAssigned)
-	{
-		bNetworkPhysicsTickOffsetAssigned = true;
-		NetworkPhysicsTickOffset = Timestamp.ServerFrame - Timestamp.LocalFrame;
-	}
+	bNetworkPhysicsTickOffsetAssigned = true;
+	NetworkPhysicsTickOffset = Timestamp.ServerFrame - Timestamp.LocalFrame;
 }
 
 void APlayerController::ClientAckTimeDilation_Implementation(float TimeDilation, int32 ServerStep)
