World Partition Editor: fix performance regression by creating temporary actor descriptor handles during editor painting from 30260347.

diff --git a/Engine/Source/Runtime/Engine/Private/WorldPartition/WorldPartitionEditorSpatialHash.cpp b/Engine/Source/Runtime/Engine/Private/WorldPartition/WorldPartitionEditorSpatialHash.cpp
--- a/Engine/Source/Runtime/Engine/Private/WorldPartition/WorldPartitionEditorSpatialHash.cpp
+++ b/Engine/Source/Runtime/Engine/Private/WorldPartition/WorldPartitionEditorSpatialHash.cpp
@@ -301,7 +301,7 @@ void UWorldPartitionEditorSpatialHash::UnhashActor(FWorldPartitionHandle& InActo
 
 int32 UWorldPartitionEditorSpatialHash::ForEachIntersectingActor(const FBox& Box, TFunctionRef<void(FWorldPartitionActorDescInstance*)> InOperation, const FForEachIntersectingActorParams& Params)
 {
-	TSet<FWorldPartitionHandle> IntersectedActors;
+	TSet<FGuid> IntersectedActors;
 
 	const int MininumumLevel = Params.MinimumBox.IsSet() ? GetLevelForBox(*Params.MinimumBox) : 0;
 
@@ -314,7 +314,7 @@ int32 UWorldPartitionEditorSpatialHash::ForEachIntersectingActor(const FBox& Box
 				if (ActorHandle.IsValid())
 				{
 					bool bWasAlreadyInSet;
-					IntersectedActors.Add(ActorHandle, &bWasAlreadyInSet);
+					IntersectedActors.Add(ActorHandle.GetInstance()->GetGuid(), &bWasAlreadyInSet);
 
 					if (!bWasAlreadyInSet)
 					{
@@ -337,7 +337,7 @@ int32 UWorldPartitionEditorSpatialHash::ForEachIntersectingActor(const FBox& Box
 				if (Box.Intersect(ActorHandle.GetInstance()->GetEditorBounds()))
 				{
 					bool bWasAlreadyInSet;
-					IntersectedActors.Add(ActorHandle, &bWasAlreadyInSet);
+					IntersectedActors.Add(ActorHandle.GetInstance()->GetGuid(), &bWasAlreadyInSet);
 				
 					if (!bWasAlreadyInSet)
 					{
