Fix bug with AudioDecompress

Reset `bIsStoringEndOfFile`
#tests PIE

diff --git a/Engine/Source/Runtime/Engine/Private/AudioDecompress.cpp b/Engine/Source/Runtime/Engine/Private/AudioDecompress.cpp
--- a/Engine/Source/Runtime/Engine/Private/AudioDecompress.cpp
+++ b/Engine/Source/Runtime/Engine/Private/AudioDecompress.cpp
@@ -359,6 +359,7 @@ bool IStreamedCompressedInfo::StreamCompressedData(uint8* Destination, bool bLoo
 							CurrentSampleCount = StreamSeekToAudioFrames * NumChannels;
 							LastPCMByteSize = 0;
 							LastPCMOffset = 0;
+							bStoringEndOfFile = false;
 						}
 						else // Seek failed (off the end of the chunk).
 						{
