[shaders] move error processing so it doesn't only happen in the "handle potential retry" path, but instead happens during unique error gathering (and so debug info path/filenames/shadertype/permutation etc can be properly reported if global shader compilation fails during a cook)

diff --git a/Engine/Source/Runtime/Engine/Private/ShaderCompiler/ShaderCompiler.cpp b/Engine/Source/Runtime/Engine/Private/ShaderCompiler/ShaderCompiler.cpp
--- a/Engine/Source/Runtime/Engine/Private/ShaderCompiler/ShaderCompiler.cpp
+++ b/Engine/Source/Runtime/Engine/Private/ShaderCompiler/ShaderCompiler.cpp
@@ -2797,93 +2797,36 @@ bool FShaderCompileUtilities::DoWriteTasks(const TArray<FShaderCommonCompileJobP
 	return DoWriteTasksInner(QueuedJobs, InTransferFile, BuildDistributionController, bUseRelativePaths, bCompressTaskFile);
 }
 
-static void ProcessErrors(const FShaderCompileJob& CurrentJob, TArray<FString>& UniqueErrors, FString& ErrorString)
+struct FShaderErrorInfo
+{
+	TArray<FShaderCommonCompileJob*> ErrorJobs;
+	TArray<FString> UniqueErrors;
+	TArray<FString> UniqueErrorPrefixes;
+	TArray<FString> UniqueWarnings;
+	TArray<EShaderPlatform> ErrorPlatforms;
+	FString TargetShaderPlatformString;
+};
+
+static void BuildErrorStringAndReport(const FShaderErrorInfo& ErrorInfo, FString& ErrorString)
 {
 	bool bReportedDebugInfo = false;
 
-	for (int32 ErrorIndex = 0; ErrorIndex < CurrentJob.Output.Errors.Num(); ErrorIndex++)
+	for (int32 ErrorIndex = 0; ErrorIndex < ErrorInfo.UniqueErrors.Num(); ErrorIndex++)
 	{
-		FShaderCompilerError CurrentError = CurrentJob.Output.Errors[ErrorIndex];
-		int32 UniqueError = INDEX_NONE;
+		FString UniqueErrorString = ErrorInfo.UniqueErrorPrefixes[ErrorIndex] + ErrorInfo.UniqueErrors[ErrorIndex] + TEXT("\n");
 
-		if (UniqueErrors.Find(CurrentError.GetErrorString(), UniqueError))
+		if (FPlatformMisc::IsDebuggerPresent())
 		{
-			// This unique error is being processed, remove it from the array
-			UniqueErrors.RemoveAt(UniqueError);
-
-			// Extract source location from error message if the shader backend doesn't provide it separated from the stripped message
-			CurrentError.ExtractSourceLocation();
-
-			// Remap filenames
-			if (CurrentError.ErrorVirtualFilePath == TEXT("/Engine/Generated/Material.ush"))
-			{
-				// MaterialTemplate.usf is dynamically included as Material.usf
-				// Currently the material translator does not add new lines when filling out MaterialTemplate.usf,
-				// So we don't need the actual filled out version to find the line of a code bug.
-				CurrentError.ErrorVirtualFilePath = TEXT("/Engine/Private/MaterialTemplate.ush");
-			}
-			else if (CurrentError.ErrorVirtualFilePath.Contains(TEXT("memory")))
-			{
-				check(CurrentJob.Key.ShaderType);
-
-				// Files passed to the shader compiler through memory will be named memory
-				// Only the shader's main file is passed through memory without a filename
-				CurrentError.ErrorVirtualFilePath = FString(CurrentJob.Key.ShaderType->GetShaderFilename());
-			}
-			else if (CurrentError.ErrorVirtualFilePath == TEXT("/Engine/Generated/VertexFactory.ush"))
-			{
-				// VertexFactory.usf is dynamically included from whichever vertex factory the shader was compiled with.
-				check(CurrentJob.Key.VFType);
-				CurrentError.ErrorVirtualFilePath = FString(CurrentJob.Key.VFType->GetShaderFilename());
-			}
-			else if (CurrentError.ErrorVirtualFilePath == TEXT("") && CurrentJob.Key.ShaderType)
-			{
-				// Some shader compiler errors won't have a file and line number, so we just assume the error happened in file containing the entrypoint function.
-				CurrentError.ErrorVirtualFilePath = FString(CurrentJob.Key.ShaderType->GetShaderFilename());
-			}
-
-			FString UniqueErrorPrefix;
-
-			// If we dumped the shader info, add it before the first error string
-			if (!GIsBuildMachine && !bReportedDebugInfo && CurrentJob.Input.DumpDebugInfoPath.Len() > 0)
-			{
-				UniqueErrorPrefix += FString::Printf(TEXT("Shader debug info dumped to: \"%s\"\n"), *CurrentJob.Input.DumpDebugInfoPath);
-				bReportedDebugInfo = true;
-			}
-
-			if (CurrentJob.Key.ShaderType)
-			{
-				// Construct a path that will enable VS.NET to find the shader file, relative to the solution
-				const FString SolutionPath = FPaths::RootDir();
-				FString ShaderFilePath = IFileManager::Get().ConvertToAbsolutePathForExternalAppForRead(*CurrentError.GetShaderSourceFilePath());
-				UniqueErrorPrefix += FString::Printf(TEXT("%s(%s): Shader %s, Permutation %d, VF %s:\n\t"),
-					*ShaderFilePath,
-					*CurrentError.ErrorLineString,
-					CurrentJob.Key.ShaderType->GetName(),
-					CurrentJob.Key.PermutationId,
-					CurrentJob.Key.VFType ? CurrentJob.Key.VFType->GetName() : TEXT("None"));
-			}
-			else
-			{
-				UniqueErrorPrefix += FString::Printf(TEXT("%s(0): "),
-					*CurrentJob.Input.VirtualSourceFilePath);
-			}
-
-			FString UniqueErrorString = UniqueErrorPrefix + CurrentError.GetErrorStringWithLineMarker() + TEXT("\n");
-
-			if (FPlatformMisc::IsDebuggerPresent())
-			{
-				// Using OutputDebugString to avoid any text getting added before the filename,
-				// Which will throw off VS.NET's ability to take you directly to the file and line of the error when double clicking it in the output window.
-				FPlatformMisc::LowLevelOutputDebugStringf(*UniqueErrorString);
-			}
-			else
-			{
-				UE_LOG(LogShaderCompilers, Warning, TEXT("%s"), *UniqueErrorString);
-			}
-
-			ErrorString += UniqueErrorString;
+			// Using OutputDebugString to avoid any text getting added before the filename,
+			// Which will throw off VS.NET's ability to take you directly to the file and line of the error when double clicking it in the output window.
+			FPlatformMisc::LowLevelOutputDebugStringf(*UniqueErrorString);
+		}
+		else
+		{
+			UE_LOG(LogShaderCompilers, Warning, TEXT("%s"), *UniqueErrorString);
 		}
+
+		ErrorString += UniqueErrorString;
 	}
 }
 
@@ -3275,15 +3218,6 @@ FSCWErrorCode::ECode FShaderCompileUtilities::DoReadTaskResults(const TArray<FSh
 	return FSCWErrorCode::Success;
 }
 
-struct FShaderErrorInfo
-{
-	TArray<FShaderCommonCompileJob*> ErrorJobs;
-	TArray<FString> UniqueErrors;
-	TArray<FString> UniqueWarnings;
-	TArray<EShaderPlatform> ErrorPlatforms;
-	FString TargetShaderPlatformString;
-};
-
 #if WITH_EDITOR
 static bool CheckSingleJob(const FShaderCompileJob& SingleJob, TArray<FString>& OutErrors)
 {
@@ -3330,19 +3264,85 @@ static bool CheckSingleJob(const FShaderCompileJob& SingleJob, TArray<FString>&
 };
 #endif // WITH_EDITOR
 
-static int32 AddErrorsForFailedJobFiltered(FShaderCompileJob& CurrentJob, FShaderErrorInfo& OutShaderErrorInfo, const TCHAR* FilterMessage)
+static int32 AddAndProcessErrorsForFailedJobFiltered(FShaderCompileJob& CurrentJob, FShaderErrorInfo& OutShaderErrorInfo, const TCHAR* FilterMessage)
 {
 	int32 NumAddedErrors = 0;
 
+	bool bReportedDebugInfo = false;
+
 	for (int32 ErrorIndex = 0; ErrorIndex < CurrentJob.Output.Errors.Num(); ErrorIndex++)
 	{
-		const FShaderCompilerError& CurrentError = CurrentJob.Output.Errors[ErrorIndex];
+		FShaderCompilerError& CurrentError = CurrentJob.Output.Errors[ErrorIndex];
+		FString CurrentErrorString = CurrentError.GetErrorString();
 
 		// Include warnings if LogShaders is unsuppressed, otherwise only include filtered messages
 		if (UE_LOG_ACTIVE(LogShaders, Log) || FilterMessage == nullptr || CurrentError.StrippedErrorMessage.Contains(FilterMessage))
 		{
-			OutShaderErrorInfo.UniqueErrors.AddUnique(CurrentJob.Output.Errors[ErrorIndex].GetErrorString());
-			OutShaderErrorInfo.ErrorJobs.AddUnique(&CurrentJob);
+			// Extract source location from error message if the shader backend doesn't provide it separated from the stripped message
+			CurrentError.ExtractSourceLocation();
+
+			// Remap filenames
+			if (CurrentError.ErrorVirtualFilePath == TEXT("/Engine/Generated/Material.ush"))
+			{
+				// MaterialTemplate.usf is dynamically included as Material.usf
+				// Currently the material translator does not add new lines when filling out MaterialTemplate.usf,
+				// So we don't need the actual filled out version to find the line of a code bug.
+				CurrentError.ErrorVirtualFilePath = TEXT("/Engine/Private/MaterialTemplate.ush");
+			}
+			else if (CurrentError.ErrorVirtualFilePath.Contains(TEXT("memory")))
+			{
+				check(CurrentJob.Key.ShaderType);
+
+				// Files passed to the shader compiler through memory will be named memory
+				// Only the shader's main file is passed through memory without a filename
+				CurrentError.ErrorVirtualFilePath = FString(CurrentJob.Key.ShaderType->GetShaderFilename());
+			}
+			else if (CurrentError.ErrorVirtualFilePath == TEXT("/Engine/Generated/VertexFactory.ush"))
+			{
+				// VertexFactory.usf is dynamically included from whichever vertex factory the shader was compiled with.
+				check(CurrentJob.Key.VFType);
+				CurrentError.ErrorVirtualFilePath = FString(CurrentJob.Key.VFType->GetShaderFilename());
+			}
+			else if (CurrentError.ErrorVirtualFilePath == TEXT("") && CurrentJob.Key.ShaderType)
+			{
+				// Some shader compiler errors won't have a file and line number, so we just assume the error happened in file containing the entrypoint function.
+				CurrentError.ErrorVirtualFilePath = FString(CurrentJob.Key.ShaderType->GetShaderFilename());
+			}
+
+			if (OutShaderErrorInfo.UniqueErrors.Find(CurrentErrorString) == INDEX_NONE)
+			{
+				// build up additional info in a "prefix" string; only do this once for each unique error
+				FString UniqueErrorPrefix;
+
+				// If we dumped the shader info, add it before the first error string
+				if (!GIsBuildMachine && !bReportedDebugInfo && CurrentJob.Input.DumpDebugInfoPath.Len() > 0)
+				{
+					UniqueErrorPrefix += FString::Printf(TEXT("Shader debug info dumped to: \"%s\"\n"), *CurrentJob.Input.DumpDebugInfoPath);
+					bReportedDebugInfo = true;
+				}
+
+				if (CurrentJob.Key.ShaderType)
+				{
+					// Construct a path that will enable VS.NET to find the shader file, relative to the solution
+					const FString SolutionPath = FPaths::RootDir();
+					FString ShaderFilePath = IFileManager::Get().ConvertToAbsolutePathForExternalAppForRead(*CurrentError.GetShaderSourceFilePath());
+					UniqueErrorPrefix += FString::Printf(TEXT("%s(%s): Shader %s, Permutation %d, VF %s:\n\t"),
+						*ShaderFilePath,
+						*CurrentError.ErrorLineString,
+						CurrentJob.Key.ShaderType->GetName(),
+						CurrentJob.Key.PermutationId,
+						CurrentJob.Key.VFType ? CurrentJob.Key.VFType->GetName() : TEXT("None"));
+				}
+				else
+				{
+					UniqueErrorPrefix += FString::Printf(TEXT("%s(0): "),
+						*CurrentJob.Input.VirtualSourceFilePath);
+				}
+
+				OutShaderErrorInfo.UniqueErrors.Add(CurrentErrorString);
+				OutShaderErrorInfo.UniqueErrorPrefixes.Add(UniqueErrorPrefix);
+				OutShaderErrorInfo.ErrorJobs.AddUnique(&CurrentJob);
+			}
 			++NumAddedErrors;
 		}
 	}
@@ -3350,7 +3350,7 @@ static int32 AddErrorsForFailedJobFiltered(FShaderCompileJob& CurrentJob, FShade
 	return NumAddedErrors;
 }
 
-static void AddErrorsForFailedJob(FShaderCompileJob& CurrentJob, FShaderErrorInfo& OutShaderErrorInfo)
+static void AddAndProcessErrorsForFailedJob(FShaderCompileJob& CurrentJob, FShaderErrorInfo& OutShaderErrorInfo)
 {
 	OutShaderErrorInfo.ErrorPlatforms.AddUnique((EShaderPlatform)CurrentJob.Input.Target.Platform);
 
@@ -3363,9 +3363,9 @@ static void AddErrorsForFailedJob(FShaderCompileJob& CurrentJob, FShaderErrorInf
 
 	// If we filter all error messages because they are interpreted as warnings, we have to assume all error messages are in fact errors and not warnings.
 	// In that case, add jobs again without a filter; e.g. when the stripped message starts with "Internal exception".
-	if (AddErrorsForFailedJobFiltered(CurrentJob, OutShaderErrorInfo, TEXT("error")) == 0)
+	if (AddAndProcessErrorsForFailedJobFiltered(CurrentJob, OutShaderErrorInfo, TEXT("error")) == 0)
 	{
-		AddErrorsForFailedJobFiltered(CurrentJob, OutShaderErrorInfo, nullptr);
+		AddAndProcessErrorsForFailedJobFiltered(CurrentJob, OutShaderErrorInfo, nullptr);
 	}
 }
 
@@ -6828,7 +6828,7 @@ static bool GatherUniqueErrors(const TArray<FShaderCommonCompileJobPtr>& Complet
 			FShaderCompileJob* SingleJob = CurrentJob.GetSingleShaderJob();
 			if (SingleJob)
 			{
-				AddErrorsForFailedJob(*SingleJob, OutShaderErrorInfo);
+				AddAndProcessErrorsForFailedJob(*SingleJob, OutShaderErrorInfo);
 			}
 			else
 			{
@@ -6836,7 +6836,7 @@ static bool GatherUniqueErrors(const TArray<FShaderCommonCompileJobPtr>& Complet
 				check(PipelineJob);
 				for (TRefCountPtr<FShaderCompileJob>& CommonJob : PipelineJob->StageJobs)
 				{
-					AddErrorsForFailedJob(*CommonJob, OutShaderErrorInfo);
+					AddAndProcessErrorsForFailedJob(*CommonJob, OutShaderErrorInfo);
 				}
 			}
 		}
@@ -6905,7 +6905,6 @@ bool FShaderCompilingManager::HandlePotentialRetryOnError(TMap<int32, FShaderMap
 				}
 			}
 
-
 			if (UE_LOG_ACTIVE(LogShaders, Log) 
 				// Always log detailed errors when a special engine material or global shader fails to compile, as those will be fatal errors
 				|| bSpecialEngineMaterial 
@@ -6919,31 +6918,15 @@ bool FShaderCompilingManager::HandlePotentialRetryOnError(TMap<int32, FShaderMap
 				FString ErrorString = FString::Printf(TEXT("%i Shader compiler errors compiling %s for platform %s:"), ShaderErrorInfo.UniqueErrors.Num(), MaterialName, *ShaderErrorInfo.TargetShaderPlatformString);
 				UE_LOG(LogShaderCompilers, Warning, TEXT("%s"), *ErrorString);
 				ErrorString += TEXT("\n");
-				bool bAnyErrorLikelyToBeCodeError = false;
 
-				for (int32 JobIndex = 0; JobIndex < CompleteJobs.Num(); JobIndex++)
+				bool bAnyErrorLikelyToBeCodeError = false;
+				for (const FShaderCommonCompileJob* Job : ShaderErrorInfo.ErrorJobs)
 				{
-					const FShaderCommonCompileJob& CurrentJob = *CompleteJobs[JobIndex];
-					if (!CurrentJob.bSucceeded)
-					{
-						bAnyErrorLikelyToBeCodeError |= CurrentJob.bErrorsAreLikelyToBeCode;
-						const auto* SingleJob = CurrentJob.GetSingleShaderJob();
-						if (SingleJob)
-						{
-							ProcessErrors(*SingleJob, ShaderErrorInfo.UniqueErrors, ErrorString);
-						}
-						else
-						{
-							const auto* PipelineJob = CurrentJob.GetShaderPipelineJob();
-							check(PipelineJob);
-							for (auto CommonJob : PipelineJob->StageJobs)
-							{
-								ProcessErrors(*CommonJob, ShaderErrorInfo.UniqueErrors, ErrorString);
-							}
-						}
-					}
+					bAnyErrorLikelyToBeCodeError |= Job->bErrorsAreLikelyToBeCode;
 				}
 
+				BuildErrorStringAndReport(ShaderErrorInfo, ErrorString);
+
 				if (UE_LOG_ACTIVE(LogShaders, Log) && (bAnyErrorLikelyToBeCodeError || bPromptToRetryFailedShaderCompiles || bSpecialEngineMaterial))
 				{
 					// Use debug break in debug with the debugger attached, otherwise message box
@@ -10297,8 +10280,9 @@ void ProcessCompiledGlobalShaders(const TArray<FShaderCommonCompileJobPtr>& Comp
 	GatherUniqueErrors(CompilationResults, ShaderErrorInfo);
 
 	// Report unique errors for global shaders.
-	for (const FString& ErrorString : ShaderErrorInfo.UniqueErrors)
+	for (int32 ErrorIndex = 0; ErrorIndex < ShaderErrorInfo.UniqueErrors.Num(); ++ErrorIndex)
 	{
+		FString ErrorString = ShaderErrorInfo.UniqueErrorPrefixes[ErrorIndex] + ShaderErrorInfo.UniqueErrors[ErrorIndex];
 		UE_LOGFMT_NSLOC(LogShaders, Error, "Shaders", "GlobalShaderCompileError", "{ErrorMessage}", 
 			("ErrorMessage", ErrorString));
 	}
