World Partition: fix FWorldPartitionExternalDirtyActorsTracker to properly handle deleted actors.
Happens when undoing a new actor placement, the dirty actor tracker will register it again if we redo the delete operation.

diff --git a/Engine/Source/Runtime/Engine/Private/WorldPartition/WorldPartition.cpp b/Engine/Source/Runtime/Engine/Private/WorldPartition/WorldPartition.cpp
--- a/Engine/Source/Runtime/Engine/Private/WorldPartition/WorldPartition.cpp
+++ b/Engine/Source/Runtime/Engine/Private/WorldPartition/WorldPartition.cpp
@@ -272,7 +272,7 @@ void UWorldPartition::FWorldPartitionExternalDirtyActorsTracker::Tick(float InDe
 	for (auto& [Actor, Reference] : NonDirtyActors)
 	{
 		// Resolve reference for newly added actors
-		if (!Reference.IsValid())
+		if (!Reference.IsValid() && Actor.IsValid())
 		{
 			Reference = FWorldPartitionReference(Owner, Actor->GetActorGuid());
 		}
