World Partition Runtime HashSet: optimize ForEachStreamingCellsQuery to early out when provided function returns false.

diff --git a/Engine/Source/Runtime/Engine/Private/WorldPartition/RuntimeHashSet/StaticSpatialIndex.cpp b/Engine/Source/Runtime/Engine/Private/WorldPartition/RuntimeHashSet/StaticSpatialIndex.cpp
--- a/Engine/Source/Runtime/Engine/Private/WorldPartition/RuntimeHashSet/StaticSpatialIndex.cpp
+++ b/Engine/Source/Runtime/Engine/Private/WorldPartition/RuntimeHashSet/StaticSpatialIndex.cpp
@@ -9,15 +9,19 @@ void FStaticSpatialIndex::FListImpl::Init(const TArray<TPair<FBox, uint32>>& InE
 	Algo::Transform(InElements, Elements, [](const TPair<FBox, uint32>& Element) { return Element.Value; });
 }
 
-void FStaticSpatialIndex::FListImpl::ForEachElement(TFunctionRef<void(uint32 InValueIndex)> Func) const
+bool FStaticSpatialIndex::FListImpl::ForEachElement(TFunctionRef<bool(uint32 InValueIndex)> Func) const
 {
 	for (uint32 ValueIndex : Elements)
 	{
-		Func(ValueIndex);
+		if (!Func(ValueIndex))
+		{
+			return false;
+		}
 	}
+	return true;
 }
 
-void FStaticSpatialIndex::FListImpl::ForEachIntersectingElement(const FBox& InBox, TFunctionRef<void(uint32 InValueIndex)> Func) const
+bool  FStaticSpatialIndex::FListImpl::ForEachIntersectingElement(const FBox& InBox, TFunctionRef<bool(uint32 InValueIndex)> Func) const
 {
 	for (uint32 ValueIndex : Elements)
 	{
@@ -25,12 +29,16 @@ void FStaticSpatialIndex::FListImpl::ForEachIntersectingElement(const FBox& InBo
 
 		if (Box.Intersect(InBox))
 		{
-			Func(ValueIndex);
+			if (!Func(ValueIndex))
+			{
+				return false;
+			}
 		}
 	}
+	return true;
 }
 
-void FStaticSpatialIndex::FListImpl::ForEachIntersectingElement(const FSphere& InSphere, TFunctionRef<void(uint32 InValueIndex)> Func) const
+bool FStaticSpatialIndex::FListImpl::ForEachIntersectingElement(const FSphere& InSphere, TFunctionRef<bool(uint32 InValueIndex)> Func) const
 {
 	for (uint32 ValueIndex : Elements)
 	{
@@ -38,9 +46,13 @@ void FStaticSpatialIndex::FListImpl::ForEachIntersectingElement(const FSphere& I
 
 		if (FMath::SphereAABBIntersection(InSphere, Box))
 		{
-			Func(ValueIndex);
+			if (!Func(ValueIndex))
+			{
+				return false;
+			}
 		}
 	}
+	return true;
 }
 
 void FStaticSpatialIndex::FRTreeImpl::Init(const TArray<TPair<FBox, uint32>>& InElements)
@@ -97,35 +109,42 @@ void FStaticSpatialIndex::FRTreeImpl::Init(const TArray<TPair<FBox, uint32>>& In
 	}
 }
 
-void FStaticSpatialIndex::FRTreeImpl::ForEachElement(TFunctionRef<void(uint32 InValueIndex)> Func) const
+bool FStaticSpatialIndex::FRTreeImpl::ForEachElement(TFunctionRef<bool(uint32 InValueIndex)> Func) const
 {
-	ForEachElementRecursive(&RootNode, Func);
+	return ForEachElementRecursive(&RootNode, Func);
 }
 
-void FStaticSpatialIndex::FRTreeImpl::ForEachElementRecursive(const FNode* Node, TFunctionRef<void(uint32 InValueIndex)> Func) const
+bool FStaticSpatialIndex::FRTreeImpl::ForEachElementRecursive(const FNode* Node, TFunctionRef<bool(uint32 InValueIndex)> Func) const
 {
 	if (Node->Content.IsType<FNode::FNodeType>())
 	{
 		for (auto& ChildNode : Node->Content.Get<FNode::FNodeType>())
 		{
-			ForEachElementRecursive(&ChildNode, Func);
+			if (!ForEachElementRecursive(&ChildNode, Func))
+			{
+				return false;
+			}
 		}
 	}
 	else
 	{
 		for (uint32 ValueIndex : Node->Content.Get<FNode::FLeafType>())
 		{
-			Func(ValueIndex);
+			if (!Func(ValueIndex))
+			{
+				return false;
+			}
 		}
 	}
+	return true;
 }
 
-void FStaticSpatialIndex::FRTreeImpl::ForEachIntersectingElement(const FBox& InBox, TFunctionRef<void(uint32 InValueIndex)> Func) const
+bool FStaticSpatialIndex::FRTreeImpl::ForEachIntersectingElement(const FBox& InBox, TFunctionRef<bool(uint32 InValueIndex)> Func) const
 {
-	ForEachIntersectingElementRecursive(&RootNode, InBox, Func);
+	return ForEachIntersectingElementRecursive(&RootNode, InBox, Func);
 }
 
-void FStaticSpatialIndex::FRTreeImpl::ForEachIntersectingElementRecursive(const FNode* Node, const FBox& InBox, TFunctionRef<void(uint32 InValueIndex)> Func) const
+bool FStaticSpatialIndex::FRTreeImpl::ForEachIntersectingElementRecursive(const FNode* Node, const FBox& InBox, TFunctionRef<bool(uint32 InValueIndex)> Func) const
 {
 	if (Node->Content.IsType<FNode::FNodeType>())
 	{
@@ -133,7 +152,10 @@ void FStaticSpatialIndex::FRTreeImpl::ForEachIntersectingElementRecursive(const
 		{
 			if (ChildNode.Box.Intersect(InBox))
 			{
-				ForEachIntersectingElementRecursive(&ChildNode, InBox, Func);
+				if (!ForEachIntersectingElementRecursive(&ChildNode, InBox, Func))
+				{
+					return false;
+				}
 			}
 		}
 	}
@@ -145,18 +167,22 @@ void FStaticSpatialIndex::FRTreeImpl::ForEachIntersectingElementRecursive(const
 
 			if (Box.Intersect(InBox))
 			{
-				Func(ValueIndex);
+				if (!Func(ValueIndex))
+				{
+					return false;
+				}
 			}
 		}
 	}
+	return true;
 }
 
-void FStaticSpatialIndex::FRTreeImpl::ForEachIntersectingElement(const FSphere& InSphere, TFunctionRef<void(uint32 InValueIndex)> Func) const
+bool FStaticSpatialIndex::FRTreeImpl::ForEachIntersectingElement(const FSphere& InSphere, TFunctionRef<bool(uint32 InValueIndex)> Func) const
 {
-	ForEachIntersectingElementRecursive(&RootNode, InSphere, Func);
+	return ForEachIntersectingElementRecursive(&RootNode, InSphere, Func);
 }
 
-void FStaticSpatialIndex::FRTreeImpl::ForEachIntersectingElementRecursive(const FNode* Node, const FSphere& InSphere, TFunctionRef<void(uint32 InValueIndex)> Func) const
+bool FStaticSpatialIndex::FRTreeImpl::ForEachIntersectingElementRecursive(const FNode* Node, const FSphere& InSphere, TFunctionRef<bool(uint32 InValueIndex)> Func) const
 {
 	if (Node->Content.IsType<FNode::FNodeType>())
 	{
@@ -164,7 +190,10 @@ void FStaticSpatialIndex::FRTreeImpl::ForEachIntersectingElementRecursive(const
 		{
 			if (FMath::SphereAABBIntersection(InSphere, ChildNode.Box))
 			{
-				ForEachIntersectingElementRecursive(&ChildNode, InSphere, Func);
+				if (!ForEachIntersectingElementRecursive(&ChildNode, InSphere, Func))
+				{
+					return false;
+				}
 			}
 		}
 	}
@@ -176,8 +205,12 @@ void FStaticSpatialIndex::FRTreeImpl::ForEachIntersectingElementRecursive(const
 
 			if (FMath::SphereAABBIntersection(InSphere, Box))
 			{
-				Func(ValueIndex);
+				if (!Func(ValueIndex))
+				{
+					return false;
+				}
 			}
 		}
 	}
+	return true;
 }
\ No newline at end of file
diff --git a/Engine/Source/Runtime/Engine/Private/WorldPartition/RuntimeHashSet/WorldPartitionRuntimeHashSet.cpp b/Engine/Source/Runtime/Engine/Private/WorldPartition/RuntimeHashSet/WorldPartitionRuntimeHashSet.cpp
--- a/Engine/Source/Runtime/Engine/Private/WorldPartition/RuntimeHashSet/WorldPartitionRuntimeHashSet.cpp
+++ b/Engine/Source/Runtime/Engine/Private/WorldPartition/RuntimeHashSet/WorldPartitionRuntimeHashSet.cpp
@@ -87,6 +87,7 @@ void UWorldPartitionRuntimeHashSet::PostLoad()
 		ForEachStreamingData([](const FRuntimePartitionStreamingData& StreamingData)
 		{
 			StreamingData.CreatePartitionsSpatialIndex();
+			return true;
 		});
 	}
 }
@@ -292,14 +293,17 @@ void UWorldPartitionRuntimeHashSet::ForEachStreamingCells(TFunctionRef<bool(cons
 	{
 		for (UWorldPartitionRuntimeCell* Cell : InCells)
 		{
-			Func(Cell);
+			if (!Func(Cell))
+			{
+				return false;
+			}
 		}
+		return true;
 	};
 
 	ForEachStreamingData([&ForEachCells](const FRuntimePartitionStreamingData& StreamingData)
 	{
-		ForEachCells(StreamingData.StreamingCells);
-		ForEachCells(StreamingData.NonStreamingCells);
+		return ForEachCells(StreamingData.StreamingCells) && ForEachCells(StreamingData.NonStreamingCells);
 	});
 }
 
@@ -335,13 +339,12 @@ void UWorldPartitionRuntimeHashSet::ForEachStreamingCellsQuery(const FWorldParti
 
 				InSpatialIndex->ForEachIntersectingElement(ShapeSphere, [&ShouldAddCell, &QuerySource, &Func](UWorldPartitionRuntimeCell* RuntimeCell)
 				{
-					if (ShouldAddCell(RuntimeCell, QuerySource))
-					{
-						Func(RuntimeCell);
-					}
+					return !ShouldAddCell(RuntimeCell, QuerySource) || Func(RuntimeCell);
 				});
 			});
 		}
+
+		return true;
 	};
 
 	auto ForEachNonStreamingCells = [&ShouldAddCell, &QuerySource, &Func](TArray<TObjectPtr<UWorldPartitionRuntimeCell>> InNonStreamingCells)
@@ -350,15 +353,18 @@ void UWorldPartitionRuntimeHashSet::ForEachStreamingCellsQuery(const FWorldParti
 		{
 			if (ShouldAddCell(Cell, QuerySource))
 			{
-				Func(Cell);
+				if (!Func(Cell))
+				{
+					return false;
+				}
 			}
 		}
+		return true;
 	};
 
 	ForEachStreamingData([&QuerySource, &ForEachStreamingCells, &ForEachNonStreamingCells](const FRuntimePartitionStreamingData& StreamingData)
 	{
-		ForEachStreamingCells(StreamingData.SpatialIndex.Get(), StreamingData.LoadingRange, StreamingData.Name);
-		ForEachNonStreamingCells(StreamingData.NonStreamingCells);
+		return ForEachStreamingCells(StreamingData.SpatialIndex.Get(), StreamingData.LoadingRange, StreamingData.Name) && ForEachNonStreamingCells(StreamingData.NonStreamingCells);
 	});
 }
 
@@ -367,13 +373,13 @@ void UWorldPartitionRuntimeHashSet::ForEachStreamingCellsSources(const TArray<FW
 	UWorldPartitionRuntimeHash::FStreamingSourceCells ActivateStreamingSourceCells;
 	UWorldPartitionRuntimeHash::FStreamingSourceCells LoadStreamingSourceCells;
 
-	auto ForEachStreamingCells = [this, &Sources, &Func, &ActivateStreamingSourceCells, &LoadStreamingSourceCells](FStaticSpatialIndexType* InSpatialIndex, int32 InLoadingRange, FName InGridName)
+	auto ForEachStreamingCells = [this, &Sources, &ActivateStreamingSourceCells, &LoadStreamingSourceCells](FStaticSpatialIndexType* InSpatialIndex, int32 InLoadingRange, FName InGridName)
 	{
 		if (InSpatialIndex)
 		{
 			for (const FWorldPartitionStreamingSource& Source : Sources)
 			{
-				Source.ForEachShape(InLoadingRange, InGridName, false, [this, &Source, InSpatialIndex, &Func, &ActivateStreamingSourceCells, &LoadStreamingSourceCells](const FSphericalSector& Shape)
+				Source.ForEachShape(InLoadingRange, InGridName, false, [this, &Source, InSpatialIndex, &ActivateStreamingSourceCells, &LoadStreamingSourceCells](const FSphericalSector& Shape)
 				{
 					const FSphere ShapeSphere(Shape.GetCenter(), Shape.GetRadius());
 
@@ -409,6 +415,7 @@ void UWorldPartitionRuntimeHashSet::ForEachStreamingCellsSources(const TArray<FW
 				});
 			}
 		}
+		return true;
 	};
 
 	auto ForEachNonStreamingCells = [this, &ActivateStreamingSourceCells, &LoadStreamingSourceCells](TArray<TObjectPtr<UWorldPartitionRuntimeCell>> InNonStreamingCells)
@@ -432,12 +439,12 @@ void UWorldPartitionRuntimeHashSet::ForEachStreamingCellsSources(const TArray<FW
 				}
 			}
 		}
+		return true;
 	};
 
 	ForEachStreamingData([&ForEachStreamingCells, &ForEachNonStreamingCells](const FRuntimePartitionStreamingData& StreamingData)
 	{
-		ForEachStreamingCells(StreamingData.SpatialIndex.Get(), StreamingData.LoadingRange, StreamingData.Name);
-		ForEachNonStreamingCells(StreamingData.NonStreamingCells);
+		return ForEachStreamingCells(StreamingData.SpatialIndex.Get(), StreamingData.LoadingRange, StreamingData.Name) && ForEachNonStreamingCells(StreamingData.NonStreamingCells);
 	});
 
 	auto ExecuteFuncOnCells = [Func](const TSet<const UWorldPartitionRuntimeCell*>& Cells, EStreamingSourceTargetState TargetState)
@@ -685,11 +692,14 @@ UWorldPartitionRuntimeHashSet::FCellUniqueId UWorldPartitionRuntimeHashSet::GetC
 }
 #endif
 
-void UWorldPartitionRuntimeHashSet::ForEachStreamingData(TFunctionRef<void(const FRuntimePartitionStreamingData&)> Func) const
+void UWorldPartitionRuntimeHashSet::ForEachStreamingData(TFunctionRef<bool(const FRuntimePartitionStreamingData&)> Func) const
 {
 	for (const FRuntimePartitionStreamingData& StreamingData : RuntimeStreamingData)
 	{
-		Func(StreamingData);
+		if (!Func(StreamingData))
+		{
+			return;
+		}
 	}
 
 	for (const TWeakObjectPtr<URuntimeHashExternalStreamingObjectBase>& InjectedExternalStreamingObject : InjectedExternalStreamingObjects)
@@ -700,7 +710,10 @@ void UWorldPartitionRuntimeHashSet::ForEachStreamingData(TFunctionRef<void(const
 			
 			for (const FRuntimePartitionStreamingData& StreamingData : ExternalStreamingObject->RuntimeStreamingData)
 			{
-				Func(StreamingData);
+				if (!Func(StreamingData))
+				{
+					return;
+				}
 			}
 		}
 	}
diff --git a/Engine/Source/Runtime/Engine/Private/WorldPartition/RuntimeHashSet/WorldPartitionRuntimeHashSetDebugDraw.cpp b/Engine/Source/Runtime/Engine/Private/WorldPartition/RuntimeHashSet/WorldPartitionRuntimeHashSetDebugDraw.cpp
--- a/Engine/Source/Runtime/Engine/Private/WorldPartition/RuntimeHashSet/WorldPartitionRuntimeHashSetDebugDraw.cpp
+++ b/Engine/Source/Runtime/Engine/Private/WorldPartition/RuntimeHashSet/WorldPartitionRuntimeHashSetDebugDraw.cpp
@@ -72,6 +72,7 @@ bool UWorldPartitionRuntimeHashSet::Draw2D(FWorldPartitionDraw2DContext& DrawCon
 				FilteredStreamingObjects.FindOrAdd(StreamingData.Name).Add(&StreamingData);
 			}
 		}
+		return true;
 	});
 
 	if (!FilteredStreamingObjects.Num())
@@ -387,5 +388,6 @@ void UWorldPartitionRuntimeHashSet::Draw3D(const TArray<FWorldPartitionStreaming
 				});
 			});
 		}
+		return true;
 	});
 }
\ No newline at end of file
diff --git a/Engine/Source/Runtime/Engine/Public/WorldPartition/RuntimeHashSet/StaticSpatialIndex.h b/Engine/Source/Runtime/Engine/Public/WorldPartition/RuntimeHashSet/StaticSpatialIndex.h
--- a/Engine/Source/Runtime/Engine/Public/WorldPartition/RuntimeHashSet/StaticSpatialIndex.h
+++ b/Engine/Source/Runtime/Engine/Public/WorldPartition/RuntimeHashSet/StaticSpatialIndex.h
@@ -6,6 +6,7 @@
 #include "Misc/TVariant.h"
 #include "UObject/ObjectPtr.h"
 #include "UObject/UObjectGlobals.h"
+#include "OverrideVoidReturnInvoker.h"
 
 struct IStaticSpatialIndexDataInterface
 {
@@ -29,27 +30,36 @@ public:
 		InitSpatialIndex();
 	}
 
-	void ForEachElement(TFunctionRef<void(const ValueType& InValue)> Func) const
+	template <class Func>
+	void ForEachElement(Func InFunc) const
 	{
-		SpatialIndex.ForEachElement([this, &Func](uint32 ValueIndex)
+		TOverrideVoidReturnInvoker Invoker(true, InFunc);
+
+		SpatialIndex.ForEachElement([this, &Invoker](uint32 ValueIndex)
 		{
-			Func(Elements[ValueIndex].Value);
+			return Invoker(Elements[ValueIndex].Value);
 		});
 	}
 
-	void ForEachIntersectingElement(const FBox& InBox, TFunctionRef<void(const ValueType& InValue)> Func) const
+	template <class Func>
+	void ForEachIntersectingElement(const FBox& InBox, Func InFunc) const
 	{
-		SpatialIndex.ForEachIntersectingElement(InBox, [this, &Func](uint32 ValueIndex)
+		TOverrideVoidReturnInvoker Invoker(true, InFunc);
+
+		SpatialIndex.ForEachIntersectingElement(InBox, [this, &Invoker](uint32 ValueIndex)
 		{
-			Func(Elements[ValueIndex].Value);
+			return Invoker(Elements[ValueIndex].Value);
 		});
 	}
 
-	void ForEachIntersectingElement(const FSphere& InSphere, TFunctionRef<void(const ValueType& InValue)> Func) const
+	template <class Func>
+	void ForEachIntersectingElement(const FSphere& InSphere, Func InFunc) const
 	{
-		SpatialIndex.ForEachIntersectingElement(InSphere, [this, &Func](uint32 ValueIndex)
+		TOverrideVoidReturnInvoker Invoker(true, InFunc);
+
+		SpatialIndex.ForEachIntersectingElement(InSphere, [this, &Invoker](uint32 ValueIndex)
 		{
-			Func(Elements[ValueIndex].Value);
+			return Invoker(Elements[ValueIndex].Value);
 		});
 	}
 
@@ -108,9 +118,9 @@ namespace FStaticSpatialIndex
 
 		void Init(const TArray<TPair<FBox, uint32>>& InElements);
 
-		void ForEachElement(TFunctionRef<void(uint32 InValueIndex)> Func) const;
-		void ForEachIntersectingElement(const FBox& InBox, TFunctionRef<void(uint32 InValueIndex)> Func) const;
-		void ForEachIntersectingElement(const FSphere& InSphere, TFunctionRef<void(uint32 InValueIndex)> Func) const;
+		bool ForEachElement(TFunctionRef<bool(uint32 InValueIndex)> Func) const;
+		bool ForEachIntersectingElement(const FBox& InBox, TFunctionRef<bool(uint32 InValueIndex)> Func) const;
+		bool ForEachIntersectingElement(const FSphere& InSphere, TFunctionRef<bool(uint32 InValueIndex)> Func) const;
 
 	private:
 		TArray<uint32> Elements;
@@ -125,9 +135,9 @@ namespace FStaticSpatialIndex
 
 		void Init(const TArray<TPair<FBox, uint32>>& InElements);
 
-		void ForEachElement(TFunctionRef<void(uint32 InValueIndex)> Func) const;
-		void ForEachIntersectingElement(const FBox& InBox, TFunctionRef<void(uint32 InValueIndex)> Func) const;
-		void ForEachIntersectingElement(const FSphere& InSphere, TFunctionRef<void(uint32 InValueIndex)> Func) const;
+		bool ForEachElement(TFunctionRef<bool(uint32 InValueIndex)> Func) const;
+		bool ForEachIntersectingElement(const FBox& InBox, TFunctionRef<bool(uint32 InValueIndex)> Func) const;
+		bool ForEachIntersectingElement(const FSphere& InSphere, TFunctionRef<bool(uint32 InValueIndex)> Func) const;
 
 	private:
 		struct FNode
@@ -139,9 +149,9 @@ namespace FStaticSpatialIndex
 			TVariant<FNodeType, FLeafType> Content;
 		};
 
-		void ForEachElementRecursive(const FNode* Node, TFunctionRef<void(uint32 InValueIndex)> Func) const;
-		void ForEachIntersectingElementRecursive(const FNode* Node, const FBox& InBox, TFunctionRef<void(uint32 InValueIndex)> Func) const;
-		void ForEachIntersectingElementRecursive(const FNode* Node, const FSphere& InSphere, TFunctionRef<void(uint32 InValueIndex)> Func) const;
+		bool ForEachElementRecursive(const FNode* Node, TFunctionRef<bool(uint32 InValueIndex)> Func) const;
+		bool ForEachIntersectingElementRecursive(const FNode* Node, const FBox& InBox, TFunctionRef<bool(uint32 InValueIndex)> Func) const;
+		bool ForEachIntersectingElementRecursive(const FNode* Node, const FSphere& InSphere, TFunctionRef<bool(uint32 InValueIndex)> Func) const;
 
 		FNode RootNode;
 	};
diff --git a/Engine/Source/Runtime/Engine/Public/WorldPartition/RuntimeHashSet/WorldPartitionRuntimeHashSet.h b/Engine/Source/Runtime/Engine/Public/WorldPartition/RuntimeHashSet/WorldPartitionRuntimeHashSet.h
--- a/Engine/Source/Runtime/Engine/Public/WorldPartition/RuntimeHashSet/WorldPartitionRuntimeHashSet.h
+++ b/Engine/Source/Runtime/Engine/Public/WorldPartition/RuntimeHashSet/WorldPartitionRuntimeHashSet.h
@@ -175,7 +175,7 @@ private:
 	FCellUniqueId GetCellUniqueId(const URuntimePartition::FCellDescInstance& InCellDescInstance) const;
 #endif
 
-	ENGINE_API void ForEachStreamingData(TFunctionRef<void(const FRuntimePartitionStreamingData&)> Func) const;
+	ENGINE_API void ForEachStreamingData(TFunctionRef<bool(const FRuntimePartitionStreamingData&)> Func) const;
 
 public:
 #if WITH_EDITORONLY_DATA
