Fix static meshes using WPO missing from ray tracing scene when WPO evaluation is toggled using UStaticMeshComponent::SetEvaluateWorldPositionOffsetInRayTracing(...)
- The proxy was left with incorrect raytracing flags.
- Fix: Update cached raytracing state when WPO flag changes.

diff --git a/Engine/Source/Runtime/Engine/Private/StaticMeshRender.cpp b/Engine/Source/Runtime/Engine/Private/StaticMeshRender.cpp
--- a/Engine/Source/Runtime/Engine/Private/StaticMeshRender.cpp
+++ b/Engine/Source/Runtime/Engine/Private/StaticMeshRender.cpp
@@ -467,6 +467,8 @@ void FStaticMeshSceneProxy::SetEvaluateWorldPositionOffsetInRayTracing(FRHIComma
 
 		ReleaseDynamicRayTracingGeometries();
 	}
+
+	GetScene().UpdateCachedRayTracingState(this);
 #endif
 }
 
