Use a console sink instead of a cvar delegate for responding to virtual texture pool cvar changes.
This avoids bug where changing device profile can call the delegate multiple times as it runs through the device profile hierachy.
In the worst/common case we didn't need to call it at all, because the resulting value didn't change.
The sink is called only at the end of the device profile change, so it avoids the problem.

diff --git a/Engine/Source/Runtime/Engine/Private/VT/VirtualTexturePoolConfig.cpp b/Engine/Source/Runtime/Engine/Private/VT/VirtualTexturePoolConfig.cpp
--- a/Engine/Source/Runtime/Engine/Private/VT/VirtualTexturePoolConfig.cpp
+++ b/Engine/Source/Runtime/Engine/Private/VT/VirtualTexturePoolConfig.cpp
@@ -22,7 +22,6 @@ static TAutoConsoleVariable<float> CVarVTPoolSizeScale(
 	TEXT("r.VT.PoolSizeScale"),
 	1.0f,
 	TEXT("Scale factor for virtual texture physical pool size.\n"),
-	FConsoleVariableDelegate::CreateLambda([](IConsoleVariable* InVariable) { OnVirtualTexturePoolConfigUpdate(); }),
 	ECVF_RenderThreadSafe | ECVF_Scalability | ECVF_ExcludeFromPreview
 );
 
@@ -30,7 +29,6 @@ static TAutoConsoleVariable<bool> CVarVTPoolAutoGrow(
 	TEXT("r.VT.PoolAutoGrow"),
 	false,
 	TEXT("Enable physical pool growing on oversubscription."),
-	FConsoleVariableDelegate::CreateLambda([](IConsoleVariable* InVariable) { OnVirtualTexturePoolConfigUpdate(); }),
 	ECVF_RenderThreadSafe
 );
 
@@ -40,7 +38,6 @@ static TAutoConsoleVariable<int32> CVarVTSplitPhysicalPoolSize(
 	TEXT("Create multiple physical pools per format to keep pools at this maximum size in tiles.\n")
 	TEXT("A value of 64 tiles will force 16bit page tables. This can be a page table memory optimization for large physical pools.\n")
 	TEXT("Defaults to 0 (off)."),
-	FConsoleVariableDelegate::CreateLambda([](IConsoleVariable* InVariable) { OnVirtualTexturePoolConfigUpdate(); }),
 	ECVF_RenderThreadSafe
 );
 
@@ -65,6 +62,8 @@ void OnVirtualTexturePoolConfigUpdate()
 	}
 }
 
+FAutoConsoleVariableSink GVirtualTexturePoolConfigCVarSink(FConsoleCommandDelegate::CreateStatic(&OnVirtualTexturePoolConfigUpdate));
+
 
 UVirtualTexturePoolConfig::UVirtualTexturePoolConfig(const FObjectInitializer& ObjectInitializer)
 	: Super(ObjectInitializer)
