Shader Compiler:  Additional bounds checking for arrays in shader job cache, in an effort to isolate an apparent stomp of the job cache's lock, only on the Mac platform.

#rnx

diff --git a/Engine/Source/Runtime/Engine/Private/ShaderCompiler/ShaderCompiler.cpp b/Engine/Source/Runtime/Engine/Private/ShaderCompiler/ShaderCompiler.cpp
--- a/Engine/Source/Runtime/Engine/Private/ShaderCompiler/ShaderCompiler.cpp
+++ b/Engine/Source/Runtime/Engine/Private/ShaderCompiler/ShaderCompiler.cpp
@@ -438,7 +438,7 @@ public:
 
 	FORCEINLINE FShaderJobData& operator[](int32 Index)
 	{
-		check(Index < NumItems);
+		check((uint32)Index < (uint32)NumItems);
 		return DataBlocks[Index / FShaderJobDataBlock::BlockSize].Data[Index & (FShaderJobDataBlock::BlockSize - 1)];
 	}
 
@@ -667,11 +667,13 @@ private:
 
 	FORCEINLINE FLockStripeData& GetStripe(EShaderCompileJobType JobType, uint32 Hash)
 	{
-		return Stripes[(int32)JobType][Hash >> STRIPE_SHIFT];
+		checkf((uint8)JobType < (uint8)NumShaderCompileJobTypes, TEXT("Out of range JobType index %u"), (uint8)JobType);
+		return Stripes[(uint8)JobType][Hash >> STRIPE_SHIFT];
 	}
 	FORCEINLINE const FLockStripeData& GetStripe(EShaderCompileJobType JobType, uint32 Hash) const
 	{
-		return Stripes[(int32)JobType][Hash >> STRIPE_SHIFT];
+		checkf((uint8)JobType < (uint8)NumShaderCompileJobTypes, TEXT("Out of range JobType index %u"), (uint8)JobType);
+		return Stripes[(uint8)JobType][Hash >> STRIPE_SHIFT];
 	}
 };
 
@@ -912,10 +914,10 @@ private:
 	FShaderCommonCompileJob* PendingSubmitJobTaskJobs = nullptr;
 
 	/** Queue of tasks that haven't been assigned to a worker yet. */
-	FShaderCommonCompileJob* PendingJobsHead[NumShaderCompileJobPriorities];
-	std::atomic_int32_t NumPendingJobs[NumShaderCompileJobPriorities];
+	TStaticArray<FShaderCommonCompileJob*, NumShaderCompileJobPriorities> PendingJobsHead;
+	TStaticArray<std::atomic_int32_t, NumShaderCompileJobPriorities> NumPendingJobs;
 #if UE_SHADERCOMPILER_FIFO_JOB_EXECUTION
-	FShaderCommonCompileJob** PendingJobsTail[NumShaderCompileJobPriorities];
+	TStaticArray<FShaderCommonCompileJob**, NumShaderCompileJobPriorities> PendingJobsTail;
 #endif
 
 	/** Number of jobs currently being compiled.  This includes PendingJobs and any jobs that have been assigned to workers but aren't complete yet. */
