Fix gray textures looking red when used as UserInterface2D
we can't output G8 platform format unless TC is set to Grayscale because that's how the shader knows to replicate R to RGB

diff --git a/Engine/Source/Runtime/Engine/Private/Texture.cpp b/Engine/Source/Runtime/Engine/Private/Texture.cpp
--- a/Engine/Source/Runtime/Engine/Private/Texture.cpp
+++ b/Engine/Source/Runtime/Engine/Private/Texture.cpp
@@ -3755,12 +3755,18 @@ FName GetDefaultTextureFormatName( const ITargetPlatform* TargetPlatform, const
 			// @todo Oodle : consider removing HasHDRSource ; user did not pick an HDR TC output format
 			TextureFormatName = NameRGBA16F;
 		}
-		else if (SourceFormat == TSF_G16)
+		else if (SourceFormat == TSF_G16 &&
+			(FormatSettings.CompressionSettings == TC_Grayscale || FormatSettings.CompressionSettings == TC_Displacementmap) )
 		{
+			// we can only output one channel formats if TC tells the shader to replicate R to RGB
+			//	see ShouldUseGreyScaleEditorVisualization and GetSamplerTypeForTexture
 			TextureFormatName = NameG16;
 		}
-		else if (SourceFormat == TSF_G8)
+		else if (SourceFormat == TSF_G8 &&
+			(FormatSettings.CompressionSettings == TC_Grayscale || FormatSettings.CompressionSettings == TC_Displacementmap) )
 		{
+			// we can only output one channel formats if TC tells the shader to replicate R to RGB
+			//	see ShouldUseGreyScaleEditorVisualization and GetSamplerTypeForTexture
 			TextureFormatName = NameG8;
 		}
 		else
