Moved LocalBlendCSBoneTransforms to use a TArrayView

Allows other allocators and C-style arrays to be passed to this function

diff --git a/Engine/Source/Runtime/Engine/Public/BonePose.h b/Engine/Source/Runtime/Engine/Public/BonePose.h
--- a/Engine/Source/Runtime/Engine/Public/BonePose.h
+++ b/Engine/Source/Runtime/Engine/Public/BonePose.h
@@ -483,7 +483,7 @@ struct FCSPose
 	* Do this safely by insuring that Parents are already in Component Space,
 	* and any Component Space children are converted back to Local Space before hand.
 	*/
-	void SafeSetCSBoneTransforms(const TArray<struct FBoneTransform> & BoneTransforms);
+	void SafeSetCSBoneTransforms(TConstArrayView<FBoneTransform> BoneTransforms);
 
 	/**
 	* Blends Component Space transforms to MeshPose in Local Space.
@@ -493,7 +493,7 @@ struct FCSPose
 	* But the blending is done in Local Space. Also we need to refresh any Children they have
 	* that has been previously converted to Component Space.
 	*/
-	void LocalBlendCSBoneTransforms(const TArray<struct FBoneTransform>& BoneTransforms, float Alpha);
+	void LocalBlendCSBoneTransforms(TConstArrayView<FBoneTransform> BoneTransforms, float Alpha);
 
 	/** This function convert component space to local space to OutPose 
 	 *
@@ -669,7 +669,7 @@ void FCSPose<PoseType>::ConvertBoneToLocalSpace(BoneIndexType BoneIndex)
 }
 
 template<class PoseType>
-void FCSPose<PoseType>::SafeSetCSBoneTransforms(const TArray<struct FBoneTransform> & BoneTransforms)
+void FCSPose<PoseType>::SafeSetCSBoneTransforms(TConstArrayView<FBoneTransform> BoneTransforms)
 {
 	checkSlow(Pose.IsValid());
 
@@ -734,7 +734,7 @@ void FCSPose<PoseType>::SafeSetCSBoneTransforms(const TArray<struct FBoneTransfo
 }
 
 template<class PoseType>
-void FCSPose<PoseType>::LocalBlendCSBoneTransforms(const TArray<struct FBoneTransform>& BoneTransforms, float Alpha)
+void FCSPose<PoseType>::LocalBlendCSBoneTransforms(TConstArrayView<FBoneTransform> BoneTransforms, float Alpha)
 {
 	// if Alpha is small enough, skip
 	if (Alpha < ZERO_ANIMWEIGHT_THRESH)
