Don't force disable bAffectDynamicIndirectLighting during PSO precaching on skeletal meshes

diff --git a/Engine/Source/Runtime/Engine/Private/Components/SkinnedMeshComponent.cpp b/Engine/Source/Runtime/Engine/Private/Components/SkinnedMeshComponent.cpp
--- a/Engine/Source/Runtime/Engine/Private/Components/SkinnedMeshComponent.cpp
+++ b/Engine/Source/Runtime/Engine/Private/Components/SkinnedMeshComponent.cpp
@@ -609,7 +609,9 @@ void USkinnedMeshComponent::CollectPSOPrecacheData(const FPSOPrecacheParams& Bas
 
 	FPSOPrecacheParams PrecachePSOParams = BasePrecachePSOParams;
 	PrecachePSOParams.bCastShadow = PrecachePSOParams.bCastShadow && bAnySectionCastsShadows;
-	PrecachePSOParams.bAffectDynamicIndirectLighting = false;
+
+	// Skinned assets shouldn't need dynamic indirect lighting but MDCs for LumenCardCapture can still be setup and created (but not actually used) causing PSO precache misses
+	//PrecachePSOParams.bAffectDynamicIndirectLighting = false;
 
 	for (FPSOPrecacheVertexFactoryDataPerMaterialIndex& VFsPerMaterial : VFsPerMaterials)
 	{
