World Partition: added FWorldPartitionActorDescInstanceView that contains an actor descriptor instance pointer and route all calls to it.
This is to avoid having to implement all virtual functions on every derivations of IWorldPartitionActorDescInstanceView.

diff --git a/Engine/Source/Runtime/Engine/Private/WorldPartition/WorldPartitionStreamingGeneration.cpp b/Engine/Source/Runtime/Engine/Private/WorldPartition/WorldPartitionStreamingGeneration.cpp
--- a/Engine/Source/Runtime/Engine/Private/WorldPartition/WorldPartitionStreamingGeneration.cpp
+++ b/Engine/Source/Runtime/Engine/Private/WorldPartition/WorldPartitionStreamingGeneration.cpp
@@ -57,13 +57,13 @@ class FGCUnsavedDirtyActorContainerInstances : public FGCObject
 {
 	friend class FStreamingGenerationUnsavedDirtyActorDescInstance;
 public:
-	//~ Begin FGCObject
+	//~ Begin FGCObject interface
 	virtual void AddReferencedObjects(FReferenceCollector& Collector)
 	{
 		Collector.AddReferencedObjects(ContainerInstances);
 	}
 	virtual FString GetReferencerName() const { return TEXT("FUnsavedDirtyActorContainer"); }
-	//~ End
+	//~ End FGCObject interface
 
 protected:
 	TSet<TObjectPtr<UActorDescContainerInstance>> ContainerInstances;
@@ -104,7 +104,7 @@ public:
 		return MakeUnique<FStreamingGenerationUnsavedDirtyActorDescInstance>(HandlingContainer, InActor);
 	}
 
-	//~ Begin FWorldPartitionActorDescInstance	
+	//~ Begin FWorldPartitionActorDescInstance interface
 	virtual void RegisterChildContainerInstance() override
 	{
 		if (!UnsavedDirtyActorContainerInstances.IsValid())
@@ -129,9 +129,9 @@ public:
 		ChildContainerInstance->Uninitialize();
 		ChildContainerInstance = nullptr;
 	}
-	//~ End
-protected:
+	//~ End FWorldPartitionActorDescInstance interface
 
+protected:
 	TUniquePtr<FWorldPartitionActorDesc> ActorDescPtr;
 	static TUniquePtr<FGCUnsavedDirtyActorContainerInstances> UnsavedDirtyActorContainerInstances;
 };
@@ -173,7 +173,7 @@ FName FStreamingGenerationActorDescView::GetRuntimeGrid() const
 		return ParentView->GetRuntimeGrid();
 	}
 
-	return ActorDescInstance->GetRuntimeGrid();
+	return Super::GetRuntimeGrid();
 }
 
 bool FStreamingGenerationActorDescView::GetIsSpatiallyLoaded() const
@@ -183,7 +183,7 @@ bool FStreamingGenerationActorDescView::GetIsSpatiallyLoaded() const
 		return false;
 	}
 
-	bool bIsSpatiallyLoaded = ActorDescInstance->GetIsSpatiallyLoaded();
+	bool bIsSpatiallyLoaded = Super::GetIsSpatiallyLoaded();
 	if (bIsSpatiallyLoaded && ParentView)
 	{
 		bIsSpatiallyLoaded = ParentView->GetIsSpatiallyLoaded();
@@ -204,7 +204,7 @@ FSoftObjectPath FStreamingGenerationActorDescView::GetHLODLayer() const
 		return RuntimedHLODLayer.GetValue();
 	}
 
-	return ActorDescInstance->GetHLODLayer();
+	return Super::GetHLODLayer();
 }
 
 const TArray<FName>& FStreamingGenerationActorDescView::GetDataLayerInstanceNames() const
@@ -225,12 +225,12 @@ const TArray<FName>& FStreamingGenerationActorDescView::GetDataLayerInstanceName
 		return ResolvedDataLayerInstanceNames.GetValue();
 	}
 
-	return ActorDescInstance->GetDataLayerInstanceNames();
+	return Super::GetDataLayerInstanceNames();
 }
 
 const TArray<FGuid>& FStreamingGenerationActorDescView::GetReferences() const
 {
-	return RuntimeReferences.IsSet() ? RuntimeReferences.GetValue() : ActorDescInstance->GetReferences();
+	return RuntimeReferences.IsSet() ? RuntimeReferences.GetValue() : Super::GetReferences();
 }
 
 const TArray<FGuid>& FStreamingGenerationActorDescView::GetEditorReferences() const
@@ -247,7 +247,7 @@ void FStreamingGenerationActorDescView::SetParentView(const FStreamingGeneration
 
 void FStreamingGenerationActorDescView::SetDataLayerInstanceNames(const TArray<FName>& InDataLayerInstanceNames)
 {
-	check(!ActorDescInstance->HasResolvedDataLayerInstanceNames());
+	check(!Super::HasResolvedDataLayerInstanceNames());
 	ResolvedDataLayerInstanceNames = InDataLayerInstanceNames;
 }
 
diff --git a/Engine/Source/Runtime/Engine/Public/WorldPartition/WorldPartitionActorDescInstanceView.h b/Engine/Source/Runtime/Engine/Public/WorldPartition/WorldPartitionActorDescInstanceView.h
new file mode 100644
--- /dev/null
+++ b/Engine/Source/Runtime/Engine/Public/WorldPartition/WorldPartitionActorDescInstanceView.h
@@ -0,0 +1,88 @@
+// Copyright Epic Games, Inc. All Rights Reserved.
+#pragma once
+
+#if WITH_EDITOR
+#include "CoreMinimal.h"
+#include "WorldPartition/WorldPartitionActorDescInstance.h"
+#include "WorldPartition/WorldPartitionActorDescInstanceViewInterface.h"
+
+/**
+ * Base class for a view on top of an actor descriptor instance pointer.
+ */
+class FWorldPartitionActorDescInstanceView : public IWorldPartitionActorDescInstanceView
+{
+public:
+	FWorldPartitionActorDescInstanceView(const FWorldPartitionActorDescInstance* InActorDescInstance)
+		: ActorDescInstance(InActorDescInstance)
+	{}
+		
+	//~ Begin IWorldPartitionActorDescInstanceView interface
+	virtual const FGuid& GetGuid() const override { return ActorDescInstance->GetGuid(); }
+
+	virtual FTopLevelAssetPath GetBaseClass() const override { return ActorDescInstance->GetBaseClass(); }
+	virtual FTopLevelAssetPath GetNativeClass() const override { return ActorDescInstance->GetNativeClass(); }
+	virtual UClass* GetActorNativeClass() const override { return ActorDescInstance->GetActorNativeClass(); }
+
+	virtual FName GetRuntimeGrid() const override { return ActorDescInstance->GetRuntimeGrid(); }
+	virtual bool GetIsSpatiallyLoaded() const override { return ActorDescInstance->GetIsSpatiallyLoaded(); }
+	virtual bool GetActorIsEditorOnly() const override { return ActorDescInstance->GetActorIsEditorOnly(); }
+	virtual bool GetActorIsRuntimeOnly() const override { return ActorDescInstance->GetActorIsRuntimeOnly(); }
+	virtual bool IsRuntimeRelevant() const override { return ActorDescInstance->IsRuntimeRelevant(); }
+	virtual bool IsEditorRelevant() const override { return ActorDescInstance->IsEditorRelevant(); }
+
+	virtual bool IsUsingDataLayerAsset() const override { return ActorDescInstance->IsUsingDataLayerAsset(); }
+	virtual const TArray<FName>& GetDataLayers() const override { return ActorDescInstance->GetDataLayers(); }
+
+	virtual bool GetActorIsHLODRelevant() const override { return ActorDescInstance->GetActorIsHLODRelevant(); }
+	virtual FSoftObjectPath GetHLODLayer() const override { return ActorDescInstance->GetHLODLayer(); }
+
+	virtual const TArray<FName>& GetTags() const override { return ActorDescInstance->GetTags(); }
+	virtual FName GetActorPackage() const override { return ActorDescInstance->GetActorPackage(); }
+	virtual FSoftObjectPath GetActorSoftPath() const override { return ActorDescInstance->GetActorSoftPath(); }
+	virtual FName GetActorLabel() const override { return ActorDescInstance->GetActorLabel(); }
+	virtual FName GetActorName() const override { return ActorDescInstance->GetActorName(); }
+	virtual FName GetFolderPath() const override { return ActorDescInstance->GetFolderPath(); }
+	virtual const FGuid& GetFolderGuid() const override { return ActorDescInstance->GetFolderGuid(); }
+
+	virtual FBox GetEditorBounds() const override { return ActorDescInstance->GetEditorBounds(); }
+	virtual FBox GetRuntimeBounds() const override { return ActorDescInstance->GetRuntimeBounds(); }
+
+	virtual bool GetProperty(FName PropertyName, FName* PropertyValue) const override { return ActorDescInstance->GetProperty(PropertyName, PropertyValue); }
+	virtual bool HasProperty(FName PropertyName) const override { return ActorDescInstance->HasProperty(PropertyName); }
+
+	virtual const TArray<FGuid>& GetReferences() const override { return ActorDescInstance->GetReferences(); }
+	virtual const TArray<FGuid>& GetEditorOnlyReferences() const override { return ActorDescInstance->GetEditorOnlyReferences(); }
+	virtual bool IsEditorOnlyReference(const FGuid& ReferenceGuid) const override { return ActorDescInstance->IsEditorOnlyReference(ReferenceGuid); }
+
+	virtual const FGuid& GetParentActor() const override { return ActorDescInstance->GetParentActor(); }
+
+	virtual FGuid GetContentBundleGuid() const override { return ActorDescInstance->GetContentBundleGuid(); }
+
+	virtual bool IsChildContainerInstance() const override { return ActorDescInstance->IsChildContainerInstance(); }
+	virtual FName GetChildContainerPackage() const override { return ActorDescInstance->GetChildContainerPackage(); }
+	virtual EWorldPartitionActorFilterType GetChildContainerFilterType() const override { return ActorDescInstance->GetChildContainerFilterType(); }
+	virtual const FWorldPartitionActorFilter* GetChildContainerFilter() const override { return ActorDescInstance->GetChildContainerFilter(); }
+	virtual bool GetChildContainerInstance(FWorldPartitionActorDesc::FContainerInstance& OutContainerInstance) const override { return ActorDescInstance->GetChildContainerInstance(OutContainerInstance); }
+
+	virtual bool IsMainWorldOnly() const override { return ActorDescInstance->IsMainWorldOnly(); }
+	virtual bool IsListedInSceneOutliner() const override { return ActorDescInstance->IsListedInSceneOutliner(); }
+	virtual const FGuid& GetSceneOutlinerParent() const override { return ActorDescInstance->GetSceneOutlinerParent(); }
+
+	virtual void CheckForErrors(IStreamingGenerationErrorHandler* ErrorHandler) const override { GetActorDesc()->CheckForErrors(this, ErrorHandler); }
+
+	virtual FString ToString(FWorldPartitionActorDesc::EToStringMode Mode) const override { return ActorDescInstance->ToString(Mode); }
+	virtual const FWorldPartitionActorDesc* GetActorDesc() const override { return ActorDescInstance->GetActorDesc(); }
+
+	virtual bool HasResolvedDataLayerInstanceNames() const override { return ActorDescInstance->HasResolvedDataLayerInstanceNames(); }
+	const TArray<FName>& GetDataLayerInstanceNames() const override { return ActorDescInstance->GetDataLayerInstanceNames(); }
+
+	virtual AActor* GetActor(bool bEvenIfPendingKill = true, bool bEvenIfUnreachable = false) const override { return ActorDescInstance->GetActor(bEvenIfPendingKill, bEvenIfUnreachable); }
+	virtual bool IsLoaded(bool bEvenIfPendingKill = false) const override { return ActorDescInstance->IsLoaded(bEvenIfPendingKill); }
+
+	virtual UActorDescContainerInstance* GetContainerInstance() const override { return ActorDescInstance->GetContainerInstance(); }
+	//~ End IWorldPartitionActorDescInstanceView interface
+
+private:
+	const FWorldPartitionActorDescInstance* ActorDescInstance;
+};
+#endif // WITH_EDITOR
\ No newline at end of file
diff --git a/Engine/Source/Runtime/Engine/Public/WorldPartition/WorldPartitionActorDescInstanceViewInterface.h b/Engine/Source/Runtime/Engine/Public/WorldPartition/WorldPartitionActorDescInstanceViewInterface.h
--- a/Engine/Source/Runtime/Engine/Public/WorldPartition/WorldPartitionActorDescInstanceViewInterface.h
+++ b/Engine/Source/Runtime/Engine/Public/WorldPartition/WorldPartitionActorDescInstanceViewInterface.h
@@ -7,19 +7,18 @@
 
 class AActor;
 class IStreamingGenerationErrorHandler;
-class UWorldPartition;
 struct FWorldPartitionActorFilter;
 enum class EWorldPartitionActorFilterType : uint8;
-#endif // WITH_EDITOR
 
 /**
- * Interface for a view on top of an actor desc, used to cache information that can be (potentially) different than the actor desc
- * itself due to streaming generation logic, etc.
+ * Interface for a view on top of an actor descriptor instance, used to cache information that can be different 
+ * than the actor descriptor instance itself.
  */
 class IWorldPartitionActorDescInstanceView
 {
-#if WITH_EDITOR
 public:
+	virtual ~IWorldPartitionActorDescInstanceView() {}
+
 	virtual const FGuid& GetGuid() const = 0;
 
 	virtual FTopLevelAssetPath GetBaseClass() const = 0;
@@ -85,6 +84,5 @@ public:
 	virtual UActorDescContainerInstance* GetContainerInstance() const = 0;
 
 	virtual void CheckForErrors(IStreamingGenerationErrorHandler* ErrorHandler) const = 0;
-#endif // WITH_EDITOR
 };
-
+#endif // WITH_EDITOR
\ No newline at end of file
diff --git a/Engine/Source/Runtime/Engine/Public/WorldPartition/WorldPartitionStreamingGeneration.h b/Engine/Source/Runtime/Engine/Public/WorldPartition/WorldPartitionStreamingGeneration.h
--- a/Engine/Source/Runtime/Engine/Public/WorldPartition/WorldPartitionStreamingGeneration.h
+++ b/Engine/Source/Runtime/Engine/Public/WorldPartition/WorldPartitionStreamingGeneration.h
@@ -7,31 +7,32 @@
 #include "WorldPartition/ActorDescContainerCollection.h"
 #include "WorldPartition/ActorDescContainerInstanceCollection.h"
 #include "WorldPartition/WorldPartitionActorDescInstance.h"
-#include "WorldPartition/WorldPartitionActorDescInstanceViewInterface.h"
+#include "WorldPartition/WorldPartitionActorDescInstanceView.h"
 
 // deprecated
 class FWorldPartitionActorDescView;
 
 class FStreamingGenerationActorDescViewMap;
 
-class FStreamingGenerationActorDescView : public IWorldPartitionActorDescInstanceView
+class FStreamingGenerationActorDescView : public FWorldPartitionActorDescInstanceView
 {
+	typedef FWorldPartitionActorDescInstanceView Super;
+
 	friend class FStreamingGenerationActorDescViewMap;
 	friend class FWorldPartitionStreamingGenerator;
 
 protected:
 	// Used for invalid reference error reporting
 	FStreamingGenerationActorDescView(const FWorldPartitionActorDescInstance* InActorDescInstance)
-		: ActorDescViewMap(nullptr)
-		, ActorDescInstance(InActorDescInstance)
+		: FWorldPartitionActorDescInstanceView(InActorDescInstance)
+		, ActorDescViewMap(nullptr)
 		, ParentView(nullptr)
 		, bIsForcedNonSpatiallyLoaded(false)
 		, bIsForcedNoRuntimeGrid(false)
 		, bIsForcedNoDataLayers(false)
 		, bIsForceNoHLODLayer(false)
 		, bIsUnsaved(false)
-	{
-	}
+	{}
 
 public:
 	FStreamingGenerationActorDescView(const FStreamingGenerationActorDescViewMap& InActorDescViewMap, const FWorldPartitionActorDescInstance* InActorDescInstance, bool bInUnsaved = false)
@@ -43,71 +44,13 @@ public:
 		
 	virtual ~FStreamingGenerationActorDescView() {}
 
-	//~ Begin FWorldPartitionActorDescViewProxy	
-	virtual const FGuid& GetGuid() const override { return ActorDescInstance->GetGuid(); }
-
-	virtual FTopLevelAssetPath GetBaseClass() const override { return ActorDescInstance->GetBaseClass(); }
-	virtual FTopLevelAssetPath GetNativeClass() const override { return ActorDescInstance->GetNativeClass(); }
-	virtual UClass* GetActorNativeClass() const override { return ActorDescInstance->GetActorNativeClass(); }
-
+	//~ Begin FWorldPartitionActorDescInstanceView interface
 	ENGINE_API virtual FName GetRuntimeGrid() const override;
 	ENGINE_API virtual bool GetIsSpatiallyLoaded() const override;
-	virtual bool GetActorIsEditorOnly() const override { return ActorDescInstance->GetActorIsEditorOnly(); }
-	virtual bool GetActorIsRuntimeOnly() const override { return ActorDescInstance->GetActorIsRuntimeOnly(); }
-	virtual bool IsRuntimeRelevant() const override { return ActorDescInstance->IsRuntimeRelevant(); }
-	virtual bool IsEditorRelevant() const override { return ActorDescInstance->IsEditorRelevant(); }
-
-	virtual bool IsUsingDataLayerAsset() const override { return ActorDescInstance->IsUsingDataLayerAsset(); }
-	virtual const TArray<FName>& GetDataLayers() const override { return ActorDescInstance->GetDataLayers(); }
-
-	virtual bool GetActorIsHLODRelevant() const override { return ActorDescInstance->GetActorIsHLODRelevant(); }
 	ENGINE_API virtual FSoftObjectPath GetHLODLayer() const override;
-
-	virtual const TArray<FName>& GetTags() const override { return ActorDescInstance->GetTags(); }
-	virtual FName GetActorPackage() const override { return ActorDescInstance->GetActorPackage(); }
-	virtual FSoftObjectPath GetActorSoftPath() const override { return ActorDescInstance->GetActorSoftPath(); }
-	virtual FName GetActorLabel() const override { return ActorDescInstance->GetActorLabel(); }
-	virtual FName GetActorName() const override { return ActorDescInstance->GetActorName(); }
-	virtual FName GetFolderPath() const override { return ActorDescInstance->GetFolderPath(); }
-	virtual const FGuid& GetFolderGuid() const override { return ActorDescInstance->GetFolderGuid(); }
-
-	virtual FBox GetEditorBounds() const override { return ActorDescInstance->GetEditorBounds(); }
-	virtual FBox GetRuntimeBounds() const override { return ActorDescInstance->GetRuntimeBounds(); }
-
-	virtual bool GetProperty(FName PropertyName, FName* PropertyValue) const override { return ActorDescInstance->GetProperty(PropertyName, PropertyValue); }
-	virtual bool HasProperty(FName PropertyName) const override { return ActorDescInstance->HasProperty(PropertyName); }
-
 	ENGINE_API virtual const TArray<FGuid>& GetReferences() const override;
-	virtual const TArray<FGuid>& GetEditorOnlyReferences() const override { return ActorDescInstance->GetEditorOnlyReferences(); }
-	virtual bool IsEditorOnlyReference(const FGuid& ReferenceGuid) const override { return ActorDescInstance->IsEditorOnlyReference(ReferenceGuid); }
-
-	virtual const FGuid& GetParentActor() const override { return ActorDescInstance->GetParentActor(); }
-
-	virtual FGuid GetContentBundleGuid() const override { return ActorDescInstance->GetContentBundleGuid(); }
-
-	virtual bool IsChildContainerInstance() const override { return ActorDescInstance->IsChildContainerInstance(); }
-	virtual FName GetChildContainerPackage() const override { return ActorDescInstance->GetChildContainerPackage(); }
-	virtual EWorldPartitionActorFilterType GetChildContainerFilterType() const override { return ActorDescInstance->GetChildContainerFilterType(); }
-	virtual const FWorldPartitionActorFilter* GetChildContainerFilter() const override { return ActorDescInstance->GetChildContainerFilter(); }
-	virtual bool GetChildContainerInstance(FWorldPartitionActorDesc::FContainerInstance& OutContainerInstance) const override { return ActorDescInstance->GetChildContainerInstance(OutContainerInstance); }
-
-	virtual bool IsMainWorldOnly() const override { return ActorDescInstance->IsMainWorldOnly(); }
-	virtual bool IsListedInSceneOutliner() const override { return ActorDescInstance->IsListedInSceneOutliner(); }
-	virtual const FGuid& GetSceneOutlinerParent() const override { return ActorDescInstance->GetSceneOutlinerParent(); }
-
-	virtual void CheckForErrors(IStreamingGenerationErrorHandler* ErrorHandler) const override { GetActorDesc()->CheckForErrors(this, ErrorHandler); }
-
-	virtual FString ToString(FWorldPartitionActorDesc::EToStringMode Mode) const override { return ActorDescInstance->ToString(Mode); }
-	virtual const FWorldPartitionActorDesc* GetActorDesc() const override { return ActorDescInstance->GetActorDesc(); }
-
-	virtual bool HasResolvedDataLayerInstanceNames() const override { return ActorDescInstance->HasResolvedDataLayerInstanceNames(); }
 	ENGINE_API const TArray<FName>& GetDataLayerInstanceNames() const override;
-
-	virtual AActor* GetActor(bool bEvenIfPendingKill = true, bool bEvenIfUnreachable = false) const override { return ActorDescInstance->GetActor(bEvenIfPendingKill, bEvenIfUnreachable); }
-	virtual bool IsLoaded(bool bEvenIfPendingKill = false) const override { return ActorDescInstance->IsLoaded(bEvenIfPendingKill); }
-
-	virtual UActorDescContainerInstance* GetContainerInstance() const override { return ActorDescInstance->GetContainerInstance(); }
-	//~ End
+	//~ End FWorldPartitionActorDescInstanceView interface
 
 	bool IsUnsaved() const { return bIsUnsaved; }
 	const FStreamingGenerationActorDescViewMap& GetActorDescViewMap() const { return *ActorDescViewMap; }
@@ -139,8 +82,6 @@ public:
 
 private:
 	const FStreamingGenerationActorDescViewMap* ActorDescViewMap;
-	const FWorldPartitionActorDescInstance* ActorDescInstance;
-		
 	const FStreamingGenerationActorDescView* ParentView;
 
 	bool bIsForcedNonSpatiallyLoaded;
