- Precache PSOs for Niagara Components as well to make sure material and mesh overrides on the component are precached as well (not all the NiagaraAssets are precached anymore)
- Precache PSOs again on mesh or material override changes (delayed because overrides are also delayed applied)

diff --git a/Engine/Plugins/FX/Niagara/Source/Niagara/Private/NiagaraComponent.cpp b/Engine/Plugins/FX/Niagara/Source/Niagara/Private/NiagaraComponent.cpp
--- a/Engine/Plugins/FX/Niagara/Source/Niagara/Private/NiagaraComponent.cpp
+++ b/Engine/Plugins/FX/Niagara/Source/Niagara/Private/NiagaraComponent.cpp
@@ -566,6 +566,7 @@ UNiagaraComponent::UNiagaraComponent(const FObjectInitializer& ObjectInitializer
 	, bIsCulledByScalability(false)
 	, bDuringUpdateContextReset(false)
 	, bDesiredPauseState(false)
+	, bRecachePSOs(false)
 	//, bIsChangingAutoAttachment(false)
 	, ScalabilityManagerHandle(INDEX_NONE)
 	, ForceUpdateTransformTime(0.0f)
@@ -1159,6 +1160,8 @@ bool UNiagaraComponent::InitializeSystem()
 			}
 		}
 
+		PrecachePSOs();
+
 #if WITH_EDITORONLY_DATA
 		OnSystemInstanceChangedDelegate.Broadcast();
 #endif
@@ -1656,6 +1659,12 @@ void UNiagaraComponent::PostSystemTick_GameThread()
 			SystemInstanceController->PostTickRenderers(*RenderData);
 		}
 	}
+
+	if (bRecachePSOs)
+	{
+		PrecachePSOs();
+		bRecachePSOs = false;
+	}
 }
 
 void UNiagaraComponent::OnSystemComplete(bool bExternalCompletion)
@@ -2310,7 +2319,8 @@ FPrimitiveSceneProxy* UNiagaraComponent::CreateSceneProxy()
 #if UE_WITH_PSO_PRECACHING
 	if (Asset != nullptr)
 	{
-		// PSO request should have been handled by the Asset itself, so here we just ensure the request gets boosted
+		// If PSOs not precached yet then rely on PSOs reuqested by the shared asset itself - this will make sure
+		// those pending requests are boosted if still pending
 		if (!bPSOPrecacheCalled)
 		{
 			PrecacheAssetPSOs(Asset);
@@ -2329,6 +2339,26 @@ FPrimitiveSceneProxy* UNiagaraComponent::CreateSceneProxy()
 	return Proxy;
 }
 
+void UNiagaraComponent::CollectPSOPrecacheData(const FPSOPrecacheParams& BasePrecachePSOParams, FComponentPSOPrecacheParamsList& OutParams)
+{
+	if (SystemInstanceController.IsValid())
+	{
+		FMaterialInterfacePSOPrecacheParamsList PSOPrecacheParamsList;
+		SystemInstanceController->CollectPSOPrecacheData(BasePrecachePSOParams, PSOPrecacheParamsList);
+
+		// Translate to component PSO precache params - will be removed with upcoming change but make CL bigger than needed
+		for (FMaterialInterfacePSOPrecacheParams& PSOPrecacheParams : PSOPrecacheParamsList)
+		{
+			FComponentPSOPrecacheParams& NewEntry = OutParams.AddDefaulted_GetRef();
+			NewEntry.Priority = PSOPrecacheParams.Priority;
+			NewEntry.MaterialInterface = PSOPrecacheParams.MaterialInterface;
+			NewEntry.PSOPrecacheParams = PSOPrecacheParams.PSOPrecacheParams;
+			NewEntry.VertexFactoryDataList = PSOPrecacheParams.VertexFactoryDataList;
+			
+		}
+	}
+}
+
 void UNiagaraComponent::GetUsedMaterials(TArray<UMaterialInterface*>& OutMaterials, bool bGetDebugMaterials) const
 {
 	if (CullProxy)
@@ -2741,6 +2771,9 @@ void UNiagaraComponent::SetVariableMaterial(FName InVariableName, UMaterialInter
 	{
 		OverrideParameters.SetUObject(InValue, VariableDesc);
 	}
+
+	bRecachePSOs = true;
+
 #if WITH_EDITOR
 	SetParameterOverride(VariableDesc, FNiagaraVariant(InValue));
 #endif
@@ -2758,6 +2791,9 @@ void UNiagaraComponent::SetVariableStaticMesh(FName InVariableName, UStaticMesh*
 	{
 		OverrideParameters.SetUObject(InValue, VariableDesc);
 	}
+
+	bRecachePSOs = true;
+
 #if WITH_EDITOR
 	SetParameterOverride(VariableDesc, FNiagaraVariant(InValue));
 #endif
diff --git a/Engine/Plugins/FX/Niagara/Source/Niagara/Private/NiagaraMeshRendererProperties.cpp b/Engine/Plugins/FX/Niagara/Source/Niagara/Private/NiagaraMeshRendererProperties.cpp
--- a/Engine/Plugins/FX/Niagara/Source/Niagara/Private/NiagaraMeshRendererProperties.cpp
+++ b/Engine/Plugins/FX/Niagara/Source/Niagara/Private/NiagaraMeshRendererProperties.cpp
@@ -763,18 +763,16 @@ void UNiagaraMeshRendererProperties::GetUsedMaterials(const FNiagaraEmitterInsta
 	}
 }
 
-void UNiagaraMeshRendererProperties::CollectPSOPrecacheData(FPSOPrecacheParamsList& OutParams)
+void UNiagaraMeshRendererProperties::CollectPSOPrecacheData(const FNiagaraEmitterInstance* InEmitter, FPSOPrecacheParamsList& OutParams) const
 {
 	const FVertexFactoryType* VFType = GetVertexFactoryType();
 	bool bSupportsManualVertexFetch = VFType->SupportsManualVertexFetch(GMaxRHIFeatureLevel);
 		
 	for (int32 MeshIndex = 0; MeshIndex < Meshes.Num(); ++MeshIndex)
 	{
-		// Don't have an instance yet to retrieve the possible material override data from
-		FNiagaraEmitterInstance* EmitterInstance = nullptr;
 		INiagaraRenderableMeshInterface* RenderableMeshInterface = nullptr;
 		UStaticMesh* StaticMesh = nullptr;
-		NiagaraMeshRendererPropertiesInternal::ResolveRenderableMeshInternal(Meshes[MeshIndex], EmitterInstance, RenderableMeshInterface, StaticMesh);
+		NiagaraMeshRendererPropertiesInternal::ResolveRenderableMeshInternal(Meshes[MeshIndex], InEmitter, RenderableMeshInterface, StaticMesh);
 		if (StaticMesh)
 		{
 			TArray<UMaterialInterface*> OrderedMeshMaterials;
diff --git a/Engine/Plugins/FX/Niagara/Source/Niagara/Private/NiagaraRibbonRendererProperties.cpp b/Engine/Plugins/FX/Niagara/Source/Niagara/Private/NiagaraRibbonRendererProperties.cpp
--- a/Engine/Plugins/FX/Niagara/Source/Niagara/Private/NiagaraRibbonRendererProperties.cpp
+++ b/Engine/Plugins/FX/Niagara/Source/Niagara/Private/NiagaraRibbonRendererProperties.cpp
@@ -205,7 +205,7 @@ void UNiagaraRibbonRendererProperties::GetUsedMaterials(const FNiagaraEmitterIns
 	OutMaterials.Add(MaterialInterface ? MaterialInterface : ToRawPtr(Material));
 }
 
-void UNiagaraRibbonRendererProperties::CollectPSOPrecacheData(FPSOPrecacheParamsList& OutParams)
+void UNiagaraRibbonRendererProperties::CollectPSOPrecacheData(const FNiagaraEmitterInstance* InEmitter, FPSOPrecacheParamsList& OutParams) const
 {
 	const FVertexFactoryType* VFType = GetVertexFactoryType();
 	UMaterialInterface* MaterialInterface = ToRawPtr(Material);
diff --git a/Engine/Plugins/FX/Niagara/Source/Niagara/Private/NiagaraSpriteRendererProperties.cpp b/Engine/Plugins/FX/Niagara/Source/Niagara/Private/NiagaraSpriteRendererProperties.cpp
--- a/Engine/Plugins/FX/Niagara/Source/Niagara/Private/NiagaraSpriteRendererProperties.cpp
+++ b/Engine/Plugins/FX/Niagara/Source/Niagara/Private/NiagaraSpriteRendererProperties.cpp
@@ -131,7 +131,7 @@ void UNiagaraSpriteRendererProperties::GetUsedMaterials(const FNiagaraEmitterIns
 	OutMaterials.Add(MaterialInterface ? MaterialInterface : ToRawPtr(Material));
 }
 
-void UNiagaraSpriteRendererProperties::CollectPSOPrecacheData(FPSOPrecacheParamsList& OutParams)
+void UNiagaraSpriteRendererProperties::CollectPSOPrecacheData(const FNiagaraEmitterInstance* InEmitter, FPSOPrecacheParamsList& OutParams) const
 {
 	const FVertexFactoryType* VFType = GetVertexFactoryType();
 	UMaterialInterface* MaterialInterface = ToRawPtr(Material);
diff --git a/Engine/Plugins/FX/Niagara/Source/Niagara/Private/NiagaraSystem.cpp b/Engine/Plugins/FX/Niagara/Source/Niagara/Private/NiagaraSystem.cpp
--- a/Engine/Plugins/FX/Niagara/Source/Niagara/Private/NiagaraSystem.cpp
+++ b/Engine/Plugins/FX/Niagara/Source/Niagara/Private/NiagaraSystem.cpp
@@ -1173,7 +1173,11 @@ void UNiagaraSystem::PrecachePSOs()
 		return;
 	}
 
-	TArray<VFsPerMaterialData, TInlineAllocator<2>> VFsPerMaterials;
+	FMaterialInterfacePSOPrecacheParamsList MaterialInterfacePSOPrecacheParamsList;
+
+	FMaterialInterfacePSOPrecacheParams NewEntry;
+	NewEntry.PSOPrecacheParams.SetMobility(EComponentMobility::Movable);
+	NewEntry.PSOPrecacheParams.bRenderCustomDepth = true;
 
 	for (const FNiagaraEmitterHandle& EmitterHandle : GetEmitterHandles())
 	{
@@ -1188,29 +1192,28 @@ void UNiagaraSystem::PrecachePSOs()
 					{
 						return;
 					}
-					bool bDisableBackfaceCulling = Properties->IsBackfaceCullingDisabled();
+
+					NewEntry.PSOPrecacheParams.bDisableBackFaceCulling = Properties->IsBackfaceCullingDisabled();
+
+					// Don't have an instance yet to retrieve the possible material override data from
+					FNiagaraEmitterInstance* EmitterInstance = nullptr;
 
 					UNiagaraRendererProperties::FPSOPrecacheParamsList PSOPrecacheParamsList;
-					Properties->CollectPSOPrecacheData(PSOPrecacheParamsList);
+					Properties->CollectPSOPrecacheData(EmitterInstance, PSOPrecacheParamsList);
 
 					for (UNiagaraRendererProperties::FPSOPrecacheParams& PSOPrecacheParams: PSOPrecacheParamsList)
 					{
-						UMaterialInterface* MaterialInterface = PSOPrecacheParams.MaterialInterface;
-						VFsPerMaterialData* VFsPerMaterial = VFsPerMaterials.FindByPredicate([MaterialInterface, bDisableBackfaceCulling](const VFsPerMaterialData& Other) 
-						{ 
-							return (Other.MaterialInterface == MaterialInterface &&
-									Other.bDisableBackfaceCulling == bDisableBackfaceCulling);
-						});
-						if (VFsPerMaterial == nullptr)
-						{
-							VFsPerMaterial = &VFsPerMaterials.AddDefaulted_GetRef();
-							VFsPerMaterial->MaterialInterface = MaterialInterface;
-							VFsPerMaterial->bDisableBackfaceCulling = bDisableBackfaceCulling;
-						}
+						NewEntry.MaterialInterface = PSOPrecacheParams.MaterialInterface;
+						NewEntry.VertexFactoryDataList = PSOPrecacheParams.VertexFactoryDataList;
+
+						NewEntry.PSOPrecacheParams.bReverseCulling = false;
+						AddMaterialInterfacePSOPrecacheParamsToList(NewEntry, MaterialInterfacePSOPrecacheParamsList);
 
-						for (FPSOPrecacheVertexFactoryData& VFData : PSOPrecacheParams.VertexFactoryDataList)
+						// Also precache with reverse culling if not two sided because we don't know of the component using the asset will have negative determinant
+						if (!NewEntry.PSOPrecacheParams.bDisableBackFaceCulling)
 						{
-							VFsPerMaterial->VertexFactoryData.AddUnique(VFData);
+							NewEntry.PSOPrecacheParams.bReverseCulling = true;
+							AddMaterialInterfacePSOPrecacheParamsToList(NewEntry, MaterialInterfacePSOPrecacheParamsList);
 						}
 					}
 				}
@@ -1218,7 +1221,7 @@ void UNiagaraSystem::PrecachePSOs()
 		}
 	}
 
-	LaunchPSOPrecaching(VFsPerMaterials);
+	LaunchPSOPrecaching(MaterialInterfacePSOPrecacheParamsList);
 }
 
 
diff --git a/Engine/Plugins/FX/Niagara/Source/Niagara/Private/NiagaraSystemInstanceController.cpp b/Engine/Plugins/FX/Niagara/Source/Niagara/Private/NiagaraSystemInstanceController.cpp
--- a/Engine/Plugins/FX/Niagara/Source/Niagara/Private/NiagaraSystemInstanceController.cpp
+++ b/Engine/Plugins/FX/Niagara/Source/Niagara/Private/NiagaraSystemInstanceController.cpp
@@ -100,6 +100,53 @@ void FNiagaraSystemInstanceController::NotifyRenderersComplete(FNiagaraSystemRen
 	}
 }
 
+void FNiagaraSystemInstanceController::CollectPSOPrecacheData(const FPSOPrecacheParams& BasePrecachePSOParams, FMaterialInterfacePSOPrecacheParamsList& List) const
+{
+	if (!SystemInstance.IsValid())
+	{
+		return;
+	}
+
+	FMaterialInterfacePSOPrecacheParams NewEntry;
+	NewEntry.PSOPrecacheParams.SetMobility(EComponentMobility::Movable);
+	NewEntry.PSOPrecacheParams.bRenderCustomDepth = true;
+
+	for (const FNiagaraEmitterInstanceRef& EmitterInst : SystemInstance->GetEmitters())
+	{
+		EmitterInst->ForEachEnabledRenderer(
+			[&](const UNiagaraRendererProperties* Properties)
+			{	
+				const FVertexFactoryType* VFType = Properties->GetVertexFactoryType();
+				if (VFType == nullptr)
+				{
+					return;
+				}
+
+				NewEntry.PSOPrecacheParams.bDisableBackFaceCulling = Properties->IsBackfaceCullingDisabled();
+
+				UNiagaraRendererProperties::FPSOPrecacheParamsList PSOPrecacheParamsList;
+				Properties->CollectPSOPrecacheData(&EmitterInst.Get(), PSOPrecacheParamsList);
+
+				for (UNiagaraRendererProperties::FPSOPrecacheParams& PSOPrecacheParams : PSOPrecacheParamsList)
+				{
+					NewEntry.MaterialInterface = PSOPrecacheParams.MaterialInterface;
+					NewEntry.VertexFactoryDataList = PSOPrecacheParams.VertexFactoryDataList;
+
+					NewEntry.PSOPrecacheParams.bReverseCulling = false;
+					AddMaterialInterfacePSOPrecacheParamsToList(NewEntry, List);
+
+					// Also precache with reverse culling if not two sided because we don't know of the component using the asset will have negative determinant
+					if (!NewEntry.PSOPrecacheParams.bDisableBackFaceCulling)
+					{
+						NewEntry.PSOPrecacheParams.bReverseCulling = true;
+						AddMaterialInterfacePSOPrecacheParamsToList(NewEntry, List);
+					}
+				}
+			}
+		);
+	}
+}
+
 void FNiagaraSystemInstanceController::GetUsedMaterials(TArray<UMaterialInterface*>& OutMaterials)
 {
 	if (!SystemInstance.IsValid())
diff --git a/Engine/Plugins/FX/Niagara/Source/Niagara/Public/NiagaraComponent.h b/Engine/Plugins/FX/Niagara/Source/Niagara/Public/NiagaraComponent.h
--- a/Engine/Plugins/FX/Niagara/Source/Niagara/Public/NiagaraComponent.h
+++ b/Engine/Plugins/FX/Niagara/Source/Niagara/Public/NiagaraComponent.h
@@ -263,6 +263,7 @@ public:
 	NIAGARA_API virtual void GetUsedMaterials(TArray<UMaterialInterface*>& OutMaterials, bool bGetDebugMaterials = false) const override;
 	NIAGARA_API virtual void GetStreamingRenderAssetInfo(FStreamingTextureLevelContext& LevelContext, TArray<FStreamingRenderAssetPrimitiveInfo>& OutStreamingRenderAssets) const override;
 	NIAGARA_API virtual void OnAttachmentChanged() override;
+	NIAGARA_API virtual void CollectPSOPrecacheData(const FPSOPrecacheParams& BasePrecachePSOParams, FComponentPSOPrecacheParamsList& OutParams) override;
 	//~ End UPrimitiveComponent Interface
 
 	//~ Begin USceneComponent Interface
@@ -843,6 +844,9 @@ private:
 	/** Stores the current state for pause/unpause desired by the use. Allows us to pause/unpause correctly while also pausing/unpausing via scalability. */
 	uint32 bDesiredPauseState : 1;
 
+	/** Request recache the PSOs */
+	uint32 bRecachePSOs : 1;
+
 	/** Restore relative transform from auto attachment and optionally detach from parent (regardless of whether it was an auto attachment). */
 	NIAGARA_API void CancelAutoAttachment(bool bDetachFromParent);
 
diff --git a/Engine/Plugins/FX/Niagara/Source/Niagara/Public/NiagaraMeshRendererProperties.h b/Engine/Plugins/FX/Niagara/Source/Niagara/Public/NiagaraMeshRendererProperties.h
--- a/Engine/Plugins/FX/Niagara/Source/Niagara/Public/NiagaraMeshRendererProperties.h
+++ b/Engine/Plugins/FX/Niagara/Source/Niagara/Public/NiagaraMeshRendererProperties.h
@@ -274,7 +274,7 @@ public:
 	NIAGARA_API virtual const FVertexFactoryType* GetVertexFactoryType() const override;
 	virtual bool IsSimTargetSupported(ENiagaraSimTarget InSimTarget) const override { return true; };
 	NIAGARA_API virtual bool PopulateRequiredBindings(FNiagaraParameterStore& InParameterStore) override;
-	NIAGARA_API virtual void CollectPSOPrecacheData(FPSOPrecacheParamsList& OutParams) override;
+	NIAGARA_API virtual void CollectPSOPrecacheData(const FNiagaraEmitterInstance* InEmitter, FPSOPrecacheParamsList& OutParams) const override;
 
 #if WITH_EDITORONLY_DATA
 	NIAGARA_API virtual const TArray<FNiagaraVariable>& GetOptionalAttributes() override;
diff --git a/Engine/Plugins/FX/Niagara/Source/Niagara/Public/NiagaraRendererProperties.h b/Engine/Plugins/FX/Niagara/Source/Niagara/Public/NiagaraRendererProperties.h
--- a/Engine/Plugins/FX/Niagara/Source/Niagara/Public/NiagaraRendererProperties.h
+++ b/Engine/Plugins/FX/Niagara/Source/Niagara/Public/NiagaraRendererProperties.h
@@ -334,7 +334,7 @@ public:
 		FPSOPrecacheVertexFactoryDataList VertexFactoryDataList;
 	};
 	typedef TArray<FPSOPrecacheParams, TInlineAllocator<2> > FPSOPrecacheParamsList;
-	virtual void CollectPSOPrecacheData(FPSOPrecacheParamsList& OutParams) {};
+	virtual void CollectPSOPrecacheData(const FNiagaraEmitterInstance* InEmitter, FPSOPrecacheParamsList& OutParams) const {};
 
 #if WITH_EDITORONLY_DATA
 
diff --git a/Engine/Plugins/FX/Niagara/Source/Niagara/Public/NiagaraRibbonRendererProperties.h b/Engine/Plugins/FX/Niagara/Source/Niagara/Public/NiagaraRibbonRendererProperties.h
--- a/Engine/Plugins/FX/Niagara/Source/Niagara/Public/NiagaraRibbonRendererProperties.h
+++ b/Engine/Plugins/FX/Niagara/Source/Niagara/Public/NiagaraRibbonRendererProperties.h
@@ -219,7 +219,7 @@ public:
 	NIAGARA_API virtual bool IsBackfaceCullingDisabled() const override;
 	virtual bool IsSimTargetSupported(ENiagaraSimTarget InSimTarget) const override { return true; };
 	NIAGARA_API virtual bool PopulateRequiredBindings(FNiagaraParameterStore& InParameterStore) override;
-	NIAGARA_API virtual void CollectPSOPrecacheData(FPSOPrecacheParamsList& OutParams) override;
+	NIAGARA_API virtual void CollectPSOPrecacheData(const FNiagaraEmitterInstance* InEmitter, FPSOPrecacheParamsList& OutParams) const override;
 
 #if WITH_EDITOR
 	NIAGARA_API virtual const TArray<FNiagaraVariable>& GetOptionalAttributes() override;
diff --git a/Engine/Plugins/FX/Niagara/Source/Niagara/Public/NiagaraSpriteRendererProperties.h b/Engine/Plugins/FX/Niagara/Source/Niagara/Public/NiagaraSpriteRendererProperties.h
--- a/Engine/Plugins/FX/Niagara/Source/Niagara/Public/NiagaraSpriteRendererProperties.h
+++ b/Engine/Plugins/FX/Niagara/Source/Niagara/Public/NiagaraSpriteRendererProperties.h
@@ -137,7 +137,7 @@ public:
 	NIAGARA_API virtual const FVertexFactoryType* GetVertexFactoryType() const override;
 	virtual bool IsSimTargetSupported(ENiagaraSimTarget InSimTarget) const override { return true; };
 	NIAGARA_API virtual bool PopulateRequiredBindings(FNiagaraParameterStore& InParameterStore)  override;
-	NIAGARA_API virtual void CollectPSOPrecacheData(FPSOPrecacheParamsList& OutParams) override;
+	NIAGARA_API virtual void CollectPSOPrecacheData(const FNiagaraEmitterInstance* InEmitter, FPSOPrecacheParamsList& OutParams) const override;
 #if WITH_EDITOR
 	NIAGARA_API virtual const TArray<FNiagaraVariable>& GetOptionalAttributes() override;
 	NIAGARA_API virtual void GetAdditionalVariables(TArray<FNiagaraVariableBase>& OutArray) const override;
diff --git a/Engine/Plugins/FX/Niagara/Source/Niagara/Public/NiagaraSystemInstanceController.h b/Engine/Plugins/FX/Niagara/Source/Niagara/Public/NiagaraSystemInstanceController.h
--- a/Engine/Plugins/FX/Niagara/Source/Niagara/Public/NiagaraSystemInstanceController.h
+++ b/Engine/Plugins/FX/Niagara/Source/Niagara/Public/NiagaraSystemInstanceController.h
@@ -92,6 +92,8 @@ public:
 	UMaterialInterface* GetMaterialOverride(const UNiagaraRendererProperties* InProps, int32 InMaterialSubIndex) const;
 	void SetOnMaterialsUpdated(const FOnMaterialsUpdated& Delegate) { OnMaterialsUpdatedDelegate = Delegate; }
 
+	void CollectPSOPrecacheData(const FPSOPrecacheParams& BasePrecachePSOParams, FMaterialInterfacePSOPrecacheParamsList& List) const;
+
 	void GetMaterialStreamingInfo(FNiagaraMaterialAndScaleArray& OutMaterialAndScales) const;
 	void GetStreamingMeshInfo(const FBoxSphereBounds& OwnerBounds, FStreamingTextureLevelContext& LevelContext, TArray<FStreamingRenderAssetPrimitiveInfo>& OutStreamingRenderAssets) const;
 
diff --git a/Engine/Source/Runtime/Engine/Classes/Particles/ParticleSystem.h b/Engine/Source/Runtime/Engine/Classes/Particles/ParticleSystem.h
--- a/Engine/Source/Runtime/Engine/Classes/Particles/ParticleSystem.h
+++ b/Engine/Source/Runtime/Engine/Classes/Particles/ParticleSystem.h
@@ -159,15 +159,8 @@ public:
 	const TArray<FMaterialPSOPrecacheRequestID>& GetMaterialPSOPrecacheRequestIDs() const { return MaterialPSOPrecacheRequestIDs; }
 
 protected:
-	struct VFsPerMaterialData
-	{
-		UMaterialInterface* MaterialInterface = nullptr;
-		EPrimitiveType PrimitiveType = PT_TriangleList; // must match FPSOPrecacheParams::PrimitiveType default value
-		bool bDisableBackfaceCulling = false;  // must match FPSOPrecacheParams::bDisableBackfaceCulling default value
-		FPSOPrecacheVertexFactoryDataList VertexFactoryData;
-	};
-
-	ENGINE_API void LaunchPSOPrecaching(TArrayView<VFsPerMaterialData> VFsPerMaterials);
+	
+	ENGINE_API void LaunchPSOPrecaching(const FMaterialInterfacePSOPrecacheParamsList& VFsPerMaterials);
 
 	FGraphEventRef PrecachePSOsEvent;
 	TArray<FMaterialPSOPrecacheRequestID> MaterialPSOPrecacheRequestIDs;
diff --git a/Engine/Source/Runtime/Engine/Private/PSOPrecache.cpp b/Engine/Source/Runtime/Engine/Private/PSOPrecache.cpp
--- a/Engine/Source/Runtime/Engine/Private/PSOPrecache.cpp
+++ b/Engine/Source/Runtime/Engine/Private/PSOPrecache.cpp
@@ -90,3 +90,24 @@ FPSOPrecacheVertexFactoryData::FPSOPrecacheVertexFactoryData(
 {
 	CustomDefaultVertexDeclaration = PipelineStateCache::GetOrCreateVertexDeclaration(ElementList);
 }
+
+void AddMaterialInterfacePSOPrecacheParamsToList(const FMaterialInterfacePSOPrecacheParams& EntryToAdd, FMaterialInterfacePSOPrecacheParamsList& List)
+{
+	FMaterialInterfacePSOPrecacheParams* CurrentEntry = List.FindByPredicate([EntryToAdd](const FMaterialInterfacePSOPrecacheParams& Other)
+		{
+			return (Other.Priority			== EntryToAdd.Priority &&
+					Other.MaterialInterface == EntryToAdd.MaterialInterface &&
+					Other.PSOPrecacheParams == EntryToAdd.PSOPrecacheParams);
+		});
+	if (CurrentEntry)
+	{
+		for (const FPSOPrecacheVertexFactoryData& VFData : EntryToAdd.VertexFactoryDataList)
+		{
+			CurrentEntry->VertexFactoryDataList.AddUnique(VFData);
+		}
+	}
+	else
+	{
+		List.Add(EntryToAdd);
+	}
+}
diff --git a/Engine/Source/Runtime/Engine/Private/PSOPrecacheMaterial.cpp b/Engine/Source/Runtime/Engine/Private/PSOPrecacheMaterial.cpp
--- a/Engine/Source/Runtime/Engine/Private/PSOPrecacheMaterial.cpp
+++ b/Engine/Source/Runtime/Engine/Private/PSOPrecacheMaterial.cpp
@@ -528,6 +528,17 @@ void FMaterialPSOPrecacheCollectionTask::DoTask(ENamedThreads::Type CurrentThrea
 	CollectionGraphEvent->DispatchSubsequents();
 }
 
+void PrecacheMaterialPSOs(const FMaterialInterfacePSOPrecacheParamsList& PSOPrecacheParamsList, TArray<FMaterialPSOPrecacheRequestID>& OutMaterialPSOPrecacheRequestIDs, FGraphEventArray& OutGraphEvents)
+{
+	for (const FMaterialInterfacePSOPrecacheParams& MaterialPSOPrecacheParams : PSOPrecacheParamsList)
+	{
+		if (MaterialPSOPrecacheParams.MaterialInterface)
+		{
+			OutGraphEvents.Append(MaterialPSOPrecacheParams.MaterialInterface->PrecachePSOs(MaterialPSOPrecacheParams.VertexFactoryDataList, MaterialPSOPrecacheParams.PSOPrecacheParams, MaterialPSOPrecacheParams.Priority, OutMaterialPSOPrecacheRequestIDs));
+		}
+	}
+}
+
 FMaterialPSOPrecacheRequestID PrecacheMaterialPSOs(const FMaterialPSOPrecacheParams& MaterialPSOPrecacheParams, EPSOPrecachePriority Priority, FGraphEventArray& GraphEvents)
 {
 	return GMaterialPSORequestManager.PrecachePSOs(MaterialPSOPrecacheParams, Priority, GraphEvents);
diff --git a/Engine/Source/Runtime/Engine/Private/Particles/ParticleComponents.cpp b/Engine/Source/Runtime/Engine/Private/Particles/ParticleComponents.cpp
--- a/Engine/Source/Runtime/Engine/Private/Particles/ParticleComponents.cpp
+++ b/Engine/Source/Runtime/Engine/Private/Particles/ParticleComponents.cpp
@@ -43,6 +43,7 @@
 #include "DeviceProfiles/DeviceProfile.h"
 #include "DeviceProfiles/DeviceProfileManager.h"
 #include "PSOPrecache.h"
+#include "PSOPrecacheMaterial.h"
 #include "PrimitiveSceneProxy.h"
 #include "RenderingThread.h"
 #include "SceneInterface.h"
@@ -147,30 +148,10 @@ void UFXSystemAsset::PostInitProperties()
 #endif
 }
 
-void UFXSystemAsset::LaunchPSOPrecaching(TArrayView<VFsPerMaterialData> VFsPerMaterials)
+void UFXSystemAsset::LaunchPSOPrecaching(const FMaterialInterfacePSOPrecacheParamsList& PSOPrecacheParamsList)
 {
-	FPSOPrecacheParams PreCachePSOParams;
-	PreCachePSOParams.SetMobility(EComponentMobility::Movable);
-	PreCachePSOParams.bRenderCustomDepth = true;
-
 	FGraphEventArray PrecachePSOsEvents;
-	for (VFsPerMaterialData& VFsPerMaterial : VFsPerMaterials)
-	{
-		if (VFsPerMaterial.MaterialInterface)
-		{
-			PreCachePSOParams.PrimitiveType = (EPrimitiveType)VFsPerMaterial.PrimitiveType;
-			PreCachePSOParams.bDisableBackFaceCulling = VFsPerMaterial.bDisableBackfaceCulling;
-			PreCachePSOParams.bReverseCulling = false;
-			PrecachePSOsEvents.Append(VFsPerMaterial.MaterialInterface->PrecachePSOs(VFsPerMaterial.VertexFactoryData, PreCachePSOParams, EPSOPrecachePriority::Medium, MaterialPSOPrecacheRequestIDs));
-
-			// Also precache with reverse culling if not two sided because we don't know of the component using the asset will have negative determinant
-			if (!PreCachePSOParams.bDisableBackFaceCulling)
-			{
-				PreCachePSOParams.bReverseCulling = true;
-				PrecachePSOsEvents.Append(VFsPerMaterial.MaterialInterface->PrecachePSOs(VFsPerMaterial.VertexFactoryData, PreCachePSOParams, EPSOPrecachePriority::Medium, MaterialPSOPrecacheRequestIDs));
-			}
-		}
-	}
+	PrecacheMaterialPSOs(PSOPrecacheParamsList, MaterialPSOPrecacheRequestIDs, PrecachePSOsEvents);
 
 	// Create task to signal that the PSO precache events are done by adding them as prerequisite to the task.
 	if (PrecachePSOsEvents.Num() > 0)
@@ -2731,7 +2712,11 @@ void UParticleSystem::PrecachePSOs()
 		return;
 	}
 
-	TArray<VFsPerMaterialData, TInlineAllocator<4>> VFsPerMaterials;
+	FMaterialInterfacePSOPrecacheParamsList PSOPrecacheParamsList;
+
+	FMaterialInterfacePSOPrecacheParams NewEntry;
+	NewEntry.PSOPrecacheParams.SetMobility(EComponentMobility::Movable);
+	NewEntry.PSOPrecacheParams.bRenderCustomDepth = true;
 
 	// No per component emitter materials known at this point in time
 	TArray<UMaterialInterface*> EmptyEmitterMaterials;
@@ -2771,29 +2756,17 @@ void UParticleSystem::PrecachePSOs()
 
 				for (UMaterialInterface* MaterialInterface : Materials)
 				{
-					EPrimitiveType PrimitiveType = PrecacheParams.PrimitiveType;
+					NewEntry.MaterialInterface = MaterialInterface;
+					NewEntry.VertexFactoryDataList = PrecacheParams.VertexFactoryDataList;
+					NewEntry.PSOPrecacheParams.PrimitiveType = PrecacheParams.PrimitiveType;
 
-					VFsPerMaterialData* VFsPerMaterial = VFsPerMaterials.FindByPredicate([MaterialInterface, PrimitiveType](const VFsPerMaterialData& Other)
-						{
-							return Other.MaterialInterface == MaterialInterface && Other.PrimitiveType == PrimitiveType;
-						});
-					if (VFsPerMaterial == nullptr)
-					{
-						VFsPerMaterial = &VFsPerMaterials.AddDefaulted_GetRef();
-						VFsPerMaterial->MaterialInterface = MaterialInterface;
-						VFsPerMaterial->PrimitiveType = PrimitiveType;
-					}
-					
-					for (FPSOPrecacheVertexFactoryData VFData : PrecacheParams.VertexFactoryDataList)
-					{
-						VFsPerMaterial->VertexFactoryData.AddUnique(VFData);
-					}
+					AddMaterialInterfacePSOPrecacheParamsToList(NewEntry, PSOPrecacheParamsList);
 				}
 			}
 		}
 	}
 
-	LaunchPSOPrecaching(VFsPerMaterials);
+	LaunchPSOPrecaching(PSOPrecacheParamsList);
 }
 
 void UParticleSystem::Serialize(FArchive& Ar)
diff --git a/Engine/Source/Runtime/Engine/Public/PSOPrecache.h b/Engine/Source/Runtime/Engine/Public/PSOPrecache.h
--- a/Engine/Source/Runtime/Engine/Public/PSOPrecache.h
+++ b/Engine/Source/Runtime/Engine/Public/PSOPrecache.h
@@ -184,6 +184,18 @@ struct FPSOPrecacheVertexFactoryDataPerMaterialIndex
 
 typedef TArray<FPSOPrecacheVertexFactoryDataPerMaterialIndex, TInlineAllocator<4> > FPSOPrecacheVertexFactoryDataPerMaterialIndexList;
 
+struct FMaterialInterfacePSOPrecacheParams
+{
+	EPSOPrecachePriority Priority = EPSOPrecachePriority::Medium;
+	UMaterialInterface* MaterialInterface = nullptr;
+	FPSOPrecacheParams PSOPrecacheParams;
+	FPSOPrecacheVertexFactoryDataList VertexFactoryDataList;
+};
+
+typedef TArray<FMaterialInterfacePSOPrecacheParams, TInlineAllocator<4> > FMaterialInterfacePSOPrecacheParamsList;
+
+extern ENGINE_API void AddMaterialInterfacePSOPrecacheParamsToList(const FMaterialInterfacePSOPrecacheParams& EntryToAdd, FMaterialInterfacePSOPrecacheParamsList& List);
+
 /**
  * Wrapper class around the initializer to collect some extra validation data during PSO collection on the different collectors
  */
diff --git a/Engine/Source/Runtime/Engine/Public/PSOPrecacheMaterial.h b/Engine/Source/Runtime/Engine/Public/PSOPrecacheMaterial.h
--- a/Engine/Source/Runtime/Engine/Public/PSOPrecacheMaterial.h
+++ b/Engine/Source/Runtime/Engine/Public/PSOPrecacheMaterial.h
@@ -121,6 +121,11 @@ private:
 	uint32 Index;
 };
 
+/**
+ * Precache all PSOs for given material data
+ */
+extern ENGINE_API void PrecacheMaterialPSOs(const FMaterialInterfacePSOPrecacheParamsList& PSOPrecacheParamsList, TArray<FMaterialPSOPrecacheRequestID>& OutMaterialPSOPrecacheRequestIDs, FGraphEventArray& OutGraphEvents);
+
 /**
  * Precache all PSOs for given material and parameters
  */
