World Partition Data Layers: move all data layers specific code out of UActorDescContainerInstance into UWorldPartition.

diff --git a/Engine/Source/Runtime/Engine/Private/WorldPartition/ActorDescContainerInstance.cpp b/Engine/Source/Runtime/Engine/Private/WorldPartition/ActorDescContainerInstance.cpp
--- a/Engine/Source/Runtime/Engine/Private/WorldPartition/ActorDescContainerInstance.cpp
+++ b/Engine/Source/Runtime/Engine/Private/WorldPartition/ActorDescContainerInstance.cpp
@@ -10,7 +10,6 @@
 #include "WorldPartition/WorldPartition.h"
 #include "WorldPartition/WorldPartitionActorDesc.h"
 #include "WorldPartition/WorldPartitionActorDescUtils.h"
-#include "WorldPartition/DataLayer/DataLayerManager.h"
 #include "Editor.h"
 #endif
 
@@ -507,24 +506,6 @@ void UActorDescContainerInstance::OnObjectsReplaced(const TMap<UObject*, UObject
 	}
 }
 
-const UDataLayerManager* UActorDescContainerInstance::GetResolvingDataLayerManager() const
-{
-	if (UWorldPartition* WorldPartition = GetWorldPartition())
-	{
-		if (UWorld* OwningWorld = WorldPartition->GetWorld(); OwningWorld && !OwningWorld->IsGameWorld())
-		{
-			if (UWorldPartition* OwningWorldPartition = OwningWorld->GetWorldPartition())
-			{
-				return UDataLayerManager::GetDataLayerManager(OwningWorldPartition);
-			}
-		}
-
-		return UDataLayerManager::GetDataLayerManager(WorldPartition);
-	}
-
-	return nullptr;
-}
-
 void UActorDescContainerInstance::OnActorDescAdded(FWorldPartitionActorDesc* InActorDesc)
 {
 	FWorldPartitionActorDescInstance* NewActorDescInstance = AddActor(InActorDesc);
@@ -536,11 +517,6 @@ void UActorDescContainerInstance::OnActorDescAdded(FWorldPartitionActorDesc* InA
 
 	if (UWorldPartition* WorldPartition = GetWorldPartition())
 	{
-		if (const UDataLayerManager* DataLayerManager = GetResolvingDataLayerManager())
-		{
-			DataLayerManager->ResolveActorDescInstanceDataLayers(NewActorDescInstance);
-		}
-
 		WorldPartition->OnActorDescInstanceAdded(NewActorDescInstance);
 	}
 
@@ -581,11 +557,6 @@ void UActorDescContainerInstance::OnActorDescUpdated(FWorldPartitionActorDesc* I
 
 	if (UWorldPartition* WorldPartition = GetWorldPartition())
 	{
-		if (const UDataLayerManager* DataLayerManager = GetResolvingDataLayerManager())
-		{
-			DataLayerManager->ResolveActorDescInstanceDataLayers(ActorDescInstance->Get());
-		}
-
 		WorldPartition->OnActorDescInstanceUpdated(ActorDescInstance->Get());
 	}
 
diff --git a/Engine/Source/Runtime/Engine/Private/WorldPartition/WorldPartition.cpp b/Engine/Source/Runtime/Engine/Private/WorldPartition/WorldPartition.cpp
--- a/Engine/Source/Runtime/Engine/Private/WorldPartition/WorldPartition.cpp
+++ b/Engine/Source/Runtime/Engine/Private/WorldPartition/WorldPartition.cpp
@@ -782,6 +782,18 @@ UDataLayerManager* UWorldPartition::GetDataLayerManager() const
 	return DataLayerManager;
 }
 
+UDataLayerManager* UWorldPartition::GetResolvingDataLayerManager() const
+{
+	if (UWorld* OwningWorld = GetWorld(); OwningWorld && !OwningWorld->IsGameWorld())
+	{
+		if (UWorldPartition* OwningWorldPartition = OwningWorld->GetWorldPartition())
+		{
+			return UDataLayerManager::GetDataLayerManager(OwningWorldPartition);
+		}
+	}
+	return GetDataLayerManager();
+}
+
 bool UWorldPartition::IsInitialized() const
 {
 	return InitState == EWorldPartitionInitState::Initialized;
@@ -1196,6 +1208,11 @@ bool UWorldPartition::IsSimulating(bool bIncludeTestEnableSimulationStreamingSou
 #if WITH_EDITOR
 void UWorldPartition::OnActorDescInstanceAdded(FWorldPartitionActorDescInstance* NewActorDescInstance)
 {
+	if (const UDataLayerManager* ResolvingDataLayerManager = GetResolvingDataLayerManager())
+	{
+		ResolvingDataLayerManager->ResolveActorDescInstanceDataLayers(NewActorDescInstance);
+	}
+
 	NewActorDescInstance->SetForceNonSpatiallyLoadded(!IsStreamingEnabledInEditor());
 
 	HashActorDescInstance(NewActorDescInstance);
@@ -1248,6 +1265,11 @@ void UWorldPartition::OnActorDescInstanceUpdating(FWorldPartitionActorDescInstan
 
 void UWorldPartition::OnActorDescInstanceUpdated(FWorldPartitionActorDescInstance* ActorDescInstance)
 {
+	if (const UDataLayerManager* ResolvingDataLayerManager = GetResolvingDataLayerManager())
+	{
+		ResolvingDataLayerManager->ResolveActorDescInstanceDataLayers(ActorDescInstance);
+	}
+
 	HashActorDescInstance(ActorDescInstance);
 
 	if (WorldPartitionEditor)
diff --git a/Engine/Source/Runtime/Engine/Public/WorldPartition/ActorDescContainerInstance.h b/Engine/Source/Runtime/Engine/Public/WorldPartition/ActorDescContainerInstance.h
--- a/Engine/Source/Runtime/Engine/Public/WorldPartition/ActorDescContainerInstance.h
+++ b/Engine/Source/Runtime/Engine/Public/WorldPartition/ActorDescContainerInstance.h
@@ -13,7 +13,6 @@
 class UActorDescContainer;
 class FWorldPartitionActorDesc;
 class UWorldPartition;
-class UDataLayerManager;
 struct FWorldPartitionRuntimeCellPropertyOverride;
 
 class FActorDescInstanceList : public TActorDescList<FWorldPartitionActorDescInstance> { };
@@ -128,7 +127,6 @@ protected:
 	virtual FWorldPartitionActorDesc* GetActorDesc(const FGuid& InActorGuid) const;
 	virtual FWorldPartitionActorDesc* GetActorDescChecked(const FGuid& InActorGuid) const;
 
-	const UDataLayerManager* GetResolvingDataLayerManager() const;
 private:
 	void OnContainerUpdated(FName ContainerPackage);
 	void OnContainerReplaced(UActorDescContainer* InOldContainer, UActorDescContainer* InNewContainer);
diff --git a/Engine/Source/Runtime/Engine/Public/WorldPartition/WorldPartition.h b/Engine/Source/Runtime/Engine/Public/WorldPartition/WorldPartition.h
--- a/Engine/Source/Runtime/Engine/Public/WorldPartition/WorldPartition.h
+++ b/Engine/Source/Runtime/Engine/Public/WorldPartition/WorldPartition.h
@@ -456,6 +456,7 @@ public:
 	ENGINE_API void EnableStreamingIn();
 
 	ENGINE_API UDataLayerManager* GetDataLayerManager() const;
+	ENGINE_API UDataLayerManager* GetResolvingDataLayerManager() const;
 
 	inline EWorldPartitionDataLayersLogicOperator GetDataLayersLogicOperator() const { return DataLayersLogicOperator; }
 
