Fix memory corruption when ACF_None is used as a rotation format in various animation codecs

diff --git a/Engine/Source/Runtime/Engine/Private/Animation/AnimCompress.cpp b/Engine/Source/Runtime/Engine/Private/Animation/AnimCompress.cpp
--- a/Engine/Source/Runtime/Engine/Private/Animation/AnimCompress.cpp
+++ b/Engine/Source/Runtime/Engine/Private/Animation/AnimCompress.cpp
@@ -180,7 +180,7 @@ void UAnimCompress::PackQuaternionToStream(
 {
 	if ( Format == ACF_None )
 	{
-		UnalignedWriteToStream( ByteStream, &Quat, sizeof(FQuat) );
+		UnalignedWriteToStream( ByteStream, &Quat, sizeof(Quat) );
 	}
 	else if ( Format == ACF_Float96NoW )
 	{
@@ -327,7 +327,7 @@ void UAnimCompress::BitwiseCompressAnimationTracks(
 			AnimData.CompressedScaleOffsets.AddUninitialized(NumTracks);
 		}
 
-		const int32 MaxSize = CompressibleAnimData.RawAnimationData.Num() * CompressibleAnimData.NumberOfKeys * sizeof(FVector3f) + sizeof(FQuat4f) + sizeof(FVector3f);
+		const int32 MaxSize = CompressibleAnimData.RawAnimationData.Num() * CompressibleAnimData.NumberOfKeys * (sizeof(FVector3f) + sizeof(FQuat4f) + sizeof(FVector3f));
 		AnimData.CompressedByteStream.Reset(MaxSize);
 
 		for (int32 TrackIndex = 0; TrackIndex < NumTracks; ++TrackIndex)
