Fix ACF_IntervalFixed32NoW codec format for translation/scale when used with linear key reduction

diff --git a/Engine/Source/Runtime/Engine/Private/Animation/AnimEncoding_VariableKeyLerp.ispc b/Engine/Source/Runtime/Engine/Private/Animation/AnimEncoding_VariableKeyLerp.ispc
--- a/Engine/Source/Runtime/Engine/Private/Animation/AnimEncoding_VariableKeyLerp.ispc
+++ b/Engine/Source/Runtime/Engine/Private/Animation/AnimEncoding_VariableKeyLerp.ispc
@@ -113,13 +113,14 @@ export void GetVariableKeyLerpPoseTranslations(
 		const int NumTransKeys = TrackData[1];
 		const unsigned int8* TransStream = CompressedByteStream + TransKeysOffset;
 
-		const unsigned int8* Frames= TransStream +(NumTransKeys*CompressedTranslationStrideNum);
+		const int32 TransStreamOffset = ((FORMAT == ACF_IntervalFixed32NoW) && NumTransKeys > 1) ? (sizeof(uniform float)*6) : 0; // offset past Min and Range data
+
+		const unsigned int8* Frames = TransStream + TransStreamOffset + (NumTransKeys * CompressedTranslationStrideNum);
 		const unsigned int8* FrameTable= Align(Frames, 4);
 
 		int32 Index0;
 		int32 Index1;
 		const float Alpha = TimeToIndex(InterpolationType, CompressedNumberOfFrames, FrameTable, RelativePos, NumTransKeys, Index0, Index1);
-		const int32 TransStreamOffset = ((FORMAT == ACF_IntervalFixed32NoW) && NumTransKeys > 1) ? (sizeof(uniform float)*6) : 0; // offset past Min and Range data
 
 		const unsigned int8* KeyData0 = TransStream + TransStreamOffset + Index0*CompressedTranslationStrideNum;
 		FVector3f P0;
@@ -178,13 +179,14 @@ export void GetVariableKeyLerpPoseScales(
 		const int NumScaleKeys = ScaleData[1];
 		const unsigned int8* ScaleStream = CompressedByteStream + ScaleKeysOffset;
 
-		const unsigned int8* Frames= ScaleStream +(NumScaleKeys*CompressedScaleStrideNum);
+		const int32 ScaleStreamOffset = ((FORMAT == ACF_IntervalFixed32NoW) && NumScaleKeys > 1) ? (sizeof(uniform float)*6) : 0; // offset past Min and Range data
+
+		const unsigned int8* Frames= ScaleStream + ScaleStreamOffset + (NumScaleKeys * CompressedScaleStrideNum);
 		const unsigned int8* FrameTable= Align(Frames, 4);
 
 		int32 Index0;
 		int32 Index1;
 		const float Alpha = TimeToIndex(InterpolationType, CompressedNumberOfFrames, FrameTable, RelativePos, NumScaleKeys, Index0, Index1);
-		const int32 ScaleStreamOffset = ((FORMAT == ACF_IntervalFixed32NoW) && NumScaleKeys > 1) ? (sizeof(uniform float)*6) : 0; // offset past Min and Range data
 
 		const unsigned int8* KeyData0 = ScaleStream + ScaleStreamOffset + Index0*CompressedScaleStrideNum;
 		FVector3f P0;
diff --git a/Engine/Source/Runtime/Engine/Public/AnimEncoding_VariableKeyLerp.h b/Engine/Source/Runtime/Engine/Public/AnimEncoding_VariableKeyLerp.h
--- a/Engine/Source/Runtime/Engine/Public/AnimEncoding_VariableKeyLerp.h
+++ b/Engine/Source/Runtime/Engine/Public/AnimEncoding_VariableKeyLerp.h
@@ -258,13 +258,13 @@ FORCEINLINE_DEBUGGABLE void AEFVariableKeyLerp<FORMAT>::GetBoneAtomTranslation(F
 	int32 NumTransKeys = TrackData[1];
 	const uint8* RESTRICT TransStream = AnimData.CompressedByteStream.GetData() + TransKeysOffset;
 
-	const uint8* RESTRICT FrameTable= TransStream +(NumTransKeys*CompressedTranslationStrides[FORMAT]*CompressedTranslationNum[FORMAT]);
+	const int32 TransStreamOffset = ((FORMAT == ACF_IntervalFixed32NoW) && NumTransKeys > 1) ? (sizeof(float) * 6) : 0; // offset past Min and Range data
+	const uint8* RESTRICT FrameTable = TransStream + TransStreamOffset + (NumTransKeys * CompressedTranslationStrides[FORMAT] * CompressedTranslationNum[FORMAT]);
 	FrameTable= Align(FrameTable, 4);
 
 	int32 Index0;
 	int32 Index1;
 	float Alpha = TimeToIndex(DecompContext.Interpolation, AnimData.CompressedNumberOfKeys, FrameTable, DecompContext.GetRelativePosition(), NumTransKeys, Index0, Index1);
-	const int32 TransStreamOffset = ((FORMAT == ACF_IntervalFixed32NoW) && NumTransKeys > 1) ? (sizeof(float)*6) : 0; // offset past Min and Range data
 
 	if (Index0 != Index1)
 	{
@@ -302,13 +302,13 @@ FORCEINLINE_DEBUGGABLE void AEFVariableKeyLerp<FORMAT>::GetBoneAtomScale(FTransf
 	const int32 NumScaleKeys = AnimData.CompressedScaleOffsets.GetOffsetData(TrackIndex, 1);
 	const uint8* RESTRICT ScaleStream = AnimData.CompressedByteStream.GetData() + ScaleKeysOffset;
 
-	const uint8* RESTRICT FrameTable= ScaleStream +(NumScaleKeys*CompressedScaleStrides[FORMAT]*CompressedScaleNum[FORMAT]);
+	const int32 ScaleStreamOffset = ((FORMAT == ACF_IntervalFixed32NoW) && NumScaleKeys > 1) ? (sizeof(float) * 6) : 0; // offset past Min and Range data
+	const uint8* RESTRICT FrameTable = ScaleStream + ScaleStreamOffset + (NumScaleKeys * CompressedScaleStrides[FORMAT] * CompressedScaleNum[FORMAT]);
 	FrameTable= Align(FrameTable, 4);
 
 	int32 Index0;
 	int32 Index1;
 	float Alpha = TimeToIndex(DecompContext.Interpolation, AnimData.CompressedNumberOfKeys, FrameTable, DecompContext.GetRelativePosition(), NumScaleKeys, Index0, Index1);
-	const int32 ScaleStreamOffset = ((FORMAT == ACF_IntervalFixed32NoW) && NumScaleKeys > 1) ? (sizeof(float)*6) : 0; // offset past Min and Range data
 
 	if (Index0 != Index1)
 	{
