Implemented automatically generated UVWs as fallback for the SVT sample node.

#rnx

diff --git a/Engine/Source/Runtime/Engine/Private/Materials/MaterialExpressions.cpp b/Engine/Source/Runtime/Engine/Private/Materials/MaterialExpressions.cpp
--- a/Engine/Source/Runtime/Engine/Private/Materials/MaterialExpressions.cpp
+++ b/Engine/Source/Runtime/Engine/Private/Materials/MaterialExpressions.cpp
@@ -28105,16 +28105,37 @@ int32 UMaterialExpressionSparseVolumeTextureSample::Compile(class FMaterialCompi
 			return INDEX_NONE;
 		}
 
+		// Get sampling coordinates/UVWs. Fall back to computing them as (LocalPosition - LocalBoundsMin) / BoundsSize.
+		int32 CoordinateIndex = INDEX_NONE;
+		if (Coordinates.GetTracedInput().Expression)
+		{
+			CoordinateIndex = Coordinates.Compile(Compiler);
+		}
+		else
+		{
+			int32 WorldPositionLWCIndex = Compiler->WorldPosition(WPT_Default);
+			int32 LocalPositionLWCIndex = Compiler->TransformPosition(MCB_World, MCB_Local, WorldPositionLWCIndex);
+			int32 LocalPositionIndex = Compiler->ValidCast(LocalPositionLWCIndex, MCT_Float3); // LocalPosition is likely of LWC type. Float precision is fine past this point, so cast the LWC-ness away.
+			int32 BoundsSizeIndex = Compiler->ObjectLocalBounds(1);
+			int32 BoundsMinIndex = Compiler->ObjectLocalBounds(2);
+			int32 RelativeLocalPositionIndex = Compiler->Sub(LocalPositionIndex, BoundsMinIndex);
+			CoordinateIndex = Compiler->Div(RelativeLocalPositionIndex, BoundsSizeIndex);
+
+			if (CoordinateIndex == INDEX_NONE)
+			{
+				CompilerError(Compiler, TEXT("Failed to generate fallback UVW input for sparse volume texture"));
+			}
+		}
+
 		UMaterialExpression* MipLevelExpression = MipLevel.GetTracedInput().Expression;
 
 		// Shared inputs for both potential samples
-		int32 UVWIndex = CompileWithDefaultFloat3(Compiler, Coordinates, 0.0f, 0.0f, 0.0f);
 		int32 PhysicalTileDataIdxIndex = Compiler->Constant(OutputIndex);
 		int32 MipLevelInputIndex = MipLevelExpression ? MipLevel.Compile(Compiler) : INDEX_NONE;
 		
 		// Sample the first mip
 		int32 MipLevel0Index = MipLevelExpression ? Compiler->Floor(MipLevelInputIndex) : Compiler->Constant(0.0f);
-		int32 VoxelCoordMip0Index = Compiler->SparseVolumeTextureSamplePageTable(SparseVolumeTextureIndex, UVWIndex, MipLevel0Index, SamplerSource);
+		int32 VoxelCoordMip0Index = Compiler->SparseVolumeTextureSamplePageTable(SparseVolumeTextureIndex, CoordinateIndex, MipLevel0Index, SamplerSource);
 		int32 Mip0SampleIndex = Compiler->SparseVolumeTextureSamplePhysicalTileData(SparseVolumeTextureIndex, VoxelCoordMip0Index, PhysicalTileDataIdxIndex);
 
 		if (MipLevelExpression)
@@ -28122,7 +28143,7 @@ int32 UMaterialExpressionSparseVolumeTextureSample::Compile(class FMaterialCompi
 			// Sample the second mip
 			// SVT_TODO: Try to optimize out this second sample if LerpAlpha == 0. Might need to do that in HLSL.
 			int32 MipLevel1Index = Compiler->Ceil(MipLevelInputIndex);
-			int32 VoxelCoordMip1Index = Compiler->SparseVolumeTextureSamplePageTable(SparseVolumeTextureIndex, UVWIndex, MipLevel1Index, SamplerSource);
+			int32 VoxelCoordMip1Index = Compiler->SparseVolumeTextureSamplePageTable(SparseVolumeTextureIndex, CoordinateIndex, MipLevel1Index, SamplerSource);
 			int32 Mip1SampleIndex = Compiler->SparseVolumeTextureSamplePhysicalTileData(SparseVolumeTextureIndex, VoxelCoordMip1Index, PhysicalTileDataIdxIndex);
 
 			// Lerp
