VulkanRHI: MorphTargetVertexInfoBuffers was being used before its InitRHI() call.  This would cause errors for null descriptors in Vulkan.  Skip updates when buffer is not ready (but keep the clear of the vertex buffer).

#rnx

diff --git a/Engine/Source/Runtime/Engine/Private/SkeletalRenderGPUSkin.cpp b/Engine/Source/Runtime/Engine/Private/SkeletalRenderGPUSkin.cpp
--- a/Engine/Source/Runtime/Engine/Private/SkeletalRenderGPUSkin.cpp
+++ b/Engine/Source/Runtime/Engine/Private/SkeletalRenderGPUSkin.cpp
@@ -1218,7 +1218,7 @@ void FSkeletalMeshObjectGPUSkin::FSkeletalMeshObjectLOD::UpdateMorphVertexBuffer
 			RHICmdList.ClearUAVUint(MorphVertexBuffer.GetUAV(), FUintVector4(0, 0, 0, 0));
 		}
 
-		if (MorphTargetVertexInfoBuffers.GetNumMorphs() > 0)
+		if (MorphTargetVertexInfoBuffers.IsRHIIntialized() && (MorphTargetVertexInfoBuffers.GetNumMorphs() > 0))
 		{
 			{
 				SCOPED_DRAW_EVENTF(RHICmdList, MorphUpdateScatter, TEXT("Scatter"));
