Create the UActorPartitionSubsystem world subsystem for inactive worlds

diff --git a/Engine/Source/Runtime/Engine/Private/ActorPartition/ActorPartitionSubsystem.cpp b/Engine/Source/Runtime/Engine/Private/ActorPartition/ActorPartitionSubsystem.cpp
--- a/Engine/Source/Runtime/Engine/Private/ActorPartition/ActorPartitionSubsystem.cpp
+++ b/Engine/Source/Runtime/Engine/Private/ActorPartition/ActorPartitionSubsystem.cpp
@@ -326,6 +326,11 @@ bool UActorPartitionSubsystem::IsLevelPartition() const
 	return !UWorld::IsPartitionedWorld(GetWorld());
 }
 
+bool UActorPartitionSubsystem::DoesSupportWorldType(const EWorldType::Type WorldType) const
+{
+	return Super::DoesSupportWorldType(WorldType) || WorldType == EWorldType::Inactive || WorldType == EWorldType::EditorPreview;
+}
+
 #if WITH_EDITOR
 void UActorPartitionSubsystem::Initialize(FSubsystemCollectionBase& Collection)
 {
diff --git a/Engine/Source/Runtime/Engine/Public/ActorPartition/ActorPartitionSubsystem.h b/Engine/Source/Runtime/Engine/Public/ActorPartition/ActorPartitionSubsystem.h
--- a/Engine/Source/Runtime/Engine/Public/ActorPartition/ActorPartitionSubsystem.h
+++ b/Engine/Source/Runtime/Engine/Public/ActorPartition/ActorPartitionSubsystem.h
@@ -175,6 +175,9 @@ public:
 #endif
 	ENGINE_API bool IsLevelPartition() const;
 
+protected:
+	ENGINE_API virtual bool DoesSupportWorldType(const EWorldType::Type WorldType) const override;
+
 private:
 
 #if WITH_EDITOR
