Fix to allow changing ArchiveCrc32 operator<<(FArchive&)  template generalization for structs to only apply for direct serialization to ArchiveCrc32 and if there is no better conversion available.

Details in Jira:

diff --git a/Engine/Source/Runtime/Engine/Private/WorldPartition/HLOD/HLODSourceActorsFromCell.cpp b/Engine/Source/Runtime/Engine/Private/WorldPartition/HLOD/HLODSourceActorsFromCell.cpp
--- a/Engine/Source/Runtime/Engine/Private/WorldPartition/HLOD/HLODSourceActorsFromCell.cpp
+++ b/Engine/Source/Runtime/Engine/Private/WorldPartition/HLOD/HLODSourceActorsFromCell.cpp
@@ -41,7 +41,11 @@ uint32 UWorldPartitionHLODSourceActorsFromCell::GetHLODHash(const TArray<FWorldP
 	TArray<FWorldPartitionRuntimeCellObjectMapping>& MutableSourceActors(const_cast<TArray<FWorldPartitionRuntimeCellObjectMapping>&>(InSourceActors));
 
 	FArchiveCrc32 Ar;
-	Ar << MutableSourceActors;
+	for (FWorldPartitionRuntimeCellObjectMapping& Mapping : MutableSourceActors)
+	{
+		Ar << Mapping;
+	}
+	
 	return Ar.GetCrc();
 }
 
