Reworked Skeleton Mesh Linkup map, to avoid leaving dead data in the LinkUpTable.

diff --git a/Engine/Plugins/Experimental/AnimNext/Source/AnimNext/Private/GenerationTools.cpp b/Engine/Plugins/Experimental/AnimNext/Source/AnimNext/Private/GenerationTools.cpp
--- a/Engine/Plugins/Experimental/AnimNext/Source/AnimNext/Private/GenerationTools.cpp
+++ b/Engine/Plugins/Experimental/AnimNext/Source/AnimNext/Private/GenerationTools.cpp
@@ -123,10 +123,7 @@ struct FCompareBoneIndexType
 				}
 
 				// Removing const here because the linkup lazily builds the mapping and caches it
-				const int32 SkelMeshLinkupIndex = const_cast<USkeleton*>(Skeleton)->GetMeshLinkupIndex(SkeletalMesh);
-				check(SkelMeshLinkupIndex != INDEX_NONE);
-
-				const FSkeletonToMeshLinkup& LinkupTable = Skeleton->LinkupCache[SkelMeshLinkupIndex];
+				const FSkeletonToMeshLinkup& LinkupTable = const_cast<USkeleton*>(Skeleton)->FindOrAddMeshLinkupData(SkeletalMesh);
 
 				// Generate a Skeleton to LOD look up table
 				const int32 NumOrderedBones = LODBoneIndexToMeshBoneIndexMap.Num();
diff --git a/Engine/Plugins/Experimental/PhysicsControl/Source/PhysicsControl/Private/AnimNode_RigidBodyWithControl.cpp b/Engine/Plugins/Experimental/PhysicsControl/Source/PhysicsControl/Private/AnimNode_RigidBodyWithControl.cpp
--- a/Engine/Plugins/Experimental/PhysicsControl/Source/PhysicsControl/Private/AnimNode_RigidBodyWithControl.cpp
+++ b/Engine/Plugins/Experimental/PhysicsControl/Source/PhysicsControl/Private/AnimNode_RigidBodyWithControl.cpp
@@ -1002,9 +1002,7 @@ void FAnimNode_RigidBodyWithControl::InitPhysics(const UAnimInstance* InAnimInst
 
 	ensure(SkeletonAsset == SkeletalMeshAsset->GetSkeleton());
 
-	const int32 SkelMeshLinkupIndex = SkeletonAsset->GetMeshLinkupIndex(SkeletalMeshAsset);
-	ensure(SkelMeshLinkupIndex != INDEX_NONE);
-	const FSkeletonToMeshLinkup& SkeletonToMeshLinkupTable = SkeletonAsset->LinkupCache[SkelMeshLinkupIndex];
+	const FSkeletonToMeshLinkup& SkeletonToMeshLinkupTable = SkeletonAsset->FindOrAddMeshLinkupData(SkeletalMeshAsset);
 	const TArray<int32>& MeshToSkeletonBoneIndex = SkeletonToMeshLinkupTable.MeshToSkeletonTable;
 	
 	const int32 NumSkeletonBones = SkeletonAsset->GetReferenceSkeleton().GetNum();
diff --git a/Engine/Source/Runtime/AnimGraphRuntime/Private/BoneControllers/AnimNode_RigidBody.cpp b/Engine/Source/Runtime/AnimGraphRuntime/Private/BoneControllers/AnimNode_RigidBody.cpp
--- a/Engine/Source/Runtime/AnimGraphRuntime/Private/BoneControllers/AnimNode_RigidBody.cpp
+++ b/Engine/Source/Runtime/AnimGraphRuntime/Private/BoneControllers/AnimNode_RigidBody.cpp
@@ -1037,10 +1037,8 @@ void FAnimNode_RigidBody::InitPhysics(const UAnimInstance* InAnimInstance)
 
 	ensure(SkeletonAsset == SkeletalMeshAsset->GetSkeleton());
 
-	const int32 SkelMeshLinkupIndex = SkeletonAsset->GetMeshLinkupIndex(SkeletalMeshAsset);
-	ensure(SkelMeshLinkupIndex != INDEX_NONE);
-	const FSkeletonToMeshLinkup& SkeletonToMeshLinkupTable = SkeletonAsset->LinkupCache[SkelMeshLinkupIndex];
-	const TArray<int32>& MeshToSkeletonBoneIndex = SkeletonToMeshLinkupTable.MeshToSkeletonTable;
+	const FSkeletonToMeshLinkup& LinkupTable = SkeletonAsset->FindOrAddMeshLinkupData(SkeletalMeshAsset);
+	const TArray<int32>& MeshToSkeletonBoneIndex = LinkupTable.MeshToSkeletonTable;
 	
 	const int32 NumSkeletonBones = SkeletonAsset->GetReferenceSkeleton().GetNum();
 	SkeletonBoneIndexToBodyIndex.Reset(NumSkeletonBones);
diff --git a/Engine/Source/Runtime/Engine/Classes/Animation/Skeleton.h b/Engine/Source/Runtime/Engine/Classes/Animation/Skeleton.h
--- a/Engine/Source/Runtime/Engine/Classes/Animation/Skeleton.h
+++ b/Engine/Source/Runtime/Engine/Classes/Animation/Skeleton.h
@@ -23,6 +23,7 @@
 #include "Interfaces/Interface_AssetUserData.h"
 #include "Interfaces/Interface_PreviewMeshProvider.h"
 #include "Delegates/DelegateCombinations.h"
+#include "UObject/ObjectKey.h"
 #include "Skeleton.generated.h"
 
 class UAnimCurveMetaData;
@@ -636,15 +637,24 @@ public:
 
 	typedef TArray<FBoneNode> FBoneTreeType;
 
-	//Use this Lock everytime you change or access LinkupCache and SkelMesh2LinkupCache member.
+	
+	UE_DEPRECATED(5.4, "Public access to this member variable is deprecated.")
 	FCriticalSection LinkupCacheLock;
 
 	/** Non-serialised cache of linkups between different skeletal meshes and this Skeleton. */
+	UE_DEPRECATED(5.4, "Public access to this member variable is deprecated.")
 	TArray<struct FSkeletonToMeshLinkup> LinkupCache;
 
 private:
-	/** Runtime built mapping table between SkinnedAssets, and LinkupCache array indices. */
+	UE_DEPRECATED(5.4, "Please, use SkinnedAssetLinkupCache.")
 	TMap<TWeakObjectPtr<USkinnedAsset>, int32> SkinnedAsset2LinkupCache;
+	
+	//Use this Lock everytime you change or access SkinnedAssetLinkupCache member.
+	FRWLock SkinnedAssetLinkupCacheLock;
+
+	/** Runtime built mapping table between SkinnedAssets and Mesh Linkup Data*/
+	TMap<TObjectKey<USkinnedAsset>, TUniquePtr<FSkeletonToMeshLinkup>> SkinnedAssetLinkupCache;
+
 public:
 	UE_DEPRECATED(5.1, "Public access to this member variable is deprecated.")
 	TMap<TWeakObjectPtr<USkeletalMesh>, int32> SkelMesh2LinkupCache;
@@ -795,11 +805,22 @@ public:
 	/** Clears all cache data **/
 	ENGINE_API void ClearCacheData();
 
-	/** 
+	UE_DEPRECATED(5.4, "Please use FindOrAddMeshLinkupData.")
+	ENGINE_API int32 GetMeshLinkupIndex(const USkinnedAsset* InSkinnedAsset);
+
+	/**
 	 * Find a mesh linkup table (mapping of skeleton bone tree indices to refpose indices) for a particular SkinnedAsset
 	 * If one does not already exist, create it now.
 	 */
-	ENGINE_API int32 GetMeshLinkupIndex(const USkinnedAsset* InSkinnedAsset);
+	ENGINE_API const FSkeletonToMeshLinkup& FindOrAddMeshLinkupData(const USkinnedAsset* InSkinnedAsset);
+
+	/**
+	 * Adds a new Mesh Linkup Table to the map  for a particular SkinnedAsset
+	 *
+	 * @param	InSkinnedAsset	: SkinnedAsset to build look up for
+	 * @return	Const ref to the added FSkeletonToMeshLinkup unique ptr
+	 */
+	ENGINE_API const FSkeletonToMeshLinkup& AddMeshLinkupData(const USkinnedAsset* InSkinnedAsset);
 
 	/** 
 	 * Merge Bones (RequiredBones from InSkinnedAsset) to BoneTrees if not exists
@@ -958,8 +979,17 @@ protected:
 	 * @param	InSkinnedAsset	: SkinnedAsset to build look up for
 	 * @return	Index of LinkupCache that this SkelMesh is linked to
 	 */
+	UE_DEPRECATED(5.4, "Please use BuildLinkup with USkinnedAsset and FSkeletonToMeshLinkup parameters.")
 	int32 BuildLinkup(const USkinnedAsset* InSkinnedAsset);
 
+	/**
+	 * Build Look up between SkinnedAsset to BoneTree
+	 *
+	 * @param	InSkinnedAsset			: SkinnedAsset to build look up for
+	 * @param	FSkeletonToMeshLinkup	: Out mesh linkup data
+	 */
+	void BuildLinkupData(const USkinnedAsset* InSkinnedAsset, FSkeletonToMeshLinkup& NewMeshLinkup);
+
 #if WITH_EDITORONLY_DATA
 	/**
 	 * Refresh All Retarget Sources
diff --git a/Engine/Source/Runtime/Engine/Private/Animation/Skeleton.cpp b/Engine/Source/Runtime/Engine/Private/Animation/Skeleton.cpp
--- a/Engine/Source/Runtime/Engine/Private/Animation/Skeleton.cpp
+++ b/Engine/Source/Runtime/Engine/Private/Animation/Skeleton.cpp
@@ -733,26 +733,60 @@ bool USkeleton::IsCompatibleMesh(const USkinnedAsset* InSkinnedAsset, bool bDoPa
 
 void USkeleton::ClearCacheData()
 {
+	{
+		FRWScopeLock Lock(SkinnedAssetLinkupCacheLock, SLT_Write);
+		SkinnedAssetLinkupCache.Empty();
+	}
+
+	PRAGMA_DISABLE_DEPRECATION_WARNINGS
 	FScopeLock ScopeLock(&LinkupCacheLock);
 	LinkupCache.Empty();
-	PRAGMA_DISABLE_DEPRECATION_WARNINGS
 	SkelMesh2LinkupCache.Empty();
-	PRAGMA_ENABLE_DEPRECATION_WARNINGS
 	SkinnedAsset2LinkupCache.Empty();
+	PRAGMA_ENABLE_DEPRECATION_WARNINGS
+}
+
+const FSkeletonToMeshLinkup& USkeleton::FindOrAddMeshLinkupData(const USkinnedAsset* InSkinnedAsset)
+{
+	const TUniquePtr<FSkeletonToMeshLinkup>* SkeletonToMeshLinkupPtr = nullptr;
+
+	{
+		FRWScopeLock Lock(SkinnedAssetLinkupCacheLock, SLT_ReadOnly);
+		SkeletonToMeshLinkupPtr = SkinnedAssetLinkupCache.Find(InSkinnedAsset);
+	}
+
+	return (SkeletonToMeshLinkupPtr != nullptr)
+		? *SkeletonToMeshLinkupPtr->Get()
+		: AddMeshLinkupData(InSkinnedAsset);
+}
+
+const FSkeletonToMeshLinkup& USkeleton::AddMeshLinkupData(const USkinnedAsset* InSkinnedAsset)
+{
+	TUniquePtr<FSkeletonToMeshLinkup> TmpLinkup = MakeUnique<FSkeletonToMeshLinkup>();
+	BuildLinkupData(InSkinnedAsset, *TmpLinkup.Get());
+
+	{
+		FRWScopeLock Lock(SkinnedAssetLinkupCacheLock, SLT_Write);
+		return *SkinnedAssetLinkupCache.Add(TObjectKey<USkinnedAsset>(InSkinnedAsset), MoveTemp(TmpLinkup)).Get();
+	}
 }
 
 int32 USkeleton::GetMeshLinkupIndex(const USkinnedAsset* InSkinnedAsset)
 {
 	int32* IndexPtr = nullptr;
 	{
+		PRAGMA_DISABLE_DEPRECATION_WARNINGS
 		FScopeLock ScopeLock(&LinkupCacheLock);
 		IndexPtr = SkinnedAsset2LinkupCache.Find(MakeWeakObjectPtr(const_cast<USkinnedAsset*>(InSkinnedAsset)));
+		PRAGMA_ENABLE_DEPRECATION_WARNINGS
 	}
 	int32 LinkupIndex = INDEX_NONE;
 
 	if ( IndexPtr == NULL )
 	{
+		PRAGMA_DISABLE_DEPRECATION_WARNINGS
 		LinkupIndex = BuildLinkup(InSkinnedAsset);
+		PRAGMA_ENABLE_DEPRECATION_WARNINGS
 	}
 	else
 	{
@@ -760,9 +794,11 @@ int32 USkeleton::GetMeshLinkupIndex(const USkinnedAsset* InSkinnedAsset)
 	}
 
 	{
+		PRAGMA_DISABLE_DEPRECATION_WARNINGS
 		FScopeLock ScopeLock(&LinkupCacheLock);
 		// make sure it's not out of range
 		check(LinkupIndex < LinkupCache.Num());
+		PRAGMA_ENABLE_DEPRECATION_WARNINGS
 	}
 
 	return LinkupIndex;
@@ -770,26 +806,52 @@ int32 USkeleton::GetMeshLinkupIndex(const USkinnedAsset* InSkinnedAsset)
 
 void USkeleton::RemoveLinkup(const USkinnedAsset* InSkinnedAsset)
 {
-	FScopeLock ScopeLock(&LinkupCacheLock);
+	{
+		FRWScopeLock Lock(SkinnedAssetLinkupCacheLock, SLT_Write);
+		SkinnedAssetLinkupCache.Remove(InSkinnedAsset);
+	}
+
 	PRAGMA_DISABLE_DEPRECATION_WARNINGS
+	FScopeLock ScopeLock(&LinkupCacheLock);
 	if (const USkeletalMesh* const InSkelMesh = Cast<const USkeletalMesh>(InSkinnedAsset))
 	{
 		SkelMesh2LinkupCache.Remove(MakeWeakObjectPtr(const_cast<USkeletalMesh*>(InSkelMesh)));
 	}
-	PRAGMA_ENABLE_DEPRECATION_WARNINGS
 	SkinnedAsset2LinkupCache.Remove(MakeWeakObjectPtr(const_cast<USkinnedAsset*>(InSkinnedAsset)));
+	PRAGMA_ENABLE_DEPRECATION_WARNINGS
 }
 
 int32 USkeleton::BuildLinkup(const USkinnedAsset* InSkinnedAsset)
 {
-	const FReferenceSkeleton& SkeletonRefSkel = ReferenceSkeleton;
-	const FReferenceSkeleton& MeshRefSkel = InSkinnedAsset->GetRefSkeleton();
-
 	// @todoanim : need to refresh NULL SkeletalMeshes from Cache
 	// since now they're autoweak pointer, they will go away if not used
 	// so whenever map transition happens, this links will need to clear up
 	FSkeletonToMeshLinkup NewMeshLinkup;
 
+	BuildLinkupData(InSkinnedAsset, NewMeshLinkup);
+
+	int32 NewIndex = INDEX_NONE;
+	//LinkupCache lock scope
+	{
+		PRAGMA_DISABLE_DEPRECATION_WARNINGS
+		FScopeLock ScopeLock(&LinkupCacheLock);
+		NewIndex = LinkupCache.Add(NewMeshLinkup);
+		check(NewIndex != INDEX_NONE);
+			if (const USkeletalMesh* const InSkelMesh = Cast<const USkeletalMesh>(InSkinnedAsset))
+			{
+				SkelMesh2LinkupCache.Add(MakeWeakObjectPtr(const_cast<USkeletalMesh*>(InSkelMesh)), NewIndex);
+			}
+		SkinnedAsset2LinkupCache.Add(MakeWeakObjectPtr(const_cast<USkinnedAsset*>(InSkinnedAsset)), NewIndex);
+		PRAGMA_ENABLE_DEPRECATION_WARNINGS
+	}
+	return NewIndex;
+}
+
+void USkeleton::BuildLinkupData(const USkinnedAsset* InSkinnedAsset, FSkeletonToMeshLinkup& NewMeshLinkup)
+{
+	const FReferenceSkeleton& SkeletonRefSkel = ReferenceSkeleton;
+	const FReferenceSkeleton& MeshRefSkel = InSkinnedAsset->GetRefSkeleton();
+
 	// First, make sure the Skeleton has all the bones the SkeletalMesh possesses.
 	// This can get out of sync if a mesh was imported on that Skeleton, but the Skeleton was not saved.
 
@@ -842,22 +904,6 @@ int32 USkeleton::BuildLinkup(const USkinnedAsset* InSkinnedAsset)
 		const int32 MeshBoneIndex = MeshRefSkel.FindBoneIndex( SkeletonRefSkel.GetBoneName(SkeletonBoneIndex) );
 		NewMeshLinkup.SkeletonToMeshTable[SkeletonBoneIndex] = MeshBoneIndex;
 	}
-
-	int32 NewIndex = INDEX_NONE;
-	//LinkupCache lock scope
-	{
-		FScopeLock ScopeLock(&LinkupCacheLock);
-		NewIndex = LinkupCache.Add(NewMeshLinkup);
-		check(NewIndex != INDEX_NONE);
-		PRAGMA_DISABLE_DEPRECATION_WARNINGS
-		if (const USkeletalMesh* const InSkelMesh = Cast<const USkeletalMesh>(InSkinnedAsset))
-		{
-			SkelMesh2LinkupCache.Add(MakeWeakObjectPtr(const_cast<USkeletalMesh*>(InSkelMesh)), NewIndex);
-		}
-		PRAGMA_ENABLE_DEPRECATION_WARNINGS
-		SkinnedAsset2LinkupCache.Add(MakeWeakObjectPtr(const_cast<USkinnedAsset*>(InSkinnedAsset)), NewIndex);
-	}
-	return NewIndex;
 }
 
 
@@ -866,7 +912,7 @@ void USkeleton::RebuildLinkup(const USkinnedAsset* InSkinnedAsset)
 	// remove the key
 	RemoveLinkup(InSkinnedAsset);
 	// build new one
-	BuildLinkup(InSkinnedAsset);
+	AddMeshLinkupData(InSkinnedAsset);
 }
 
 void USkeleton::UpdateReferencePoseFromMesh(const USkinnedAsset* InSkinnedAsset)
@@ -1113,31 +1159,19 @@ int32 USkeleton::GetRawAnimationTrackIndex(const int32 InSkeletonBoneIndex, cons
 	return INDEX_NONE;
 }
 
+
 int32 USkeleton::GetSkeletonBoneIndexFromMeshBoneIndex(const USkinnedAsset* InSkinnedAsset, const int32 MeshBoneIndex)
 {
 	check(MeshBoneIndex != INDEX_NONE);
-	const int32 LinkupCacheIdx = GetMeshLinkupIndex(InSkinnedAsset);
-
-	//LinkupCache lock scope
-	{
-		FScopeLock ScopeLock(&LinkupCacheLock);
-		const FSkeletonToMeshLinkup& LinkupTable = LinkupCache[LinkupCacheIdx];
-		return LinkupTable.MeshToSkeletonTable[MeshBoneIndex];
-	}
+	const FSkeletonToMeshLinkup& LinkupTable = FindOrAddMeshLinkupData(InSkinnedAsset);
+	return LinkupTable.MeshToSkeletonTable[MeshBoneIndex];
 }
 
-
 int32 USkeleton::GetMeshBoneIndexFromSkeletonBoneIndex(const USkinnedAsset* InSkinnedAsset, const int32 SkeletonBoneIndex)
 {
 	check(SkeletonBoneIndex != INDEX_NONE);
-	const int32 LinkupCacheIdx = GetMeshLinkupIndex(InSkinnedAsset);
-
-	//LinkupCache lock scope
-	{
-		FScopeLock ScopeLock(&LinkupCacheLock);
-		const FSkeletonToMeshLinkup& LinkupTable = LinkupCache[LinkupCacheIdx];
-		return LinkupTable.SkeletonToMeshTable[SkeletonBoneIndex];
-	}
+	const FSkeletonToMeshLinkup& LinkupTable = FindOrAddMeshLinkupData(InSkinnedAsset);
+	return LinkupTable.SkeletonToMeshTable[SkeletonBoneIndex];
 }
 
 
diff --git a/Engine/Source/Runtime/Engine/Private/BoneContainer.cpp b/Engine/Source/Runtime/Engine/Private/BoneContainer.cpp
--- a/Engine/Source/Runtime/Engine/Private/BoneContainer.cpp
+++ b/Engine/Source/Runtime/Engine/Private/BoneContainer.cpp
@@ -553,10 +553,8 @@ bool FBoneContainer::BoneIsChildOf(const FCompactPoseBoneIndex& BoneIndex, const
 void FBoneContainer::RemapFromSkelMesh(USkeletalMesh const & SourceSkeletalMesh, USkeleton& TargetSkeleton)
 {
 	LLM_SCOPE_BYNAME(TEXT("Animation/BoneContainer"));
-	int32 const SkelMeshLinkupIndex = TargetSkeleton.GetMeshLinkupIndex(&SourceSkeletalMesh);
-	check(SkelMeshLinkupIndex != INDEX_NONE);
-
-	FSkeletonToMeshLinkup const & LinkupTable = TargetSkeleton.LinkupCache[SkelMeshLinkupIndex];
+	
+	const FSkeletonToMeshLinkup& LinkupTable = TargetSkeleton.FindOrAddMeshLinkupData(&SourceSkeletalMesh);
 
 	// Copy LinkupTable arrays for now.
 	// @laurent - Long term goal is to trim that down based on LOD, so we can get rid of the BoneIndicesArray and branch cost of testing if PoseBoneIndex is in that required bone index array.
