Fixed memory reporting for Texture2DArray ResidentMips

diff --git a/Engine/Source/Runtime/Engine/Classes/Engine/Texture2DArray.cpp b/Engine/Source/Runtime/Engine/Classes/Engine/Texture2DArray.cpp
--- a/Engine/Source/Runtime/Engine/Classes/Engine/Texture2DArray.cpp
+++ b/Engine/Source/Runtime/Engine/Classes/Engine/Texture2DArray.cpp
@@ -250,7 +250,11 @@ uint32 UTexture2DArray::CalcTextureMemorySize(int32 MipCount) const
 
 uint32 UTexture2DArray::CalcTextureMemorySizeEnum(ETextureMipCount Enum) const
 {
-	if (Enum == TMC_ResidentMips || Enum == TMC_AllMipsBiased) 
+	if (Enum == TMC_ResidentMips)
+	{
+		return CalcTextureMemorySize(GetNumResidentMips());
+	}
+	else if (Enum == TMC_AllMipsBiased) 
 	{
 		return CalcTextureMemorySize(GetNumMips() - GetCachedLODBias());
 	}
