Content Bundle: better error messages for failures in FContentBundle::DoInitialize.

diff --git a/Engine/Source/Runtime/Engine/Private/WorldPartition/ContentBundle/ContentBundle.cpp b/Engine/Source/Runtime/Engine/Private/WorldPartition/ContentBundle/ContentBundle.cpp
--- a/Engine/Source/Runtime/Engine/Private/WorldPartition/ContentBundle/ContentBundle.cpp
+++ b/Engine/Source/Runtime/Engine/Private/WorldPartition/ContentBundle/ContentBundle.cpp
@@ -36,12 +36,12 @@ void FContentBundle::DoInitialize()
 		}
 		else
 		{
-			UE_LOG(LogContentBundle, Error, TEXT("%s No streaming object found in package %s."), *ContentBundle::Log::MakeDebugInfoString(*this), *GetExternalStreamingObjectPackagePath());
+			UE_LOG(LogContentBundle, Error, TEXT("%s No streaming object found in package %s. No content will be injected."), *ContentBundle::Log::MakeDebugInfoString(*this), *GetExternalStreamingObjectPackagePath());
 		}
 	}
 	else
 	{
-		UE_LOG(LogContentBundle, Log, TEXT("%s No streaming object found. No content will be injected."), *ContentBundle::Log::MakeDebugInfoString(*this));
+		UE_LOG(LogContentBundle, Log, TEXT("%s Streaming package %s not found. No content will be injected."), *ContentBundle::Log::MakeDebugInfoString(*this), *GetExternalStreamingObjectPackagePath());
 	}
 #endif
 
