PR #11322: Print asset name when encountering a mesh card with corrupt derived data

diff --git a/Engine/Source/Runtime/Engine/Private/MeshCardRepresentation.cpp b/Engine/Source/Runtime/Engine/Private/MeshCardRepresentation.cpp
--- a/Engine/Source/Runtime/Engine/Private/MeshCardRepresentation.cpp
+++ b/Engine/Source/Runtime/Engine/Private/MeshCardRepresentation.cpp
@@ -165,6 +165,11 @@ void FCardRepresentationData::CacheDerivedData(const FString& InDDCKey, const IT
 		COOK_STAT(Timer.AddHit(DerivedData.Num()));
 		FMemoryReader Ar(DerivedData, /*bIsPersistent=*/ true);
 		Ar << *this;
+
+		if (Ar.IsError())
+		{
+			UE_LOG(LogStaticMesh, Error, TEXT("Error while deserializing Mesh Card derived data for %s from DDC (key %s)"), *Mesh->GetPathName(), *InDDCKey);
+		}
 	}
 	else
 	{
