Actor Coloration - Make sure to update uniform buffer on call to FPrimitiveSceneProxy::SetPrimitiveColor_GameThread()
* This prevented PrimitiveColor changes and required proxies to be recreated in order for their colors to change.

diff --git a/Engine/Source/Runtime/Engine/Private/PrimitiveSceneProxy.cpp b/Engine/Source/Runtime/Engine/Private/PrimitiveSceneProxy.cpp
--- a/Engine/Source/Runtime/Engine/Private/PrimitiveSceneProxy.cpp
+++ b/Engine/Source/Runtime/Engine/Private/PrimitiveSceneProxy.cpp
@@ -1361,7 +1361,16 @@ void FPrimitiveSceneProxy::SetPrimitiveColor_GameThread(const FLinearColor& InPr
 	ENQUEUE_RENDER_COMMAND(SetSelectionOutlineColorIndex)(
 		[this, InPrimitiveColor](FRHICommandListImmediate&)
 		{
-			PrimitiveColor = InPrimitiveColor;
+			if (PrimitiveColor != InPrimitiveColor)
+			{
+				PrimitiveColor = InPrimitiveColor;
+
+				if (PrimitiveSceneInfo)
+				{
+					Scene->RequestUniformBufferUpdate(*PrimitiveSceneInfo);
+					Scene->RequestGPUSceneUpdate(*PrimitiveSceneInfo, EPrimitiveDirtyState::ChangedOther);
+				}
+			}
 		});
 }
 #endif
