FActorPrimitiveColorHandler - Allow refreshing primitive colors for a list of actors or components

diff --git a/Engine/Source/Runtime/Engine/Classes/GameFramework/ActorPrimitiveColorHandler.h b/Engine/Source/Runtime/Engine/Classes/GameFramework/ActorPrimitiveColorHandler.h
--- a/Engine/Source/Runtime/Engine/Classes/GameFramework/ActorPrimitiveColorHandler.h
+++ b/Engine/Source/Runtime/Engine/Classes/GameFramework/ActorPrimitiveColorHandler.h
@@ -34,12 +34,18 @@ public:
 
 	FActorPrimitiveColorHandler();
 	static FActorPrimitiveColorHandler& Get();
+
 	void RegisterPrimitiveColorHandler(FName InHandlerName, FText InHandlerText, const FFunc& InHandlerFunc);
 	void UnregisterPrimitiveColorHandler(FName InHandlerName);
+	void GetRegisteredPrimitiveColorHandlers(TArray<FPrimitiveColorHandler>& OutPrimitiveColorHandlers) const;
+
+	FName GetActivePrimitiveColorHandler() const;
 	bool SetActivePrimitiveColorHandler(FName InHandlerName, UWorld* InWorld);
+
 	void RefreshPrimitiveColorHandler(FName InHandlerName, UWorld* InWorld);
-	FName GetActivePrimitiveColorHandler() const;
-	void GetRegisteredPrimitiveColorHandlers(TArray<FPrimitiveColorHandler>& OutPrimitiveColorHandlers) const;
+	void RefreshPrimitiveColorHandler(FName InHandlerName, const TArray<AActor*>& InActors);
+	void RefreshPrimitiveColorHandler(FName InHandlerName, const TArray<UPrimitiveComponent*>& InPrimitiveComponents);
+
 	FLinearColor GetPrimitiveColor(const UPrimitiveComponent* InPrimitiveComponent) const;
 
 private:
diff --git a/Engine/Source/Runtime/Engine/Private/ActorPrimitiveColorHandler.cpp b/Engine/Source/Runtime/Engine/Private/ActorPrimitiveColorHandler.cpp
--- a/Engine/Source/Runtime/Engine/Private/ActorPrimitiveColorHandler.cpp
+++ b/Engine/Source/Runtime/Engine/Private/ActorPrimitiveColorHandler.cpp
@@ -72,15 +72,47 @@ void FActorPrimitiveColorHandler::RefreshPrimitiveColorHandler(FName InHandlerNa
 	{
 		for (TActorIterator<AActor> It(InWorld); It; ++It)
 		{
-			TInlineComponentArray<UPrimitiveComponent*> PrimitiveComponents;
-			It->GetComponents(PrimitiveComponents);
+			It->ForEachComponent<UPrimitiveComponent>(false, [this](UPrimitiveComponent* PrimitiveComponent)
+			{
+				if (PrimitiveComponent->IsRegistered())
+				{
+					PrimitiveComponent->PushPrimitiveColorToProxy(GetPrimitiveColor(PrimitiveComponent));
+				}
+			});
+		}
+	}
+#endif
+}
 
-			for (UPrimitiveComponent* PrimitiveComponent : PrimitiveComponents)
+void FActorPrimitiveColorHandler::RefreshPrimitiveColorHandler(FName InHandlerName, const TArray<AActor*>& InActors)
+{
+#if ENABLE_ACTOR_PRIMITIVE_COLOR_HANDLER
+	if (ActivePrimitiveColorHandlerName == InHandlerName)
+	{
+		for (AActor* Actor : InActors)
+		{
+			Actor->ForEachComponent<UPrimitiveComponent>(false, [this](UPrimitiveComponent* PrimitiveComponent)
 			{
 				if (PrimitiveComponent->IsRegistered())
 				{
 					PrimitiveComponent->PushPrimitiveColorToProxy(GetPrimitiveColor(PrimitiveComponent));
 				}
+			});
+		}
+	}
+#endif
+}
+
+void FActorPrimitiveColorHandler::RefreshPrimitiveColorHandler(FName InHandlerName, const TArray<UPrimitiveComponent*>& InPrimitiveComponents)
+{
+#if ENABLE_ACTOR_PRIMITIVE_COLOR_HANDLER
+	if (ActivePrimitiveColorHandlerName == InHandlerName)
+	{
+		for (UPrimitiveComponent* PrimitiveComponent : InPrimitiveComponents)
+		{
+			if (PrimitiveComponent && PrimitiveComponent->IsRegistered())
+			{
+				PrimitiveComponent->PushPrimitiveColorToProxy(GetPrimitiveColor(PrimitiveComponent));
 			}
 		}
 	}
