[shaders] don't print retry message if we hit a global shader compile fail while running a commandlet (retries only apply to the editor)

diff --git a/Engine/Source/Runtime/Engine/Private/ShaderCompiler/ShaderCompiler.cpp b/Engine/Source/Runtime/Engine/Private/ShaderCompiler/ShaderCompiler.cpp
--- a/Engine/Source/Runtime/Engine/Private/ShaderCompiler/ShaderCompiler.cpp
+++ b/Engine/Source/Runtime/Engine/Private/ShaderCompiler/ShaderCompiler.cpp
@@ -10301,18 +10301,21 @@ void ProcessCompiledGlobalShaders(const TArray<FShaderCommonCompileJobPtr>& Comp
 	const int32 UniqueErrorCount = ShaderErrorInfo.UniqueErrors.Num();
 	if (UniqueErrorCount)
 	{
+		const TCHAR* RetryMsg = TEXT(" Enable 'r.ShaderDevelopmentMode' in ConsoleVariables.ini for retries.");
 		if (AreShaderErrorsFatal())
 		{
-			UE_LOGFMT_NSLOC(LogShaders, Fatal, "Shaders", "GlobalShadersCompilationFailed", "{NumErrors} Shader compiler errors compiling GlobalShaders for platform {Platform}.  Enable 'r.ShaderDevelopmentMode' in ConsoleVariables.ini for retries.",
+			UE_LOGFMT_NSLOC(LogShaders, Fatal, "Shaders", "GlobalShadersCompilationFailed", "{NumErrors} Shader compiler errors compiling GlobalShaders for platform {Platform}. {RetryMsg}",
 				("NumErrors", UniqueErrorCount),
-				("Platform", ShaderErrorInfo.TargetShaderPlatformString)
+				("Platform", ShaderErrorInfo.TargetShaderPlatformString),
+				("RetryMsg", IsRunningCommandlet() ? TEXT("") : RetryMsg)
 			);
 		}
 		else
 		{
-			UE_LOGFMT_NSLOC(LogShaders, Error, "Shaders", "GlobalShadersCompilationFailed", "{NumErrors} Shader compiler errors compiling GlobalShaders for platform {Platform}.  Enable 'r.ShaderDevelopmentMode' in ConsoleVariables.ini for retries.",
+			UE_LOGFMT_NSLOC(LogShaders, Error, "Shaders", "GlobalShadersCompilationFailed", "{NumErrors} Shader compiler errors compiling GlobalShaders for platform {Platform}. {RetryMsg}",
 				("NumErrors", UniqueErrorCount),
-				("Platform", ShaderErrorInfo.TargetShaderPlatformString)
+				("Platform", ShaderErrorInfo.TargetShaderPlatformString),
+				("RetryMsg", IsRunningCommandlet() ? TEXT("") : RetryMsg)
 			);
 		}
 	}
