Fix async loading crash when sub-objects aren't yet loaded

diff --git a/Engine/Source/Runtime/Engine/Private/Animation/AnimSequenceBase.cpp b/Engine/Source/Runtime/Engine/Private/Animation/AnimSequenceBase.cpp
--- a/Engine/Source/Runtime/Engine/Private/Animation/AnimSequenceBase.cpp
+++ b/Engine/Source/Runtime/Engine/Private/Animation/AnimSequenceBase.cpp
@@ -139,10 +139,8 @@ void UAnimSequenceBase::PostLoad()
 		    ValidateModel();
 			if (UObject* DataModelObject = DataModelInterface.GetObject())
 			{
-				if (DataModelObject->HasAnyFlags(RF_NeedPostLoad | RF_NeedPostLoadSubobjects))
-				{
-					DataModelObject->ConditionalPostLoad();
-				}
+				DataModelObject->ConditionalPostLoad();
+				DataModelObject->ConditionalPostLoadSubobjects();
 			}
 		    GetController();
 		    BindToModelModificationEvent();
