SparseVolumeTextures now use Clamping UV addressing behavior by default. Also fixed SVT trying to create GPU resources when FApp::CanEverRender() is false.

#rnx

diff --git a/Engine/Source/Runtime/Engine/Private/SparseVolumeTexture/SparseVolumeTexture.cpp b/Engine/Source/Runtime/Engine/Private/SparseVolumeTexture/SparseVolumeTexture.cpp
--- a/Engine/Source/Runtime/Engine/Private/SparseVolumeTexture/SparseVolumeTexture.cpp
+++ b/Engine/Source/Runtime/Engine/Private/SparseVolumeTexture/SparseVolumeTexture.cpp
@@ -911,7 +911,7 @@ bool USparseVolumeTextureFrame::Initialize(USparseVolumeTexture* InOwner, int32
 
 bool USparseVolumeTextureFrame::CreateTextureRenderResources()
 {
-	if (!TextureRenderResources)
+	if (!TextureRenderResources && FApp::CanEverRender())
 	{
 		TextureRenderResources = new UE::SVT::FTextureRenderResources();
 		TextureRenderResources->SetGlobalVolumeResolution_GameThread(Owner->GetVolumeResolution());
@@ -1072,6 +1072,12 @@ bool UStreamableSparseVolumeTexture::BeginInitialize(int32 NumExpectedFrames)
 	check(FormatA == PF_Unknown);
 	check(FormatB == PF_Unknown);
 
+	// This is different to other texture types which all seem to default to wrap. However, for SVT content it is more common to want a single volume without wrapping
+	// and since changing the addressing mode currently results in fairly costly recomputing of derived data, defaulting to clamp should result in a better user experience in most cases.
+	AddressX = TA_Clamp;
+	AddressY = TA_Clamp;
+	AddressZ = TA_Clamp;
+
 	InitState = EInitState_Pending;
 
 	return true;
