Prefer UBA over XGE when multiple shader controllers are available.

#rnx

diff --git a/Engine/Source/Runtime/Engine/Private/ShaderCompiler/ShaderCompiler.cpp b/Engine/Source/Runtime/Engine/Private/ShaderCompiler/ShaderCompiler.cpp
--- a/Engine/Source/Runtime/Engine/Private/ShaderCompiler/ShaderCompiler.cpp
+++ b/Engine/Source/Runtime/Engine/Private/ShaderCompiler/ShaderCompiler.cpp
@@ -5425,6 +5425,24 @@ bool FShaderCompilingManager::AllTargetPlatformSupportsRemoteShaderCompiling()
 	return true;
 }
 
+// Returns a rank for the preference of distributed shader controllers; Higher is better.
+static int32 GetShaderControllerPreferenceRank(IDistributedBuildController& Controller)
+{
+	const FString Name = Controller.GetName();
+	if (Name.StartsWith(TEXT("UBA")))
+	{
+		return 2;
+	}
+	else if (Name.StartsWith(TEXT("XGE")))
+	{
+		return 1;
+	}
+	else
+	{
+		return 0;
+	}
+}
+
 IDistributedBuildController* FShaderCompilingManager::FindRemoteCompilerController() const
 {
 	// no controllers needed if not compiling
@@ -5435,31 +5453,27 @@ IDistributedBuildController* FShaderCompilingManager::FindRemoteCompilerControll
 
 	TArray<IDistributedBuildController*> AvailableControllers = IModularFeatures::Get().GetModularFeatureImplementations<IDistributedBuildController>(IDistributedBuildController::GetModularFeatureType());
 
-	auto FindXGE = [](IDistributedBuildController* Controller)
-	{
-		return Controller != nullptr
-			&& Controller->GetName().StartsWith("XGE");
-	};
+	// Prefer UBA, then XGE, and fallback to any other controller otherwise
+	int32 SupportedControllerPreferenceRank = 0;
+	IDistributedBuildController* SupportedController = nullptr;
 
-	// Prefer XGE
-	if (IDistributedBuildController** ControllerPtr = AvailableControllers.FindByPredicate(FindXGE))
+	for (IDistributedBuildController* Controller : AvailableControllers)
 	{
-		IDistributedBuildController* Controller = *ControllerPtr;
-		if (Controller && Controller->IsSupported())
+		if (Controller != nullptr && Controller->IsSupported())
 		{
-			Controller->InitializeController();
-			return Controller;
+			const int32 PreferenceRank = GetShaderControllerPreferenceRank(*Controller);
+			if (SupportedController == nullptr || SupportedControllerPreferenceRank < PreferenceRank)
+			{
+				SupportedController = Controller;
+				SupportedControllerPreferenceRank = PreferenceRank;
+			}
 		}
 	}
 
-	// Fall back on whatever is available
-	for (IDistributedBuildController* Controller : AvailableControllers)
+	if (SupportedController != nullptr)
 	{
-		if (Controller != nullptr && Controller->IsSupported())
-		{
-			Controller->InitializeController();
-			return Controller;
-		}
+		SupportedController->InitializeController();
+		return SupportedController;
 	}
 
 	return nullptr;
