Avoid mutating the reference skeleton when we aren't loading

diff --git a/Engine/Source/Runtime/Engine/Private/Animation/Skeleton.cpp b/Engine/Source/Runtime/Engine/Private/Animation/Skeleton.cpp
--- a/Engine/Source/Runtime/Engine/Private/Animation/Skeleton.cpp
+++ b/Engine/Source/Runtime/Engine/Private/Animation/Skeleton.cpp
@@ -511,9 +511,7 @@ void USkeleton::Serialize( FArchive& Ar )
 		}
 	}
 
-	// This is crashing when live coding in debug - the ObjectReferenceCollector 
-	// is multithreaded and so I assume that is causing some subtle problem
-	if (!Ar.IsObjectReferenceCollector())
+	if (Ar.IsLoading())
 	{
 		const bool bRebuildNameMap = false;
 		ReferenceSkeleton.RebuildRefSkeleton(this, bRebuildNameMap);
