SkeletalMesh: changed morph targets validity test

diff --git a/Engine/Source/Runtime/Engine/Private/Rendering/SkeletalMeshLODImporterData.cpp b/Engine/Source/Runtime/Engine/Private/Rendering/SkeletalMeshLODImporterData.cpp
--- a/Engine/Source/Runtime/Engine/Private/Rendering/SkeletalMeshLODImporterData.cpp
+++ b/Engine/Source/Runtime/Engine/Private/Rendering/SkeletalMeshLODImporterData.cpp
@@ -1815,14 +1815,25 @@ bool FSkeletalMeshImportData::GetMeshDescription(FMeshDescription& OutMeshDescri
 
 	// Register morph targets.
 	TSet<FString> ValidMorphTargets;
-	if (ensure(MorphTargetNames.Num() == MorphTargets.Num()))
+	if (ensure(MorphTargetNames.Num() == MorphTargets.Num() && MorphTargetModifiedPoints.Num() == MorphTargets.Num()))
 	{
 		for (int32 MorphIndex = 0; MorphIndex < MorphTargets.Num(); MorphIndex++)
 		{
 			const FSkeletalMeshImportData& MorphTargetMesh = MorphTargets[MorphIndex];
-			
-			// As long as we have the same point count, we have a valid morph target mesh. 
-			if (!ensure(MorphTargetMesh.Points.Num() == Points.Num()))
+			const TSet<uint32>& ModifiedPoints = MorphTargetModifiedPoints[MorphIndex];
+
+			// ensure they both have the same number of points
+			if (!ensure(MorphTargetMesh.Points.Num() == ModifiedPoints.Num()))
+			{
+				continue;
+			}
+
+			// ensure modified points indices are valid
+			const uint32* InvalidIndex = Algo::FindByPredicate(ModifiedPoints, [this](uint32 PointIndex)
+			{
+				return !Points.IsValidIndex(PointIndex);
+			});
+			if (!ensure(InvalidIndex == nullptr))
 			{
 				continue;
 			}
@@ -2074,10 +2085,12 @@ bool FSkeletalMeshImportData::GetMeshDescription(FMeshDescription& OutMeshDescri
 		const TSet<uint32>& ModifiedPoints = MorphTargetModifiedPoints[MorphTargetIndex]; 
 		FMorphTargetVertexAttributesRef MorphTargetRef = MeshAttributes.GetVertexMorphTarget(*MorphTargetName);
 
+		int32 ModifiedPointIndex = 0;
 		for (uint32 PointIndex: ModifiedPoints)
 		{
 			const FVertexID VertexID = VertexIDMap[PointIndex];
-			MorphTargetRef.SetPositionDelta(VertexID, MorphPoints[PointIndex] - Points[PointIndex]);
+			MorphTargetRef.SetPositionDelta(VertexID, MorphPoints[ModifiedPointIndex] - Points[PointIndex]);
+			ModifiedPointIndex++;
 		}
 	}
 
