HeterogeneousComponent now supports user supplied MIDs. If the given UMaterialInterface is already a MID, the HV component simply reuses it and otherwise creates a new MID. The SVT is retrieved from the currently used MID, which allows users to update the SVT directly on the MID.

#rnx

diff --git a/Engine/Source/Runtime/Engine/Classes/Components/HeterogeneousVolumeComponent.h b/Engine/Source/Runtime/Engine/Classes/Components/HeterogeneousVolumeComponent.h
--- a/Engine/Source/Runtime/Engine/Classes/Components/HeterogeneousVolumeComponent.h
+++ b/Engine/Source/Runtime/Engine/Classes/Components/HeterogeneousVolumeComponent.h
@@ -103,12 +103,14 @@ public:
 	virtual int32 GetNumMaterials() const override { return 1; }
 	virtual bool ShouldRenderSelected() const override { return true; }
 	virtual void GetUsedMaterials(TArray<UMaterialInterface*>& OutMaterials, bool bGetDebugMaterials = false) const override;
+	virtual void SetMaterial(int32 ElementIndex, class UMaterialInterface* Material) override;
 	//~ End UPrimitiveComponent Interface.
 
 private:
 	const USparseVolumeTexture* PreviousSVT;
 
-	const USparseVolumeTexture* GetSparseVolumeTexture() const;
+	static USparseVolumeTexture* GetSparseVolumeTexture(UMaterialInterface* MaterialInterface, int32 ParameterIndex, FName* OutParamName = nullptr);
+	static UMaterialInstanceDynamic* CreateOrCastToMID(UMaterialInterface* MaterialInterface);
 	void OnSparseVolumeTextureChanged(const USparseVolumeTexture* SparseVolumeTexture);
 };
 
diff --git a/Engine/Source/Runtime/Engine/Private/Components/HeterogeneousVolumeComponent.cpp b/Engine/Source/Runtime/Engine/Private/Components/HeterogeneousVolumeComponent.cpp
--- a/Engine/Source/Runtime/Engine/Private/Components/HeterogeneousVolumeComponent.cpp
+++ b/Engine/Source/Runtime/Engine/Private/Components/HeterogeneousVolumeComponent.cpp
@@ -337,40 +337,66 @@ void UHeterogeneousVolumeComponent::PostLoad()
 		const UMaterial* Material = MaterialInterface->GetMaterial();
 		if (Material && Material->MaterialDomain == EMaterialDomain::MD_Volume)
 		{
-			Material->GetRenderProxy();
 			MaterialInterface->CheckMaterialUsage(MATUSAGE_HeterogeneousVolumes);
 		}
 
-		MaterialInstanceDynamic = UMaterialInstanceDynamic::Create(MaterialInterface, nullptr);
+		MaterialInstanceDynamic = CreateOrCastToMID(MaterialInterface);
 	}
 }
 
-const USparseVolumeTexture* UHeterogeneousVolumeComponent::GetSparseVolumeTexture() const
+USparseVolumeTexture* UHeterogeneousVolumeComponent::GetSparseVolumeTexture(UMaterialInterface* MaterialInterface, int32 ParameterIndex, FName* OutParamName)
 {
 	USparseVolumeTexture* SparseVolumeTexture = nullptr;
 
-	uint32 MaterialIndex = 0;
-	UMaterialInterface* MaterialInterface = GetMaterial(MaterialIndex);
 	if (MaterialInterface)
 	{
+		// Get parameter infos for all SVTs in the material
 		TArray<FMaterialParameterInfo> ParameterInfo;
 		TArray<FGuid> ParameterIds;
 		MaterialInterface->GetAllSparseVolumeTextureParameterInfo(ParameterInfo, ParameterIds);
 
-		if (!ParameterInfo.IsEmpty())
+		// Get the SVT object
+		if (ParameterInfo.IsValidIndex(ParameterIndex))
 		{
-			MaterialInterface->GetSparseVolumeTextureParameterValue(ParameterInfo[MaterialIndex], SparseVolumeTexture);
+			MaterialInterface->GetSparseVolumeTextureParameterValue(ParameterInfo[ParameterIndex], SparseVolumeTexture);
+		}
+
+		// The SVT in MaterialInterface might be a frame of a UStreamableSparseVolumeTexture. In that case we try to get the owning SVT object.
+		if (SparseVolumeTexture && SparseVolumeTexture->IsA<USparseVolumeTextureFrame>())
+		{
+			USparseVolumeTextureFrame* Frame = Cast<USparseVolumeTextureFrame>(SparseVolumeTexture);
+			UObject* FrameOuter = Frame->GetOuter();
+			check(FrameOuter->IsA<UStreamableSparseVolumeTexture>());
+			SparseVolumeTexture = Cast<USparseVolumeTexture>(FrameOuter);
+			check(SparseVolumeTexture);
+		}
+
+		if (SparseVolumeTexture && OutParamName)
+		{
+			*OutParamName = ParameterInfo[ParameterIndex].Name;
 		}
 	}
 
 	return SparseVolumeTexture;
 }
 
+UMaterialInstanceDynamic* UHeterogeneousVolumeComponent::CreateOrCastToMID(UMaterialInterface* MaterialInterface)
+{
+	if (MaterialInterface->IsA<UMaterialInstanceDynamic>())
+	{
+		return Cast<UMaterialInstanceDynamic>(MaterialInterface);
+	}
+	else
+	{
+		return UMaterialInstanceDynamic::Create(MaterialInterface, nullptr);
+	}
+}
+
 void UHeterogeneousVolumeComponent::OnSparseVolumeTextureChanged(const USparseVolumeTexture* SparseVolumeTexture)
 {
 	if (SparseVolumeTexture)
 	{
-		VolumeResolution = SparseVolumeTexture ? SparseVolumeTexture->GetVolumeResolution() : FIntVector(1);
+		VolumeResolution = SparseVolumeTexture->GetVolumeResolution();
 		StartFrame = 0;
 		EndFrame = SparseVolumeTexture->GetNumFrames() - 1;
 		Frame = FMath::Clamp(Frame, StartFrame, EndFrame);
@@ -392,6 +418,8 @@ void UHeterogeneousVolumeComponent::PostEditChangeProperty(FPropertyChangedEvent
 {
 	Super::PostEditChangeProperty(PropertyChangedEvent);
 
+	const int32 SVTParameterIndex = 0;
+
 	FName PropertyName;
 	if (PropertyChangedEvent.Property)
 	{
@@ -400,6 +428,7 @@ void UHeterogeneousVolumeComponent::PostEditChangeProperty(FPropertyChangedEvent
 
 	if (PropertyName == GET_MEMBER_NAME_CHECKED(UHeterogeneousVolumeComponent, OverrideMaterials))
 	{
+		MaterialInstanceDynamic = nullptr; // Reset internal MID. We either create a new one from the new material or leave it as null if the material was unset
 		uint32 MaterialIndex = 0;
 		UMaterialInterface* MaterialInterface = GetMaterial(MaterialIndex);
 		if (MaterialInterface)
@@ -407,19 +436,18 @@ void UHeterogeneousVolumeComponent::PostEditChangeProperty(FPropertyChangedEvent
 			const UMaterial* Material = MaterialInterface->GetMaterial();
 			if (Material && Material->MaterialDomain == EMaterialDomain::MD_Volume)
 			{
-				Material->GetRenderProxy();
 				MaterialInterface->CheckMaterialUsage(MATUSAGE_HeterogeneousVolumes);
 			}
 
-			MaterialInstanceDynamic = UMaterialInstanceDynamic::Create(MaterialInterface, nullptr);
-			OnSparseVolumeTextureChanged(GetSparseVolumeTexture());
+			MaterialInstanceDynamic = CreateOrCastToMID(MaterialInterface);
 		}
+		OnSparseVolumeTextureChanged(GetSparseVolumeTexture(MaterialInstanceDynamic, SVTParameterIndex));
 	}
 
 	if (PropertyName == GET_MEMBER_NAME_CHECKED(UHeterogeneousVolumeComponent, VolumeResolution))
 	{
 		// Prevent resolution changes when using SVT
-		const USparseVolumeTexture* SparseVolumeTexture = GetSparseVolumeTexture();
+		const USparseVolumeTexture* SparseVolumeTexture = GetSparseVolumeTexture(MaterialInstanceDynamic, SVTParameterIndex);
 		if (SparseVolumeTexture)
 		{
 			VolumeResolution = SparseVolumeTexture->GetVolumeResolution();
@@ -428,7 +456,7 @@ void UHeterogeneousVolumeComponent::PostEditChangeProperty(FPropertyChangedEvent
 
 	if (PropertyName == GET_MEMBER_NAME_CHECKED(UHeterogeneousVolumeComponent, Frame))
 	{
-		const USparseVolumeTexture* SparseVolumeTexture = GetSparseVolumeTexture();
+		const USparseVolumeTexture* SparseVolumeTexture = GetSparseVolumeTexture(MaterialInstanceDynamic, SVTParameterIndex);
 		if (SparseVolumeTexture)
 		{
 			Frame = FMath::Clamp(Frame, StartFrame, EndFrame);
@@ -437,7 +465,7 @@ void UHeterogeneousVolumeComponent::PostEditChangeProperty(FPropertyChangedEvent
 
 	if (PropertyName == GET_MEMBER_NAME_CHECKED(UHeterogeneousVolumeComponent, StartFrame))
 	{
-		const USparseVolumeTexture* SparseVolumeTexture = GetSparseVolumeTexture();
+		const USparseVolumeTexture* SparseVolumeTexture = GetSparseVolumeTexture(MaterialInstanceDynamic, SVTParameterIndex);
 		if (SparseVolumeTexture)
 		{
 			StartFrame = FMath::Clamp(StartFrame, 0, EndFrame);
@@ -447,7 +475,7 @@ void UHeterogeneousVolumeComponent::PostEditChangeProperty(FPropertyChangedEvent
 
 	if (PropertyName == GET_MEMBER_NAME_CHECKED(UHeterogeneousVolumeComponent, EndFrame))
 	{
-		const USparseVolumeTexture* SparseVolumeTexture = GetSparseVolumeTexture();
+		const USparseVolumeTexture* SparseVolumeTexture = GetSparseVolumeTexture(MaterialInstanceDynamic, SVTParameterIndex);
 		if (SparseVolumeTexture)
 		{
 			const int32 FrameCount = SparseVolumeTexture->GetNumFrames();
@@ -468,67 +496,66 @@ void UHeterogeneousVolumeComponent::GetUsedMaterials(TArray<UMaterialInterface*>
 	}
 }
 
+void UHeterogeneousVolumeComponent::SetMaterial(int32 ElementIndex, UMaterialInterface* Material)
+{
+	Super::SetMaterial(ElementIndex, Material);
+	if (ElementIndex == 0)
+	{
+		MaterialInstanceDynamic = CreateOrCastToMID(Material);
+		OnSparseVolumeTextureChanged(GetSparseVolumeTexture(MaterialInstanceDynamic, 0 /*SVTParameterIndex*/));
+	}
+}
+
 void UHeterogeneousVolumeComponent::TickComponent(float DeltaTime, ELevelTick TickType, FActorComponentTickFunction* ThisTickFunction)
 {
 	Super::TickComponent(DeltaTime, TickType, ThisTickFunction);
 
-	uint32 MaterialIndex = 0;
-	UMaterialInterface* MaterialInterface = UMeshComponent::GetMaterial(MaterialIndex);
-	if (MaterialInterface)
+	if (MaterialInstanceDynamic)
 	{
-		// Get all SVT params in the material
-		TArray<FMaterialParameterInfo> ParameterInfo;
-		TArray<FGuid> ParameterIds;
-		MaterialInterface->GetAllSparseVolumeTextureParameterInfo(ParameterInfo, ParameterIds);
+		const int32 SVTParameterIndex = 0;
+		FName SVTParameterName;
+		USparseVolumeTexture* SparseVolumeTexture = GetSparseVolumeTexture(MaterialInstanceDynamic, SVTParameterIndex, &SVTParameterName);
 
-		if (!ParameterInfo.IsEmpty())
-		{
-			USparseVolumeTexture* SparseVolumeTexture = nullptr;
-			bool bValid = MaterialInterface->GetSparseVolumeTextureParameterValue(ParameterInfo[MaterialIndex], SparseVolumeTexture);
 #if WITH_EDITOR
-			// Detect an update to the material
-			if (SparseVolumeTexture != PreviousSVT)
-			{
-				OnSparseVolumeTextureChanged(SparseVolumeTexture);
-			}
+		// Detect an update to the material
+		if (SparseVolumeTexture != PreviousSVT)
+		{
+			OnSparseVolumeTextureChanged(SparseVolumeTexture);
+		}
 #endif // WITH_EDITOR
 
-			if (SparseVolumeTexture)
-			{
-				const int32 FrameCount = SparseVolumeTexture->GetNumFrames();
-
-				// Determine active frame based on animation controls if playing
-				if (bPlaying)
-				{
-					Frame += DeltaTime * FrameRate;
-				}
+		if (SparseVolumeTexture)
+		{
+			const int32 FrameCount = SparseVolumeTexture->GetNumFrames();
 
-				if (bLooping)
-				{
-					float FrameRange = EndFrame - StartFrame + 1;
-					Frame = FMath::Fmod(Frame - StartFrame, (float)FrameRange) + StartFrame;
-				}
-				else
-				{
-					Frame = FMath::Clamp(Frame, StartFrame, EndFrame);
-				}
+			// Determine active frame based on animation controls if playing
+			if (bPlaying)
+			{
+				Frame += DeltaTime * FrameRate;
+			}
 
-				bool bIsBlocking = bIssueBlockingRequests != 0;
-				USparseVolumeTextureFrame* SparseVolumeTextureFrame = USparseVolumeTextureFrame::GetFrameAndIssueStreamingRequest(SparseVolumeTexture, Frame, MipLevel, bIsBlocking);
-				if (SparseVolumeTextureFrame)
-				{
-					FIntVector PerFrameVolumeResolution = SparseVolumeTextureFrame->GetVolumeResolution();
-					if (VolumeResolution != PerFrameVolumeResolution)
-					{
-						MarkRenderTransformDirty();
-					}
-				}
+			if (bLooping)
+			{
+				float FrameRange = EndFrame - StartFrame + 1;
+				Frame = FMath::Fmod(Frame - StartFrame, (float)FrameRange) + StartFrame;
+			}
+			else
+			{
+				Frame = FMath::Clamp(Frame, StartFrame, EndFrame);
+			}
 
-				if (MaterialInstanceDynamic)
+			bool bIsBlocking = bIssueBlockingRequests != 0;
+			USparseVolumeTextureFrame* SparseVolumeTextureFrame = USparseVolumeTextureFrame::GetFrameAndIssueStreamingRequest(SparseVolumeTexture, Frame, MipLevel, bIsBlocking);
+			if (SparseVolumeTextureFrame)
+			{
+				FIntVector PerFrameVolumeResolution = SparseVolumeTextureFrame->GetVolumeResolution();
+				if (VolumeResolution != PerFrameVolumeResolution)
 				{
-					MaterialInstanceDynamic->SetSparseVolumeTextureParameterValue(ParameterInfo[MaterialIndex].Name, SparseVolumeTextureFrame);
+					MarkRenderTransformDirty();
 				}
 			}
+
+			MaterialInstanceDynamic->SetSparseVolumeTextureParameterValue(SVTParameterName, SparseVolumeTextureFrame);
 		}
 	}
 }
