Improving multiple PIE session VisLog capabilities.  The TimeStamp used is now a global Engine TimeStamp.  The WorldTimeStamp is also recorded separately for ease of debugging.  By using this Engine TimeStamp during PIE sesions, we can correctly record the time for client & server events, allowing us to see the effects of replicated data more clearly on the timeline.

When using bForceUniqueLogNames, we record offsets to an initial Engine TimeStamp which gives us similar functionality to just using WorldTimeStamp in the Log Visualizer (i.e. multiple runs will line-up at time 0).
When not using bForceUniqueLogNames, we only record the offset when the data is cleared in the Log Visualizer.  That will allow multiple runs to record properly on a single timeline (rather than having new data overwrite old data).

bForceUniqueLogNames is now the default setting.

diff --git a/Engine/Source/Developer/LogVisualizer/Private/LogVisualizer.cpp b/Engine/Source/Developer/LogVisualizer/Private/LogVisualizer.cpp
--- a/Engine/Source/Developer/LogVisualizer/Private/LogVisualizer.cpp
+++ b/Engine/Source/Developer/LogVisualizer/Private/LogVisualizer.cpp
@@ -9,6 +9,7 @@
 #include "LogVisualizerSettings.h"
 #include "VisualLoggerDatabase.h"
 #include "VisualLoggerCameraController.h"
+#include "VisualLogger/VisualLogger.h"
 
 #if WITH_EDITOR
 #include "Editor/EditorEngine.h"
@@ -79,6 +80,7 @@ void FLogVisualizer::Shutdown()
 void FLogVisualizer::Reset()
 {
 	TimeSliderController->SetTimesliderArgs(FVisualLoggerTimeSliderArgs());
+	FVisualLogger::Get().OnDataReset();
 }
 
 FLogVisualizer& FLogVisualizer::Get()
diff --git a/Engine/Source/Developer/LogVisualizer/Private/SVisualLoggerStatusView.cpp b/Engine/Source/Developer/LogVisualizer/Private/SVisualLoggerStatusView.cpp
--- a/Engine/Source/Developer/LogVisualizer/Private/SVisualLoggerStatusView.cpp
+++ b/Engine/Source/Developer/LogVisualizer/Private/SVisualLoggerStatusView.cpp
@@ -165,13 +165,13 @@ void SVisualLoggerStatusView::GenerateStatusData(const FVisualLogDevice::FVisual
 	if (bAddHeader)
 	{
 		FLogStatusItem* HeaderItem = new FLogStatusItem();
-		HeaderItem->HeaderText = FString::Printf(TEXT("%s at Time: %.2fs"), *LogEntry.OwnerName.ToString(), LogEntry.Entry.TimeStamp);;
+		HeaderItem->HeaderText = FString::Printf(TEXT("%s at WorldTime: %.2lfs"), *LogEntry.OwnerName.ToString(), LogEntry.Entry.WorldTimeStamp);
 		StatusItems.Add(MakeShareable(HeaderItem));
 	}
 	else
 	{
-		FString TimestampDesc = FString::Printf(TEXT("%.2fs"), LogEntry.Entry.TimeStamp);
-		StatusItems.Add(MakeShareable(new FLogStatusItem(LOCTEXT("VisLogTimestamp", "Time").ToString(), TimestampDesc)));
+		FString TimestampDesc = FString::Printf(TEXT("%.2lfs"), LogEntry.Entry.WorldTimeStamp);
+		StatusItems.Add(MakeShareable(new FLogStatusItem(LOCTEXT("VisLogTimestamp", "WorldTime").ToString(), TimestampDesc)));
 	}
 
 	for (int32 CategoryIndex = 0; CategoryIndex < LogEntry.Entry.Status.Num(); CategoryIndex++)
diff --git a/Engine/Source/Developer/LogVisualizer/Private/SVisualLoggerTimelineBar.cpp b/Engine/Source/Developer/LogVisualizer/Private/SVisualLoggerTimelineBar.cpp
--- a/Engine/Source/Developer/LogVisualizer/Private/SVisualLoggerTimelineBar.cpp
+++ b/Engine/Source/Developer/LogVisualizer/Private/SVisualLoggerTimelineBar.cpp
@@ -72,8 +72,8 @@ FReply SVisualLoggerTimelineBar::OnMouseMove(const FGeometry& MyGeometry, const
 		if (Items.IsValidIndex(MouseMoveClosestItemIndex))
 		{
 			const FVisualLogEntry& CurrentEntry = Items[MouseMoveClosestItemIndex].Entry;
-				
-			TooltipBuilder = FString::Printf(TEXT("Time: %.2f"), CurrentEntry.TimeStamp);
+			
+			TooltipBuilder = FString::Printf(TEXT("Time: %.2lf WorldTime: %.2lf"), CurrentEntry.TimeStamp, CurrentEntry.WorldTimeStamp);
 
 			for (const FVisualLogShapeElement& Shape : CurrentEntry.ElementsToDraw)
 			{
diff --git a/Engine/Source/Runtime/Core/Public/UObject/UE5MainStreamObjectVersion.h b/Engine/Source/Runtime/Core/Public/UObject/UE5MainStreamObjectVersion.h
--- a/Engine/Source/Runtime/Core/Public/UObject/UE5MainStreamObjectVersion.h
+++ b/Engine/Source/Runtime/Core/Public/UObject/UE5MainStreamObjectVersion.h
@@ -369,6 +369,9 @@ struct FUE5MainStreamObjectVersion
 
 		// Cloth collection single lod schema
 		ClothCollectionSingleLodSchema,
+
+		// Visual Logger format now includes a WorldTimeStamp in addition to TimeStamp for easier debugging between multiple instances.
+		VisualLoggerAddedSeparateWorldTime,
 		
 		// -----<new versions can be added above this line>-------------------------------------------------
 		VersionPlusOne,
diff --git a/Engine/Source/Runtime/Engine/Private/VisualLogger/VisualLogger.cpp b/Engine/Source/Runtime/Engine/Private/VisualLogger/VisualLogger.cpp
--- a/Engine/Source/Runtime/Engine/Private/VisualLogger/VisualLogger.cpp
+++ b/Engine/Source/Runtime/Engine/Private/VisualLogger/VisualLogger.cpp
@@ -19,10 +19,26 @@
 #include "Serialization/CustomVersion.h"
 #endif
 
-
 DEFINE_LOG_CATEGORY(LogVisual);
 #if ENABLE_VISUAL_LOG
 
+namespace UE::VisLog::Private
+{
+	/** Denotes the consecutive run number, giving us the ability to differentiate UObjectName_0 in Run 1 from UObjectName_0 in Run 2 */
+	static int32 UniqueRunNumber = 0;
+
+	/** Cached prefix string when using bForceUniqueLogNames (denotes UniqueRunNumber) */
+	static FString UniqueLogPrefix;
+
+#if WITH_EDITOR
+	namespace EditorOnly
+	{
+		/** When in Editor, use this TimeStamp as a Base time so our values don't grow too large (and are canonically around 0.0). */
+		static double EditorBaseTimeStamp = 0.0;
+	}
+#endif
+}
+
 DEFINE_STAT(STAT_VisualLog);
 
 // Unfortunately needs to be a #define since it uses GET_TYPED_VARARGS_RESULT which uses the va_list stuff which operates on the
@@ -221,17 +237,43 @@ bool FVisualLogger::CheckVisualLogInputInternal(const UObject* Object, const FNa
 
 double FVisualLogger::GetTimeStampForObject(const UObject* Object) const
 {
+	// Licensees can write their own synchronized clock that can work across the network
+	// such an implementation could use the exchanged ServerTime and the relative ClientTime offset.
 	if (GetTimeStampFunc)
 	{
 		return GetTimeStampFunc(Object);
 	}
 
-	if (const UWorld* World = GetWorldForVisualLogger(Object))
+	const UWorld* WorldForTimeStamp = nullptr;
+#if WITH_EDITOR
+	// When we're in the Editor, use a Global Engine TimeStamp so that we can synchronize between Client and Server instances.
+	UEditorEngine* EditorEngine = GIsEditor ? Cast<UEditorEngine>(GEngine) : nullptr;
+	if (EditorEngine)
 	{
-		return World->TimeSeconds;
+		// We will always have the Editor world to use.  This will ensure a consistent clock since it does not reset
+		// when more clients are added or removed and can exist before a PIE session is started.
+		WorldForTimeStamp = EditorEngine->GetEditorWorldContext().World();
+		if (ensureMsgf(WorldForTimeStamp, TEXT("We always expect to have an EditorWorld in Editor")))
+		{
+			using namespace UE::VisLog::Private;
+			if (EditorOnly::EditorBaseTimeStamp <= 0.0)
+			{
+				EditorOnly::EditorBaseTimeStamp = WorldForTimeStamp->TimeSeconds;
+			}
+
+			return WorldForTimeStamp->TimeSeconds - EditorOnly::EditorBaseTimeStamp;
+		}
+	}
+
+#endif
+
+	// This will be the fallback mode in standalone.  We do not have a synchronized clock.
+	if (!WorldForTimeStamp)
+	{
+		WorldForTimeStamp = GetWorldForVisualLogger(Object);
 	}
 
-	return 0;
+	return WorldForTimeStamp ? WorldForTimeStamp->TimeSeconds : 0.0;
 }
 
 void FVisualLogger::SetGetTimeStampFunc(const TFunction<double(const UObject*)> Function)
@@ -408,6 +450,8 @@ FVisualLogEntry* FVisualLogger::GetEntryToWrite(const UObject* Object, const dou
 
 FVisualLogEntry* FVisualLogger::GetEntryToWriteInternal(const UObject* Object, const double TimeStamp, const ECreateIfNeeded ShouldCreate)
 {
+	using namespace UE::VisLog::Private;
+
 	// No redirection needed, it should have been done at the time of the thread entry was computed
 	const UObject* LogOwner = Object;
 	if (LogOwner == nullptr)
@@ -436,12 +480,17 @@ FVisualLogEntry* FVisualLogger::GetEntryToWriteInternal(const UObject* Object, c
 		// It's first and only one usage of LogOwner as regular object to get names. We assume once that LogOwner is correct here and only here.
 		CurrentEntry = &CurrentEntryPerObject.Add(LogOwner);
 
+		if (bForceUniqueLogNames && UniqueLogPrefix.IsEmpty())
+		{
+			UniqueLogPrefix = FString::Printf(TEXT("[%d] "), UniqueRunNumber);
+		}
+
 		const UWorld* World = GetWorldForVisualLogger(LogOwner);
 		const bool bIsStandalone = (World == nullptr || World->GetNetMode() == NM_Standalone);
 		const FName LogName(*FString::Printf(TEXT("%s%s%s"),
-			bIsStandalone ? TEXT("") : *FString::Printf(TEXT("(%s) "), *ToString(World->GetNetMode())),
-			*LogOwner->GetName(),
-			bForceUniqueLogNames ? *FString::Printf(TEXT(" [%d]"), LogOwner->GetUniqueID()) : TEXT("")));
+			*UniqueLogPrefix,
+			bIsStandalone ? TEXT("") : *FString::Printf(TEXT("(%s) "), *GetDebugStringForWorld(World)),
+			*LogOwner->GetName()));
 
 		ObjectToNameMap.Add(LogOwner, LogName);
 		ObjectToClassNameMap.Add(LogOwner, *(LogOwner->GetClass()->GetName()));
@@ -457,6 +506,14 @@ FVisualLogEntry* FVisualLogger::GetEntryToWriteInternal(const UObject* Object, c
 		checkf(CurrentEntry != nullptr, TEXT("bInitializeEntry can only be true when CurrentEntry is valid."));
 		CurrentEntry->InitializeEntry(TimeStamp);
 
+		// Let's record the World Time as the local instance sees it.
+		if (const TWeakObjectPtr<const UWorld>* WorldWeakPtr = ObjectToWorldMap.Find(LogOwner))
+		{
+			if (const UWorld* World = WorldWeakPtr->Get())
+			{
+				CurrentEntry->WorldTimeStamp = World->TimeSeconds;
+			}
+		}
 
 		if (const AActor* ObjectAsActor = Cast<AActor>(LogOwner))
 		{
@@ -704,15 +761,59 @@ FVisualLogger::FVisualLogger()
 
 		return true;
 	});
+
+#if WITH_EDITOR
+	// When PIE is Starting, we want to reset the EditorBaseTimeStamp on the very first PIE World Instance (but not the other instances).
+	// This gives us a pseudo-world time (that starts at zero) while still being synchronized across multiple PIE instances.
+	PIEStartedHandle = FWorldDelegates::OnPIEStarted.AddLambda(
+		[this](UGameInstance*) {
+			// We only want to do this if we are using unique log names, otherwise it makes more sense
+			// to just continue logging at later timestamps on the same timeline as the previous one.
+			if (bForceUniqueLogNames)
+			{
+				if (GEngine->IsSettingUpPlayWorld())
+				{
+					int32 OldUniqueNumber = UE::VisLog::Private::UniqueRunNumber;
+					OnDataReset();
+					UE::VisLog::Private::UniqueRunNumber = OldUniqueNumber + 1;
+				}
+
+			}
+		});
+#endif
 }
 
-void FVisualLogger::Shutdown()
+void FVisualLogger::TearDown()
 {
+#if WITH_EDITOR
+	FWorldDelegates::OnPIEStarted.Remove(PIEStartedHandle);
+#endif
+
 	SetIsRecording(false);
 	SetIsRecordingToFile(false);
 	RemoveDevice(&FVisualLoggerBinaryFileDevice::Get());
 }
 
+void FVisualLogger::OnDataReset()
+{
+	using namespace UE::VisLog::Private;
+
+	UniqueRunNumber = 0;
+	UniqueLogPrefix.Empty();
+
+#if WITH_EDITOR
+	UEditorEngine* EditorEngine = GIsEditor ? Cast<UEditorEngine>(GEngine) : nullptr;
+	if (EditorEngine)
+	{
+		const UWorld* EditorWorld = EditorEngine->GetEditorWorldContext().World();
+		if (EditorWorld)
+		{
+			EditorOnly::EditorBaseTimeStamp = EditorWorld->TimeSeconds;
+		}
+	}
+#endif
+}
+
 void FVisualLogger::Cleanup(UWorld* OldWorld, const bool bReleaseMemory)
 {
 	const bool WasRecordingToFile = IsRecordingToFile();
@@ -1006,6 +1107,11 @@ void FVisualLogger::SetIsRecordingToTrace(const bool InIsRecording)
 	bIsRecordingToTrace = InIsRecording;
 }
 
+void FVisualLogger::SetUseUniqueNames(const bool bEnable)
+{
+	bForceUniqueLogNames = bEnable;
+	UE::VisLog::Private::UniqueLogPrefix.Empty();
+}
 
 void FVisualLogger::DiscardRecordingToFile()
 {
@@ -1086,7 +1192,7 @@ protected:
 				FVisualLogger::Get().AddCategoryToAllowList(*Category);
 				return true;
 			}
-			else if (FModuleManager::Get().LoadModulePtr<IModuleInterface>("LogVisualizer"))
+			else if (Command.IsEmpty() && FModuleManager::Get().LoadModulePtr<IModuleInterface>("LogVisualizer"))
 			{
 				FGlobalTabmanager::Get()->TryInvokeTab(FName(TEXT("VisualLogger")));
 				return true;
diff --git a/Engine/Source/Runtime/Engine/Private/VisualLogger/VisualLoggerTypes.cpp b/Engine/Source/Runtime/Engine/Private/VisualLogger/VisualLoggerTypes.cpp
--- a/Engine/Source/Runtime/Engine/Private/VisualLogger/VisualLoggerTypes.cpp
+++ b/Engine/Source/Runtime/Engine/Private/VisualLogger/VisualLoggerTypes.cpp
@@ -8,6 +8,7 @@
 #include "UObject/Interface.h"
 #include "UObject/UE5MainStreamObjectVersion.h"
 #include "VisualLogger/VisualLoggerCustomVersion.h"
+#include "VisualLogger/VisualLogger.h"
 #endif
 
 namespace
@@ -67,12 +68,19 @@ FVisualLogEntry::FVisualLogEntry(const AActor* InActor, TArray<TWeakObjectPtr<UO
 	if (!IsValid(InActor))
 	{
 		Reset();
+		return;
 	}
 
-	TimeStamp = InActor->GetWorld()->TimeSeconds; // Should we be bypassing FVisualLogger::GetTimeStampForObject() ?
+	WorldTimeStamp = InActor->GetWorld()->TimeSeconds;
 	Location = InActor->GetActorLocation();
 	bIsLocationValid = true;
 
+#if ENABLE_VISUAL_LOG
+	TimeStamp = FVisualLogger::Get().GetTimeStampForObject(InActor);
+#else
+	TimeStamp = WorldTimeStamp;
+#endif
+
 	const IVisualLoggerDebugSnapshotInterface* DebugSnapshotInterface = Cast<const IVisualLoggerDebugSnapshotInterface>(InActor);
 	if (DebugSnapshotInterface)
 	{
@@ -127,12 +135,14 @@ void FVisualLogEntry::InitializeEntry(const double InTimeStamp)
 {
 	Reset();
 	TimeStamp = InTimeStamp;
+	WorldTimeStamp = InTimeStamp;
 	bIsInitialized = true;
 }
 
 void FVisualLogEntry::Reset()
 {
-	TimeStamp = -1.;
+	TimeStamp = -1.0;
+	WorldTimeStamp = -1.0;
 	Location = FVector::ZeroVector;
 	bIsLocationValid = false;
 	Events.Reset();
@@ -599,10 +609,7 @@ FArchive& operator<<(FArchive& Ar, FVisualLogEntry& LogEntry)
 	const int32 VLogsOldVer = Ar.CustomVer(EVisualLoggerVersion::GUID);
 	const int32 VLogsStreamObjectVer = Ar.CustomVer(FUE5MainStreamObjectVersion::GUID);
 
-	// @todo replace with FUE5MainStreamObjectVersion::VisualLoggerTimeStampAsDouble when it becomes available
-	constexpr int32 TEMPVisualLoggerTimeStampAsDouble = 98;
-
-	if (VLogsStreamObjectVer >= TEMPVisualLoggerTimeStampAsDouble)
+	if (VLogsStreamObjectVer >= FUE5MainStreamObjectVersion::VisualLoggerTimeStampAsDouble)
 	{
 		Ar << LogEntry.TimeStamp;
 	}
@@ -613,6 +620,15 @@ FArchive& operator<<(FArchive& Ar, FVisualLogEntry& LogEntry)
 		LogEntry.TimeStamp = TimeStampFlt;
 	}
 
+	if (VLogsStreamObjectVer < FUE5MainStreamObjectVersion::VisualLoggerAddedSeparateWorldTime)
+	{
+		LogEntry.WorldTimeStamp = LogEntry.TimeStamp;
+	}
+	else
+	{
+		Ar << LogEntry.WorldTimeStamp;
+	}
+
 	if (VLogsOldVer >= EVisualLoggerVersion::LargeWorldCoordinatesAndLocationValidityFlag)
 	{
 		Ar << LogEntry.Location;
diff --git a/Engine/Source/Runtime/Engine/Public/VisualLogger/VisualLogger.h b/Engine/Source/Runtime/Engine/Public/VisualLogger/VisualLogger.h
--- a/Engine/Source/Runtime/Engine/Public/VisualLogger/VisualLogger.h
+++ b/Engine/Source/Runtime/Engine/Public/VisualLogger/VisualLogger.h
@@ -660,12 +660,18 @@ public:
 
 	virtual ~FVisualLogger() {}
 
+	UE_DEPRECATED(5.4, "Following the base class convention and using the name TearDown. Since FVisualLogger::Get() is used internally everywhere, this class isn't designed to be inherited.")
+	ENGINE_API virtual void Shutdown() {}
+
 	// called on engine shutdown to flush all, etc.
-	ENGINE_API virtual void Shutdown();
+	ENGINE_API virtual void TearDown() override;
 
 	// Removes all logged data 
 	ENGINE_API void Cleanup(UWorld* OldWorld, bool bReleaseMemory = false);
 
+	// Use when a visual logger device has discarded all of its data, waiting for new data
+	ENGINE_API void OnDataReset();
+
 	/** Set log owner redirection from one object to another, to combine logs */
 	static void Redirect(const UObject* FromObject, const UObject* ToObject)
 	{ 
@@ -723,7 +729,7 @@ public:
 	bool IsRecordingOnServer() const { return !!bIsRecordingOnServer; }
 
 	/** Configure whether VisLog should be using decorated, unique names */
-	void SetUseUniqueNames(const bool bEnable) { bForceUniqueLogNames = bEnable; }
+	ENGINE_API void SetUseUniqueNames(const bool bEnable);
 
 	/** Add visual logger output device */
 	void AddDevice(FVisualLogDevice* InDevice) { OutputDevices.AddUnique(InDevice); }
@@ -760,10 +766,10 @@ public:
 	/** internal check for each usage of visual logger */
 	static ENGINE_API bool CheckVisualLogInputInternal(const UObject* Object, const FName& CategoryName, ELogVerbosity::Type Verbosity, UWorld **OutWorld, FVisualLogEntry **OutCurrentEntry);
 	
-	/** Returns time stamp for object */
+	/** Returns a current time stamp to associate with a recorded event that occurred on Object (used for ordering events on a timeline) */
 	ENGINE_API double GetTimeStampForObject(const UObject* Object) const;
 
-	/** Sets function to call to get a timestamp instead of the default implementation (i.e. world time) */
+	/** Sets function to call to get a timestamp instead of the default implementation (e.g. using a network synchronized clock instead of the local world time) */
 	ENGINE_API void SetGetTimeStampFunc(TFunction<double(const UObject*)> Function);
 
 	typedef TMap<FObjectKey, TArray<TWeakObjectPtr<const UObject> > > FOwnerToChildrenRedirectionMap;
@@ -885,9 +891,14 @@ protected:
 	/** Delegate to set project specific file name for vlogs */
 	FVisualLogFilenameGetterDelegate LogFileNameGetter;
 
-	/** function to call when getting the time stamp */
+	/** Function to call when recording the absolute time stamp of an event. Useful for manually aligning events across multiple instances (e.g. using FPlatformTime::Seconds() rather than WorldTime) */
 	TFunction<double(const UObject*)> GetTimeStampFunc;
 
+#if WITH_EDITOR
+	/** Handle for registering with PIEStarted to reset the EditorBaseTimeStamp */
+	FDelegateHandle PIEStartedHandle;
+#endif
+
 	// if set we are recording and collecting all vlog data
 	static ENGINE_API int32 bIsRecording;
 };
diff --git a/Engine/Source/Runtime/Engine/Public/VisualLogger/VisualLoggerTypes.h b/Engine/Source/Runtime/Engine/Public/VisualLogger/VisualLoggerTypes.h
--- a/Engine/Source/Runtime/Engine/Public/VisualLogger/VisualLoggerTypes.h
+++ b/Engine/Source/Runtime/Engine/Public/VisualLogger/VisualLoggerTypes.h
@@ -172,7 +172,10 @@ struct FVisualLogDataBlock
 struct FVisualLogEntry
 {
 #if ENABLE_VISUAL_LOG
+	/** For absolute position of events along a timeline (can involve multiple worlds/game instances such as clients and server) */
 	double TimeStamp;
+	/** The time of the event according to its UWorld (can vary widely between game instances such as clients and server) */
+	double WorldTimeStamp;
 	FVector Location;
 	uint8 bPassedClassAllowList : 1;
 	uint8 bPassedObjectAllowList : 1;	
@@ -188,7 +191,11 @@ struct FVisualLogEntry
 	TArray<FVisualLogDataBlock>	DataBlocks;
 
 	FVisualLogEntry() { Reset(); }
+
+	UE_DEPRECATED(5.4, "To be removed.  Build up the FVisualLogEntry manually or write your own helper function which doesn't rely on Children being IVisualLoggerDebugSnapshotInterface")
 	ENGINE_API FVisualLogEntry(const AActor* InActor, TArray<TWeakObjectPtr<UObject> >* Children);
+	
+	UE_DEPRECATED(5.4, "To be removed.  Build up the FVisualLogEntry manually or write your own helper function which doesn't rely on Children being IVisualLoggerDebugSnapshotInterface")
 	ENGINE_API FVisualLogEntry(double InTimeStamp, FVector InLocation, const UObject* Object, TArray<TWeakObjectPtr<UObject> >* Children);
 
 	bool ShouldLog(const ECreateIfNeeded ShouldCreate) const
diff --git a/Engine/Source/Runtime/Launch/Private/LaunchEngineLoop.cpp b/Engine/Source/Runtime/Launch/Private/LaunchEngineLoop.cpp
--- a/Engine/Source/Runtime/Launch/Private/LaunchEngineLoop.cpp
+++ b/Engine/Source/Runtime/Launch/Private/LaunchEngineLoop.cpp
@@ -5046,7 +5046,7 @@ void FEngineLoop::Exit()
 
 	// shutdown visual logger and flush all data
 #if ENABLE_VISUAL_LOG
-	FVisualLogger::Get().Shutdown();
+	FVisualLogger::Get().TearDown();
 #endif
 
 	FAssetCompilingManager::Get().Shutdown();
