Fix an issue where the world partition might be initialized during the advanced copy step.
Any world can be converted to a world partition world before subsystems might have initialized.
Add support for Staging a plugin created by a cooked editor.

diff --git a/Engine/Source/Runtime/Engine/Private/Level.cpp b/Engine/Source/Runtime/Engine/Private/Level.cpp
--- a/Engine/Source/Runtime/Engine/Private/Level.cpp
+++ b/Engine/Source/Runtime/Engine/Private/Level.cpp
@@ -3070,8 +3070,11 @@ void ULevel::OnLevelLoaded()
 
 			if (bIsMainWorldLevel || bInitializeForEditor)
 			{
-				FTransform Transform = LevelStreaming ? LevelStreaming->LevelTransform : FTransform::Identity;
-				WorldPartition->Initialize(OwningWorld, Transform);
+				if (!WorldPartition->IsInitialized())
+				{
+					FTransform Transform = LevelStreaming ? LevelStreaming->LevelTransform : FTransform::Identity;
+					WorldPartition->Initialize(OwningWorld, Transform);
+				}
 			}
 		}
 	}
diff --git a/Engine/Source/Runtime/Engine/Private/WorldPartition/ContentBundle/ContentBundleWorldSubsystem.cpp b/Engine/Source/Runtime/Engine/Private/WorldPartition/ContentBundle/ContentBundleWorldSubsystem.cpp
--- a/Engine/Source/Runtime/Engine/Private/WorldPartition/ContentBundle/ContentBundleWorldSubsystem.cpp
+++ b/Engine/Source/Runtime/Engine/Private/WorldPartition/ContentBundle/ContentBundleWorldSubsystem.cpp
@@ -33,7 +33,14 @@ void UContentBundleManager::Initialize()
 		PIEDuplicateHelper->Initialize();
 	}
 #endif
-
+	// Any world can be converted to a world partition world before the content bundle manager has been initialized.
+	if (UWorldPartition* WorldPartition = GetWorld()->GetWorldPartition())
+	{
+		if (WorldPartition->IsInitialized())
+		{
+			OnWorldPartitionInitialized(WorldPartition);
+		}
+	}
 	GetWorld()->OnWorldPartitionInitialized().AddUObject(this, &UContentBundleManager::OnWorldPartitionInitialized);
 	GetWorld()->OnWorldPartitionUninitialized().AddUObject(this, &UContentBundleManager::OnWorldPartitionUninitialized);
 }
diff --git a/Engine/Source/Runtime/Engine/Private/WorldPartition/WorldPartitionSubsystem.cpp b/Engine/Source/Runtime/Engine/Private/WorldPartition/WorldPartitionSubsystem.cpp
--- a/Engine/Source/Runtime/Engine/Private/WorldPartition/WorldPartitionSubsystem.cpp
+++ b/Engine/Source/Runtime/Engine/Private/WorldPartition/WorldPartitionSubsystem.cpp
@@ -393,6 +393,13 @@ void UWorldPartitionSubsystem::Initialize(FSubsystemCollectionBase& Collection)
 	}
 #endif
 
+	if (UWorldPartition* WorldPartition = GetWorld()->GetWorldPartition())
+	{
+		if (WorldPartition->IsInitialized())
+		{
+			OnWorldPartitionInitialized(WorldPartition);
+		}
+	}
 	GetWorld()->OnWorldPartitionInitialized().AddUObject(this, &UWorldPartitionSubsystem::OnWorldPartitionInitialized);
 	GetWorld()->OnWorldPartitionUninitialized().AddUObject(this, &UWorldPartitionSubsystem::OnWorldPartitionUninitialized);
 	if (GetWorld()->IsGameWorld())
