Non-Unity fix for VisLog changes

diff --git a/Engine/Source/Runtime/Engine/Private/VisualLogger/VisualLogger.cpp b/Engine/Source/Runtime/Engine/Private/VisualLogger/VisualLogger.cpp
--- a/Engine/Source/Runtime/Engine/Private/VisualLogger/VisualLogger.cpp
+++ b/Engine/Source/Runtime/Engine/Private/VisualLogger/VisualLogger.cpp
@@ -11,6 +11,7 @@
 #include "VisualLogger/VisualLoggerCustomVersion.h"
 #include "VisualLogger/VisualLoggerTraceDevice.h"
 #include "VisualLogger/VisualLoggerDebugSnapshotInterface.h"
+#include "UnrealEngine.h"
 
 #if WITH_EDITOR
 #include "Editor/EditorEngine.h"
