World Partition Actor Descriptor Patcher: fix FSoftObjectPath patching through FMemoryWriter.

diff --git a/Engine/Source/Runtime/Engine/Private/WorldPartition/WorldPartitionActorDescArchive.cpp b/Engine/Source/Runtime/Engine/Private/WorldPartition/WorldPartitionActorDescArchive.cpp
--- a/Engine/Source/Runtime/Engine/Private/WorldPartition/WorldPartitionActorDescArchive.cpp
+++ b/Engine/Source/Runtime/Engine/Private/WorldPartition/WorldPartitionActorDescArchive.cpp
@@ -140,7 +140,7 @@ FArchive& FActorDescArchivePatcher::operator<<(FSoftObjectPath& Value)
 	TGuardValue<bool> GuardIsPatching(bIsPatching, true);
 	FActorDescArchive::operator<<(Value);
 	AssetDataPatcher->DoPatch(Value);
-	OutAr << Value;
+	Value.SerializePathWithoutFixup(OutAr);
 	return *this;
 }
 
