Misc new HLSL translator bug fixes
AcquireHLSLExpressionOr* return values need tested before passing possible nullptr further
Incorrect dropping of return value in NewErrorExpression functions, added compile test

#tests materials, editor, win64

diff --git a/Engine/Source/Runtime/Engine/Private/Materials/MaterialExpressionHLSL.cpp b/Engine/Source/Runtime/Engine/Private/Materials/MaterialExpressionHLSL.cpp
--- a/Engine/Source/Runtime/Engine/Private/Materials/MaterialExpressionHLSL.cpp
+++ b/Engine/Source/Runtime/Engine/Private/Materials/MaterialExpressionHLSL.cpp
@@ -266,8 +266,15 @@ bool UMaterialExpressionNamedRerouteDeclaration::GenerateHLSLExpression(FMateria
 
 bool UMaterialExpressionNamedRerouteUsage::GenerateHLSLExpression(FMaterialHLSLGenerator& Generator, UE::HLSLTree::FScope& Scope, int32 OutputIndex, UE::HLSLTree::FExpression const*& OutExpression) const
 {
-	OutExpression = Generator.AcquireExpression(Scope, 0, Declaration, 0, UE::HLSLTree::FSwizzleParameters());
-	return OutExpression != nullptr;
+	if (IsDeclarationValid())
+	{
+		OutExpression = Generator.AcquireExpression(Scope, 0, Declaration, 0, UE::HLSLTree::FSwizzleParameters());
+	}
+	else
+	{
+		OutExpression = Generator.NewErrorExpression(TEXT("Invalid named reroute variable"));
+	}
+	return true;
 }
 
 bool UMaterialExpressionGenericConstant::GenerateHLSLExpression(FMaterialHLSLGenerator& Generator, UE::HLSLTree::FScope& Scope, int32 OutputIndex, UE::HLSLTree::FExpression const*& OutExpression) const
@@ -1509,6 +1516,10 @@ bool UMaterialExpressionTextureSample::GenerateHLSLExpressionBase(FMaterialHLSLG
 	}
 
 	const FExpression* TexCoordExpression = Coordinates.AcquireHLSLExpressionOrExternalInput(Generator, Scope, Material::MakeInputTexCoord(ConstCoordinate));
+	if (!TexCoordExpression)
+	{
+		return false;
+	}
 	const FExpression* MipLevelExpression = nullptr;
 	FExpressionDerivatives TexCoordDerivatives;
 	switch (MipValueMode)
@@ -3848,6 +3859,10 @@ bool UMaterialExpressionPanner::GenerateHLSLExpression(FMaterialHLSLGenerator& G
 
 	const FExpression* ExpressionTime = Time.AcquireHLSLExpressionOrExternalInput(Generator, Scope, Material::EExternalInput::GameTime);
 	const FExpression* ExpressionSpeed = Speed.AcquireHLSLExpressionOrConstant(Generator, Scope, FVector2f(SpeedX, SpeedY));
+	if (!ExpressionTime || !ExpressionSpeed)
+	{
+		return false;
+	}
 	const FExpression* ExpressionOffset = Generator.GetTree().NewMul(ExpressionSpeed, ExpressionTime);
 	if (bFractionalPart)
 	{
diff --git a/Engine/Source/Runtime/Engine/Private/Materials/MaterialHLSLGenerator.cpp b/Engine/Source/Runtime/Engine/Private/Materials/MaterialHLSLGenerator.cpp
--- a/Engine/Source/Runtime/Engine/Private/Materials/MaterialHLSLGenerator.cpp
+++ b/Engine/Source/Runtime/Engine/Private/Materials/MaterialHLSLGenerator.cpp
@@ -544,6 +544,7 @@ const UE::HLSLTree::FExpression* FMaterialHLSLGenerator::AcquireExpression(UE::H
 		const FExpression** Found = GeneratedExpressionMap.Find(KeyHash);
 		Expression = Found ? *Found : nullptr;
 
+		check(MaterialExpression);
 		if (!Expression && MaterialExpression->GenerateHLSLExpression(*this, Scope, OutputIndex, Expression))
 		{
 			GeneratedExpressionMap.Add(KeyHash, Expression);
diff --git a/Engine/Source/Runtime/Engine/Private/Materials/MaterialHLSLTree.cpp b/Engine/Source/Runtime/Engine/Private/Materials/MaterialHLSLTree.cpp
--- a/Engine/Source/Runtime/Engine/Private/Materials/MaterialHLSLTree.cpp
+++ b/Engine/Source/Runtime/Engine/Private/Materials/MaterialHLSLTree.cpp
@@ -1788,6 +1788,11 @@ void FExpressionStaticTerrainLayerWeight::EmitValueShader(FEmitContext& Context,
 		TextureValue.Texture = bTextureArrayEnabled ? GEngine->WeightMapArrayPlaceholderTexture : GEngine->WeightMapPlaceholderTexture;
 		TextureValue.ParameterInfo.Name = BuildWeightmapName(TEXT("Weightmap"), WeightmapIndex, !bTextureArrayEnabled);
 
+		if (!EmitTexCoordValue)
+		{
+			EmitTexCoordValue = TexCoordExpression->GetValueShader(Context, Scope, TexCoordType);
+		}
+
 		if (bTextureArrayEnabled)
 		{
 			EmitTexCoordValue = Context.EmitExpression(Scope, Shader::EValueType::Float3, TEXT("float3((%).xy,%)"), EmitTexCoordValue, WeightmapIndex);
diff --git a/Engine/Source/Runtime/Engine/Private/Materials/MaterialShared.cpp b/Engine/Source/Runtime/Engine/Private/Materials/MaterialShared.cpp
--- a/Engine/Source/Runtime/Engine/Private/Materials/MaterialShared.cpp
+++ b/Engine/Source/Runtime/Engine/Private/Materials/MaterialShared.cpp
@@ -435,11 +435,7 @@ const UE::HLSLTree::FExpression* FExpressionInput::AcquireHLSLExpression(FMateri
 		{
 			return Generator.NewErrorExpressionf(TEXT("Missing input '%s'"), *LocalInputName.ToString());
 		}
-		else
-		{
-			Generator.NewErrorExpression(TEXT("Missing input"));
-		}
-		return nullptr;
+		return Generator.NewErrorExpression(TEXT("Missing input"));
 	}
 	return TryAcquireHLSLExpression(Generator, Scope, InputIndex);
 }
diff --git a/Engine/Source/Runtime/Engine/Public/MaterialHLSLGenerator.h b/Engine/Source/Runtime/Engine/Public/MaterialHLSLGenerator.h
--- a/Engine/Source/Runtime/Engine/Public/MaterialHLSLGenerator.h
+++ b/Engine/Source/Runtime/Engine/Public/MaterialHLSLGenerator.h
@@ -111,13 +111,13 @@ public:
 	const UE::HLSLTree::FExpression* NewSwizzle(const UE::HLSLTree::FSwizzleParameters& Params, const UE::HLSLTree::FExpression* Input);
 
 	template<typename StringType>
-	inline const UE::HLSLTree::FExpression* NewErrorExpression(const StringType& InError)
+	[[nodiscard]] inline const UE::HLSLTree::FExpression* NewErrorExpression(const StringType& InError)
 	{
 		return InternalNewErrorExpression(FStringView(InError));
 	}
 
 	template<typename FormatType, typename... Types>
-	const UE::HLSLTree::FExpression* NewErrorExpressionf(const FormatType& Format, Types... Args)
+	[[nodiscard]] const UE::HLSLTree::FExpression* NewErrorExpressionf(const FormatType& Format, Types... Args)
 	{
 		TStringBuilder<1024> String;
 		String.Appendf(Format, Forward<Types>(Args)...);
