Fix incorrect format for world height streaming mips.

diff --git a/Engine/Source/Runtime/Engine/Private/Components/RuntimeVirtualTextureComponent.cpp b/Engine/Source/Runtime/Engine/Private/Components/RuntimeVirtualTextureComponent.cpp
--- a/Engine/Source/Runtime/Engine/Private/Components/RuntimeVirtualTextureComponent.cpp
+++ b/Engine/Source/Runtime/Engine/Private/Components/RuntimeVirtualTextureComponent.cpp
@@ -310,9 +310,20 @@ private:
 	TArray<URuntimeVirtualTexture*> VirtualTextures;
 };
 
+static TextureCompressionSettings GetCompressionSettingFromLayerFormat(EPixelFormat LayerFormat)
+{
+	switch (LayerFormat)
+	{
+	case PF_BC5: return TC_Normalmap;
+	case PF_BC4: return TC_Alpha;
+	case PF_G16: return TC_Grayscale;
+	}
+	return TC_Default;
+}
+
 static void GetLayerFormatSettings(FTextureFormatSettings& OutFormatSettings, EPixelFormat LayerFormat, bool IsLayerYCoCg, bool IsLayerSRGB, bool IsLayerLQCompression)
 {
-	OutFormatSettings.CompressionSettings = IsLayerLQCompression ? TC_LQ : (LayerFormat == PF_BC5 ? TC_Normalmap : (LayerFormat == PF_BC4 ? TC_Alpha : TC_Default));
+	OutFormatSettings.CompressionSettings = IsLayerLQCompression ? TC_LQ : GetCompressionSettingFromLayerFormat(LayerFormat);
 	OutFormatSettings.CompressionNone = LayerFormat == PF_B8G8R8A8 || LayerFormat == PF_G16;
 	OutFormatSettings.CompressionNoAlpha = LayerFormat == PF_DXT1 || LayerFormat == PF_BC5 || LayerFormat == PF_R5G6B5_UNORM;
 	OutFormatSettings.CompressionForceAlpha = LayerFormat == PF_DXT5;
