Reduce the GT render streaming parallel for to only use background tasks to avoid taking up the entire task graph for multiple ms, delaying critical path rendering thread work.

diff --git a/Engine/Source/Runtime/Engine/Private/Streaming/StreamingManagerTexture.cpp b/Engine/Source/Runtime/Engine/Private/Streaming/StreamingManagerTexture.cpp
--- a/Engine/Source/Runtime/Engine/Private/Streaming/StreamingManagerTexture.cpp
+++ b/Engine/Source/Runtime/Engine/Private/Streaming/StreamingManagerTexture.cpp
@@ -1729,7 +1729,7 @@ public:
 	}
 	static FORCEINLINE ENamedThreads::Type GetDesiredThread()
 	{
-		return ENamedThreads::AnyNormalThreadNormalTask;
+		return ENamedThreads::AnyBackgroundThreadNormalTask;
 	}
 	static FORCEINLINE ESubsequentsMode::Type GetSubsequentsMode()
 	{
