World Partition: fix invalid HLOD warning when creating a new empty open world map.

diff --git a/Engine/Source/Runtime/Engine/Private/WorldPartition/RuntimeHashSet/WorldPartitionRuntimeHashSet.cpp b/Engine/Source/Runtime/Engine/Private/WorldPartition/RuntimeHashSet/WorldPartitionRuntimeHashSet.cpp
--- a/Engine/Source/Runtime/Engine/Private/WorldPartition/RuntimeHashSet/WorldPartitionRuntimeHashSet.cpp
+++ b/Engine/Source/Runtime/Engine/Private/WorldPartition/RuntimeHashSet/WorldPartitionRuntimeHashSet.cpp
@@ -104,6 +104,44 @@ void UWorldPartitionRuntimeHashSet::SetDefaultValues()
 	RuntimePartitionDesc.MainLayer = NewObject<URuntimePartitionLHGrid>(this, NAME_None);
 	RuntimePartitionDesc.MainLayer->Name = RuntimePartitionDesc.Name;
 	RuntimePartitionDesc.MainLayer->SetDefaultValues();
+
+	UWorldPartition* WorldPartition = GetTypedOuter<UWorldPartition>();
+	check(WorldPartition);
+
+	if (const UHLODLayer* HLODLayer = WorldPartition->GetDefaultHLODLayer())
+	{
+		uint32 HLODIndex = 0;
+		while (HLODLayer)
+		{
+			FRuntimePartitionHLODSetup& HLODSetup = RuntimePartitionDesc.HLODSetups.AddDefaulted_GetRef();
+
+			HLODSetup.Name = HLODLayer->GetFName();
+			HLODSetup.bIsSpatiallyLoaded = HLODLayer->IsSpatiallyLoaded();
+			HLODSetup.HLODLayers = { HLODLayer };
+
+			if (HLODSetup.bIsSpatiallyLoaded)
+			{
+				URuntimePartitionLHGrid* HLODLHGrid = NewObject<URuntimePartitionLHGrid>(this, NAME_None);
+				HLODLHGrid->CellSize = CastChecked<URuntimePartitionLHGrid>(RuntimePartitionDesc.MainLayer)->CellSize * (2 << HLODIndex);
+				HLODLHGrid->LoadingRange = RuntimePartitionDesc.MainLayer->LoadingRange * (2 << HLODIndex);
+				HLODSetup.PartitionLayer = HLODLHGrid;
+			}
+			else
+			{
+				HLODSetup.PartitionLayer = NewObject<URuntimePartitionPersistent>(this, NAME_None);
+				HLODSetup.PartitionLayer->LoadingRange = 0;
+			}
+
+			HLODSetup.PartitionLayer->Name = HLODSetup.Name;
+			HLODSetup.PartitionLayer->bBlockOnSlowStreaming = false;
+			HLODSetup.PartitionLayer->bClientOnlyVisible = true;
+			HLODSetup.PartitionLayer->Priority = 0;
+			HLODSetup.PartitionLayer->HLODIndex = HLODIndex;
+
+			HLODLayer = HLODLayer->GetParentLayer();
+			HLODIndex++;
+		}
+	}
 }
 
 void UWorldPartitionRuntimeHashSet::FlushStreaming()
diff --git a/Engine/Source/Runtime/Engine/Private/WorldPartition/RuntimeHashSet/WorldPartitionRuntimeHashSetConversions.cpp b/Engine/Source/Runtime/Engine/Private/WorldPartition/RuntimeHashSet/WorldPartitionRuntimeHashSetConversions.cpp
--- a/Engine/Source/Runtime/Engine/Private/WorldPartition/RuntimeHashSet/WorldPartitionRuntimeHashSetConversions.cpp
+++ b/Engine/Source/Runtime/Engine/Private/WorldPartition/RuntimeHashSet/WorldPartitionRuntimeHashSetConversions.cpp
@@ -42,13 +42,14 @@ UWorldPartitionRuntimeHashSet* UWorldPartitionRuntimeHashSet::CreateFrom(const U
 
 				if (const UHLODLayer* HLODLayer = ActorDescView.GetHLODLayer().IsValid() ? Cast<UHLODLayer>(ActorDescView.GetHLODLayer().TryLoad()) : DefaultHLODLayer)
 				{
-					uint32 HLODLIndex = 0;
+					uint32 HLODIndex = 0;
 					while (HLODLayer)
 					{
 						const FName GridName = ActorDescView.GetRuntimeGrid().IsNone() ? SpatialHash->Grids[0].GridName : ActorDescView.GetRuntimeGrid();
-						const FName HLODGridName = HLODLayer->GetRuntimeGrid(HLODLIndex++);
+						const FName HLODGridName = HLODLayer->GetRuntimeGrid(HLODIndex);
 						GridHLODLayersMap.FindOrAdd(GridName).FindOrAdd(HLODGridName).Add(HLODLayer);
 						HLODLayer = HLODLayer->GetParentLayer();
+						HLODIndex++;
 					}
 				}
 			});
