fix comment

diff --git a/Engine/Source/Runtime/Engine/Private/WorldPartition/LevelInstance/LevelInstanceActorDesc.cpp b/Engine/Source/Runtime/Engine/Private/WorldPartition/LevelInstance/LevelInstanceActorDesc.cpp
--- a/Engine/Source/Runtime/Engine/Private/WorldPartition/LevelInstance/LevelInstanceActorDesc.cpp
+++ b/Engine/Source/Runtime/Engine/Private/WorldPartition/LevelInstance/LevelInstanceActorDesc.cpp
@@ -173,7 +173,7 @@ bool FLevelInstanceActorDesc::GetChildContainerInstance(const FWorldPartitionAct
 		OutContainerInstance.ContainerInstance = ContainerInstance;
 		OutContainerInstance.ClusterMode = EContainerClusterMode::Partitioned;
 
-		// @todo_pat: this is to validate that new parenting of container instance code is equivalent
+		// @todo_ow: this is to validate that new parenting of container instance code is equivalent
 		OutContainerInstance.Transform = GetChildContainerTransform();
 
 		return true;
