World Partition: fix infinite loading of streaming cells in PIE if the only actors they need to load are editor-only actors, referenced by dirty, non-editor-only actors.

diff --git a/Engine/Source/Runtime/Engine/Private/WorldPartition/WorldPartitionLevelHelper.cpp b/Engine/Source/Runtime/Engine/Private/WorldPartition/WorldPartitionLevelHelper.cpp
--- a/Engine/Source/Runtime/Engine/Private/WorldPartition/WorldPartitionLevelHelper.cpp
+++ b/Engine/Source/Runtime/Engine/Private/WorldPartition/WorldPartitionLevelHelper.cpp
@@ -421,64 +421,71 @@ bool FWorldPartitionLevelHelper::LoadActors(UWorld* InOuterWorld, ULevel* InDest
 
 bool FWorldPartitionLevelHelper::LoadActors(const FLoadActorsParams& InParams)
 {
-	UPackage* DestPackage = InParams.DestLevel ? InParams.DestLevel->GetPackage() : nullptr;
-	FString ShortLevelPackageName = DestPackage? FPackageName::GetShortName(DestPackage->GetFName()) : FString();
+	TArray<FWorldPartitionRuntimeCellObjectMapping*> ActorPackagesToLoad;
+	TMap<FActorContainerID, FLinkerInstancingContext> LinkerInstancingContexts;
 
-	struct FLoadProgress
+	if (!InParams.ActorPackages.IsEmpty())
 	{
-		int32 NumPendingLoadRequests = 0;
-		int32 NumFailedLoadedRequests = 0;
-	};
-	TSharedPtr<FLoadProgress> LoadProgress = MakeShared<FLoadProgress>();
+		ActorPackagesToLoad.Reserve(InParams.ActorPackages.Num());
 
-	// Actors to load
-	TArray<FWorldPartitionRuntimeCellObjectMapping*> ActorPackages;
-	ActorPackages.Reserve(InParams.ActorPackages.Num());
-
-	TMap<FActorContainerID, FLinkerInstancingContext> LinkerInstancingContexts;
-	// Add Main container context
-	LinkerInstancingContexts.Add(FActorContainerID::GetMainContainerID(), MoveTemp(InParams.InstancingContext));
+		// Add main container context
+		LinkerInstancingContexts.Add(FActorContainerID::GetMainContainerID(), MoveTemp(InParams.InstancingContext));
 			
-	for (FWorldPartitionRuntimeCellObjectMapping& PackageObjectMapping : InParams.ActorPackages)
-	{
-		FLinkerInstancingContext* Context = LinkerInstancingContexts.Find(PackageObjectMapping.ContainerID);
-		if (!Context)
+		for (FWorldPartitionRuntimeCellObjectMapping& PackageObjectMapping : InParams.ActorPackages)
 		{
-			check(!PackageObjectMapping.ContainerID.IsMainContainer());
+			FLinkerInstancingContext* Context = LinkerInstancingContexts.Find(PackageObjectMapping.ContainerID);
+			if (!Context)
+			{
+				check(!PackageObjectMapping.ContainerID.IsMainContainer());
 		
-			const FString DestLevelPackageName = InParams.DestLevel ? InParams.DestLevel->GetPackage()->GetName() : FString();
-			const FName ContainerPackageInstanceName(GetContainerPackage(PackageObjectMapping.ContainerID, PackageObjectMapping.ContainerPackage.ToString(), DestLevelPackageName));
+				const FString DestLevelPackageName = InParams.DestLevel ? InParams.DestLevel->GetPackage()->GetName() : FString();
+				const FName ContainerPackageInstanceName(GetContainerPackage(PackageObjectMapping.ContainerID, PackageObjectMapping.ContainerPackage.ToString(), DestLevelPackageName));
 
-			FLinkerInstancingContext& NewContext = LinkerInstancingContexts.Add(PackageObjectMapping.ContainerID);
+				FLinkerInstancingContext& NewContext = LinkerInstancingContexts.Add(PackageObjectMapping.ContainerID);
 
-			// Make sure here we don't remap the SoftObjectPaths through the linker when loading the embedded actor packages. 
-			// A remapping will happen in the packaged loaded callback later in this method.
-			NewContext.SetSoftObjectPathRemappingEnabled(false); 
+				// Make sure here we don't remap the SoftObjectPaths through the linker when loading the embedded actor packages. 
+				// A remapping will happen in the packaged loaded callback later in this method.
+				NewContext.SetSoftObjectPathRemappingEnabled(false); 
 			
-			NewContext.AddTag(ULevel::DontLoadExternalObjectsTag);
-			NewContext.AddPackageMapping(PackageObjectMapping.ContainerPackage, ContainerPackageInstanceName);
-			Context = &NewContext;
-		}
+				NewContext.AddTag(ULevel::DontLoadExternalObjectsTag);
+				NewContext.AddPackageMapping(PackageObjectMapping.ContainerPackage, ContainerPackageInstanceName);
+				Context = &NewContext;
+			}
 		
-		const FName ContainerPackageInstanceName = Context->RemapPackage(PackageObjectMapping.ContainerPackage);
+			const FName ContainerPackageInstanceName = Context->RemapPackage(PackageObjectMapping.ContainerPackage);
 
-		if (PackageObjectMapping.bIsEditorOnly || PackageObjectMapping.ContainerPackage != ContainerPackageInstanceName)
-		{
-			const FName ActorPackageName = *FPackageName::ObjectPathToPackageName(PackageObjectMapping.Package.ToString());
-			const FName ActorPackageInstanceName = PackageObjectMapping.bIsEditorOnly ? NAME_None : FName(*ULevel::GetExternalActorPackageInstanceName(ContainerPackageInstanceName.ToString(), ActorPackageName.ToString()));
+			if (PackageObjectMapping.bIsEditorOnly || PackageObjectMapping.ContainerPackage != ContainerPackageInstanceName)
+			{
+				const FName ActorPackageName = *FPackageName::ObjectPathToPackageName(PackageObjectMapping.Package.ToString());
+				const FName ActorPackageInstanceName = PackageObjectMapping.bIsEditorOnly ? NAME_None : FName(*ULevel::GetExternalActorPackageInstanceName(ContainerPackageInstanceName.ToString(), ActorPackageName.ToString()));
 
-			Context->AddPackageMapping(ActorPackageName, ActorPackageInstanceName);
-	}
+				Context->AddPackageMapping(ActorPackageName, ActorPackageInstanceName);
+			}
 
-		if (!PackageObjectMapping.bIsEditorOnly)
-		{
-			ActorPackages.Add(&PackageObjectMapping);
+			if (!PackageObjectMapping.bIsEditorOnly)
+			{
+				ActorPackagesToLoad.Add(&PackageObjectMapping);
+			}
 		}
 	}
 
-	LoadProgress->NumPendingLoadRequests = ActorPackages.Num();
+	if (ActorPackagesToLoad.IsEmpty())
+	{
+		InParams.CompletionCallback(true);
+		return true;
+	}
 
-	for (FWorldPartitionRuntimeCellObjectMapping* PackageObjectMapping : ActorPackages)
+	struct FLoadProgress
+	{
+		int32 NumPendingLoadRequests;
+		int32 NumFailedLoadedRequests;
+	};
+
+	TSharedPtr<FLoadProgress> LoadProgress = MakeShared<FLoadProgress>();
+	LoadProgress->NumPendingLoadRequests = ActorPackagesToLoad.Num();
+	LoadProgress->NumFailedLoadedRequests = 0;
+
+	for (FWorldPartitionRuntimeCellObjectMapping* PackageObjectMapping : ActorPackagesToLoad)
 	{
 		FLoadPackageAsyncDelegate CompletionCallback = FLoadPackageAsyncDelegate::CreateLambda([LoadProgress, PackageObjectMapping, PackageReferencer = InParams.PackageReferencer, OuterWorld = InParams.OuterWorld, DestLevel = InParams.DestLevel, CompletionCallback = InParams.CompletionCallback](const FName& LoadedPackageName, UPackage* LoadedPackage, EAsyncLoadingResult::Type Result)
 		{
@@ -637,10 +644,10 @@ bool FWorldPartitionLevelHelper::LoadActors(const FLoadActorsParams& InParams)
 
 		if (InParams.bLoadAsync)
 		{
-			FPackagePath PackagePath = FPackagePath::FromPackageNameChecked(PackageToLoad);
-
-			check(DestPackage);
-			const EPackageFlags PackageFlags = DestPackage->HasAnyPackageFlags(PKG_PlayInEditor) ? PKG_PlayInEditor : PKG_None;
+			check(InParams.DestLevel);
+			const UPackage* DestPackage = InParams.DestLevel->GetPackage();
+			const EPackageFlags PackageFlags = InParams.DestLevel->GetPackage()->HasAnyPackageFlags(PKG_PlayInEditor) ? PKG_PlayInEditor : PKG_None;
+			const FPackagePath PackagePath = FPackagePath::FromPackageNameChecked(PackageToLoad);
 			::LoadPackageAsync(PackagePath, PackageName, CompletionCallback, PackageFlags, DestPackage->GetPIEInstanceID(), 0, &ContainerInstancingContext);
 		}
 		else
