Fix a crash when cleaning up a transient package level multiple times overflows the FName length
Related to CL: 22686882

diff --git a/Engine/Source/Runtime/Engine/Private/Level.cpp b/Engine/Source/Runtime/Engine/Private/Level.cpp
--- a/Engine/Source/Runtime/Engine/Private/Level.cpp
+++ b/Engine/Source/Runtime/Engine/Private/Level.cpp
@@ -469,7 +469,7 @@ void ULevel::CleanupLevel(bool bCleanupResources, bool bUnloadFromEditor)
 				ForEachObjectWithPackage(InPackage, [](UObject* Object) { Object->ClearFlags(RF_Standalone); return true; }, false);
 			}
 #endif
-			if (bTrashPackage)
+			if (bTrashPackage && (InPackage != GetTransientPackage()))
 			{
 				// Rename package to make sure it won't be reused
 				FName NewPackageName = MakeUniqueObjectName(nullptr, UPackage::StaticClass(), FName(*FString::Printf(TEXT("%s_Trashed"), *InPackage->GetName())));
