Added more specific info to debug bone buffer crash.

diff --git a/Engine/Source/Runtime/Engine/Public/GPUSkinVertexFactory.h b/Engine/Source/Runtime/Engine/Public/GPUSkinVertexFactory.h
--- a/Engine/Source/Runtime/Engine/Public/GPUSkinVertexFactory.h
+++ b/Engine/Source/Runtime/Engine/Public/GPUSkinVertexFactory.h
@@ -243,7 +243,8 @@ public:
 			if(!RetPtr->VertexBufferRHI.IsValid())
 			{
 				// this only should happen if we request the old data
-				check(bPrevious);
+				checkf(bPrevious, TEXT("Trying to access current bone buffer for reading, but it is null. BoneBuffer[0] = %p, BoneBuffer[1] = %p, CurrentRevisionNumber = %u, PreviousRevisionNumber = %u"),
+					BoneBuffer[0].VertexBufferRHI.GetReference(), BoneBuffer[1].VertexBufferRHI.GetReference(), CurrentRevisionNumber, PreviousRevisionNumber);
 
 				// if we don't have any old data we use the current one
 				RetPtr = &GetBoneBufferInternal(false);
