Fix for uninitialized struct warning. I think / hope.

diff --git a/Engine/Source/Runtime/Engine/Classes/Engine/Texture.h b/Engine/Source/Runtime/Engine/Classes/Engine/Texture.h
--- a/Engine/Source/Runtime/Engine/Classes/Engine/Texture.h
+++ b/Engine/Source/Runtime/Engine/Classes/Engine/Texture.h
@@ -53,11 +53,11 @@ struct FTextureSourceLayerColorInfo
 
 	/** Per channel min value of the colors for all blocks in mip0 in linear space */
 	UPROPERTY(VisibleAnywhere, Category = TextureSource)
-	FLinearColor ColorMin;
+	FLinearColor ColorMin = FLinearColor(EForceInit::ForceInitToZero);
 
 	/** Per channel max value of the colors for all blocks in mip0 in linear space */
 	UPROPERTY(VisibleAnywhere, Category = TextureSource)
-	FLinearColor ColorMax;
+	FLinearColor ColorMax = FLinearColor(EForceInit::ForceInitToZero);
 };
 
 USTRUCT()
