Fixed SetNum that should have been a Reserve

#rnx

diff --git a/Engine/Source/Runtime/Engine/Private/DataChannel.cpp b/Engine/Source/Runtime/Engine/Private/DataChannel.cpp
--- a/Engine/Source/Runtime/Engine/Private/DataChannel.cpp
+++ b/Engine/Source/Runtime/Engine/Private/DataChannel.cpp
@@ -2961,7 +2961,7 @@ void UActorChannel::ReceivedBunch( FInBunch & Bunch )
 			uint16 NumMustBeMappedGUIDs = 0;
 			Bunch << NumMustBeMappedGUIDs;
 
-			QueuedObjectsToTrack.SetNum(NumMustBeMappedGUIDs);
+			QueuedObjectsToTrack.Reserve(NumMustBeMappedGUIDs);
 			//UE_LOG( LogNetTraffic, Warning, TEXT( "Read must be mapped GUID's. NumMustBeMappedGUIDs: %i" ), NumMustBeMappedGUIDs );
 
 			FNetGUIDCache* GuidCache = Connection->Driver->GuidCache.Get();
