Removing heavy include dependencies in NavRelevantInterface.h

diff --git a/Engine/Plugins/2D/Paper2D/Source/Paper2D/Private/PaperGroupedSpriteComponent.cpp b/Engine/Plugins/2D/Paper2D/Source/Paper2D/Private/PaperGroupedSpriteComponent.cpp
--- a/Engine/Plugins/2D/Paper2D/Source/Paper2D/Private/PaperGroupedSpriteComponent.cpp
+++ b/Engine/Plugins/2D/Paper2D/Source/Paper2D/Private/PaperGroupedSpriteComponent.cpp
@@ -2,6 +2,7 @@
 
 #include "PaperGroupedSpriteComponent.h"
 #include "AI/NavigationSystemBase.h"
+#include "AI/Navigation/NavigationRelevantData.h"
 #include "Engine/CollisionProfile.h"
 
 #include "GroupedSpriteSceneProxy.h"
diff --git a/Engine/Plugins/Experimental/Water/Source/Runtime/Private/WaterBodyComponent.cpp b/Engine/Plugins/Experimental/Water/Source/Runtime/Private/WaterBodyComponent.cpp
--- a/Engine/Plugins/Experimental/Water/Source/Runtime/Private/WaterBodyComponent.cpp
+++ b/Engine/Plugins/Experimental/Water/Source/Runtime/Private/WaterBodyComponent.cpp
@@ -8,6 +8,7 @@
 #include "LandscapeProxy.h"
 #include "Materials/MaterialInstanceDynamic.h"
 #include "AI/NavigationSystemBase.h"
+#include "AI/Navigation/NavigationRelevantData.h"
 #include "Spatial/MeshAABBTree3.h"
 #include "WaterModule.h"
 #include "WaterBodyActor.h"
diff --git a/Engine/Plugins/Experimental/Water/Source/Runtime/Public/WaterBodyComponent.h b/Engine/Plugins/Experimental/Water/Source/Runtime/Public/WaterBodyComponent.h
--- a/Engine/Plugins/Experimental/Water/Source/Runtime/Public/WaterBodyComponent.h
+++ b/Engine/Plugins/Experimental/Water/Source/Runtime/Public/WaterBodyComponent.h
@@ -42,6 +42,7 @@ class AWaterZone;
 class ALandscapeProxy;
 class UMaterialInstanceDynamic;
 class FTokenizedMessage;
+class UNavAreaBase;
 namespace UE::Geometry { class FDynamicMesh3; }
 struct FMeshDescription;
 
diff --git a/Engine/Source/Runtime/Engine/Classes/AI/Navigation/NavRelevantInterface.h b/Engine/Source/Runtime/Engine/Classes/AI/Navigation/NavRelevantInterface.h
--- a/Engine/Source/Runtime/Engine/Classes/AI/Navigation/NavRelevantInterface.h
+++ b/Engine/Source/Runtime/Engine/Classes/AI/Navigation/NavRelevantInterface.h
@@ -2,14 +2,20 @@
 
 #pragma once
 
+#if UE_ENABLE_INCLUDE_ORDER_DEPRECATED_IN_5_4
 #include "CoreMinimal.h"
+#endif
 #include "UObject/ObjectMacros.h"
 #include "UObject/Interface.h"
-#include "AI/Navigation/NavigationTypes.h"
+#include "AI/Navigation/NavigationTypes.h"	//ENavDataGatheringMode
+#if UE_ENABLE_INCLUDE_ORDER_DEPRECATED_IN_5_4
 #include "AI/NavigationModifier.h"
+#include "AI/NavigationRelevantData.h"
+#endif
 #include "NavRelevantInterface.generated.h"
 
 struct FNavigableGeometryExport;
+struct FNavigationRelevantData;
 class UBodySetup;
 
 /** Determines if a NavRelevant object contains custom collision for navigation/AI */
@@ -32,146 +38,6 @@ enum Type : int
 };
 }
 
-struct FNavigationRelevantDataFilter 
-{
-	/** pass when actor has geometry */
-	uint32 bIncludeGeometry : 1;
-	/** pass when actor has any offmesh link modifier */
-	uint32 bIncludeOffmeshLinks : 1;
-	/** pass when actor has any area modifier */
-	uint32 bIncludeAreas : 1;
-	/** pass when actor has any modifier with meta area */
-	uint32 bIncludeMetaAreas : 1;
-	/** fail if from level loading (only valid in WP dynamic mode) */
-	uint32 bExcludeLoadedData : 1;
-
-	FNavigationRelevantDataFilter() 
-		: bIncludeGeometry(false)
-		, bIncludeOffmeshLinks(false)
-		, bIncludeAreas(false)
-		, bIncludeMetaAreas(false)
-		, bExcludeLoadedData(false)
-	{}
-};
-
-// @todo consider optional structures that can contain a delegate instead of 
-// actual copy of collision data
-struct FNavigationRelevantData : public TSharedFromThis<FNavigationRelevantData, ESPMode::ThreadSafe>
-{
-	DECLARE_DELEGATE_RetVal_OneParam(bool, FFilterNavDataDelegate, const struct FNavDataConfig*);
-
-	/** CollisionData should always start with this struct for validation purposes */
-	struct FCollisionDataHeader
-	{
-		int32 DataSize;
-
-		static bool IsValid(const uint8* RawData, int32 RawDataSize);
-	};
-
-	/** exported geometry (used by recast navmesh as FRecastGeometryCache) */
-	TNavStatArray<uint8> CollisionData;
-
-	/** cached voxels (used by recast navmesh as FRecastVoxelCache) */
-	TNavStatArray<uint8> VoxelData;
-
-	/** bounds of geometry (unreal coords) */
-	FBox Bounds;
-
-	/** Gathers per instance data for navigation geometry in a specified area box */
-	FNavDataPerInstanceTransformDelegate NavDataPerInstanceTransformDelegate;
-
-	/** called to check if hosted geometry should be used for given FNavDataConfig. If not set then "true" is assumed.
-	 *  Might want to set bUseVirtualGeometryFilteringAndDirtying to true in the Navmesh class you are excluding geometry from.
-	 *  This will improve cpu performance by stopping the navmesh from dirtying tiles requested by actors being excluded by this delegate.
-	 */
-	FFilterNavDataDelegate ShouldUseGeometryDelegate;
-
-	/** additional modifiers: areas and external links */
-	FCompositeNavModifier Modifiers;
-
-	/** UObject these data represents */
-	TWeakObjectPtr<UObject> SourceObject;
-
-	/** get set to true when lazy navigation exporting is enabled and this navigation data has "potential" of
-	*	containing geometry data. First access will result in gathering the data and setting this flag back to false.
-	*	Mind that this flag can go back to 'true' if related data gets cleared out. */
-	uint32 bPendingLazyGeometryGathering : 1;
-	uint32 bPendingLazyModifiersGathering : 1;
-	uint32 bPendingChildLazyModifiersGathering : 1;
-
-	uint32 bSupportsGatheringGeometrySlices : 1;
-
-	/** Indicates that this data will not dirty the navmesh when added or removed from the octree. */
-	uint32 bShouldSkipDirtyAreaOnAddOrRemove : 1;
-
-	/** From level loading (only valid in WP dynamic mode) */
-	uint32 bLoadedData : 1;
-
-	FNavigationRelevantData(UObject& Source)
-		: SourceObject(&Source)
-		, bPendingLazyGeometryGathering(false)
-		, bPendingLazyModifiersGathering(false)
-		, bPendingChildLazyModifiersGathering(false)
-		, bSupportsGatheringGeometrySlices(false)
-		, bShouldSkipDirtyAreaOnAddOrRemove(false)
-		, bLoadedData(false)
-	{}
-
-	FORCEINLINE bool HasGeometry() const { return VoxelData.Num() || CollisionData.Num(); }
-	FORCEINLINE bool HasModifiers() const { return !Modifiers.IsEmpty(); }
-	FORCEINLINE bool HasDynamicModifiers() const { return Modifiers.IsDynamic(); }
-	FORCEINLINE bool IsPendingLazyGeometryGathering() const { return bPendingLazyGeometryGathering; }
-	FORCEINLINE bool IsPendingLazyModifiersGathering() const { return bPendingLazyModifiersGathering; }
-	FORCEINLINE bool IsPendingChildLazyModifiersGathering() const { return bPendingChildLazyModifiersGathering; }
-	FORCEINLINE bool NeedAnyPendingLazyModifiersGathering() const { return bPendingLazyModifiersGathering || bPendingChildLazyModifiersGathering; }
-	FORCEINLINE bool SupportsGatheringGeometrySlices() const { return bSupportsGatheringGeometrySlices; }
-
-	/**
-	 * Indicates that this object will not dirty the navmesh when added or removed from the octree.
-	 * In this case we expect it to manually dirty areas (e.g. using OnObjectBoundsChanged).
-	 */
-	FORCEINLINE bool ShouldSkipDirtyAreaOnAddOrRemove() const { return bShouldSkipDirtyAreaOnAddOrRemove; }
-
-	FORCEINLINE bool IsEmpty() const { return !HasGeometry() && !HasModifiers(); }
-	FORCEINLINE SIZE_T GetAllocatedSize() const { return CollisionData.GetAllocatedSize() + VoxelData.GetAllocatedSize() + Modifiers.GetAllocatedSize(); }
-	FORCEINLINE SIZE_T GetGeometryAllocatedSize() const { return CollisionData.GetAllocatedSize() + VoxelData.GetAllocatedSize(); }
-	FORCEINLINE int32 GetDirtyFlag() const
-	{
-		const bool bSetGeometryFlag = HasGeometry() || IsPendingLazyGeometryGathering() ||
-			Modifiers.GetFillCollisionUnderneathForNavmesh() || Modifiers.GetMaskFillCollisionUnderneathForNavmesh() ||
-			(Modifiers.GetNavMeshResolution() != ENavigationDataResolution::Invalid);
-		
-		return (bSetGeometryFlag ? ENavigationDirtyFlag::Geometry : 0) |
-			((HasDynamicModifiers() || NeedAnyPendingLazyModifiersGathering()) ? ENavigationDirtyFlag::DynamicModifier : 0) |
-			(Modifiers.HasAgentHeightAdjust() ? ENavigationDirtyFlag::UseAgentHeight : 0);
-	}
-
-	FORCEINLINE FCompositeNavModifier GetModifierForAgent(const struct FNavAgentProperties* NavAgent = nullptr) const
-	{
-		return Modifiers.HasMetaAreas() ? Modifiers.GetInstantiatedMetaModifier(NavAgent, SourceObject) : Modifiers;
-	}
-
-	ENGINE_API bool HasPerInstanceTransforms() const;
-	ENGINE_API bool IsMatchingFilter(const FNavigationRelevantDataFilter& Filter) const;
-	ENGINE_API void Shrink();
-	ENGINE_API bool IsCollisionDataValid() const;
-
-	void ValidateAndShrink()
-	{
-		if (IsCollisionDataValid())
-		{
-			Shrink();
-		}
-		else
-		{
-			CollisionData.Empty();
-		}
-	}
-
-	FORCEINLINE UObject* GetOwner() const { return SourceObject.Get(); }
-	ENGINE_API FORCEINLINE decltype(SourceObject)& GetOwnerPtr() { return SourceObject; }
-};
-
 UINTERFACE(MinimalAPI, meta=(CannotImplementInterfaceInBlueprint))
 class UNavRelevantInterface : public UInterface
 {
diff --git a/Engine/Source/Runtime/Engine/Classes/AI/Navigation/NavigationRelevantData.h b/Engine/Source/Runtime/Engine/Classes/AI/Navigation/NavigationRelevantData.h
new file mode 100644
--- /dev/null
+++ b/Engine/Source/Runtime/Engine/Classes/AI/Navigation/NavigationRelevantData.h
@@ -0,0 +1,147 @@
+﻿// Copyright Epic Games, Inc. All Rights Reserved.
+
+#pragma once
+
+#include "UObject/ObjectMacros.h"
+#include "AI/Navigation/NavigationTypes.h"
+#include "AI/NavigationModifier.h"
+
+struct FNavigationRelevantDataFilter 
+{
+	/** pass when actor has geometry */
+	uint32 bIncludeGeometry : 1;
+	/** pass when actor has any offmesh link modifier */
+	uint32 bIncludeOffmeshLinks : 1;
+	/** pass when actor has any area modifier */
+	uint32 bIncludeAreas : 1;
+	/** pass when actor has any modifier with meta area */
+	uint32 bIncludeMetaAreas : 1;
+	/** fail if from level loading (only valid in WP dynamic mode) */
+	uint32 bExcludeLoadedData : 1;
+
+	FNavigationRelevantDataFilter() 
+		: bIncludeGeometry(false)
+		, bIncludeOffmeshLinks(false)
+		, bIncludeAreas(false)
+		, bIncludeMetaAreas(false)
+		, bExcludeLoadedData(false)
+	{}
+};
+
+// @todo consider optional structures that can contain a delegate instead of 
+// actual copy of collision data
+struct FNavigationRelevantData : public TSharedFromThis<FNavigationRelevantData, ESPMode::ThreadSafe>
+{
+	DECLARE_DELEGATE_RetVal_OneParam(bool, FFilterNavDataDelegate, const struct FNavDataConfig*);
+
+	/** CollisionData should always start with this struct for validation purposes */
+	struct FCollisionDataHeader
+	{
+		int32 DataSize;
+
+		static bool IsValid(const uint8* RawData, int32 RawDataSize);
+	};
+
+	/** exported geometry (used by recast navmesh as FRecastGeometryCache) */
+	TNavStatArray<uint8> CollisionData;
+
+	/** cached voxels (used by recast navmesh as FRecastVoxelCache) */
+	TNavStatArray<uint8> VoxelData;
+
+	/** bounds of geometry (unreal coords) */
+	FBox Bounds;
+
+	/** Gathers per instance data for navigation geometry in a specified area box */
+	FNavDataPerInstanceTransformDelegate NavDataPerInstanceTransformDelegate;
+
+	/** called to check if hosted geometry should be used for given FNavDataConfig. If not set then "true" is assumed.
+	 *  Might want to set bUseVirtualGeometryFilteringAndDirtying to true in the Navmesh class you are excluding geometry from.
+	 *  This will improve cpu performance by stopping the navmesh from dirtying tiles requested by actors being excluded by this delegate.
+	 */
+	FFilterNavDataDelegate ShouldUseGeometryDelegate;
+
+	/** additional modifiers: areas and external links */
+	FCompositeNavModifier Modifiers;
+
+	/** UObject these data represents */
+	TWeakObjectPtr<UObject> SourceObject;
+
+	/** get set to true when lazy navigation exporting is enabled and this navigation data has "potential" of
+	*	containing geometry data. First access will result in gathering the data and setting this flag back to false.
+	*	Mind that this flag can go back to 'true' if related data gets cleared out. */
+	uint32 bPendingLazyGeometryGathering : 1;
+	uint32 bPendingLazyModifiersGathering : 1;
+	uint32 bPendingChildLazyModifiersGathering : 1;
+
+	uint32 bSupportsGatheringGeometrySlices : 1;
+
+	/** Indicates that this data will not dirty the navmesh when added or removed from the octree. */
+	uint32 bShouldSkipDirtyAreaOnAddOrRemove : 1;
+
+	/** From level loading (only valid in WP dynamic mode) */
+	uint32 bLoadedData : 1;
+
+	FNavigationRelevantData(UObject& Source)
+		: SourceObject(&Source)
+		, bPendingLazyGeometryGathering(false)
+		, bPendingLazyModifiersGathering(false)
+		, bPendingChildLazyModifiersGathering(false)
+		, bSupportsGatheringGeometrySlices(false)
+		, bShouldSkipDirtyAreaOnAddOrRemove(false)
+		, bLoadedData(false)
+	{}
+
+	FORCEINLINE bool HasGeometry() const { return VoxelData.Num() || CollisionData.Num(); }
+	FORCEINLINE bool HasModifiers() const { return !Modifiers.IsEmpty(); }
+	FORCEINLINE bool HasDynamicModifiers() const { return Modifiers.IsDynamic(); }
+	FORCEINLINE bool IsPendingLazyGeometryGathering() const { return bPendingLazyGeometryGathering; }
+	FORCEINLINE bool IsPendingLazyModifiersGathering() const { return bPendingLazyModifiersGathering; }
+	FORCEINLINE bool IsPendingChildLazyModifiersGathering() const { return bPendingChildLazyModifiersGathering; }
+	FORCEINLINE bool NeedAnyPendingLazyModifiersGathering() const { return bPendingLazyModifiersGathering || bPendingChildLazyModifiersGathering; }
+	FORCEINLINE bool SupportsGatheringGeometrySlices() const { return bSupportsGatheringGeometrySlices; }
+
+	/**
+	 * Indicates that this object will not dirty the navmesh when added or removed from the octree.
+	 * In this case we expect it to manually dirty areas (e.g. using OnObjectBoundsChanged).
+	 */
+	FORCEINLINE bool ShouldSkipDirtyAreaOnAddOrRemove() const { return bShouldSkipDirtyAreaOnAddOrRemove; }
+
+	FORCEINLINE bool IsEmpty() const { return !HasGeometry() && !HasModifiers(); }
+	FORCEINLINE SIZE_T GetAllocatedSize() const { return CollisionData.GetAllocatedSize() + VoxelData.GetAllocatedSize() + Modifiers.GetAllocatedSize(); }
+	FORCEINLINE SIZE_T GetGeometryAllocatedSize() const { return CollisionData.GetAllocatedSize() + VoxelData.GetAllocatedSize(); }
+	FORCEINLINE int32 GetDirtyFlag() const
+	{
+		const bool bSetGeometryFlag = HasGeometry() || IsPendingLazyGeometryGathering() ||
+			Modifiers.GetFillCollisionUnderneathForNavmesh() || Modifiers.GetMaskFillCollisionUnderneathForNavmesh() ||
+			(Modifiers.GetNavMeshResolution() != ENavigationDataResolution::Invalid);
+		
+		return (bSetGeometryFlag ? ENavigationDirtyFlag::Geometry : 0) |
+			((HasDynamicModifiers() || NeedAnyPendingLazyModifiersGathering()) ? ENavigationDirtyFlag::DynamicModifier : 0) |
+			(Modifiers.HasAgentHeightAdjust() ? ENavigationDirtyFlag::UseAgentHeight : 0);
+	}
+
+	FORCEINLINE FCompositeNavModifier GetModifierForAgent(const struct FNavAgentProperties* NavAgent = nullptr) const
+	{
+		return Modifiers.HasMetaAreas() ? Modifiers.GetInstantiatedMetaModifier(NavAgent, SourceObject) : Modifiers;
+	}
+
+	ENGINE_API bool HasPerInstanceTransforms() const;
+	ENGINE_API bool IsMatchingFilter(const FNavigationRelevantDataFilter& Filter) const;
+	ENGINE_API void Shrink();
+	ENGINE_API bool IsCollisionDataValid() const;
+
+	void ValidateAndShrink()
+	{
+		if (IsCollisionDataValid())
+		{
+			Shrink();
+		}
+		else
+		{
+			CollisionData.Empty();
+		}
+	}
+
+	FORCEINLINE UObject* GetOwner() const { return SourceObject.Get(); }
+	ENGINE_API FORCEINLINE decltype(SourceObject)& GetOwnerPtr() { return SourceObject; }
+};
diff --git a/Engine/Source/Runtime/Engine/Private/AI/Navigation/NavigationTypes.cpp b/Engine/Source/Runtime/Engine/Private/AI/Navigation/NavigationTypes.cpp
--- a/Engine/Source/Runtime/Engine/Private/AI/Navigation/NavigationTypes.cpp
+++ b/Engine/Source/Runtime/Engine/Private/AI/Navigation/NavigationTypes.cpp
@@ -3,6 +3,7 @@
 #include "AI/Navigation/NavigationTypes.h"
 #include "AI/NavigationSystemBase.h"
 #include "AI/Navigation/NavQueryFilter.h"
+#include "AI/Navigation/NavigationRelevantData.h"
 #include "Engine/Level.h"
 #include "EngineStats.h"
 #include "Components/ShapeComponent.h"
diff --git a/Engine/Source/Runtime/Engine/Private/Components/PrimitiveComponent.cpp b/Engine/Source/Runtime/Engine/Private/Components/PrimitiveComponent.cpp
--- a/Engine/Source/Runtime/Engine/Private/Components/PrimitiveComponent.cpp
+++ b/Engine/Source/Runtime/Engine/Private/Components/PrimitiveComponent.cpp
@@ -18,6 +18,7 @@
 #include "HLOD/HLODBatchingPolicy.h"
 #include "PSOPrecache.h"
 #include "AI/NavigationSystemBase.h"
+#include "AI/Navigation/NavigationRelevantData.h"
 #include "GameFramework/PlayerController.h"
 #include "GameFramework/PhysicsVolume.h"
 #include "GameFramework/WorldSettings.h"
diff --git a/Engine/Source/Runtime/Engine/Private/Components/ShapeComponent.cpp b/Engine/Source/Runtime/Engine/Private/Components/ShapeComponent.cpp
--- a/Engine/Source/Runtime/Engine/Private/Components/ShapeComponent.cpp
+++ b/Engine/Source/Runtime/Engine/Private/Components/ShapeComponent.cpp
@@ -3,6 +3,7 @@
 
 #include "Components/ShapeComponent.h"
 #include "AI/NavigationSystemBase.h"
+#include "AI/Navigation/NavigationRelevantData.h"
 #include "BodySetupEnums.h"
 #include "PhysicsEngine/BoxElem.h"
 #include "PhysicsEngine/SphereElem.h"
diff --git a/Engine/Source/Runtime/Engine/Private/Components/StaticMeshComponent.cpp b/Engine/Source/Runtime/Engine/Private/Components/StaticMeshComponent.cpp
--- a/Engine/Source/Runtime/Engine/Private/Components/StaticMeshComponent.cpp
+++ b/Engine/Source/Runtime/Engine/Private/Components/StaticMeshComponent.cpp
@@ -40,6 +40,7 @@
 #include "AI/Navigation/NavCollisionBase.h"
 #include "Engine/StaticMeshSocket.h"
 #include "AI/NavigationSystemBase.h"
+#include "AI/Navigation/NavigationRelevantData.h"
 #include "PhysicsEngine/BodySetup.h"
 #include "ComponentRecreateRenderStateContext.h"
 #include "Engine/StaticMesh.h"
diff --git a/Engine/Source/Runtime/Engine/Private/InstancedStaticMesh.cpp b/Engine/Source/Runtime/Engine/Private/InstancedStaticMesh.cpp
--- a/Engine/Source/Runtime/Engine/Private/InstancedStaticMesh.cpp
+++ b/Engine/Source/Runtime/Engine/Private/InstancedStaticMesh.cpp
@@ -16,6 +16,7 @@
 #include "UnrealEngine.h"
 #include "AI/NavigationSystemHelpers.h"
 #include "AI/Navigation/NavCollisionBase.h"
+#include "AI/Navigation/NavigationRelevantData.h"
 #include "MeshDrawShaderBindings.h"
 #include "Misc/UObjectToken.h"
 #include "Misc/DelayedAutoRegister.h"
diff --git a/Engine/Source/Runtime/NavigationSystem/Private/NavLinkComponent.cpp b/Engine/Source/Runtime/NavigationSystem/Private/NavLinkComponent.cpp
--- a/Engine/Source/Runtime/NavigationSystem/Private/NavLinkComponent.cpp
+++ b/Engine/Source/Runtime/NavigationSystem/Private/NavLinkComponent.cpp
@@ -3,6 +3,7 @@
 #include "NavLinkComponent.h"
 #include "NavLinkRenderingProxy.h"
 #include "AI/NavigationSystemHelpers.h"
+#include "AI/Navigation/NavigationRelevantData.h"
 #include "NavigationSystemTypes.h"
 #include "NavAreas/NavArea_Default.h"
 #include "Engine/CollisionProfile.h"
diff --git a/Engine/Source/Runtime/NavigationSystem/Private/SplineNavModifierComponent.cpp b/Engine/Source/Runtime/NavigationSystem/Private/SplineNavModifierComponent.cpp
--- a/Engine/Source/Runtime/NavigationSystem/Private/SplineNavModifierComponent.cpp
+++ b/Engine/Source/Runtime/NavigationSystem/Private/SplineNavModifierComponent.cpp
@@ -3,6 +3,7 @@
 #include "SplineNavModifierComponent.h"
 
 #include "AI/NavigationSystemBase.h"
+#include "AI/Navigation/NavigationRelevantData.h"
 #include "Components/SplineComponent.h"
 #include "VisualLogger/VisualLogger.h"
 
diff --git a/Engine/Source/Runtime/NavigationSystem/Public/NavMesh/RecastNavMeshGenerator.h b/Engine/Source/Runtime/NavigationSystem/Public/NavMesh/RecastNavMeshGenerator.h
--- a/Engine/Source/Runtime/NavigationSystem/Public/NavMesh/RecastNavMeshGenerator.h
+++ b/Engine/Source/Runtime/NavigationSystem/Public/NavMesh/RecastNavMeshGenerator.h
@@ -5,6 +5,7 @@
 #include "Stats/Stats.h"
 #include "AI/Navigation/NavigationTypes.h"
 #include "AI/Navigation/NavigationInvokerPriority.h"
+#include "AI/Navigation/NavigationRelevantData.h"
 #include "EngineDefines.h"
 #include "AI/NavigationModifier.h"
 #include "NavigationOctree.h"
@@ -31,7 +32,6 @@ class FNavMeshBuildContext;
 class FRecastNavMeshGenerator;
 struct FTileRasterizationContext;
 struct BuildContext;
-struct FNavigationRelevantData;
 struct dtTileCacheLayer;
 struct FKAggregateGeom;
 struct FTileCacheCompressor;
diff --git a/Engine/Source/Runtime/NavigationSystem/Public/NavModifierComponent.h b/Engine/Source/Runtime/NavigationSystem/Public/NavModifierComponent.h
--- a/Engine/Source/Runtime/NavigationSystem/Public/NavModifierComponent.h
+++ b/Engine/Source/Runtime/NavigationSystem/Public/NavModifierComponent.h
@@ -11,6 +11,7 @@
 
 struct FNavigationRelevantData;
 class UBodySetup;
+enum class ENavigationDataResolution : uint8;
 
 UCLASS(ClassGroup = (Navigation), meta = (BlueprintSpawnableComponent), hidecategories = (Activation), config = Engine, defaultconfig, MinimalAPI)
 class UNavModifierComponent : public UNavRelevantComponent
diff --git a/Engine/Source/Runtime/NavigationSystem/Public/NavModifierVolume.h b/Engine/Source/Runtime/NavigationSystem/Public/NavModifierVolume.h
--- a/Engine/Source/Runtime/NavigationSystem/Public/NavModifierVolume.h
+++ b/Engine/Source/Runtime/NavigationSystem/Public/NavModifierVolume.h
@@ -11,6 +11,8 @@
 #include "GameFramework/Volume.h"
 #include "NavModifierVolume.generated.h"
 
+enum class ENavigationDataResolution : uint8;
+
 struct FNavigationRelevantData;
 
 /** 
diff --git a/Engine/Source/Runtime/NavigationSystem/Public/NavigationOctree.h b/Engine/Source/Runtime/NavigationSystem/Public/NavigationOctree.h
--- a/Engine/Source/Runtime/NavigationSystem/Public/NavigationOctree.h
+++ b/Engine/Source/Runtime/NavigationSystem/Public/NavigationOctree.h
@@ -10,6 +10,7 @@
 #include "EngineStats.h"
 #include "AI/NavigationModifier.h"
 #include "AI/Navigation/NavRelevantInterface.h"
+#include "AI/Navigation/NavigationRelevantData.h"
 #include "Math/GenericOctree.h"
 #include "HAL/LowLevelMemTracker.h"
 
