Renaming a plugin and adding proper core redirects to the project defaultengine.ini was failing on World Partition maps and Level Instance World Assets

Fix missing CoreRedirect fixups:
- AssetRegistryHelpers now calls FixupCoreRedirect on the passed in SoftObjectPath before handling redirector asset fixups
- Temp fix: FWorldPartitionActorDescUtils::GetActorDescriptorFromAssetData calls FixupCoreRedirect on the FAssetData ObjectPath as currently this isn't done on all FAssetDatas (pending jiras linked)

#rnx

diff --git a/Engine/Source/Runtime/AssetRegistry/Private/AssetRegistryHelpers.cpp b/Engine/Source/Runtime/AssetRegistry/Private/AssetRegistryHelpers.cpp
--- a/Engine/Source/Runtime/AssetRegistry/Private/AssetRegistryHelpers.cpp
+++ b/Engine/Source/Runtime/AssetRegistry/Private/AssetRegistryHelpers.cpp
@@ -235,6 +235,7 @@ UAssetRegistryHelpers::FTemporaryCachingModeScope::~FTemporaryCachingModeScope()
 void UAssetRegistryHelpers::FixupRedirectedAssetPath(FSoftObjectPath& InOutSoftObjectPath)
 {
 	FSoftObjectPath FoundRedirection;
+	InOutSoftObjectPath.FixupCoreRedirects();
 
 #if WITH_EDITOR
 	FoundRedirection = GRedirectCollector.GetAssetPathRedirection(InOutSoftObjectPath);
diff --git a/Engine/Source/Runtime/Engine/Private/WorldPartition/WorldPartitionActorDescUtils.cpp b/Engine/Source/Runtime/Engine/Private/WorldPartition/WorldPartitionActorDescUtils.cpp
--- a/Engine/Source/Runtime/Engine/Private/WorldPartition/WorldPartitionActorDescUtils.cpp
+++ b/Engine/Source/Runtime/Engine/Private/WorldPartition/WorldPartitionActorDescUtils.cpp
@@ -77,10 +77,15 @@ TUniquePtr<FWorldPartitionActorDesc> FWorldPartitionActorDescUtils::GetActorDesc
 {
 	if (IsValidActorDescriptorFromAssetData(InAssetData))
 	{
+		// Fixing here for now until following jira tasks are closed to unblock users that are renaming plugins:
+		// @todo_ow: remove once https://jira.it.epicgames.com/browse/UE-168245 & https://jira.it.epicgames.com/browse/UE-144431 are closed
+		FSoftObjectPath ActorPath = InAssetData.GetSoftObjectPath();
+		ActorPath.FixupCoreRedirects();
+
 		FWorldPartitionActorDescInitData ActorDescInitData = FWorldPartitionActorDescInitData()
 			.SetNativeClass(GetActorNativeClassFromAssetData(InAssetData))
 			.SetPackageName(InAssetData.PackageName)
-			.SetActorPath(InAssetData.GetSoftObjectPath());
+			.SetActorPath(ActorPath);
 
 		FString ActorMetaDataStr;
 		verify(InAssetData.GetTagValue(NAME_ActorMetaData, ActorMetaDataStr));
