Fix animation compression not dispatching when the curve codec is changed and the animation has no bone data

diff --git a/Engine/Source/Runtime/Engine/Private/Animation/AnimSequence.cpp b/Engine/Source/Runtime/Engine/Private/Animation/AnimSequence.cpp
--- a/Engine/Source/Runtime/Engine/Private/Animation/AnimSequence.cpp
+++ b/Engine/Source/Runtime/Engine/Private/Animation/AnimSequence.cpp
@@ -1151,7 +1151,7 @@ void UAnimSequence::PostEditChangeProperty(FPropertyChangedEvent& PropertyChange
 
 	// @Todo fix me: This is temporary fix to make sure they always have compressed data
 	ValidateModel();
-	if (IsDataModelValid() && DataModelInterface->GetNumBoneTracks() && bNeedPostProcess)
+	if (IsDataModelValid() && bNeedPostProcess)
 	{		
 		ClearAllCachedCookedPlatformData();
 		BeginCacheDerivedDataForCurrentPlatform();
