PIEFixupSerializer - Don't trigger serialization of compilable assets

diff --git a/Engine/Source/Runtime/Engine/Private/PieFixupSerializer.cpp b/Engine/Source/Runtime/Engine/Private/PieFixupSerializer.cpp
--- a/Engine/Source/Runtime/Engine/Private/PieFixupSerializer.cpp
+++ b/Engine/Source/Runtime/Engine/Private/PieFixupSerializer.cpp
@@ -15,13 +15,8 @@ namespace
 }
 
 FPIEFixupSerializer::FPIEFixupSerializer(UObject* InRoot, int32 InPIEInstanceID)
-	: SoftObjectPathFixupFunction(DefaultSoftObjectPathFixupFunction)
-	, Root(InRoot)
-	, PIEInstanceID(InPIEInstanceID)
+	: FPIEFixupSerializer(InRoot, InPIEInstanceID, DefaultSoftObjectPathFixupFunction)
 {
-	this->ArShouldSkipBulkData = true;
-	this->ArIsObjectReferenceCollector = true;
-	this->ArIsModifyingWeakAndStrongReferences = true;
 }
 
 FPIEFixupSerializer::FPIEFixupSerializer(UObject* InRoot, int32 InPIEInstanceID, TFunctionRef<void(int32, FSoftObjectPath&)> InSoftObjectPathFixupFunction)
@@ -32,6 +27,9 @@ FPIEFixupSerializer::FPIEFixupSerializer(UObject* InRoot, int32 InPIEInstanceID,
 	this->ArShouldSkipBulkData = true;
 	this->ArIsObjectReferenceCollector = true;
 	this->ArIsModifyingWeakAndStrongReferences = true;
+
+	// Don't trigger serialization of compilable assets
+	SetShouldSkipCompilingAssets(true);
 }
 
 bool FPIEFixupSerializer::ShouldSkipProperty(const FProperty* InProperty) const
