World Partition - Minimap: Hide useless properties categories for the AWorldPartitionMinimap actor

diff --git a/Engine/Source/Runtime/Engine/Public/WorldPartition/WorldPartitionMiniMap.h b/Engine/Source/Runtime/Engine/Public/WorldPartition/WorldPartitionMiniMap.h
--- a/Engine/Source/Runtime/Engine/Public/WorldPartition/WorldPartitionMiniMap.h
+++ b/Engine/Source/Runtime/Engine/Public/WorldPartition/WorldPartitionMiniMap.h
@@ -11,7 +11,7 @@ class UTexture2D;
 /**
  * A mini map to preview the world in world partition window. (editor-only)
  */
-UCLASS(hidecategories = (Actor, Advanced, Display, Events, Object, Attachment, Info, Input, Blueprint, Layers, Tags, Replication, Physics, Cooking), notplaceable, MinimalAPI)
+UCLASS(hidecategories = (Actor, Advanced, Display, Events, Object, Attachment, Info, Input, Blueprint, Layers, Tags, Replication, Physics, Cooking, Networking, LevelInstance), notplaceable, MinimalAPI)
 class AWorldPartitionMiniMap : public AInfo
 {
 	GENERATED_BODY()
