Iris - Fixed bug causing PIE clients to use server's GameNetDriver Iris configuration.

#rnx

diff --git a/Engine/Source/Runtime/Engine/Private/UnrealEngine.cpp b/Engine/Source/Runtime/Engine/Private/UnrealEngine.cpp
--- a/Engine/Source/Runtime/Engine/Private/UnrealEngine.cpp
+++ b/Engine/Source/Runtime/Engine/Private/UnrealEngine.cpp
@@ -13792,8 +13792,8 @@ namespace UE::Private
 			}
 
 #if WITH_EDITOR
-			// In PIE let's cheat and make sure the clients follow what the server's net driver is using
-			if (Context.WorldType == EWorldType::PIE && !Context.RunAsDedicated)
+			// In PIE let's make sure the clients follow what the server's game net driver is using
+			if (Context.WorldType == EWorldType::PIE && !Context.RunAsDedicated && InNetDriverName == NAME_GameNetDriver)
 			{
 				if (FWorldContext* ServerPIEContext = GEngine->GetWorldContextFromPIEInstance(0))
 				{
