[MovementComponent]
Update moving base when we are playing a root motion montage

#rnx

diff --git a/Engine/Source/Runtime/Engine/Private/Components/CharacterMovementComponent.cpp b/Engine/Source/Runtime/Engine/Private/Components/CharacterMovementComponent.cpp
--- a/Engine/Source/Runtime/Engine/Private/Components/CharacterMovementComponent.cpp
+++ b/Engine/Source/Runtime/Engine/Private/Components/CharacterMovementComponent.cpp
@@ -1953,6 +1953,8 @@ void UCharacterMovementComponent::SimulateRootMotion(float DeltaSeconds, const F
 	{
 		FScopedCapsuleMovementUpdate ScopedMovementUpdate(UpdatedComponent, bEnableScopedMovementUpdates);
 
+		MaybeUpdateBasedMovement(DeltaSeconds);
+
 		// Convert Local Space Root Motion to world space. Do it right before used by physics to make sure we use up to date transforms, as translation is relative to rotation.
 		const FTransform WorldSpaceRootMotionTransform = ConvertLocalRootMotionToWorld(LocalRootMotionTransform, DeltaSeconds);
 		RootMotionParams.Set( WorldSpaceRootMotionTransform );
@@ -1987,6 +1989,8 @@ void UCharacterMovementComponent::SimulateRootMotion(float DeltaSeconds, const F
 			const FQuat NewActorRotationQuat = RootMotionRotationQuat * UpdatedComponent->GetComponentQuat();
 			MoveUpdatedComponent(FVector::ZeroVector, NewActorRotationQuat, true);
 		}
+
+		MaybeSaveBaseLocation();
 	}
 
 	// Root Motion has been used, clear
