Backout CL30592377 with a better fix : TemporarilyDisableStreaming needs to happen before Modify but it needs to wait on async compilation.
#tests editor

diff --git a/Engine/Source/Runtime/Engine/Private/Texture2D.cpp b/Engine/Source/Runtime/Engine/Private/Texture2D.cpp
--- a/Engine/Source/Runtime/Engine/Private/Texture2D.cpp
+++ b/Engine/Source/Runtime/Engine/Private/Texture2D.cpp
@@ -1415,6 +1415,7 @@ void UTexture2D::UpdateTextureRegions(int32 MipIndex, uint32 NumRegions, const F
 #if WITH_EDITOR
 void UTexture2D::TemporarilyDisableStreaming()
 {
+	BlockOnAnyAsyncBuild();
 	if( !bTemporarilyDisableStreaming )
 	{
 		bTemporarilyDisableStreaming = true;
