Add syncing of the virtual texture pool config version to render thread.
Fixes issue on adaptive VT where during pool resize, the render thread will allocate VTs into the new physical pool before it is ready.

diff --git a/Engine/Source/Runtime/Engine/Classes/VT/VirtualTexturePoolConfig.h b/Engine/Source/Runtime/Engine/Classes/VT/VirtualTexturePoolConfig.h
--- a/Engine/Source/Runtime/Engine/Classes/VT/VirtualTexturePoolConfig.h
+++ b/Engine/Source/Runtime/Engine/Classes/VT/VirtualTexturePoolConfig.h
@@ -82,10 +82,6 @@ class UVirtualTexturePoolConfig : public UDeveloperSettings
 	UPROPERTY(Transient, EditAnywhere, Category = PoolConfig, meta = (TitleProperty = "Formats"))
 	TArray<FVirtualTextureSpacePoolConfig> TransientPools;
 
-	/** Version number used to help track configuration changes. */
-	UPROPERTY(Transient)
-	uint32 Version = 0;
-
 	/** Find a matching config by first searching the TransientPools and then the Pools. */
 	void FindPoolConfig(TEnumAsByte<EPixelFormat> const* InFormats, int32 InNumLayers, int32 InTileSize, FVirtualTextureSpacePoolConfig& OutConfig) const;
 
diff --git a/Engine/Source/Runtime/Engine/Private/VT/VirtualTexturePoolConfig.cpp b/Engine/Source/Runtime/Engine/Private/VT/VirtualTexturePoolConfig.cpp
--- a/Engine/Source/Runtime/Engine/Private/VT/VirtualTexturePoolConfig.cpp
+++ b/Engine/Source/Runtime/Engine/Private/VT/VirtualTexturePoolConfig.cpp
@@ -4,6 +4,7 @@
 
 #include "Async/TaskGraphInterfaces.h"
 #include "HAL/IConsoleManager.h"
+#include "RenderingThread.h"
 #include "VT/VirtualTextureRecreate.h"
 
 #if WITH_EDITOR
@@ -16,8 +17,6 @@
 
 #define LOCTEXT_NAMESPACE "VirtualTexturePool"
 
-void OnVirtualTexturePoolConfigUpdate();
-
 static TAutoConsoleVariable<float> CVarVTPoolSizeScale(
 	TEXT("r.VT.PoolSizeScale"),
 	1.0f,
@@ -64,6 +63,28 @@ void OnVirtualTexturePoolConfigUpdate()
 
 FAutoConsoleVariableSink GVirtualTexturePoolConfigCVarSink(FConsoleCommandDelegate::CreateStatic(&OnVirtualTexturePoolConfigUpdate));
 
+/** Version number used to help track configuration changes. */
+struct FVirtualTexturePoolVersion
+{
+	uint32 Version_GameThread = 0;
+	uint32 Version_RenderThread = 0;
+
+	uint32 Get() const
+	{
+		return IsInRenderingThread() ? Version_RenderThread : Version_GameThread;
+	}
+
+	void Increment_GameThread()
+	{
+		uint32 Version = ++Version_GameThread;
+		ENQUEUE_RENDER_COMMAND(IncrementVersion)([this, Version](FRHICommandListImmediate& RHICmdList)
+		{
+			Version_RenderThread = Version;
+		});
+	}
+};
+FVirtualTexturePoolVersion GVirtualTexturePoolVersion;
+
 
 UVirtualTexturePoolConfig::UVirtualTexturePoolConfig(const FObjectInitializer& ObjectInitializer)
 	: Super(ObjectInitializer)
@@ -147,7 +168,7 @@ bool UVirtualTexturePoolConfig::AddOrModifyTransientPoolConfig(FVirtualTextureSp
 			if (Config.SizeInMegabyte != InConfig.SizeInMegabyte)
 			{
 				Config.SizeInMegabyte = InConfig.SizeInMegabyte;
-				Version++;
+				GVirtualTexturePoolVersion.Increment_GameThread();
 				return true;
 			}
 
@@ -156,7 +177,7 @@ bool UVirtualTexturePoolConfig::AddOrModifyTransientPoolConfig(FVirtualTextureSp
 		}
 	}
 
-	Version++;
+	GVirtualTexturePoolVersion.Increment_GameThread();
 	TransientPools.Add(InConfig);
 	return true;
 }
@@ -166,7 +187,7 @@ bool UVirtualTexturePoolConfig::AddOrModifyTransientPoolConfig(FVirtualTextureSp
 void UVirtualTexturePoolConfig::PostEditChangeProperty(struct FPropertyChangedEvent& PropertyChangedEvent)
 {
 	Super::PostEditChangeProperty(PropertyChangedEvent);
-	Version++;
+	GVirtualTexturePoolVersion.Increment_GameThread();
 	VirtualTexture::Recreate();
 }
 
@@ -282,9 +303,7 @@ uint32 VirtualTexturePool::GetConfigHash()
 	uint32 Hash = GetTypeHash(GetPoolSizeScale());
 	Hash = HashCombine(Hash, GetTypeHash(GetPoolAutoGrow()));
 	Hash = HashCombine(Hash, GetTypeHash(GetSplitPhysicalPoolSize()));
-
-	UVirtualTexturePoolConfig const* PoolConfig = GetDefault<UVirtualTexturePoolConfig>();
-	Hash = HashCombine(Hash, GetTypeHash(PoolConfig->Version));
+	Hash = HashCombine(Hash, GetTypeHash(GVirtualTexturePoolVersion.Get()));
 
 	return Hash;
 }
