Fixed build breakage: incorrect header path when using the deprecated include order.

#rnx

diff --git a/Engine/Source/Runtime/Engine/Classes/AI/Navigation/NavRelevantInterface.h b/Engine/Source/Runtime/Engine/Classes/AI/Navigation/NavRelevantInterface.h
--- a/Engine/Source/Runtime/Engine/Classes/AI/Navigation/NavRelevantInterface.h
+++ b/Engine/Source/Runtime/Engine/Classes/AI/Navigation/NavRelevantInterface.h
@@ -10,7 +10,7 @@
 #include "AI/Navigation/NavigationTypes.h"	//ENavDataGatheringMode
 #if UE_ENABLE_INCLUDE_ORDER_DEPRECATED_IN_5_4
 #include "AI/NavigationModifier.h"
-#include "AI/NavigationRelevantData.h"
+#include "AI/Navigation/NavigationRelevantData.h"
 #endif
 #include "NavRelevantInterface.generated.h"
 
