Removing texture build metadata due to ddc complexities. It will be added to the metadata system in the new build flow SoonTM. Additionally, the alpha information for the source mips is no longer necessary as that is computed elsewhere now.

diff --git a/Engine/Source/Developer/TextureBuild/Private/TextureBuildFunction.cpp b/Engine/Source/Developer/TextureBuild/Private/TextureBuildFunction.cpp
--- a/Engine/Source/Developer/TextureBuild/Private/TextureBuildFunction.cpp
+++ b/Engine/Source/Developer/TextureBuild/Private/TextureBuildFunction.cpp
@@ -36,9 +36,6 @@ DEFINE_LOG_CATEGORY_STATIC(LogTextureBuildFunction, Log, All);
 //
 static const FGuid TextureBuildFunctionVersion(TEXT("B20676CE-A786-43EE-96F0-2620A4C38ACA"));
 
-// This is mirrored in TextureDerivedData.cpp
-static const FGuid TextureMetadataDerivedDataVer(0xB9106D68, 0xA61B4F2A, 0x8105E16F, 0x48799976);
-
 static void ReadCbField(FCbFieldView Field, bool& OutValue) { OutValue = Field.AsBool(OutValue); }
 static void ReadCbField(FCbFieldView Field, int32& OutValue) { OutValue = Field.AsInt32(OutValue); }
 static void ReadCbField(FCbFieldView Field, uint8& OutValue) { OutValue = Field.AsUInt8(OutValue); }
@@ -295,7 +292,6 @@ FGuid FTextureBuildFunction::GetVersion() const
 {
 	UE::DerivedData::FBuildVersionBuilder Builder;
 	Builder << TextureBuildFunctionVersion;
-	Builder << TextureMetadataDerivedDataVer;
 	ITextureFormat* TextureFormat = nullptr;
 	GetVersion(Builder, TextureFormat);
 	if (TextureFormat)
@@ -453,7 +449,9 @@ void FTextureBuildFunction::Build(UE::DerivedData::FBuildContext& Context) const
 			Context.AddValue(UE::DerivedData::FValueId::FromName(ANSITEXTVIEW("CPUCopyRawData")), CPUCopyData);
 		}
 
-		Context.AddValue(UE::DerivedData::FValueId::FromName(ANSITEXTVIEW("TextureBuildMetadata")), BuildMetadata.ToCompactBinaryWithDefaults());
+		// This will get added to the build metadata in a later cl.
+		//Context.AddValue(UE::DerivedData::FValueId::FromName(ANSITEXTVIEW("TextureBuildMetadata")), BuildMetadata.ToCompactBinaryWithDefaults());
+
 		Context.AddValue(UE::DerivedData::FValueId::FromName(ANSITEXTVIEW("EncodedTextureDescription")), UE::TextureBuildUtilities::EncodedTextureDescription::ToCompactBinary(TextureDescription));
 		Context.AddValue(UE::DerivedData::FValueId::FromName(ANSITEXTVIEW("EncodedTextureExtendedData")), UE::TextureBuildUtilities::EncodedTextureExtendedData::ToCompactBinary(ExtendedData));
 
@@ -523,7 +521,8 @@ void GenericTextureTilingBuildFunction(UE::DerivedData::FBuildContext& Context,
 		UE::TextureBuildUtilities::TextureEngineParameters::FromCompactBinary(EngineParameters, EngineParametersCb);
 	}
 
-	UE::TextureBuildUtilities::FTextureBuildMetadata BuildMetadata(FCbObject(Context.FindInput(ANSITEXTVIEW("TextureBuildMetadata"))));
+	// This will get added to the build metadata in a later cl.
+	//UE::TextureBuildUtilities::FTextureBuildMetadata BuildMetadata(FCbObject(Context.FindInput(ANSITEXTVIEW("TextureBuildMetadata"))));
 
 	UE_LOG(LogTextureBuildFunction, Display, TEXT("Tiling %s with %s -> %d source mip(s) with a tail of %d..."), *Context.GetName(), StringCast<TCHAR>(*BuildFunctionName).Get(), TextureDescription.NumMips, TextureExtendedData.NumMipsInTail);
 
@@ -617,5 +616,6 @@ void GenericTextureTilingBuildFunction(UE::DerivedData::FBuildContext& Context,
 
 	Context.AddValue(UE::DerivedData::FValueId::FromName(UTF8TEXTVIEW("EncodedTextureDescription")), UE::TextureBuildUtilities::EncodedTextureDescription::ToCompactBinary(TextureDescription));
 	Context.AddValue(UE::DerivedData::FValueId::FromName(UTF8TEXTVIEW("EncodedTextureExtendedData")), UE::TextureBuildUtilities::EncodedTextureExtendedData::ToCompactBinary(TextureExtendedData));
-	Context.AddValue(UE::DerivedData::FValueId::FromName(UTF8TEXTVIEW("TextureBuildMetadata")), BuildMetadata.ToCompactBinaryWithDefaults());
+	// This will get added to the build metadata in a later cl.
+	//Context.AddValue(UE::DerivedData::FValueId::FromName(UTF8TEXTVIEW("TextureBuildMetadata")), BuildMetadata.ToCompactBinaryWithDefaults());
 }
diff --git a/Engine/Source/Developer/TextureBuildUtilities/Private/TextureBuildUtilities.cpp b/Engine/Source/Developer/TextureBuildUtilities/Private/TextureBuildUtilities.cpp
--- a/Engine/Source/Developer/TextureBuildUtilities/Private/TextureBuildUtilities.cpp
+++ b/Engine/Source/Developer/TextureBuildUtilities/Private/TextureBuildUtilities.cpp
@@ -255,10 +255,6 @@ FCbObject FTextureBuildMetadata::ToCompactBinaryWithDefaults() const
 
 	FCbWriter Writer;
 	Writer.BeginObject();
-	if (bSourceMipsAlphaDetected != Defaults.bSourceMipsAlphaDetected)
-	{
-		Writer.AddBool("bSourceMipsAlphaDetected", bSourceMipsAlphaDetected);
-	}
 	if (PreEncodeMipsHash != Defaults.PreEncodeMipsHash)
 	{
 		Writer << UTF8TEXTVIEW("PreEncodeMipsHash") << PreEncodeMipsHash;
@@ -269,7 +265,6 @@ FCbObject FTextureBuildMetadata::ToCompactBinaryWithDefaults() const
 
 FTextureBuildMetadata::FTextureBuildMetadata(FCbObject InCbObject)
 {
-	bSourceMipsAlphaDetected = InCbObject["bSourceMipsAlphaDetected"].AsBool(bSourceMipsAlphaDetected);
 	PreEncodeMipsHash = InCbObject["PreEncodeMipsHash"].AsUInt64(PreEncodeMipsHash);
 }
 
diff --git a/Engine/Source/Developer/TextureBuildUtilities/Public/TextureBuildUtilities.h b/Engine/Source/Developer/TextureBuildUtilities/Public/TextureBuildUtilities.h
--- a/Engine/Source/Developer/TextureBuildUtilities/Public/TextureBuildUtilities.h
+++ b/Engine/Source/Developer/TextureBuildUtilities/Public/TextureBuildUtilities.h
@@ -49,12 +49,8 @@ namespace TextureEngineParameters
 // Carries information out of the build that we don't want to cook or save off in the runtime
 struct TEXTUREBUILDUTILITIES_API FTextureBuildMetadata
 {
-	// Whether or not the first mip level of the source has alpha detected. This is used e.g.
-	// picking between BC1 and BC3
-	bool bSourceMipsAlphaDetected = false;
-	
 	// Digests of the data at various processing stages so we can track down determinism issues
-	// that arise.
+	// that arise. Currently just the hash from before we pass to the encoders.
 	uint64 PreEncodeMipsHash = 0;
 
 	FCbObject ToCompactBinaryWithDefaults() const;
diff --git a/Engine/Source/Developer/TextureCompressor/Private/TextureCompressorModule.cpp b/Engine/Source/Developer/TextureCompressor/Private/TextureCompressorModule.cpp
--- a/Engine/Source/Developer/TextureCompressor/Private/TextureCompressorModule.cpp
+++ b/Engine/Source/Developer/TextureCompressor/Private/TextureCompressorModule.cpp
@@ -3871,14 +3871,10 @@ public:
 		TArray<FImage> IntermediateMipChain;
 
 		bool bSourceMipsAlphaDetected = false;
-		if (OutMetadata || bDoDetailedAlphaLogging)
+		if (bDoDetailedAlphaLogging)
 		{
 			bSourceMipsAlphaDetected = FImageCore::DetectAlphaChannel(SourceMips[0]);
-			
-			if (bDoDetailedAlphaLogging)
-			{
-				UE_LOG(LogTextureCompressor, Display, TEXT("[alpha] Source Mips: %d - %.*s"), bSourceMipsAlphaDetected, DebugTexturePathName.Len(), DebugTexturePathName.GetData());
-			}
+			UE_LOG(LogTextureCompressor, Display, TEXT("[alpha] Source Mips: %d - %.*s"), bSourceMipsAlphaDetected, DebugTexturePathName.Len(), DebugTexturePathName.GetData());
 		}
 
 		// allow to leave texture in sRGB in case compressor accepts other than non-F32 input source
@@ -3998,12 +3994,6 @@ public:
 				SaveImageInfos[i] = IntermediateMipChain[i];
 			}
 
-			// The metadata is about trying to determine bImageHasAlphaChannel _before_ we launch a task, which 
-			// means it must be on the actual source mips, not the post-processed mips. At some point enough textures
-			// will have this saved as part of the creation process and we can rely on it - but for now we route it
-			// back out so that when textures happen to get saved, it goes with it.
-			OutMetadata->bSourceMipsAlphaDetected = bSourceMipsAlphaDetected;
-
 			OutMetadata->PreEncodeMipsHash = ComputeMipChainHash(IntermediateMipChain);
 		}
 		
diff --git a/Engine/Source/Runtime/Engine/Private/TextureDerivedData.cpp b/Engine/Source/Runtime/Engine/Private/TextureDerivedData.cpp
--- a/Engine/Source/Runtime/Engine/Private/TextureDerivedData.cpp
+++ b/Engine/Source/Runtime/Engine/Private/TextureDerivedData.cpp
@@ -77,9 +77,6 @@ static TAutoConsoleVariable<int32> CVarTexturesCookToDerivedDataReferences(
 static const FGuid GTextureSLEDerivedDataVer(0xBD855730U, 0xA5B44BBBU, 0x89D051D0U, 0x695AC618U);
 const FGuid& GetTextureSLEDerivedDataVersion() { return GTextureSLEDerivedDataVer; }
 
-// This GUID is copied in TextureBuildFunction.cpp for the IBuild flow (TextureMetadataDerivedDataVer)
-#define TEXTURE_METADATA_DERIVEDDATA_VER TEXT("B9106D68A61B4F2A8105E16F48799976")
-
 static bool IsUsingNewDerivedData()
 {
 	struct FTextureDerivedDataSetting
@@ -557,18 +554,6 @@ void GetTextureDerivedMipKey(
 		);
 }
 
-// Get the ddc key for the texture metadata (old build flow).
-UE::DerivedData::FCacheKey GetTextureDerivedMetadataKeyFromSuffix(
-	const FString& KeySuffix
-	)
-{
-	FString Key = FDerivedDataCacheInterface::BuildCacheKey(
-		TEXT("TEXTURE"),
-		TEXTURE_DERIVEDDATA_VER,
-		*(KeySuffix + FString(TEXT("_METADATA_" TEXTURE_METADATA_DERIVEDDATA_VER))));
-	return UE::DerivedData::ConvertLegacyCacheKey(Key);
-}
-
 /**
  * Computes the derived data key for a texture with the specified compression settings.
  * @param Texture - The texture for which to compute the derived data key.
@@ -1259,8 +1244,6 @@ static void GetBuildSettingsPerFormat(
 	}
 }
 
-void UnpackTextureBuildMetadataFromPlatformData(UE::TextureBuildUtilities::FTextureBuildMetadata* BuildMetadata, const FTexturePlatformData* PlatformData);
-
 /**
  * Stores derived data in the DDC.
  * After this returns, all bulk data from streaming (non-inline) mips will be sent separately to the DDC and the BulkData for those mips removed.
@@ -1280,20 +1263,6 @@ int64 PutDerivedDataInCache(FTexturePlatformData* DerivedData, const FString& De
 	// Build the key with which to cache derived data.
 	GetTextureDerivedDataKeyFromSuffix(DerivedDataKeySuffix, DerivedDataKey);
 
-	{
-		// Store the metadata.
-		UE::TextureBuildUtilities::FTextureBuildMetadata BuildMetadata;
-		UnpackTextureBuildMetadataFromPlatformData(&BuildMetadata, DerivedData);
-		FCbObject MetadataObject = BuildMetadata.ToCompactBinaryWithDefaults();
-		UE::DerivedData::FValue Value = UE::DerivedData::FValue::Compress(MetadataObject.GetBuffer());
-
-		const UE::DerivedData::FSharedString Name(WriteToString<256>(TextureName, TEXTVIEW(" [Meta]")));
-		UE::DerivedData::FRequestOwner AsyncOwner(UE::DerivedData::EPriority::Normal);
-		const UE::DerivedData::ECachePolicy Policy = bReplaceExistingDDC ? UE::DerivedData::ECachePolicy::Store : UE::DerivedData::ECachePolicy::Default;
-		UE::DerivedData::GetCache().PutValue({ {Name, GetTextureDerivedMetadataKeyFromSuffix(DerivedDataKeySuffix), MoveTemp(Value), Policy} }, AsyncOwner);
-		AsyncOwner.KeepAlive();
-	}
-
 	FString LogString;
 
 	// Write out individual mips to the derived data cache.
diff --git a/Engine/Source/Runtime/Engine/Private/TextureDerivedDataTask.cpp b/Engine/Source/Runtime/Engine/Private/TextureDerivedDataTask.cpp
--- a/Engine/Source/Runtime/Engine/Private/TextureDerivedDataTask.cpp
+++ b/Engine/Source/Runtime/Engine/Private/TextureDerivedDataTask.cpp
@@ -59,25 +59,10 @@ static TAutoConsoleVariable<int32> CVarForceRetileTextures(
 
 
 void GetTextureDerivedDataKeyFromSuffix(const FString& KeySuffix, FString& OutKey);
-UE::DerivedData::FCacheKey GetTextureDerivedMetadataKeyFromSuffix(const FString& KeySuffix);
 static void PackTextureBuildMetadataInPlatformData(FTexturePlatformData* PlatformData, const UE::TextureBuildUtilities::FTextureBuildMetadata& BuildMetadata)
 {
-	PlatformData->bSourceMipsAlphaDetectedValid = true;
-	PlatformData->bSourceMipsAlphaDetected = BuildMetadata.bSourceMipsAlphaDetected;
 	PlatformData->PreEncodeMipsHash = BuildMetadata.PreEncodeMipsHash;
 }
-void UnpackTextureBuildMetadataFromPlatformData(UE::TextureBuildUtilities::FTextureBuildMetadata* BuildMetadata, const FTexturePlatformData* PlatformData)
-{
-	if (PlatformData->bSourceMipsAlphaDetectedValid)
-	{
-		BuildMetadata->bSourceMipsAlphaDetected = PlatformData->bSourceMipsAlphaDetected;
-	}
-	else
-	{
-		BuildMetadata->bSourceMipsAlphaDetected = false;
-	}
-	BuildMetadata->PreEncodeMipsHash = PlatformData->PreEncodeMipsHash;
-}
 
 static FTextureEngineParameters GenerateTextureEngineParameters()
 {
@@ -827,8 +812,6 @@ static void DDC1_BuildTexture(
 		uint32 NumMipsInTail;
 		uint32 ExtData;
 
-		UE::TextureBuildUtilities::FTextureBuildMetadata BuildMetadata;
-
 		// Compress the texture by calling texture compressor directly.
 		TArray<FCompressedImage2D> CompressedMips;
 		if (Compressor->BuildTexture(TextureData.Blocks[0].MipsPerLayer[0],
@@ -838,12 +821,10 @@ static void DDC1_BuildTexture(
 			CompressedMips,
 			NumMipsInTail,
 			ExtData,
-			&BuildMetadata))
+			nullptr))
 		{
 			check(CompressedMips.Num());
 
-			PackTextureBuildMetadataInPlatformData(DerivedData, BuildMetadata);
-
 			DDC1_StoreClassicTextureInDerivedData(
 				CompressedMips, DerivedData, InBuildSettingsPerLayer[0].bVolume, InBuildSettingsPerLayer[0].bTextureArray, InBuildSettingsPerLayer[0].bCubemap, 
 				NumMipsInTail, ExtData, bReplaceExistingDDC, TexturePathName, KeySuffix, BytesCached);
@@ -1488,34 +1469,23 @@ static void DDC1_FetchAndFillDerivedData(
 			FString FetchFirstKey;
 			GetTextureDerivedDataKeyFromSuffix(FetchFirstKeySuffix, FetchFirstKey);
 
-			TArray<FCacheGetValueRequest, TInlineAllocator<2>> Requests;
+			TArray<FCacheGetValueRequest, TInlineAllocator<1>> Requests;
 			Requests.Add({ SharedTexturePathName, ConvertLegacyCacheKey(FetchFirstKey), ECachePolicy::Default, 0 /* UserData */});
-			Requests.Add({ SharedTextureMetaPathName, GetTextureDerivedMetadataKeyFromSuffix(FetchFirstKeySuffix), ECachePolicy::Default, 1 /* UserData */});
 
 			FRequestOwner BlockingOwner(EPriority::Blocking);
-			FSharedBuffer MetadataBuffer;
 
-			GetCache().GetValue(Requests, BlockingOwner, [&MetadataBuffer, &RawDerivedData](FCacheGetValueResponse&& Response)
+			GetCache().GetValue(Requests, BlockingOwner, [&RawDerivedData](FCacheGetValueResponse&& Response)
 			{
 				if (Response.UserData == 0)
 				{
 					RawDerivedData = Response.Value.GetData().Decompress();
 				}
-				else
-				{
-					MetadataBuffer = Response.Value.GetData().Decompress();
-				}
 			});
 			BlockingOwner.Wait();
 
 			bGotDDCData = !RawDerivedData.IsNull();
 			if (bGotDDCData)
 			{
-				// We don't necessarily get the metadata until we force a texture rebuild.
-				if (MetadataBuffer.IsNull() == false && MetadataBuffer.GetSize() != 0)
-				{
-					PackTextureBuildMetadataInPlatformData(DerivedData, FCbObject(MetadataBuffer));
-				}
 				bUsedFetchFirst = true;
 				LocalDerivedDataKey = MoveTemp(FetchFirstKey);
 				LocalDerivedDataKeySuffix = MoveTemp(FetchFirstKeySuffix);
@@ -1529,36 +1499,21 @@ static void DDC1_FetchAndFillDerivedData(
 		LocalDerivedDataKeySuffix = MoveTemp(FetchOrBuildKeySuffix);
 		GetTextureDerivedDataKeyFromSuffix(LocalDerivedDataKeySuffix, LocalDerivedDataKey);
 
-		TArray<FCacheGetValueRequest, TInlineAllocator<2>> Requests;
+		TArray<FCacheGetValueRequest, TInlineAllocator<1>> Requests;
 		Requests.Add({ SharedTextureFastPathName, ConvertLegacyCacheKey(LocalDerivedDataKey), ECachePolicy::Default, 0 /* UserData */ });
-		Requests.Add({ SharedTextureFastMetaPathName, GetTextureDerivedMetadataKeyFromSuffix(LocalDerivedDataKeySuffix), ECachePolicy::Default, 1 /* UserData */ });
 
 		FRequestOwner BlockingOwner(EPriority::Blocking);
-		FSharedBuffer MetadataBuffer;
 
-		GetCache().GetValue(Requests, BlockingOwner, [&MetadataBuffer, &RawDerivedData](FCacheGetValueResponse&& Response)
+		GetCache().GetValue(Requests, BlockingOwner, [&RawDerivedData](FCacheGetValueResponse&& Response)
 		{
 			if (Response.UserData == 0)
 			{
 				RawDerivedData = Response.Value.GetData().Decompress();
 			}
-			else
-			{
-				MetadataBuffer = Response.Value.GetData().Decompress();
-			}
 		});
 		BlockingOwner.Wait();
 
 		bGotDDCData = !RawDerivedData.IsNull();
-		if (bGotDDCData)
-		{
-			// Only read the metadata if we actually got the main data.
-			// We don't necessarily get the metadata until we force a texture rebuild.
-			if (MetadataBuffer.IsNull() == false && MetadataBuffer.GetSize() != 0)
-			{
-				PackTextureBuildMetadataInPlatformData(DerivedData, FCbObject(MetadataBuffer));
-			}
-		}
 	}
 
 	KeySuffix = LocalDerivedDataKeySuffix;
@@ -1885,9 +1840,6 @@ bool DDC1_BuildTiledClassicTexture(
 
 	check(LinearDerivedData.GetNumMipsInTail() == 0);
 
-	UE::TextureBuildUtilities::FTextureBuildMetadata BuildMetadata;
-	UnpackTextureBuildMetadataFromPlatformData(&BuildMetadata, &LinearDerivedData);
-
 	// Have all the data - do some sanity checks as we convert to the metadata format the tiler expects.
 	TArray<FMemoryView, TInlineAllocator<MAX_TEXTURE_MIP_COUNT>> InputTextureMipViews;
 	FEncodedTextureDescription TextureDescription;
@@ -1967,8 +1919,6 @@ bool DDC1_BuildTiledClassicTexture(
 		FMemory::Memcpy(TiledMip.RawData.GetData(), MipData.GetData(), MipData.GetSize());
 	} // end for each mip
 
-	PackTextureBuildMetadataInPlatformData(DerivedData, BuildMetadata);
-
 	for (int32 MipIndex = 0; MipIndex < TextureDescription.NumMips; ++MipIndex)
 	{
 		FMemory::Free(LinearMipData[MipIndex]);
@@ -2248,10 +2198,11 @@ static bool UnpackPlatformDataFromBuild(FTexturePlatformData& OutPlatformData, U
 		}
 	}
 
-	{
-		const FValueWithId& Value = BuildOutput.GetValue(FValueId::FromName(UTF8TEXTVIEW("TextureBuildMetadata")));
-		PackTextureBuildMetadataInPlatformData(&OutPlatformData, FCbObject(Value.GetData().Decompress()));
-	}
+	// this will get pulled from the build metadata in a later cl...
+	//{
+	//	const FValueWithId& Value = BuildOutput.GetValue(FValueId::FromName(UTF8TEXTVIEW("TextureBuildMetadata")));
+	//	PackTextureBuildMetadataInPlatformData(&OutPlatformData, FCbObject(Value.GetData().Decompress()));
+	//}
 
 	// We take this as a build output, however in ideal (future) situations, this is generated prior to build launch and
 	// just routed through the build. Since we currently handle several varying situations, we just always consume it from
@@ -2834,7 +2785,6 @@ private:
 		ANSITEXTVIEW("EncodedTextureDescription"),
 		ANSITEXTVIEW("EncodedTextureExtendedData"),
 		ANSITEXTVIEW("MipTail"),
-		ANSITEXTVIEW("TextureBuildMetadata"),
 		ANSITEXTVIEW("CPUCopyImageInfo"),
 		ANSITEXTVIEW("CPUCopyRawData")
 	};
@@ -3099,8 +3049,6 @@ public:
 			DefinitionBuilder.AddConstant(UTF8TEXTVIEW("EncodedTextureDescriptionConstant"), UE::TextureBuildUtilities::EncodedTextureDescription::ToCompactBinary(*InTextureDescription));
 			DefinitionBuilder.AddConstant(UTF8TEXTVIEW("EncodedTextureExtendedDataConstant"), UE::TextureBuildUtilities::EncodedTextureExtendedData::ToCompactBinary(*InTextureExtendedData));
 		}
-		
-		DefinitionBuilder.AddInputBuild(UTF8TEXTVIEW("TextureBuildMetadata"), { InParentBuildKey, UE::DerivedData::FValueId::FromName(UTF8TEXTVIEW("TextureBuildMetadata")) });
 
 		if (InputTextureNumMips > InputTextureNumStreamingMips)
 		{
