Skeletal Mesh: Change the import storage to use FMeshDescription instead of FSkeletalMeshImportData. This includes deprecating the USkeletalMesh::LoadLODImportedData and friends, in lieu of an API mor

diff --git a/Engine/Source/Developer/SkeletalMeshUtilitiesCommon/Private/LODUtilities.cpp b/Engine/Source/Developer/SkeletalMeshUtilitiesCommon/Private/LODUtilities.cpp
--- a/Engine/Source/Developer/SkeletalMeshUtilitiesCommon/Private/LODUtilities.cpp
+++ b/Engine/Source/Developer/SkeletalMeshUtilitiesCommon/Private/LODUtilities.cpp
@@ -2,6 +2,8 @@
 
 #include "LODUtilities.h"
 
+#include "SkeletalMeshAttributes.h"
+
 #if WITH_EDITOR
 
 #include "Algo/Accumulate.h"
@@ -55,10 +57,10 @@ DEFINE_LOG_CATEGORY_STATIC(LogLODUtilities, Log, All);
 /**
 * Process and update the vertex Influences using the predefined wedges
 *
-* @param WedgeCount - The number of wedges in the corresponding mesh.
+* @param VertexCount - The number of vertices in the corresponding mesh.
 * @param Influences - BoneWeights and Ids for the corresponding vertices.
 */
-void FLODUtilities::ProcessImportMeshInfluences(const int32 WedgeCount, TArray<SkeletalMeshImportData::FRawBoneInfluence>& Influences, const FString& MeshName)
+void FLODUtilities::ProcessImportMeshInfluences(const int32 VertexCount, TArray<SkeletalMeshImportData::FRawBoneInfluence>& Influences, const FString& MeshName)
 {
 
 	// Sort influences by vertex index.
@@ -184,12 +186,12 @@ void FLODUtilities::ProcessImportMeshInfluences(const int32 WedgeCount, TArray<S
 		// warn about no influences
 		UE_LOG(LogLODUtilities, Warning, TEXT("Warning skeletal mesh (%s) has no vertex influences"), *MeshName);
 		// add one for each wedge entry
-		Influences.AddUninitialized(WedgeCount);
-		for (int32 WedgeIdx = 0; WedgeIdx < WedgeCount; WedgeIdx++)
+		Influences.AddUninitialized(VertexCount);
+		for (int32 VertexIdx = 0; VertexIdx < VertexCount; VertexIdx++)
 		{
-			Influences[WedgeIdx].VertexIndex = WedgeIdx;
-			Influences[WedgeIdx].BoneIndex = 0;
-			Influences[WedgeIdx].Weight = 1.0f;
+			Influences[VertexIdx].VertexIndex = VertexIdx;
+			Influences[VertexIdx].BoneIndex = 0;
+			Influences[VertexIdx].Weight = 1.0f;
 		}
 		for (int32 i = 0; i < Influences.Num(); i++)
 		{
@@ -371,33 +373,23 @@ void FLODUtilities::RemoveLOD(FSkeletalMeshUpdateContext& UpdateContext, int32 D
 		}
 
 		//Adjust the imported data so it point on the correct LOD index
-		if (DependentLODs.Num() > 0 && !SkeletalMesh->IsLODImportedDataEmpty(DesiredLOD))
+		if (DependentLODs.Num() > 0 && SkeletalMesh->HasMeshDescription(DesiredLOD))
 		{
-			int32 FirstDepLODIndex = DependentLODs[0];
-			FSkeletalMeshLODInfo* FirstDepLODInfo = SkeletalMesh->GetLODInfo(FirstDepLODIndex);
-			if (FirstDepLODInfo)
+			const int32 FirstDepLODIndex = DependentLODs[0];
+			if (FSkeletalMeshLODInfo* FirstDepLODInfo = SkeletalMesh->GetLODInfo(FirstDepLODIndex))
 			{
-				FSkeletalMeshImportData ToRemovedLODImportData;
-				SkeletalMesh->LoadLODImportedData(DesiredLOD, ToRemovedLODImportData);
-				//Override imported data with the original source imported data (we are depending on the LOD we want to removed)
-				SkeletalMesh->SaveLODImportedData(FirstDepLODIndex, ToRemovedLODImportData);
+				SkeletalMesh->ModifyMeshDescription(DesiredLOD);
+				FMeshDescription* SourceMeshDescription = SkeletalMesh->GetMeshDescription(DesiredLOD);
+				
+				SkeletalMesh->ModifyMeshDescription(FirstDepLODIndex);
+				SkeletalMesh->CreateMeshDescription(FirstDepLODIndex, MoveTemp(*SourceMeshDescription));
+				SkeletalMesh->CommitMeshDescription(FirstDepLODIndex);
 
 				//Manage the override original reduction source mesh data
 				if (SkelMeshModel->InlineReductionCacheDatas.IsValidIndex(FirstDepLODIndex))
 				{
-					if(SkeletalMesh->IsLODImportedDataBuildAvailable(DesiredLOD))
-					{
-						//The inline reduction cache data will be recache by the build
-						SkelMeshModel->InlineReductionCacheDatas[FirstDepLODIndex].SetCacheGeometryInfo(MAX_uint32, MAX_uint32);
-					}
-					else if(SkelMeshModel->InlineReductionCacheDatas.IsValidIndex(DesiredLOD))
-					{
-						//If there is no build copy the one from the DesiredLOD
-						uint32 CacheVertexCount = 0;
-						uint32 CacheTriangleCount = 0;
-						SkelMeshModel->InlineReductionCacheDatas[DesiredLOD].GetCacheGeometryInfo(CacheVertexCount, CacheTriangleCount);
-						SkelMeshModel->InlineReductionCacheDatas[FirstDepLODIndex].SetCacheGeometryInfo(CacheVertexCount, CacheTriangleCount);
-					}
+					//The inline reduction cache data will be recache by the build
+					SkelMeshModel->InlineReductionCacheDatas[FirstDepLODIndex].SetCacheGeometryInfo(MAX_uint32, MAX_uint32);
 				}
 
 				//Adjust Reduction settings
@@ -406,7 +398,7 @@ void FLODUtilities::RemoveLOD(FSkeletalMeshUpdateContext& UpdateContext, int32 D
 				//Do the adjustment for the other dependent LODs
 				for (int32 DependentLODsIndex = 1; DependentLODsIndex < DependentLODs.Num(); ++DependentLODsIndex)
 				{
-					int32 DepLODIndex = DependentLODs[DependentLODsIndex];
+					const int32 DepLODIndex = DependentLODs[DependentLODsIndex];
 					FSkeletalMeshLODInfo* DepLODInfo = SkeletalMesh->GetLODInfo(DepLODIndex);
 					if (!DepLODInfo)
 					{
@@ -638,11 +630,13 @@ bool FLODUtilities::SetCustomLOD(USkeletalMesh* DestinationSkeletalMesh, USkelet
 	{
 		//The imported LOD is always in LOD 0 of the SourceSkeletalMesh
 		const int32 SourceLODIndex = 0;
-		if (!SourceSkeletalMesh->IsLODImportedDataEmpty(SourceLODIndex))
+		if (SourceSkeletalMesh->HasMeshDescription(SourceLODIndex))
 		{
 			// Fix up the imported data bone indexes
 			FSkeletalMeshImportData LODImportData;
+			PRAGMA_DISABLE_DEPRECATION_WARNINGS
 			SourceSkeletalMesh->LoadLODImportedData(SourceLODIndex, LODImportData);
+			PRAGMA_ENABLE_DEPRECATION_WARNINGS
 			const int32 LODImportDataBoneNumber = LODImportData.RefBonesBinary.Num();
 			//We want to create a remap array so we can fix all influence easily
 			TArray<int32> ImportDataBoneRemap;
@@ -690,7 +684,9 @@ bool FLODUtilities::SetCustomLOD(USkeletalMesh* DestinationSkeletalMesh, USkelet
 				LODImportData.Influences.Shrink();
 			}
 			//Save the fix up remap bone index
+			PRAGMA_DISABLE_DEPRECATION_WARNINGS
 			SourceSkeletalMesh->SaveLODImportedData(SourceLODIndex, LODImportData);
+			PRAGMA_ENABLE_DEPRECATION_WARNINGS
 		}
 
 		// Fix up the ActiveBoneIndices array.
@@ -759,10 +755,12 @@ bool FLODUtilities::SetCustomLOD(USkeletalMesh* DestinationSkeletalMesh, USkelet
 	//Restore the LOD section data in case this LOD was reimport and some material match
 	if (DestImportedResource->LODModels.IsValidIndex(LodIndex) && SourceImportedResource->LODModels.IsValidIndex(0))
 	{
-		if (!DestinationSkeletalMesh->IsLODImportedDataEmpty(LodIndex))
+		if (DestinationSkeletalMesh->HasMeshDescription(LodIndex))
 		{
 			FSkeletalMeshImportData DestinationLODImportData;
+			PRAGMA_DISABLE_DEPRECATION_WARNINGS
 			DestinationSkeletalMesh->LoadLODImportedData(LodIndex, DestinationLODImportData);
+			PRAGMA_ENABLE_DEPRECATION_WARNINGS
 			for (int32 SectionIndex = 0; SectionIndex < DestinationLODImportData.Materials.Num(); ++SectionIndex)
 			{
 				ExistingOriginalPerSectionMaterialImportName.Add(FName(*DestinationLODImportData.Materials[SectionIndex].MaterialImportName));
@@ -782,11 +780,13 @@ bool FLODUtilities::SetCustomLOD(USkeletalMesh* DestinationSkeletalMesh, USkelet
 	}
 
 	const int32 SourceLODIndex = 0;
-	if (!SourceSkeletalMesh->IsLODImportedDataEmpty(SourceLODIndex))
+	if (SourceSkeletalMesh->HasMeshDescription(SourceLODIndex))
 	{
 		// Fix up the imported data bone indexes
 		FSkeletalMeshImportData SourceLODImportData;
+		PRAGMA_DISABLE_DEPRECATION_WARNINGS
 		SourceSkeletalMesh->LoadLODImportedData(SourceLODIndex, SourceLODImportData);
+		PRAGMA_ENABLE_DEPRECATION_WARNINGS
 		FLODUtilities::FSkeletalMeshMatchImportedMaterialsParameters Parameters;
 		Parameters.bIsReImport = bIsReImport;
 		Parameters.LodIndex = LodIndex;
@@ -812,8 +812,15 @@ bool FLODUtilities::SetCustomLOD(USkeletalMesh* DestinationSkeletalMesh, USkelet
 
 	// Assign new FSkeletalMeshLODModel to desired slot in selected skeletal mesh.
 	FSkeletalMeshLODModel::CopyStructure(&(DestImportedResource->LODModels[LodIndex]), &NewLODModel);
+	
 	//Copy the import data into the base skeletalmesh for the imported LOD
-	USkeletalMesh::CopyImportedData(0, SourceSkeletalMesh, LodIndex, DestinationSkeletalMesh);
+	FMeshDescription SourceMeshDescription;
+	if (SourceSkeletalMesh->CloneMeshDescription(0, SourceMeshDescription))
+	{
+		DestinationSkeletalMesh->ModifyMeshDescription(LodIndex);
+		DestinationSkeletalMesh->CreateMeshDescription(LodIndex, MoveTemp(SourceMeshDescription));
+		DestinationSkeletalMesh->CommitMeshDescription(LodIndex);
+	}
 
 
 	// If this LOD had been generated previously by automatic mesh reduction, clear that flag.
@@ -2104,7 +2111,7 @@ bool FLODUtilities::UpdateAlternateSkinWeights(USkeletalMesh* SkeletalMeshDest,
 	check(SkeletalMeshDest->GetImportedModel());
 	check(SkeletalMeshDest->GetImportedModel()->LODModels.IsValidIndex(LODIndexDest));
 	FSkeletalMeshLODModel& LODModelDest = SkeletalMeshDest->GetImportedModel()->LODModels[LODIndexDest];
-	if (SkeletalMeshDest->IsLODImportedDataEmpty(LODIndexDest))
+	if (!SkeletalMeshDest->HasMeshDescription(LODIndexDest))
 	{
 		UE_LOG(LogLODUtilities, Error, TEXT("Failed to import Skin Weight Profile as the target skeletal mesh (%s) requires reimporting first."), *SkeletalMeshDest->GetName());
 		//Very old asset will not have this data, we cannot add alternate until the asset is reimported
@@ -2112,7 +2119,9 @@ bool FLODUtilities::UpdateAlternateSkinWeights(USkeletalMesh* SkeletalMeshDest,
 	}
 	
 	FSkeletalMeshImportData ImportDataDest;
+	PRAGMA_DISABLE_DEPRECATION_WARNINGS
 	SkeletalMeshDest->LoadLODImportedData(LODIndexDest, ImportDataDest);
+	PRAGMA_ENABLE_DEPRECATION_WARNINGS
 	
 	return UpdateAlternateSkinWeights(
 		LODModelDest,
@@ -2756,14 +2765,16 @@ bool FLODUtilities::UpdateAlternateSkinWeights(
 	check(SkeletalMeshDest->GetImportedModel()->LODModels.IsValidIndex(LODIndexDest));
 	FSkeletalMeshLODModel& LODModelDest = SkeletalMeshDest->GetImportedModel()->LODModels[LODIndexDest];
 
-	if (SkeletalMeshDest->IsLODImportedDataEmpty(LODIndexDest))
+	if (!SkeletalMeshDest->HasMeshDescription(LODIndexDest))
 	{
 		UE_LOG(LogLODUtilities, Error, TEXT("Failed to import Skin Weight Profile as the target skeletal mesh (%s) requires reimporting first."), SkeletalMeshDest ? *SkeletalMeshDest->GetName() : TEXT("NULL"));
 		//Very old asset will not have this data, we cannot add alternate until the asset is reimported
 		return false;
 	}
 	FSkeletalMeshImportData ImportDataDest;
+	PRAGMA_DISABLE_DEPRECATION_WARNINGS
 	SkeletalMeshDest->LoadLODImportedData(LODIndexDest, ImportDataDest);
+	PRAGMA_ENABLE_DEPRECATION_WARNINGS
 	int32 PointNumberDest = ImportDataDest.Points.Num();
 	int32 VertexNumberDest = ImportDataDest.Points.Num();
 
@@ -2771,9 +2782,11 @@ bool FLODUtilities::UpdateAlternateSkinWeights(
 	check(SkeletalMeshSrc);
 
 	//The source model is a fresh import and the data need to be there
-	check(!SkeletalMeshSrc->IsLODImportedDataEmpty(LODIndexSrc));
+	check(SkeletalMeshSrc->HasMeshDescription(LODIndexSrc));
 	FSkeletalMeshImportData ImportDataSrc;
+	PRAGMA_DISABLE_DEPRECATION_WARNINGS
 	SkeletalMeshSrc->LoadLODImportedData(LODIndexSrc, ImportDataSrc);
+	PRAGMA_ENABLE_DEPRECATION_WARNINGS
 	
 	//Remove all unnecessary array data from the structure (this will save a lot of memory)
 	ImportDataSrc.KeepAlternateSkinningBuildDataOnly();
@@ -2802,13 +2815,10 @@ bool FLODUtilities::UpdateAlternateSkinWeights(
 	ImportDataDest.AlternateInfluences.Add(ImportDataSrc);
 
 	//Resave the bulk data with the new or refreshed data
+	PRAGMA_DISABLE_DEPRECATION_WARNINGS
 	SkeletalMeshDest->SaveLODImportedData(LODIndexDest, ImportDataDest);
+	PRAGMA_ENABLE_DEPRECATION_WARNINGS
 
-	if(!SkeletalMeshDest->IsLODImportedDataBuildAvailable(LODIndexDest))
-	{
-		//Build the alternate buffer with all the data into the bulk, in case the build data is not existing (old asset)
-		return UpdateAlternateSkinWeights(SkeletalMeshDest, ProfileNameDest, LODIndexDest, Options);
-	}
 	return true;
 }
 
@@ -2968,9 +2978,11 @@ bool FLODUtilities::UpdateLODInfoVertexAttributes(
 	bool bInCopyAttributeValues
 	)
 {
-	FSkeletalMeshImportData SkeletalMeshImportData;
-	InSkeletalMesh->LoadLODImportedData(InSourceLODIndex, SkeletalMeshImportData);
-	
+	if (!ensure(InSkeletalMesh->HasMeshDescription(InSourceLODIndex)))
+	{
+		return false;
+	}
+
 	FSkeletalMeshLODModel& TargetLODModel = InSkeletalMesh->GetImportedModel()->LODModels[InTargetLODIndex];
 	
 	TArray<FSkeletalMeshVertexAttributeInfo>& SkelMeshAttributeInfos = InSkeletalMesh->GetLODInfo(InTargetLODIndex)->VertexAttributes; 
@@ -2981,8 +2993,20 @@ bool FLODUtilities::UpdateLODInfoVertexAttributes(
 	{
 		ExistingAttributeInfos.Add(AttributeInfo.Name, MoveTemp(AttributeInfo));
 	}
+	
+	const FMeshDescription* MeshDescription = InSkeletalMesh->GetMeshDescription(InSourceLODIndex);
 
-	SkelMeshAttributeInfos.Reset(SkeletalMeshImportData.VertexAttributes.Num());
+	// NOTE: There's a current limitation that we only support single-channel attributes for rendering.
+	TMap<FName, TVertexAttributesConstRef<float>> SourceAttributes; 
+	MeshDescription->VertexAttributes().ForEachByType<float>([&SourceAttributes](const FName InAttributeName, TVertexAttributesConstRef<float> InAttributeRef)
+	{
+		if (!FSkeletalMeshAttributes::IsReservedAttributeName(InAttributeName))
+		{
+			SourceAttributes.Add(InAttributeName, InAttributeRef);
+		}
+	});
+ 
+	SkelMeshAttributeInfos.Reset(SourceAttributes.Num());
 
 	// If we're not copying the values, leave the existing data in place.
 	if (bInCopyAttributeValues)
@@ -2992,10 +3016,10 @@ bool FLODUtilities::UpdateLODInfoVertexAttributes(
 	
 	TArray<UE::Tasks::FTask> ConversionTasks;
 	
-	for (int32 AttributeIndex = 0; AttributeIndex < SkeletalMeshImportData.VertexAttributes.Num(); AttributeIndex++)
+	for (const TPair<FName, TVertexAttributesConstRef<float>> SourceAttributeInfo: SourceAttributes)
 	{
-		const SkeletalMeshImportData::FVertexAttribute& ImportAttribute = SkeletalMeshImportData.VertexAttributes[AttributeIndex];
-		const FName AttributeName(SkeletalMeshImportData.VertexAttributeNames[AttributeIndex]);
+		const TVertexAttributesConstRef<float>& SourceAttribute = SourceAttributeInfo.Value;
+		const FName AttributeName(SourceAttributeInfo.Key);
 
 		// Did this definition already exist? Try to retain as much of the existing information as possible.
 		FSkeletalMeshVertexAttributeInfo Info;
@@ -3016,18 +3040,18 @@ bool FLODUtilities::UpdateLODInfoVertexAttributes(
 			FSkeletalMeshModelVertexAttribute& ModelAttribute = TargetLODModel.VertexAttributes.FindOrAdd(AttributeName);
 
 			ModelAttribute.DataType = Info.DataType;
-			ModelAttribute.ComponentCount = ImportAttribute.ComponentCount;
+			ModelAttribute.ComponentCount = 1;
 
 			if (InTargetLODIndex == InSourceLODIndex)
 			{
 				ConversionTasks.Add(
-					UE::Tasks::Launch(UE_SOURCE_LOCATION, [&TargetLODModel, &ModelAttribute, &ImportAttribute]()
+					UE::Tasks::Launch(UE_SOURCE_LOCATION, [&TargetLODModel, &ModelAttribute, &SourceAttribute]()
 					{
 						ModelAttribute.Values.SetNumUninitialized(TargetLODModel.NumVertices);
 						for(uint32 VertexIndex = 0; VertexIndex < TargetLODModel.NumVertices; VertexIndex++)
 						{
 							const int32 ImportVertexIndex = TargetLODModel.MeshToImportVertexMap[VertexIndex];
-							ModelAttribute.Values[VertexIndex] = ImportAttribute.AttributeValues[ImportVertexIndex];
+							ModelAttribute.Values[VertexIndex] = SourceAttribute.Get(FVertexID(ImportVertexIndex));
 						}
 					})
 				);
@@ -3113,7 +3137,6 @@ void FLODUtilities::RegenerateDependentLODs(USkeletalMesh* SkeletalMesh, int32 L
 				});
 			}
 
-			SkeletalMesh->ReserveLODImportData(MaxDependentLODIndex);
 			//Reduce all dependent LODs
 			FThreadSafeBool bNeedsPackageDirtied(false);
 			
@@ -4110,7 +4133,7 @@ void FLODUtilities::ReorderMaterialSlotToBaseLod(USkeletalMesh* SkeletalMesh)
 
 	for (int32 LodIndex = 0; LodIndex < SkeletalMesh->GetLODNum(); ++LodIndex)
 	{
-		if (SkeletalMesh->IsLODImportedDataEmpty(LodIndex))
+		if (!SkeletalMesh->HasMeshDescription(LodIndex))
 		{
 			if (LodIndex == 0)
 			{
@@ -4121,7 +4144,9 @@ void FLODUtilities::ReorderMaterialSlotToBaseLod(USkeletalMesh* SkeletalMesh)
 		}
 
 		FSkeletalMeshImportData LodImportData;
+		PRAGMA_DISABLE_DEPRECATION_WARNINGS
 		SkeletalMesh->LoadLODImportedData(LodIndex, LodImportData);
+		PRAGMA_ENABLE_DEPRECATION_WARNINGS
 		const TArray<SkeletalMeshImportData::FMaterial>& ImportedMaterials = LodImportData.Materials;
 
 		//Find remap index for this LOD matching materials
@@ -4212,7 +4237,7 @@ void FLODUtilities::ReorderMaterialSlotToBaseLod(USkeletalMesh* SkeletalMesh)
 
 void FLODUtilities::RemoveUnusedMaterialSlot(USkeletalMesh* SkeletalMesh)
 {
-	if (!SkeletalMesh || !SkeletalMesh->IsLODImportedDataBuildAvailable(0))
+	if (!SkeletalMesh || !SkeletalMesh->HasMeshDescription(0))
 	{
 		return;
 	}
@@ -4415,13 +4440,14 @@ bool FLODUtilities::StripLODGeometry(USkeletalMesh* SkeletalMesh, const int32 LO
 		FScopedSkeletalMeshPostEditChange ScopePostEditChange(SkeletalMesh);
 		//This is like a re-import, we must force to use a new DDC
 		SkeletalMesh->InvalidateDeriveDataCacheGUID();
-		const bool bContainImportedData = SkeletalMesh->IsLODImportedDataEmpty(LODIndex);
-		const bool bBuildAvailable = SkeletalMesh->IsLODImportedDataBuildAvailable(LODIndex);
+		const bool bBuildAvailable = SkeletalMesh->HasMeshDescription(LODIndex);
 		FSkeletalMeshImportData ImportedData;
 		//Get the imported data if available
 		if (bBuildAvailable)
 		{
+			PRAGMA_DISABLE_DEPRECATION_WARNINGS
 			SkeletalMesh->LoadLODImportedData(LODIndex, ImportedData);
+			PRAGMA_ENABLE_DEPRECATION_WARNINGS
 		}
 		SkeletalMesh->Modify();
 		
@@ -4637,7 +4663,9 @@ bool FLODUtilities::StripLODGeometry(USkeletalMesh* SkeletalMesh, const int32 LO
 					Influence.VertexIndex = RemapVertexIndex[VertexIndex];
 				}
 			}
+			PRAGMA_DISABLE_DEPRECATION_WARNINGS
 			SkeletalMesh->SaveLODImportedData(LODIndex, StrippedImportedData);
+			PRAGMA_ENABLE_DEPRECATION_WARNINGS
 		}
 	}
 	return true;
diff --git a/Engine/Source/Developer/SkeletalMeshUtilitiesCommon/Private/StaticToSkeletalMeshConverter.cpp b/Engine/Source/Developer/SkeletalMeshUtilitiesCommon/Private/StaticToSkeletalMeshConverter.cpp
--- a/Engine/Source/Developer/SkeletalMeshUtilitiesCommon/Private/StaticToSkeletalMeshConverter.cpp
+++ b/Engine/Source/Developer/SkeletalMeshUtilitiesCommon/Private/StaticToSkeletalMeshConverter.cpp
@@ -184,7 +184,7 @@ static void CopyReductionSettings(
 
 
 static bool AddLODFromMeshDescription(
-	const FMeshDescription& InMeshDescription,
+	FMeshDescription&& InMeshDescription,
 	USkeletalMesh* InSkeletalMesh,
 	IMeshUtilities& InMeshUtilities
 	)
@@ -196,11 +196,14 @@ static bool AddLODFromMeshDescription(
 	{
 		return false;
 	}
-	
-	FSkeletalMeshLODModel& SkeletalMeshModel = ImportedModels->LODModels.Last(); 
-	
+
 	FSkeletalMeshImportData SkeletalMeshImportGeometry = FSkeletalMeshImportData::CreateFromMeshDescription(InMeshDescription);
 
+	InSkeletalMesh->CreateMeshDescription(LODIndex, MoveTemp(InMeshDescription));
+	InSkeletalMesh->CommitMeshDescription(LODIndex);
+	
+	FSkeletalMeshLODModel& SkeletalMeshModel = ImportedModels->LODModels.Last();
+
 	// We need at least one set of texture coordinates. Always.
 	SkeletalMeshModel.NumTexCoords = FMath::Max<uint32>(1, SkeletalMeshImportGeometry.NumTexCoords);
 	
@@ -226,9 +229,6 @@ static bool AddLODFromMeshDescription(
 		return false;
 	}
 	
-	InSkeletalMesh->SaveLODImportedData(LODIndex, SkeletalMeshImportGeometry);
-	InSkeletalMesh->SetLODImportedDataVersions(LODIndex, ESkeletalMeshGeoImportVersions::LatestVersion, ESkeletalMeshSkinningImportVersions::LatestVersion);
-	
 	return true;
 }
 
@@ -271,7 +271,7 @@ static bool AddLODFromStaticMeshSourceModel(
 			SkinWeights.Set(VertexID, RootBinding);
 		}
 
-		if (!AddLODFromMeshDescription(SkeletalMeshGeometry, InSkeletalMesh, InMeshUtilities))
+		if (!AddLODFromMeshDescription(MoveTemp(SkeletalMeshGeometry), InSkeletalMesh, InMeshUtilities))
 		{
 			return false;
 		}
@@ -499,7 +499,8 @@ bool FStaticToSkeletalMeshConverter::InitializeSkeletalMeshFromMeshDescriptions(
 			SkeletalLODInfo.BuildSettings.bRecomputeNormals = bInRecomputeNormals;
 			SkeletalLODInfo.BuildSettings.bRecomputeTangents = bInRecomputeTangents;
 
-			if (!AddLODFromMeshDescription(*MeshDescription, InSkeletalMesh, MeshUtilities))
+			FMeshDescription ClonedDescription(*MeshDescription);
+			if (!AddLODFromMeshDescription(MoveTemp(ClonedDescription), InSkeletalMesh, MeshUtilities))
 			{
 				// If we didn't get a model for LOD index 0, we don't have a mesh. Bail out.
 				if (bFirstSourceModel)
diff --git a/Engine/Source/Developer/SkeletalMeshUtilitiesCommon/Public/LODUtilities.h b/Engine/Source/Developer/SkeletalMeshUtilitiesCommon/Public/LODUtilities.h
--- a/Engine/Source/Developer/SkeletalMeshUtilitiesCommon/Public/LODUtilities.h
+++ b/Engine/Source/Developer/SkeletalMeshUtilitiesCommon/Public/LODUtilities.h
@@ -52,11 +52,11 @@ public:
 	/**
 	* Process and update the vertex Influences using the predefined wedges
 	* 
-	* @param WedgeCount - The number of wedges in the corresponding mesh.
+	* @param VertexCount - The number of wedges in the corresponding mesh.
 	* @param Influences - BoneWeights and Ids for the corresponding vertices. 
 	* @param MeshName	- Name of mesh, used for warning messages
 	*/
-	static void ProcessImportMeshInfluences(const int32 WedgeCount, TArray<SkeletalMeshImportData::FRawBoneInfluence>& Influences, const FString& MeshName);
+	static void ProcessImportMeshInfluences(const int32 VertexCount, TArray<SkeletalMeshImportData::FRawBoneInfluence>& Influences, const FString& MeshName);
 
 	/** Regenerate LODs of the mesh
 	*
diff --git a/Engine/Source/Editor/DetailCustomizations/Private/SkeletalMeshReductionSettingsDetails.cpp b/Engine/Source/Editor/DetailCustomizations/Private/SkeletalMeshReductionSettingsDetails.cpp
--- a/Engine/Source/Editor/DetailCustomizations/Private/SkeletalMeshReductionSettingsDetails.cpp
+++ b/Engine/Source/Editor/DetailCustomizations/Private/SkeletalMeshReductionSettingsDetails.cpp
@@ -108,7 +108,7 @@ void FSkeletalMeshReductionSettingsDetails::CustomizeChildren(TSharedRef<IProper
 		// Only able to do this for LOD1 and above, so only show the property if this is the case
 		if (LODIndex > 0)
 		{
-			bool AllowInline = SkeletalMesh && !SkeletalMesh->IsLODImportedDataEmpty(LODIndex);
+			bool AllowInline = SkeletalMesh && SkeletalMesh->HasMeshDescription(LODIndex);
 			// Add and retrieve the default widgets
 			IDetailPropertyRow& Row = StructBuilder.AddProperty(BaseLODPropertyHandle->AsShared());
 
diff --git a/Engine/Source/Editor/MeshPaint/Private/MeshPaintSkeletalMeshAdapter.cpp b/Engine/Source/Editor/MeshPaint/Private/MeshPaintSkeletalMeshAdapter.cpp
--- a/Engine/Source/Editor/MeshPaint/Private/MeshPaintSkeletalMeshAdapter.cpp
+++ b/Engine/Source/Editor/MeshPaint/Private/MeshPaintSkeletalMeshAdapter.cpp
@@ -7,6 +7,7 @@
 #include "MeshPaintHelpers.h"
 #include "MeshPaintTypes.h"
 #include "ComponentReregisterContext.h"
+#include "StaticMeshAttributes.h"
 #include "Rendering/SkeletalMeshRenderData.h"
 #include "Rendering/SkeletalMeshModel.h"
 #include "Factories/FbxSkeletalMeshImportData.h"
@@ -25,23 +26,23 @@ void PropagateVertexPaintToAsset(USkeletalMesh* SkeletalMesh, int32 LODIndex)
 {
 	struct FMatchFaceData
 	{
-		int32 SoftVerticeIndexes[3];
+		int32 SoftVertexIndexes[3];
 	};
 
-	if (!SkeletalMesh || SkeletalMesh->IsLODImportedDataEmpty(LODIndex) || !SkeletalMesh->IsLODImportedDataBuildAvailable(LODIndex))
+	if (!SkeletalMesh || SkeletalMesh->HasMeshDescription(LODIndex))
 	{
 		//We do not propagate vertex color for old asset
 		return;
 	}
 
-	auto GetMatchKey = [](const FVector& PositionA, const FVector& PositionB, const FVector& PositionC)->FSHAHash
+	auto GetMatchKey = [](const FVector3f& PositionA, const FVector3f& PositionB, const FVector3f& PositionC)->FSHAHash
 	{
 		FSHA1 SHA;
 		FSHAHash SHAHash;
 
-		SHA.Update((const uint8*)&PositionA, sizeof(FVector));
-		SHA.Update((const uint8*)&PositionB, sizeof(FVector));
-		SHA.Update((const uint8*)&PositionC, sizeof(FVector));
+		SHA.Update(reinterpret_cast<const uint8*>(&PositionA), sizeof(FVector3f));
+		SHA.Update(reinterpret_cast<const uint8*>(&PositionB), sizeof(FVector3f));
+		SHA.Update(reinterpret_cast<const uint8*>(&PositionC), sizeof(FVector3f));
 		SHA.Final();
 		SHA.GetHash(&SHAHash.Hash[0]);
 
@@ -53,56 +54,61 @@ void PropagateVertexPaintToAsset(USkeletalMesh* SkeletalMesh, int32 LODIndex)
 
 	TArray<FSoftSkinVertex> SrcVertices;
 	LODModel.GetVertices(SrcVertices);
-	
-	FSkeletalMeshImportData ImportData;
-	SkeletalMesh->LoadLODImportedData(LODIndex, ImportData);
+
+	FMeshDescription* MeshDescription = SkeletalMesh->GetMeshDescription(LODIndex);
 
 	TMap<FSHAHash, FMatchFaceData> MatchTriangles;
-	MatchTriangles.Reserve(ImportData.Wedges.Num());
+	MatchTriangles.Reserve(MeshDescription->Triangles().Num());
 
 	for (int32 IndexBufferIndex = 0, SrcIndexBufferNum = SrcIndexBuffer.Num(); IndexBufferIndex < SrcIndexBufferNum; IndexBufferIndex += 3)
 	{
-		FVector PositionA = (FVector)SrcVertices[SrcIndexBuffer[IndexBufferIndex]].Position;
-		FVector PositionB = (FVector)SrcVertices[SrcIndexBuffer[IndexBufferIndex + 1]].Position;
-		FVector PositionC = (FVector)SrcVertices[SrcIndexBuffer[IndexBufferIndex + 2]].Position;
+		const FVector3f& PositionA = SrcVertices[SrcIndexBuffer[IndexBufferIndex]].Position;
+		const FVector3f& PositionB = SrcVertices[SrcIndexBuffer[IndexBufferIndex + 1]].Position;
+		const FVector3f& PositionC = SrcVertices[SrcIndexBuffer[IndexBufferIndex + 2]].Position;
 
 		FSHAHash Key = GetMatchKey(PositionA, PositionB, PositionC);
 		FMatchFaceData MatchFaceData;
-		MatchFaceData.SoftVerticeIndexes[0] = SrcIndexBuffer[IndexBufferIndex];
-		MatchFaceData.SoftVerticeIndexes[1] = SrcIndexBuffer[IndexBufferIndex + 1];
-		MatchFaceData.SoftVerticeIndexes[2] = SrcIndexBuffer[IndexBufferIndex + 2];
+		MatchFaceData.SoftVertexIndexes[0] = SrcIndexBuffer[IndexBufferIndex];
+		MatchFaceData.SoftVertexIndexes[1] = SrcIndexBuffer[IndexBufferIndex + 1];
+		MatchFaceData.SoftVertexIndexes[2] = SrcIndexBuffer[IndexBufferIndex + 2];
 		MatchTriangles.Add(Key, MatchFaceData);
 	}
 
+	FStaticMeshAttributes MeshAttributes(*MeshDescription);
+	TVertexInstanceAttributesRef<FVector4f> ColorAttribute = MeshAttributes.GetVertexInstanceColors();
+
 	bool bAllVertexFound = true;
-	for (int32 FaceIndex = 0, FaceNum = ImportData.Faces.Num(); FaceIndex < FaceNum; ++FaceIndex)
+	for (const FTriangleID TriangleID: MeshDescription->Triangles().GetElementIDs())
 	{
-		const SkeletalMeshImportData::FTriangle& Triangle = ImportData.Faces[FaceIndex];
-		SkeletalMeshImportData::FVertex& WedgeA = ImportData.Wedges[Triangle.WedgeIndex[0]];
-		SkeletalMeshImportData::FVertex& WedgeB = ImportData.Wedges[Triangle.WedgeIndex[1]];
-		SkeletalMeshImportData::FVertex& WedgeC = ImportData.Wedges[Triangle.WedgeIndex[2]];
-		FVector PositionA = (FVector)ImportData.Points[WedgeA.VertexIndex];
-		FVector PositionB = (FVector)ImportData.Points[WedgeB.VertexIndex];
-		FVector PositionC = (FVector)ImportData.Points[WedgeC.VertexIndex];
+		TArrayView<const FVertexID> TriangleVertexIDs = MeshDescription->GetTriangleVertices(TriangleID);
+
+		FVector3f PositionA = MeshDescription->GetVertexPosition(TriangleVertexIDs[0]);
+		FVector3f PositionB = MeshDescription->GetVertexPosition(TriangleVertexIDs[1]);
+		FVector3f PositionC = MeshDescription->GetVertexPosition(TriangleVertexIDs[2]);
 
 		const FSHAHash Key = GetMatchKey(PositionA, PositionB, PositionC);
-		if (FMatchFaceData* MatchFaceData = MatchTriangles.Find(Key))
+		if (const FMatchFaceData* MatchFaceData = MatchTriangles.Find(Key))
 		{
-			WedgeA.Color = SrcVertices[MatchFaceData->SoftVerticeIndexes[0]].Color;
-			WedgeB.Color = SrcVertices[MatchFaceData->SoftVerticeIndexes[1]].Color;
-			WedgeC.Color = SrcVertices[MatchFaceData->SoftVerticeIndexes[2]].Color;
+			TArrayView<const FVertexInstanceID> TriangleVertexIndexIDs = MeshDescription->GetTriangleVertexInstances(TriangleID);
+
+			for (int32 Index = 0; Index < 3; Index++)
+			{
+				FLinearColor Color(SrcVertices[MatchFaceData->SoftVertexIndexes[Index]].Color.ReinterpretAsLinear());
+				ColorAttribute.Set(TriangleVertexIndexIDs[Index], Color);
+			}
 		}
-		else if(bAllVertexFound)
+		else if (bAllVertexFound)
 		{
 			bAllVertexFound = false;
 			FString SkeletalMeshName(SkeletalMesh->GetName());
 			UE_LOG(LogMeshPaintSkeletalMeshAdapter, Warning, TEXT("Some vertex color data could not be applied to the %s SkeletalMesh asset."), *SkeletalMeshName);
 		}
 	}
-		
-	SkeletalMesh->SaveLODImportedData(LODIndex, ImportData);
+	
+	SkeletalMesh->CommitMeshDescription(LODIndex);
 }
 
+
 bool FMeshPaintGeometryAdapterForSkeletalMeshes::Construct(UMeshComponent* InComponent, int32 InMeshLODIndex)
 {
 	SkeletalMeshComponent = Cast<USkeletalMeshComponent>(InComponent);
diff --git a/Engine/Source/Editor/Persona/Persona.Build.cs b/Engine/Source/Editor/Persona/Persona.Build.cs
--- a/Engine/Source/Editor/Persona/Persona.Build.cs
+++ b/Engine/Source/Editor/Persona/Persona.Build.cs
@@ -79,6 +79,7 @@ public class Persona : ModuleRules
 				"ToolWidgets",
 				"InteractiveToolsFramework",
 				"ContentBrowser",
+	            "SkeletalMeshDescription",
             }
         );
 
diff --git a/Engine/Source/Editor/Persona/Private/PersonaMeshDetails.cpp b/Engine/Source/Editor/Persona/Private/PersonaMeshDetails.cpp
--- a/Engine/Source/Editor/Persona/Private/PersonaMeshDetails.cpp
+++ b/Engine/Source/Editor/Persona/Private/PersonaMeshDetails.cpp
@@ -2970,7 +2970,7 @@ void FPersonaMeshDetails::AddLODLevelCategories(IDetailLayoutBuilder& DetailLayo
 				//Display the LODInfo settings
 				CustomizeLODInfoSetingsDetails(DetailLayout, LODInfoUILayout, LODInfoProperty, LODCategory);
 				
-				bool bIsbuildAvailable = SkelMesh->IsLODImportedDataBuildAvailable(LODIndex);
+				bool bIsbuildAvailable = SkelMesh->HasMeshDescription(LODIndex);
 
 				
 				//Avoid offering re-generate if the LOD is reduced on itself and does not have the original data. The user in this case has to re-import the asset to generate the data 
@@ -3965,7 +3965,7 @@ void FPersonaMeshDetails::RestoreNonReducedLOD(int32 LODIndex)
 
 	FSkeletalMeshLODInfo* CurrentLODInfo = SkelMesh->GetLODInfo(LODIndex);
 	const bool bIsReductionActive = SkelMesh->IsReductionActive(LODIndex);
-	const bool bIsLODModelbuildDataAvailable = SkelMesh->IsLODImportedDataBuildAvailable(LODIndex);
+	const bool bIsLODModelbuildDataAvailable = SkelMesh->HasMeshDescription(LODIndex);
 
 	if (CurrentLODInfo->bHasBeenSimplified
 		&& !bIsReductionActive
@@ -4003,7 +4003,7 @@ FReply FPersonaMeshDetails::ApplyLODChanges(int32 LODIndex)
 		{
 			SourceLODIndex = LODInfo->ReductionSettings.BaseLOD;
 		}
-		bool bSrcBuildDataAvailable = SkelMesh->IsLODImportedDataBuildAvailable(SourceLODIndex);
+		bool bSrcBuildDataAvailable = SkelMesh->HasMeshDescription(SourceLODIndex);
 		if (!bSrcBuildDataAvailable)
 		{
 			SkelMesh->InvalidateDeriveDataCacheGUID();
@@ -4062,7 +4062,7 @@ void FPersonaMeshDetails::RegenerateOneLOD(int32 LODIndex)
 		FScopedSkeletalMeshPostEditChange ScopedPostEditChange(SkelMesh);
 		FSkeletalMeshLODInfo& CurrentLODInfo = *(SkelMesh->GetLODInfo(LODIndex));
 		
-		bool bIsLODModelbuildDataAvailable = SkelMesh->GetImportedModel()->LODModels.IsValidIndex(LODIndex) && SkelMesh->IsLODImportedDataBuildAvailable(LODIndex);
+		bool bIsLODModelbuildDataAvailable = SkelMesh->GetImportedModel()->LODModels.IsValidIndex(LODIndex) && SkelMesh->HasMeshDescription(LODIndex);
 		if (LODIndex == CurrentLODInfo.ReductionSettings.BaseLOD
 			&& CurrentLODInfo.bHasBeenSimplified
 			&& !SkelMesh->IsReductionActive(LODIndex)
@@ -4153,7 +4153,7 @@ FReply FPersonaMeshDetails::RegenerateLOD(int32 LODIndex)
 			}
 		}
 
-		if (bIsReductionActive && !SkelMesh->IsLODImportedDataBuildAvailable(LODIndex))
+		if (bIsReductionActive && !SkelMesh->HasMeshDescription(LODIndex))
 		{
 			if (const FSkeletalMeshLODInfo* LodInfoPtr = SkelMesh->GetLODInfo(LODIndex))
 			{
@@ -4624,7 +4624,7 @@ FReply FPersonaMeshDetails::OnReimportLodClicked(EReimportButtonType InReimportT
 			//Avoid changing the settings if the skeletal mesh is using a LODSettings asset valid for this LOD
 			bool bUseLODSettingAsset = SkelMesh->GetLODSettings() != nullptr && SkelMesh->GetLODSettings()->GetNumberOfSettings() > InLODIndex;
 			//Make the reduction settings change according to the context
-			if (!bUseLODSettingAsset && SkelMesh->IsReductionActive(InLODIndex) && LODInfo->bHasBeenSimplified && SkelMesh->IsLODImportedDataEmpty(InLODIndex))
+			if (!bUseLODSettingAsset && SkelMesh->IsReductionActive(InLODIndex) && LODInfo->bHasBeenSimplified && !SkelMesh->HasMeshDescription(InLODIndex))
 			{
 				FSkeletalMeshOptimizationSettings& ReductionSettings = LODInfo->ReductionSettings;
 				//Backup the reduction settings
diff --git a/Engine/Source/Editor/Persona/Private/SkinWeightProfileHelpers.cpp b/Engine/Source/Editor/Persona/Private/SkinWeightProfileHelpers.cpp
--- a/Engine/Source/Editor/Persona/Private/SkinWeightProfileHelpers.cpp
+++ b/Engine/Source/Editor/Persona/Private/SkinWeightProfileHelpers.cpp
@@ -91,7 +91,7 @@ void FSkinWeightProfileHelpers::ImportSkinWeightProfile(USkeletalMesh* InSkeleta
 				FSlateNotificationManager::Get().AddNotification(NotificationInfo);
 			}
 
-			if(!InSkeletalMesh->IsLODImportedDataBuildAvailable(ImportSettings->LODIndex))
+			if(!InSkeletalMesh->HasMeshDescription(ImportSettings->LODIndex))
 			{
 				FLODUtilities::RegenerateDependentLODs(InSkeletalMesh, ImportSettings->LODIndex, GetTargetPlatformManagerRef().GetRunningTargetPlatform());
 			}
@@ -111,7 +111,7 @@ void FSkinWeightProfileHelpers::ImportSkinWeightProfileLOD(USkeletalMesh* InSkel
 		const bool bResult = FSkinWeightsUtilities::ImportAlternateSkinWeight(InSkeletalMesh, PickedFileName, LODIndex, ProfileName, false);
 		if (bResult)
 		{
-			if(!InSkeletalMesh->IsLODImportedDataBuildAvailable(LODIndex))
+			if(!InSkeletalMesh->HasMeshDescription(LODIndex))
 			{
 				FLODUtilities::RegenerateDependentLODs(InSkeletalMesh, LODIndex, GetTargetPlatformManagerRef().GetRunningTargetPlatform());
 			}
@@ -162,7 +162,7 @@ void FSkinWeightProfileHelpers::ReimportSkinWeightProfileLOD(USkeletalMesh* InSk
 
 			if (bResult)
 			{
-				if (!InSkeletalMesh->IsLODImportedDataBuildAvailable(LODIndex))
+				if (!InSkeletalMesh->HasMeshDescription(LODIndex))
 				{
 					// Make sure we regenerate any LOD data that is based off the now re imported LOD
 					FLODUtilities::RegenerateDependentLODs(InSkeletalMesh, LODIndex, GetTargetPlatformManagerRef().GetRunningTargetPlatform());
@@ -208,7 +208,7 @@ void FSkinWeightProfileHelpers::RemoveSkinWeightProfile(USkeletalMesh* InSkeleta
 			{
 				// Delete the alternate influences data
 				LODModel.SkinWeightProfiles.Remove(InProfileName);
-				if (!InSkeletalMesh->IsLODImportedDataBuildAvailable(LODIndex))
+				if (!InSkeletalMesh->HasMeshDescription(LODIndex))
 				{
 					// Regenerate this generated LOD
 					FLODUtilities::SimplifySkeletalMeshLOD(UpdateContext, LODIndex, GetTargetPlatformManagerRef().GetRunningTargetPlatform());
@@ -253,7 +253,7 @@ void FSkinWeightProfileHelpers::RemoveSkinWeightProfileLOD(USkeletalMesh* InSkel
 		{
 			FSkinWeightProfileHelpers::ClearSkinWeightProfileInstanceOverrides(InSkeletalMesh, InProfileName);
 
-			if (!InSkeletalMesh->IsLODImportedDataBuildAvailable(LODIndex))
+			if (!InSkeletalMesh->HasMeshDescription(LODIndex))
 			{
 				//Regenerate dependent LODs if we remove the LOD successfully
 				FLODUtilities::RegenerateDependentLODs(InSkeletalMesh, LODIndex, GetTargetPlatformManagerRef().GetRunningTargetPlatform());
diff --git a/Engine/Source/Editor/SkeletalMeshEditor/Private/SkeletalMeshEditor.cpp b/Engine/Source/Editor/SkeletalMeshEditor/Private/SkeletalMeshEditor.cpp
--- a/Engine/Source/Editor/SkeletalMeshEditor/Private/SkeletalMeshEditor.cpp
+++ b/Engine/Source/Editor/SkeletalMeshEditor/Private/SkeletalMeshEditor.cpp
@@ -126,7 +126,7 @@ bool IsReductionParentBaseLODUseSkeletalMeshBuildWorkflow(USkeletalMesh* Skeleta
 	{
 		return false;
 	}
-	if (SkeletalMesh->IsLODImportedDataBuildAvailable(TestLODIndex))
+	if (SkeletalMesh->HasMeshDescription(TestLODIndex))
 	{
 		return true;
 	}
diff --git a/Engine/Source/Editor/UnrealEd/Private/Factories/EditorFactories.cpp b/Engine/Source/Editor/UnrealEd/Private/Factories/EditorFactories.cpp
--- a/Engine/Source/Editor/UnrealEd/Private/Factories/EditorFactories.cpp
+++ b/Engine/Source/Editor/UnrealEd/Private/Factories/EditorFactories.cpp
@@ -6589,10 +6589,10 @@ EReimportResult::Type UReimportFbxSkeletalMeshFactory::Reimport( UObject* Obj, i
 	ReimportUI->SkeletalMeshImportData = ImportData;
 	const FSkeletalMeshModel* SkeletalMeshModel = SkeletalMesh->GetImportedModel();
 
-	bool bIsBuildAvailable = SkeletalMesh->IsLODImportedDataBuildAvailable(0);
+	const bool bUserDefinedGeometryExists = SkeletalMesh->HasMeshDescription(0);
 	
 	//Manage the content type from the source file index
-	ReimportUI->bAllowContentTypeImport = SkeletalMeshModel && SkeletalMeshModel->LODModels.Num() > 0 && !SkeletalMesh->IsLODImportedDataEmpty(0);
+	ReimportUI->bAllowContentTypeImport = SkeletalMeshModel && SkeletalMeshModel->LODModels.Num() > 0 && bUserDefinedGeometryExists;
 	if (!ReimportUI->bAllowContentTypeImport)
 	{
 		//No content type allow reimport All (legacy)
@@ -6693,7 +6693,7 @@ EReimportResult::Type UReimportFbxSkeletalMeshFactory::Reimport( UObject* Obj, i
 		{
 			const FSkeletalMeshLODModel& LODModel = SkeletalMesh->GetImportedModel()->LODModels[0];
 			
-			if (bIsBuildAvailable)
+			if (bUserDefinedGeometryExists)
 			{
 				//Set the build settings
 				LODInfo->BuildSettings.bComputeWeightedNormals = SKImportData->bComputeWeightedNormals;
diff --git a/Engine/Source/Editor/UnrealEd/Private/Fbx/FbxSkeletalMeshImport.cpp b/Engine/Source/Editor/UnrealEd/Private/Fbx/FbxSkeletalMeshImport.cpp
--- a/Engine/Source/Editor/UnrealEd/Private/Fbx/FbxSkeletalMeshImport.cpp
+++ b/Engine/Source/Editor/UnrealEd/Private/Fbx/FbxSkeletalMeshImport.cpp
@@ -1781,16 +1781,9 @@ USkeletalMesh* UnFbx::FFbxImporter::ImportSkeletalMesh(FImportSkeletalMeshArgs &
 	//Stack the PostEditChange call, it will call post edit change when it will go out of scope
 	FScopedSkeletalMeshPostEditChange ScopedPostEditChange(ExistingSkelMesh);
 
-	ESkeletalMeshGeoImportVersions GeoImportVersion = ESkeletalMeshGeoImportVersions::LatestVersion;
-	ESkeletalMeshSkinningImportVersions SkinningImportVersion = ESkeletalMeshSkinningImportVersions::LatestVersion;
 	//Adjust the import data from the import options
 	if (ExistingSkelMesh != nullptr)
 	{
-		if (ExistingSkelMesh->GetImportedModel() && ExistingSkelMesh->GetImportedModel()->LODModels.IsValidIndex(SafeLODIndex))
-		{
-			ExistingSkelMesh->GetLODImportedDataVersions(SafeLODIndex, GeoImportVersion, SkinningImportVersion);
-		}
-
 		if (ImportOptions->bImportAsSkeletalSkinning)
 		{
 			//Replace geometry import data by original existing skel mesh geometry data
@@ -1905,21 +1898,6 @@ USkeletalMesh* UnFbx::FFbxImporter::ImportSkeletalMesh(FImportSkeletalMeshArgs &
 	// process bone influences from import data
 	SkeletalMeshImportUtils::ProcessImportMeshInfluences(*SkelMeshImportDataPtr, SkeletalMesh->GetPathName());
 
-	//Store the original fbx import data the SkelMeshImportDataPtr should not be modified after this
-	SkeletalMesh->SaveLODImportedData(ImportLODModelIndex, *SkelMeshImportDataPtr);
-	if (ImportOptions->bImportAsSkeletalSkinning)
-	{
-		SkeletalMesh->SetLODImportedDataVersions(ImportLODModelIndex, GeoImportVersion, ESkeletalMeshSkinningImportVersions::LatestVersion);
-	}
-	else if(ImportOptions->bImportAsSkeletalGeometry)
-	{
-		SkeletalMesh->SetLODImportedDataVersions(ImportLODModelIndex, ESkeletalMeshGeoImportVersions::LatestVersion, SkinningImportVersion);
-	}
-	else
-	{
-		//We reimport both
-		SkeletalMesh->SetLODImportedDataVersions(ImportLODModelIndex, ESkeletalMeshGeoImportVersions::LatestVersion, ESkeletalMeshSkinningImportVersions::LatestVersion);
-	}
 	SkeletalMesh->ResetLODInfo();
 	FSkeletalMeshLODInfo& NewLODInfo = SkeletalMesh->AddLODInfo();
 	NewLODInfo.ReductionSettings.NumOfTrianglesPercentage = 1.0f;
@@ -1927,6 +1905,11 @@ USkeletalMesh* UnFbx::FFbxImporter::ImportSkeletalMesh(FImportSkeletalMeshArgs &
 	NewLODInfo.ReductionSettings.MaxDeviationPercentage = 0.0f;
 	NewLODInfo.LODHysteresis = 0.02f;
 
+	//Store the original fbx import data the SkelMeshImportDataPtr should not be modified after this
+	PRAGMA_DISABLE_DEPRECATION_WARNINGS
+	SkeletalMesh->SaveLODImportedData(ImportLODModelIndex, *SkelMeshImportDataPtr);
+	PRAGMA_ENABLE_DEPRECATION_WARNINGS
+	
 	SkeletalMesh->SetImportedBounds(FBoxSphereBounds((FBox)BoundingBox));
 
 	// Store whether or not this mesh has vertex colors
@@ -1936,18 +1919,6 @@ USkeletalMesh* UnFbx::FFbxImporter::ImportSkeletalMesh(FImportSkeletalMeshArgs &
 	// Pass the number of texture coordinate sets to the LODModel.  Ensure there is at least one UV coord
 	LODModel.NumTexCoords = FMath::Max<uint32>(1, SkelMeshImportDataPtr->NumTexCoords);
 
-	// Copy mesh infos into the LOD model.
-	LODModel.ImportedMeshInfos.Empty();
-	LODModel.ImportedMeshInfos.Reserve(SkelMeshImportDataPtr->MeshInfos.Num());
-	for (const SkeletalMeshImportData::FMeshInfo& MeshInfo : SkelMeshImportDataPtr->MeshInfos)
-	{
-		LODModel.ImportedMeshInfos.AddDefaulted();
-		FSkelMeshImportedMeshInfo& LODMeshInfo = LODModel.ImportedMeshInfos.Last();
-		LODMeshInfo.Name = MeshInfo.Name;
-		LODMeshInfo.NumVertices = MeshInfo.NumVertices;
-		LODMeshInfo.StartImportedVertex = MeshInfo.StartImportedVertex;
-	}
-
 	if(ImportSkeletalMeshArgs.bCreateRenderData )
 	{
 		TArray<FVector3f> LODPoints;
@@ -2470,10 +2441,11 @@ USkeletalMesh* UnFbx::FFbxImporter::ReimportSkeletalMesh(USkeletalMesh* Mesh, UF
 	UFbxSkeletalMeshImportData* SKImportData = Cast<UFbxSkeletalMeshImportData>(Mesh->GetAssetImportData());
 	if (SKImportData)
 	{
-		FSkeletalMeshLODInfo* LODInfo = Mesh->GetLODInfo(0);
-		if(LODInfo && Mesh->GetImportedModel() && Mesh->GetImportedModel()->LODModels.IsValidIndex(0))
+		constexpr int32 LODIndex = 0;
+		FSkeletalMeshLODInfo* LODInfo = Mesh->GetLODInfo(LODIndex);
+		if(LODInfo && Mesh->GetImportedModel() && Mesh->GetImportedModel()->LODModels.IsValidIndex(LODIndex))
 		{
-			if (!Mesh->IsLODImportedDataBuildAvailable(0))
+			if (!Mesh->HasMeshDescription(LODIndex))
 			{
 				//Set the build settings
 				LODInfo->BuildSettings.bComputeWeightedNormals = SKImportData->bComputeWeightedNormals;
@@ -4223,12 +4195,17 @@ void UnFbx::FFbxImporter::InsertNewLODToBaseSkeletalMesh(USkeletalMesh* InSkelet
 	}
 
 	// Release all resources before replacing the model
-	BaseSkeletalMesh->PreEditChange(NULL);
+	BaseSkeletalMesh->PreEditChange(nullptr);
 
 	// Assign new FSkeletalMeshLODModel to desired slot in selected skeletal mesh.
 	FSkeletalMeshLODModel::CopyStructure(&(DestImportedResource->LODModels[DesiredLOD]), &NewLODModel);
 	//Copy the import data into the base skeletalmesh for the imported LOD
-	USkeletalMesh::CopyImportedData(0, InSkeletalMesh, DesiredLOD, BaseSkeletalMesh);
+	FMeshDescription SourceMeshDescription;
+	if (InSkeletalMesh->CloneMeshDescription(0, SourceMeshDescription))
+	{
+		BaseSkeletalMesh->CreateMeshDescription(DesiredLOD, MoveTemp(SourceMeshDescription));
+		BaseSkeletalMesh->CommitMeshDescription(DesiredLOD);
+	}
 	
 
 	// If this LOD had been generated previously by automatic mesh reduction, clear that flag.
@@ -4416,11 +4393,13 @@ bool UnFbx::FFbxImporter::ImportSkeletalMeshLOD(USkeletalMesh* InSkeletalMesh, U
 	{
 		//The imported LOD is always in LOD 0 of the InSkeletalMesh
 		const int32 SourceLODIndex = 0;
-		if(!InSkeletalMesh->IsLODImportedDataEmpty(SourceLODIndex))
+		if(InSkeletalMesh->HasMeshDescription(SourceLODIndex))
 		{
 			// Fix up the imported data bone indexes
 			FSkeletalMeshImportData LODImportData;
+			PRAGMA_DISABLE_DEPRECATION_WARNINGS
 			InSkeletalMesh->LoadLODImportedData(SourceLODIndex, LODImportData);
+			PRAGMA_ENABLE_DEPRECATION_WARNINGS
 			const int32 LODImportDataBoneNumber = LODImportData.RefBonesBinary.Num();
 			//We want to create a remap array so we can fix all influence easily
 			TArray<int32> ImportDataBoneRemap;
@@ -4468,7 +4447,9 @@ bool UnFbx::FFbxImporter::ImportSkeletalMeshLOD(USkeletalMesh* InSkeletalMesh, U
 				LODImportData.Influences.Shrink();
 			}
 			//Save the fix up remap bone index
+			PRAGMA_DISABLE_DEPRECATION_WARNINGS
 			InSkeletalMesh->SaveLODImportedData(SourceLODIndex, LODImportData);
+			PRAGMA_ENABLE_DEPRECATION_WARNINGS
 		}
 
 		// Fix up the ActiveBoneIndices array.
@@ -4745,7 +4726,7 @@ void UnFbx::FFbxImporter::ImportMorphTargetsInternal( TArray<FbxNode*>& SkelMesh
 	if (BaseSkelMesh->GetImportedModel() && BaseSkelMesh->GetImportedModel()->LODModels.IsValidIndex(LODIndex))
 	{
 		//If we can build the skeletal mesh there is no need to build the morph target now, all the necessary build morph target data was copied before.
-		if (!BaseSkelMesh->IsLODImportedDataBuildAvailable(LODIndex))
+		if (!BaseSkelMesh->HasMeshDescription(LODIndex))
 		{
 			//Build MorphTargets
 			FLODUtilities::BuildMorphTargets(
@@ -4793,7 +4774,9 @@ void UnFbx::FFbxImporter::ImportFbxMorphTarget(TArray<FbxNode*> &SkelMeshNodeArr
 	{
 		ImportMorphTargetsInternal( SkelMeshNodeArray, BaseSkelMesh, LODIndex, BaseSkeletalMeshImportData, bMapMorphTargetToTimeZero);
 		//Save the rawMesh
+		PRAGMA_DISABLE_DEPRECATION_WARNINGS
 		BaseSkelMesh->SaveLODImportedData(LODIndex, BaseSkeletalMeshImportData);
+		PRAGMA_ENABLE_DEPRECATION_WARNINGS
 	}
 }
 
diff --git a/Engine/Source/Editor/UnrealEd/Private/FbxMeshUtils.cpp b/Engine/Source/Editor/UnrealEd/Private/FbxMeshUtils.cpp
--- a/Engine/Source/Editor/UnrealEd/Private/FbxMeshUtils.cpp
+++ b/Engine/Source/Editor/UnrealEd/Private/FbxMeshUtils.cpp
@@ -645,7 +645,7 @@ namespace FbxMeshUtils
 			TArray<FbxNode*>* MeshObject = NULL;
 
 			//Set the build options if the BuildDat is not available so it is the same option we use to import the LOD
-			if (ImportedResource && ImportedResource->LODModels.IsValidIndex(LODLevel) && !SelectedSkelMesh->IsLODImportedDataBuildAvailable(LODLevel))
+			if (ImportedResource && ImportedResource->LODModels.IsValidIndex(LODLevel) && !SelectedSkelMesh->HasMeshDescription(LODLevel))
 			{
 				FSkeletalMeshLODInfo* LODInfo = SelectedSkelMesh->GetLODInfo(LODLevel);
 				if (LODInfo)
diff --git a/Engine/Source/Editor/UnrealEd/Private/ImportUtils/SkeletalMeshImportUtils.cpp b/Engine/Source/Editor/UnrealEd/Private/ImportUtils/SkeletalMeshImportUtils.cpp
--- a/Engine/Source/Editor/UnrealEd/Private/ImportUtils/SkeletalMeshImportUtils.cpp
+++ b/Engine/Source/Editor/UnrealEd/Private/ImportUtils/SkeletalMeshImportUtils.cpp
@@ -222,7 +222,7 @@ bool SkeletalMeshImportUtils::ProcessImportMeshSkeleton(const USkeleton* Skeleto
 */
 void SkeletalMeshImportUtils::ProcessImportMeshInfluences(FSkeletalMeshImportData& ImportData, const FString& SkeletalMeshName)
 {
-	FLODUtilities::ProcessImportMeshInfluences(ImportData.Wedges.Num(), ImportData.Influences, SkeletalMeshName);
+	FLODUtilities::ProcessImportMeshInfluences(ImportData.Points.Num(), ImportData.Influences, SkeletalMeshName);
 }
 
 void SkeletalMesUtilsImpl::SaveSkeletalMeshLODModelSections(USkeletalMesh* SourceSkeletalMesh, TSharedPtr<FExistingSkelMeshData>& ExistingMeshDataPtr, int32 LodIndex, bool bSaveNonReducedMeshData)
@@ -350,9 +350,6 @@ TSharedPtr<FExistingSkelMeshData> SkeletalMeshImportUtils::SaveExistingSkelMeshD
 	{
 		ExistingMeshDataPtr->ExistingMaterials = SourceMaterials;
 	}
-	PRAGMA_DISABLE_DEPRECATION_WARNINGS
-	ExistingMeshDataPtr->ExistingRetargetBasePose = SourceSkeletalMesh->GetRetargetBasePose();
-	PRAGMA_ENABLE_DEPRECATION_WARNINGS
 
 	ExistingMeshDataPtr->ExistingInlineReductionCacheDatas = SourceMeshModel->InlineReductionCacheDatas;
 
@@ -369,9 +366,10 @@ TSharedPtr<FExistingSkelMeshData> SkeletalMeshImportUtils::SaveExistingSkelMeshD
 			const FSkeletalMeshLODModel& LODModel = SourceMeshModel->LODModels[LODIndex];
 			ExistingMeshDataPtr->ExistingLODModels.Add(FSkeletalMeshLODModel::CreateCopy(&LODModel));
 			//Store the import data for every LODs
-			FSkeletalMeshLodImportDataBackup& LodMeshImportData = ExistingMeshDataPtr->ExistingLODImportDatas.AddDefaulted_GetRef();
-			SourceSkeletalMesh->LoadLODImportedData(LODIndex, LodMeshImportData.MeshImportData);
-			SourceSkeletalMesh->GetLODImportedDataVersions(LODIndex, LodMeshImportData.MeshGeoImportVersion, LodMeshImportData.MeshSkinningImportVersion);
+			FSkeletalMeshImportData& LodMeshImportData = ExistingMeshDataPtr->ExistingLODImportDatas.AddDefaulted_GetRef();
+			PRAGMA_DISABLE_DEPRECATION_WARNINGS
+			SourceSkeletalMesh->LoadLODImportedData(LODIndex, LodMeshImportData);
+			PRAGMA_ENABLE_DEPRECATION_WARNINGS
 		}
 		check(ExistingMeshDataPtr->ExistingLODModels.Num() == SourceMeshModel->LODModels.Num());
 
@@ -710,15 +708,6 @@ void SkeletalMeshImportUtils::RestoreExistingSkelMeshData(const TSharedPtr<const
 	//Do everything we need for base LOD re-import
 	if (SafeReimportLODIndex == 0)
 	{
-		PRAGMA_DISABLE_DEPRECATION_WARNINGS
-		// this is not ideal. Ideally we'll have to save only diff with indicating which joints, 
-		// but for now, we allow them to keep the previous pose IF the element count is same
-		if (MeshData->ExistingRetargetBasePose.Num() == SkeletalMesh->GetRefSkeleton().GetRawBoneNum())
-		{
-			SkeletalMesh->SetRetargetBasePose(MeshData->ExistingRetargetBasePose);
-		}
-		PRAGMA_ENABLE_DEPRECATION_WARNINGS
-
 		// Assign sockets from old version of this SkeletalMesh.
 		// Only copy ones for bones that exist in the new mesh.
 		for (int32 i = 0; i < MeshData->ExistingSockets.Num(); i++)
@@ -846,10 +835,11 @@ void SkeletalMeshImportUtils::RestoreExistingSkelMeshData(const TSharedPtr<const
 						SkeletalMeshImportedModel->LODModels.Add(LODModelCopy);
 						SkeletalMesh->AddLODInfo(LODInfo);
 						//Restore custom LOD import data
-						FSkeletalMeshLodImportDataBackup* LodMeshImportData = const_cast<FSkeletalMeshLodImportDataBackup*>(&(MeshData->ExistingLODImportDatas[LODIndex]));
+						const FSkeletalMeshImportData& LodMeshImportData = MeshData->ExistingLODImportDatas[LODIndex];
 						//SaveLODImportdData cannot take a const structure because it use serialization(which cannot be const because same function read and write)
-						SkeletalMesh->SaveLODImportedData(LODIndex, LodMeshImportData->MeshImportData);
-						SkeletalMesh->SetLODImportedDataVersions(LODIndex, LodMeshImportData->MeshGeoImportVersion, LodMeshImportData->MeshSkinningImportVersion);
+						PRAGMA_DISABLE_DEPRECATION_WARNINGS
+						SkeletalMesh->SaveLODImportedData(LODIndex, LodMeshImportData);
+						PRAGMA_ENABLE_DEPRECATION_WARNINGS
 
 						auto FillInlineReductionData = [&MeshData, &SkeletalMeshImportedModel, LODIndex]()
 						{
@@ -887,7 +877,7 @@ void SkeletalMeshImportUtils::RestoreExistingSkelMeshData(const TSharedPtr<const
 			
 			
 			//Old asset cannot use the new build system, we need to regenerate dependent LODs
-			if (SkeletalMesh->IsLODImportedDataBuildAvailable(SafeReimportLODIndex) == false)
+			if (!SkeletalMesh->HasMeshDescription(SafeReimportLODIndex))
 			{
 				FLODUtilities::RegenerateDependentLODs(SkeletalMesh, SafeReimportLODIndex, GetTargetPlatformManagerRef().GetRunningTargetPlatform());
 			}
diff --git a/Engine/Source/Editor/UnrealEd/Public/ImportUtils/SkelImport.h b/Engine/Source/Editor/UnrealEd/Public/ImportUtils/SkelImport.h
--- a/Engine/Source/Editor/UnrealEd/Public/ImportUtils/SkelImport.h
+++ b/Engine/Source/Editor/UnrealEd/Public/ImportUtils/SkelImport.h
@@ -43,12 +43,6 @@ struct FExistingMeshLodSectionData
 	bool bDisabled;
 };
 
-struct FSkeletalMeshLodImportDataBackup
-{
-	FSkeletalMeshImportData MeshImportData;
-	ESkeletalMeshGeoImportVersions MeshGeoImportVersion = ESkeletalMeshGeoImportVersions::Before_Versionning;
-	ESkeletalMeshSkinningImportVersions MeshSkinningImportVersion = ESkeletalMeshSkinningImportVersions::Before_Versionning;
-};
 
 struct FStreamableRenderAssetData
 {
@@ -78,7 +72,7 @@ struct FExistingSkelMeshData
 	TArray<USkeletalMeshSocket*>			ExistingSockets;
 	TArray<FInlineReductionCacheData>		ExistingInlineReductionCacheDatas;
 	TIndirectArray<FSkeletalMeshLODModel>	ExistingLODModels;
-	TArray<FSkeletalMeshLodImportDataBackup>	ExistingLODImportDatas;
+	TArray<FSkeletalMeshImportData>			ExistingLODImportDatas;
 	TArray<FSkeletalMeshLODInfo>			ExistingLODInfo;
 	FReferenceSkeleton						ExistingRefSkeleton;
 	TArray<FSkeletalMaterial>				ExistingMaterials;
@@ -87,7 +81,6 @@ struct FExistingSkelMeshData
 	TArray<UPhysicsAsset*>					ExistingPhysicsAssets;
 	UPhysicsAsset*							ExistingShadowPhysicsAsset;
 	USkeleton*								ExistingSkeleton;
-	TArray<FTransform>						ExistingRetargetBasePose;
 	USkeletalMeshLODSettings*				ExistingLODSettings;
 	TSubclassOf<UAnimInstance>				ExistingPostProcessAnimBlueprint;
 	TSoftObjectPtr<UObject>					ExistingDefaultAnimatingRig;
diff --git a/Engine/Source/Runtime/Core/Private/UObject/FortniteMainBranchObjectVersion.cpp b/Engine/Source/Runtime/Core/Private/UObject/FortniteMainBranchObjectVersion.cpp
--- a/Engine/Source/Runtime/Core/Private/UObject/FortniteMainBranchObjectVersion.cpp
+++ b/Engine/Source/Runtime/Core/Private/UObject/FortniteMainBranchObjectVersion.cpp
@@ -12,7 +12,7 @@ TMap<FGuid, FGuid> FFortniteMainBranchObjectVersion::GetSystemGuids()
 	SystemGuids.Add(DevGuids.NANITE_DERIVEDDATA_VER, FGuid("D28C003F825F4C178F698938FB0E47F7"));
 	SystemGuids.Add(DevGuids.NIAGARASHADERMAP_DERIVEDDATA_VER, FGuid("D0DF49D408F34E38A98FA097E0280EC7"));
 	SystemGuids.Add(DevGuids.Niagara_LatestScriptCompileVersion, FGuid("CE33E29C6268469E8E18F1D96A55DFDA"));
-	SystemGuids.Add(DevGuids.SkeletalMeshDerivedDataVersion, FGuid("518D7C51D31D9BE4A9F4B81542AC8119"));
+	SystemGuids.Add(DevGuids.SkeletalMeshDerivedDataVersion, FGuid("41EFB33F1B72B1F49F9AA7D295871F36"));
 	SystemGuids.Add(DevGuids.STATICMESH_DERIVEDDATA_VER, FGuid("543843B103794E0E9BA4BE80FB602F79"));
 	SystemGuids.Add(DevGuids.MaterialTranslationDDCVersion, FGuid("7CBF2AF7AD5545868E3A15AA2BA900D7"));
 	return SystemGuids;
diff --git a/Engine/Source/Runtime/Core/Public/UObject/FortniteMainBranchObjectVersion.h b/Engine/Source/Runtime/Core/Public/UObject/FortniteMainBranchObjectVersion.h
--- a/Engine/Source/Runtime/Core/Public/UObject/FortniteMainBranchObjectVersion.h
+++ b/Engine/Source/Runtime/Core/Public/UObject/FortniteMainBranchObjectVersion.h
@@ -411,6 +411,9 @@ struct FFortniteMainBranchObjectVersion
 		// Added settings to Landscape HLODs
 		LandscapeAddedHLODSettings,
 
+		// Skeletal Mesh uses Mesh Description to store mesh bulk data.
+		MeshDescriptionForSkeletalMesh,
+		
 		// -----<new versions can be added above this line>-------------------------------------------------
 		VersionPlusOne,
 		LatestVersion = VersionPlusOne - 1
diff --git a/Engine/Source/Runtime/Engine/Classes/Engine/SkeletalMesh.h b/Engine/Source/Runtime/Engine/Classes/Engine/SkeletalMesh.h
--- a/Engine/Source/Runtime/Engine/Classes/Engine/SkeletalMesh.h
+++ b/Engine/Source/Runtime/Engine/Classes/Engine/SkeletalMesh.h
@@ -6,33 +6,32 @@
  * Contains the shared data that is used by all SkeletalMeshComponents (instances).
  */
 
-#include "CoreMinimal.h"
-#include "UObject/ObjectMacros.h"
-#include "UObject/Object.h"
-#include "Templates/SubclassOf.h"
-#include "Interfaces/Interface_AssetUserData.h"
-#include "RenderCommandFence.h"
-#include "EngineDefines.h"
-#include "ReferenceSkeleton.h"
-#include "GPUSkinPublicDefs.h"
+#include "Animation/MorphTarget.h"
+#include "Animation/NodeMappingProviderInterface.h"
 #include "Animation/PreviewAssetAttachComponent.h"
-#include "Interfaces/Interface_CollisionDataProvider.h"
+#include "Engine/StreamableRenderAsset.h"
+#include "EngineDefines.h"
 #include "EngineTypes.h"
-#include "SkeletalMeshSampling.h"
-#include "PerPlatformProperties.h"
+#include "Interfaces/Interface_AssetUserData.h"
+#include "Interfaces/Interface_CollisionDataProvider.h"
 #include "Misc/EnumClassFlags.h"
-#include "Animation/NodeMappingProviderInterface.h"
-#include "Animation/MorphTarget.h"
-#include "Engine/StreamableRenderAsset.h"
+#include "PerPlatformProperties.h"
 #include "PerQualityLevelProperties.h"
+#include "ReferenceSkeleton.h"
+#include "RenderCommandFence.h"
+#include "SkeletalMeshSampling.h"
+#include "SkeletalMeshSourceModel.h"
 #include "SkinnedAsset.h"
+#include "Templates/SubclassOf.h"
+#include "UObject/Object.h"
+#include "UObject/ObjectMacros.h"
 
 #if UE_ENABLE_INCLUDE_ORDER_DEPRECATED_IN_5_2
+#include "Animation/SkinWeightProfile.h"
+#include "BoneContainer.h"
 #include "Components.h"
 #include "Rendering/SkeletalMeshRenderData.h"
-#include "BoneContainer.h"
 #include "SkeletalMeshLODSettings.h"
-#include "Animation/SkinWeightProfile.h"
 #endif
 
 #include "SkeletalMesh.generated.h"
@@ -86,7 +85,7 @@ enum class ESkeletalMeshAsyncProperties : uint64
 	RetargetBasePose = 1 << 3,
 	RefBasesInvMatrix = 1 << 4,
 	MeshClothingAssets = 1 << 5,
-	UseLegacyMeshDerivedDataKey = 1 << 6,
+	SourceModels = 1 << 6,	
 	HasActiveClothingAssets = 1 << 7,
 	LODSettings = 1 << 8,
 	HasVertexColors = 1 << 9,
@@ -451,7 +450,40 @@ class USkeletalMesh : public USkinnedAsset, public IInterface_CollisionDataProvi
 #endif
 private:
 #if WITH_EDITORONLY_DATA
-	/** Imported skeletal mesh geometry information (not used at runtime). */
+	/** The imported/editable geometric source data for generating renderable data. */
+	UE_DEPRECATED(5.4, "This must be protected for async build, always use the accessors even internally.")
+	UPROPERTY()
+	TArray<FSkeletalMeshSourceModel> SourceModels;
+
+	/** Set the number of source models. If adding more LODs than currently exist, empty LODs will be created
+	 *  to fill in the gap. If the new size is smaller than the existing LODs, the source geometry for LODs
+	 *  at and higher than the number given will be thrown away.
+	 *  NOTE: Internal for now, use AddLODInfo to add LODs (incl. source models).
+	 */
+	void SetNumSourceModels(const int32 InNumSourceModels);
+
+	/** Used by both SetNumSourceModels and AddLODInfo to set the current number of source models.
+	 *  Does not take an async lock, expects the calling function to be holding it already.
+	 *  \param InNumSourceModels The number of source models to set. Can be lower or higher than the
+	 *     current number of entries and models will be added/removed as needed.
+	 */
+	void SetNumSourceModelsInternal(const int32 InNumSourceModels);
+
+	/** Remove the source model for the given LOD. LODs above it will be adjusted down.
+	 *  NOTE: Internal for now, use RemoveLODInfo to remove specific LODs.  
+	 */
+	void RemoveSourceModel(const int32 InLODIndex);
+	
+	/** Used by RemoveLODInfo to remove source models if a LOD is removed.
+	 *  Does not take an async lock, expects the calling function to be holding it already.
+	 *  \param InLODIndex The LOD index whose model data is being removed.
+	 */
+	void RemoveSourceModelInternal(const int32 InLODIndex);
+	
+	/** Processed skeletal mesh geometry information (not used at runtime). */
+	/** NOTE: This data is slated for removal in future versions of UE. Do not rely on it. Use either the
+	 *  render data, or source models, as needed.
+	 */
 	UE_DEPRECATED(5.0, "This must be protected for async build, always use the accessors even internally.")
 	TSharedPtr<FSkeletalMeshModel> ImportedModel;
 #endif
@@ -465,7 +497,6 @@ private:
 	ENGINE_API void SetSkeletalMeshRenderData(TUniquePtr<FSkeletalMeshRenderData>&& InSkeletalMeshRenderData);
 
 #if WITH_EDITORONLY_DATA
-public:
 	/*
 	 * This editor data asset is save under the skeletalmesh(skel mesh is the owner), the editor data asset is always loaded.
 	 * There is only one editor data asset possible per skeletalmesh.
@@ -473,40 +504,122 @@ public:
 	 */
 	UE_DEPRECATED(4.27, "Please do not access this member directly; use the public ImportData API.")
 	UPROPERTY()
-	mutable TObjectPtr<USkeletalMeshEditorData> MeshEditorDataObject;
+	mutable TObjectPtr<USkeletalMeshEditorData> MeshEditorDataObject_DEPRECATED;
 
-private:
-	USkeletalMeshEditorData* GetMeshEditorDataObject() const
-	{
-		WaitUntilAsyncPropertyReleased(ESkeletalMeshAsyncProperties::MeshEditorDataObject);
-		PRAGMA_DISABLE_DEPRECATION_WARNINGS
-		return MeshEditorDataObject;
-		PRAGMA_ENABLE_DEPRECATION_WARNINGS
-	}
+#endif //WITH_EDITORONLY_DATA
 
-	void SetMeshEditorDataObject(USkeletalMeshEditorData* InMeshEditorDataObject)
-	{
-		WaitUntilAsyncPropertyReleased(ESkeletalMeshAsyncProperties::MeshEditorDataObject);
-		PRAGMA_DISABLE_DEPRECATION_WARNINGS
-		MeshEditorDataObject = InMeshEditorDataObject;
-		PRAGMA_ENABLE_DEPRECATION_WARNINGS
-	}
+public:
 
-	/*
-	 * Return a valid USkeletalMeshEditorData, if the MeshEditorDataPath is invalid it will create the USkeletalMeshEditorData and set the MeshEditorDataPath to point on it.
+#if WITH_EDITORONLY_DATA
+	/** Returns the number of source models.
+	 *  \note This value is the same as the return value of GetLODNum, as these two are kept in sync.
 	 */
-	ENGINE_API USkeletalMeshEditorData& GetMeshEditorData() const;
+	ENGINE_API int32 GetNumSourceModels() const;
 
-	bool IsMeshEditorDataValid() const
+	/** Returns the source model object for the skeletal mesh of a given LOD. This source models stores
+	 *  an optional mesh description that the renderable data is generated from. If there is no mesh description
+	 *  object, then the geometry of this LOD is automatically generated from an earlier LOD.
+	 *  \param InLODIndex The LOD index to retrieve a FSkeletalMeshSourceModel object for. 
+	 */
+	ENGINE_API const FSkeletalMeshSourceModel& GetSourceModel(const int32 InLODIndex) const;
+	
+	/** Returns the source model object for the skeletal mesh of a given LOD. This source models stores
+	 *  an optional mesh description that the renderable data is generated from. If there is no mesh description
+	 *  object, then the geometry of this LOD is automatically generated from an earlier LOD.
+	 *  \param InLODIndex The LOD index to retrieve a FSkeletalMeshSourceModel object for. 
+	 */
+	ENGINE_API FSkeletalMeshSourceModel& GetSourceModel(const int32 InLODIndex);
+	
+	/** Returns the MeshDescription object associated with the given LOD index. If there's no import data
+	 *  at this LOD, and the LOD is auto-generated, then a \c nullptr is returned.
+	 *  \param InLODIndex The LOD index to retrieve a FMeshDescription object for. 
+	 */
+	ENGINE_API FMeshDescription* GetMeshDescription(const int32 InLODIndex) const;
+	
+	/** Returns the MeshDescription object associated with the given LOD index. If there's no import data
+	 *  at this LOD, and the LOD is auto-generated, then a \c nullptr is returned.
+	 *  \param InLODIndex The LOD index to clone a FMeshDescription object from.
+	 *  \param OutMeshDescription The target mesh description object to copy the source data into.
+	 *  \return \c true if valid mesh description was copied in. \c false if there was no mesh description
+	 *     the LOD is automatically generated.
+	 */
+	ENGINE_API bool CloneMeshDescription(const int32 InLODIndex, FMeshDescription& OutMeshDescription) const;
+
+	/** Returns \c true if this LOD has a mesh description to generate render data from.
+	 *  If not, then render data is automatically generated from a lower LOD level.
+	 *  Use this rather than GetMeshDescription if the only requirement is to check for existence, since
+	 *  this function won't unpack a full mesh description object from bulk storage if it doesn't have to.
+	 *  \param InLODIndex The LOD index to check for existence of a mesh description object.
+	 *  \return \c This LOD index contains a valid mesh description object.
+	 */
+	ENGINE_API bool HasMeshDescription(const int32 InLODIndex) const;
+
+	/** Creates a new, empty mesh description at the LOD index given. If there was a mesh description there
+	 *  prior, that object is destroyed.
+	 *  \param InLODIndex The LOD index to create a new mesh description object at.
+	 *  \return \c A pointer to the newly created object, or \c nullptr if the method was unable to create anything.
+	 */
+	ENGINE_API FMeshDescription* CreateMeshDescription(const int32 InLODIndex);
+
+	/** Creates a new mesh description at the LOD index given, moving the contents of the given mesh description
+	 *  into it. If there was a mesh description there prior, that object is destroyed.
+	 *  \param InLODIndex The LOD index to create a new mesh description object at.
+	 *  \param InMeshDescription The mesh description object to initialize from as a copy.
+	 *  \return \c A pointer to the newly created object, or \c nullptr if the method was unable to create anything.
+	 */
+	ENGINE_API FMeshDescription* CreateMeshDescription(const int32 InLODIndex, FMeshDescription&& InMeshDescription);
+
+	
+	/* Structure that defines parameters passed into the commit mesh description function */
+	struct FCommitMeshDescriptionParams
 	{
-		return GetMeshEditorDataObject() != nullptr;
-	}
+		FCommitMeshDescriptionParams() {}
+		
+		/** Mark the package as dirty. If calling CommitMeshDescription from a non-game thread,
+		 *  this value should be set to \c false.
+		 */
+		bool bMarkPackageDirty = true;
+
+		/** Force the render data to update. By default the render data uses the hash of the mesh to check if
+		 *  an update is required, this forces this hash to be unique, causing the render data to update whether
+		 *  the mesh has changed or not.
+		 */
+		bool bForceUpdate = false;
+	};
 
-#endif //WITH_EDITORONLY_DATA
+	/* Commits the stored mesh description object to bulk storage. This also forces the imported bounds to update
+	 * if the mesh committed is on LOD 0. If there is no stored mesh description, the bulk storage will be emptied.
+	 *  \param InLODIndex The LOD index at which to commit the mesh description to bulk storage.
+	 *  \param InParams An optional object to control how the commit is done.
+	 *  \return \c true if the commit was successful. Even if there was no mesh description to store, the commit
+	 *    is still successful.
+	 */
+	ENGINE_API bool CommitMeshDescription(const int32 InLODIndex, const FCommitMeshDescriptionParams& InParams = {});
 
-public:
+	/** Mark the bulk data as dirty so that it can be transacted for undo/redo.
+	 *  \param InLODIndex The LOD index at which to mark the bulk data dirty.
+	 *  \param bInAlwaysMarkPackageDirty Always mark the owning package dirty, even if there was not transaction open to store into.
+	 *  \return \c true if the bulk data was stored in the transaction buffer.
+	 */
+	ENGINE_API bool ModifyMeshDescription(const int32 InLODIndex, const bool bInAlwaysMarkPackageDirty = true);
 
-#if WITH_EDITORONLY_DATA
+	/** Clears a stored mesh description and releases its memory. Leaves the bulk data in place if the mesh
+	 *  description is needed again.
+	 *  \param InLODIndex The LOD index at which to clear out the mesh description.
+	 */
+	ENGINE_API void ClearMeshDescription(const int32 InLODIndex);
+
+	/** Clears a all stored mesh descriptions on all LODs and releases their memory. Leaves the bulk data in 
+	 *  place if the mesh descriptions are needed again.
+	 */
+	ENGINE_API void ClearAllMeshDescriptions();
+
+	/** Clears the stored mesh description _and_ the bulk data. Use this to mark the LOD as being auto-generated.
+	 *  Once this function is called, the LOD settings for this level should be updated to reflect how the
+	 *  geometry should be autogenerated from a prior LOD.
+	 *  \param InLODIndex The LOD index at which to clear out the mesh description and bulk data.
+	 */
+	ENGINE_API void ClearMeshDescriptionAndBulkData(const int32 InLODIndex);
 
 	//////////////////////////////////////////////////////////////////////////
 	// USkeletalMeshEditorData public skeletalmesh API
@@ -519,63 +632,57 @@ public:
 	
 
 	/* Fill the OutMesh with the imported data */
+	UE_DEPRECATED(5.4, "Use GetMeshDescription instead.")
 	ENGINE_API void LoadLODImportedData(const int32 LODIndex, FSkeletalMeshImportData& OutMesh) const;
 
 	/* Fill the asset LOD entry with the InMesh. */
-	ENGINE_API void SaveLODImportedData(const int32 LODIndex, FSkeletalMeshImportData& InMesh);
+	UE_DEPRECATED(5.4, "Use CommitMeshDescription instead.")
+	ENGINE_API void SaveLODImportedData(const int32 LODIndex, const FSkeletalMeshImportData& InMesh);
 	
 	/* Return true if the imported data has all the necessary data to use the skeletalmesh builder. Return False otherwise.
 	 * Old asset before the refactor will not be able to be build until it get fully re-import.
 	 * This value is cache in the LODModel and update when we call SaveLODImportedData.
 	 */
+	UE_DEPRECATED(5.4, "Use HasMeshDescription instead.")
 	ENGINE_API bool IsLODImportedDataBuildAvailable(const int32 LODIndex) const;
 	
 	/* Return true if the imported data is present. Return false otherwise.
 	 * Old asset before the split workflow will not have this data and will not support import geo only or skinning only.
 	 * This value is cache in the LODModel and update when we call SaveLODImportedData.
 	 */
+	UE_DEPRECATED(5.4, "Use HasMeshDescription instead.")
 	ENGINE_API bool IsLODImportedDataEmpty(const int32 LODIndex) const;
 
 	/* Get the Versions of the geo and skinning data. We use those versions to answer to IsLODImportedDataBuildAvailable function. */
+	UE_DEPRECATED(5.4, "No equivalent provided since versioning is not surfaced for mesh description bulk data.")
 	ENGINE_API void GetLODImportedDataVersions(const int32 LODIndex, ESkeletalMeshGeoImportVersions& OutGeoImportVersion, ESkeletalMeshSkinningImportVersions& OutSkinningImportVersion) const;
 
 	/* Set the Versions of the geo and skinning data. We use those versions to answer to IsLODImportedDataBuildAvailable function. */
+	UE_DEPRECATED(5.4, "No equivalent provided since versioning is not surfaced for mesh description bulk data.")
 	ENGINE_API void SetLODImportedDataVersions(const int32 LODIndex, const ESkeletalMeshGeoImportVersions& InGeoImportVersion, const ESkeletalMeshSkinningImportVersions& InSkinningImportVersion);
 
-	/* Static function that copy the LOD import data from a source s^keletal mesh to a destination skeletal mesh*/
+	/* Static function that copy the LOD import data from a source skeletal mesh to a destination skeletal mesh*/
+	UE_DEPRECATED(5.4, "Use GetMeshDescription and CreateMeshDescription instead.")
 	static ENGINE_API void CopyImportedData(int32 SrcLODIndex, USkeletalMesh* SrcSkeletalMesh, int32 DestLODIndex, USkeletalMesh* DestSkeletalMesh);
 
 	/* Allocate the space we need. Use this before calling this API in multithreaded. */
+	UE_DEPRECATED(5.4, "No equivalent provided. All LODs should be added up-front for multi-threaded use.")
 	ENGINE_API void ReserveLODImportData(int32 MaxLODIndex);
 	
+	UE_DEPRECATED(5.4, "No equivalent provided.")
 	ENGINE_API void ForceBulkDataResident(const int32 LODIndex);
 
 	/* Remove the import data for the specified LOD */
+	UE_DEPRECATED(5.4, "Use ClearMeshDescriptionAndBulkData instead.")
 	ENGINE_API void EmptyLODImportData(const int32 LODIndex);
 
 	/* Remove the import data for all the LODs */
+	UE_DEPRECATED(5.4, "Use ClearMeshDescriptionAndBulkData instead.")
 	ENGINE_API void EmptyAllImportData();
 
 	// End USkeletalMeshEditorData public skeletalmesh API
 	//////////////////////////////////////////////////////////////////////////
 
-	/* Structure that defines parameters passed into the commit mesh description function */
-	struct FCommitMeshDescriptionParams
-	{
-		FCommitMeshDescriptionParams()
-			: bModify(true)
-		{}
-
-		bool bModify;
-	};
-
-	/* Fills OutMeshDescription from the associated InLODIndex. */
-	ENGINE_API void GetMeshDescription(const int32 InLODIndex, FMeshDescription& OutMeshDescription) const;
-
-	/* Sets InMeshDescription to the associated InLODIndex. */
-	ENGINE_API void CommitMeshDescription(const int32 InLODIndex, const FMeshDescription& InMeshDescription,
-		const FCommitMeshDescriptionParams& InParams = FCommitMeshDescriptionParams());
-
 	/** Get the number of imported vertices. This returns 0 if GetImportedModel() returns a nullptr.
 	  * This is the number of vertices as they appear in the source asset, for example 8 for a cube. */
 	ENGINE_API int32 GetNumImportedVertices() const;
@@ -2327,11 +2434,6 @@ private:
 	 */
 	ENGINE_API void PostLoadValidateUserSectionData();
 
-	/*
-	 * This function ensure skeletalmesh each non generated LOD has some imported data. If there is no import data it will create it from the LODModel data.
-	 */
-	ENGINE_API void PostLoadEnsureImportDataExist();
-
 	/*
 	 * This function will ensure we have valid tangent in all LODs. If we found an invalid tangent axis, we will try to set it with the cross product of the two other axis.
 	 * If the two other axes are also bad, it will simply apply the triangle normals, which will facet the mesh.
@@ -2887,37 +2989,6 @@ public:
 	/** Releases all allocated Skin Weight Profile resources, assumes none are currently in use */
 	ENGINE_API void ReleaseSkinWeightProfileResources();
 
-#if WITH_EDITORONLY_DATA
-	/*Transient data use when we postload an old asset to use legacy ddc key, it is turn off so if the user change the asset it go back to the latest ddc code*/
-	UE_DEPRECATED(4.27, "Please do not access this member directly; use USkeletalMesh::GetUseLegacyMeshDerivedDataKey() or USkeletalMesh::SetUseLegacyMeshDerivedDataKey().")
-	bool UseLegacyMeshDerivedDataKey = false;
-
-	static FName GetUseLegacyMeshDerivedDataKeyMemberName()
-	{
-		PRAGMA_DISABLE_DEPRECATION_WARNINGS
-		return GET_MEMBER_NAME_CHECKED(USkeletalMesh, UseLegacyMeshDerivedDataKey);
-		PRAGMA_ENABLE_DEPRECATION_WARNINGS
-	}
-
-	// USkinnedAsset interface
-	virtual bool GetUseLegacyMeshDerivedDataKey() const override
-	{
-		WaitUntilAsyncPropertyReleased(ESkeletalMeshAsyncProperties::UseLegacyMeshDerivedDataKey, ESkinnedAssetAsyncPropertyLockType::ReadOnly);
-		PRAGMA_DISABLE_DEPRECATION_WARNINGS
-		return UseLegacyMeshDerivedDataKey;
-		PRAGMA_ENABLE_DEPRECATION_WARNINGS
-	}
-
-	void SetUseLegacyMeshDerivedDataKey(const bool InUseLegacyMeshDerivedDataKey)
-	{
-		WaitUntilAsyncPropertyReleased(ESkeletalMeshAsyncProperties::UseLegacyMeshDerivedDataKey);
-		PRAGMA_DISABLE_DEPRECATION_WARNINGS
-		UseLegacyMeshDerivedDataKey = InUseLegacyMeshDerivedDataKey;
-		PRAGMA_ENABLE_DEPRECATION_WARNINGS
-	}
-
-#endif
-
 protected:
 	/** Set of skin weight profiles associated with this mesh */
 	UE_DEPRECATED(5.0, "This must be protected for async build, always use the accessors even internally.")
diff --git a/Engine/Source/Runtime/Engine/Classes/Engine/SkeletalMeshEditorData.h b/Engine/Source/Runtime/Engine/Classes/Engine/SkeletalMeshEditorData.h
--- a/Engine/Source/Runtime/Engine/Classes/Engine/SkeletalMeshEditorData.h
+++ b/Engine/Source/Runtime/Engine/Classes/Engine/SkeletalMeshEditorData.h
@@ -57,6 +57,8 @@ public:
 	/* End of UObject overrides */
 
 private:
+	friend class USkeletalMesh;
+	
 	/** Imported raw mesh data. Optional, only the imported mesh LOD has this, generated LOD or old asset will be empty. */
 	TArray<TSharedRef<FRawSkeletalMeshBulkData>> RawSkeletalMeshBulkDatas;
 
diff --git a/Engine/Source/Runtime/Engine/Classes/Engine/SkinnedAsset.h b/Engine/Source/Runtime/Engine/Classes/Engine/SkinnedAsset.h
--- a/Engine/Source/Runtime/Engine/Classes/Engine/SkinnedAsset.h
+++ b/Engine/Source/Runtime/Engine/Classes/Engine/SkinnedAsset.h
@@ -274,7 +274,8 @@ public:
 #endif // WITH_EDITOR
 
 #if WITH_EDITORONLY_DATA
-	virtual bool GetUseLegacyMeshDerivedDataKey() const	{ return false; }
+	UE_DEPRECATED(5.4, "This function and associated functionality is no longer supported.")
+	static bool GetUseLegacyMeshDerivedDataKey()	{ return false; }
 
 	/** Get the source mesh data. */
 	ENGINE_API virtual class FSkeletalMeshModel* GetImportedModel() const
