Skeletal Mesh: Change the import storage to use FMeshDescription instead of FSkeletalMeshImportData. This includes deprecating the USkeletalMesh::LoadLODImportedData and friends, in lieu of an API mor

diff --git a/Engine/Source/Runtime/Engine/Private/Rendering/SkeletalMeshLODImporterData.cpp b/Engine/Source/Runtime/Engine/Private/Rendering/SkeletalMeshLODImporterData.cpp
--- a/Engine/Source/Runtime/Engine/Private/Rendering/SkeletalMeshLODImporterData.cpp
+++ b/Engine/Source/Runtime/Engine/Private/Rendering/SkeletalMeshLODImporterData.cpp
@@ -2,8 +2,11 @@
 
 #include "Rendering/SkeletalMeshLODImporterData.h"
 
+#include "UObject/Package.h"
+
 #if WITH_EDITOR
 
+#include "Algo/AnyOf.h"
 #include "Engine/SkeletalMesh.h"
 #include "Engine/SkinnedAssetCommon.h"
 #include "ImportUtils/SkeletalMeshImportUtils.h"
@@ -147,7 +150,10 @@ bool FSkeletalMeshImportData::ReplaceSkeletalMeshGeometryImportData(const USkele
 
 	//Load the original skeletal mesh import data
 	FSkeletalMeshImportData OriginalSkeletalMeshImportData;
+	
+	PRAGMA_DISABLE_DEPRECATION_WARNINGS
 	SkeletalMesh->LoadLODImportedData(LodIndex, OriginalSkeletalMeshImportData);
+	PRAGMA_ENABLE_DEPRECATION_WARNINGS
 
 	//Backup the new geometry and rig to be able to apply the rig to the old geometry
 	FSkeletalMeshImportData NewGeometryAndRigData = *ImportData;
@@ -209,7 +215,9 @@ bool FSkeletalMeshImportData::ReplaceSkeletalMeshRigImportData(const USkeletalMe
 
 	//Load the original skeletal mesh import data
 	FSkeletalMeshImportData OriginalSkeletalMeshImportData;
+PRAGMA_DISABLE_DEPRECATION_WARNINGS
 	SkeletalMesh->LoadLODImportedData(LodIndex, OriginalSkeletalMeshImportData);
+PRAGMA_ENABLE_DEPRECATION_WARNINGS
 
 	ImportData->RefBonesBinary.Reset();
 	ImportData->RefBonesBinary += OriginalSkeletalMeshImportData.RefBonesBinary;
@@ -1422,6 +1430,45 @@ void FSkeletalMeshImportData::ComputeSmoothGroupFromNormals()
 	}
 }
 
+void FSkeletalMeshImportData::SetMorphTargets(
+	const TArray<TObjectPtr<UMorphTarget>>& InMorphTargets,
+	int32 InLODIndex,
+	const TArray<int32>& InVertexMap
+	)
+{
+	MorphTargets.Reset();
+	MorphTargetNames.Reset();
+	MorphTargetModifiedPoints.Reset();
+
+	for (UMorphTarget* MorphTargetSrc: InMorphTargets)
+	{
+		if (!MorphTargetSrc->HasDataForLOD(InLODIndex))
+		{
+			continue;
+		}
+
+		const FMorphTargetLODModel& MorphTargetModel = MorphTargetSrc->GetMorphLODModels()[InLODIndex];
+
+		MorphTargetNames.Add(MorphTargetSrc->GetName());
+
+		FSkeletalMeshImportData& MorphTargetDst = MorphTargets.AddDefaulted_GetRef();
+		TSet<uint32>& ModifiedPoints = MorphTargetModifiedPoints.AddDefaulted_GetRef();
+
+		for (const FMorphTargetDelta& Delta: MorphTargetModel.Vertices)
+		{
+			if (InVertexMap.IsValidIndex(Delta.SourceIdx))
+			{
+				const int32 MappedIndex = InVertexMap[Delta.SourceIdx];
+				if (Points.IsValidIndex(MappedIndex))
+				{
+					ModifiedPoints.Add(MappedIndex);
+					MorphTargetDst.Points.Add(Delta.PositionDelta + Points[MappedIndex]);
+				}
+			}
+		}
+	}
+}
+
 void FSkeletalMeshImportData::CleanUpUnusedMaterials()
 {
 	if (Materials.Num() <= 0)
@@ -1736,6 +1783,8 @@ void FSkeletalMeshImportData::SplitVerticesBySmoothingGroups()
 }
 
 static void CopySkinWeightsToAttribute(
+	const USkeletalMesh *InSkeletalMesh,
+	const FName InProfileName,
 	const TArray<SkeletalMeshImportData::FRawBoneInfluence>& InInfluences,
 	const TArray<FVertexID>& InVertexIDMap,
 	const TMap<int32, int32>* InBoneIndexMap,
@@ -1753,11 +1802,29 @@ static void CopySkinWeightsToAttribute(
 		return A.VertexIndex < B.VertexIndex;
 	});
 
-	// Do the base skin weights first. We do the alternative skin weights later, since they may require geometric remapping. 
+	// Do the base skin weights first. We do the alternate skin weights later, since they may require geometric remapping. 
 	TArray<FBoneWeight> BoneWeights;
 	for(int32 StartStride = 0, EndStride = 0; EndStride != SortedInfluences.Num(); StartStride = EndStride)
 	{
 		const int32 VertexIndex = SortedInfluences[StartStride].VertexIndex;
+
+		// There exist meshes where the influence map got auto-filled with 100% weight on root in by using the wedge count of the raw mesh,
+		// due to missing weights (e.g. static mesh imported as a skeletal mesh), and so may refer to vertices that don't exist.
+		// We just stop when we get to the broken set and ignore the rest.
+		if (VertexIndex >= InVertexIDMap.Num())
+		{
+			if (InSkeletalMesh)
+			{
+				FString ProfileDetail;
+				if (!InProfileName.IsNone())
+				{
+					ProfileDetail = FString::Printf(TEXT(" on profile '%s'"), *InProfileName.ToString());
+				}
+				UE_ASSET_LOG(LogSkeletalMeshLODImporterData, Warning, InSkeletalMesh, TEXT("Influences%s refer to non-existent vertices. Please re-import to fix."),
+					*ProfileDetail);
+			}
+			break;
+		}
 		
 		EndStride = StartStride + 1;
 		while (EndStride < SortedInfluences.Num() && VertexIndex == SortedInfluences[EndStride].VertexIndex)
@@ -1784,7 +1851,7 @@ static void CopySkinWeightsToAttribute(
 	}
 }
 
-bool FSkeletalMeshImportData::GetMeshDescription(FMeshDescription& OutMeshDescription) const
+bool FSkeletalMeshImportData::GetMeshDescription(const USkeletalMesh* InSkeletalMesh, FMeshDescription& OutMeshDescription) const
 {
 	using namespace UE::AnimationCore;
 	
@@ -1799,8 +1866,22 @@ bool FSkeletalMeshImportData::GetMeshDescription(FMeshDescription& OutMeshDescri
 		return false;
 	}
 
+	MeshAttributes.RegisterImportPointIndexAttribute(); 
+	
+	if (!MeshInfos.IsEmpty())
+	{
+		MeshAttributes.RegisterSourceGeometryPartsAttributes();
+	}
+	
 	TVertexAttributesRef<FVector3f> VertexPositions = MeshAttributes.GetVertexPositions();
 	FSkinWeightsVertexAttributesRef VertexSkinWeights = MeshAttributes.GetVertexSkinWeights();
+	
+	TVertexAttributesRef<int32> ImportPointIndex;
+	if (!PointToRawMap.IsEmpty())
+	{
+		ImportPointIndex = OutMeshDescription.VertexAttributes().GetAttributesRef<int32>(MeshAttribute::Vertex::ImportPointIndex); 
+	}
+	
 	TVertexInstanceAttributesRef<FVector3f> VertexInstanceNormals = MeshAttributes.GetVertexInstanceNormals();
 	TVertexInstanceAttributesRef<FVector3f> VertexInstanceTangents = MeshAttributes.GetVertexInstanceTangents();
 	TVertexInstanceAttributesRef<float> VertexInstanceBinormalSigns = MeshAttributes.GetVertexInstanceBinormalSigns();
@@ -1813,35 +1894,60 @@ bool FSkeletalMeshImportData::GetMeshDescription(FMeshDescription& OutMeshDescri
 	FSkeletalMeshAttributes::FBoneParentIndexAttributesRef BoneParentIndices = MeshAttributes.GetBoneParentIndices();
 	FSkeletalMeshAttributes::FBonePoseAttributesRef BonePoses = MeshAttributes.GetBonePoses();
 
-	// Register morph targets.
-	TSet<FString> ValidMorphTargets;
-	if (ensure(MorphTargetNames.Num() == MorphTargets.Num() && MorphTargetModifiedPoints.Num() == MorphTargets.Num()))
+	// Register morph targets, rename as needed, in case there are invalid names provided (e.g. empty strings, duplicates)
+	TMap<int32, FString> ValidMorphTargets;
+	if (ensure(MorphTargetNames.Num() == MorphTargets.Num()) &&
+		ensure(MorphTargets.Num() == MorphTargetModifiedPoints.Num()))
 	{
-		for (int32 MorphIndex = 0; MorphIndex < MorphTargets.Num(); MorphIndex++)
+		TSet<FString> CheckDuplicateSet;
+		
+		for (int32 MorphTargetIndex = 0; MorphTargetIndex < MorphTargets.Num(); MorphTargetIndex++)
 		{
-			const FSkeletalMeshImportData& MorphTargetMesh = MorphTargets[MorphIndex];
-			const TSet<uint32>& ModifiedPoints = MorphTargetModifiedPoints[MorphIndex];
-
-			// ensure they both have the same number of points
-			if (!ensure(MorphTargetMesh.Points.Num() == ModifiedPoints.Num()))
+			const FString MorphTargetName(MorphTargetNames[MorphTargetIndex]);
+			FString MorphTargetNameUsed(MorphTargetName);
+			const TArray<FVector3f>& MorphPoints = MorphTargets[MorphTargetIndex].Points;
+			const TSet<uint32>& ModifiedPoints = MorphTargetModifiedPoints[MorphTargetIndex];
+			
+			if (MorphPoints.Num() != ModifiedPoints.Num())
 			{
+				if (InSkeletalMesh)
+				{
+					UE_ASSET_LOG(LogSkeletalMeshLODImporterData, Warning, InSkeletalMesh, TEXT("Morph target '%s' has mismatched compressed and modified points. Morph target will be dropped. Re-import to fix."),
+						*MorphTargetName);  
+				}
 				continue;
 			}
 
-			// ensure modified points indices are valid
-			const uint32* InvalidIndex = Algo::FindByPredicate(ModifiedPoints, [this](uint32 PointIndex)
-			{
-				return !Points.IsValidIndex(PointIndex);
-			});
-			if (!ensure(InvalidIndex == nullptr))
+			// Ensure modified points indices are valid
+			if (Algo::AnyOf(ModifiedPoints, [this](uint32 PointIndex) { return !Points.IsValidIndex(PointIndex); }))
 			{
+				UE_ASSET_LOG(LogSkeletalMeshLODImporterData, Warning, InSkeletalMesh, TEXT("Morph target '%s' has modified points that don't exist on the base mesh. Morph target will be dropped. Re-import to fix."),
+					*MorphTargetName);  
 				continue;
 			}
 
-			const FString MorphTargetName(MorphTargetNames[MorphIndex]);
-			if (ensure(MeshAttributes.RegisterMorphTargetAttribute(*MorphTargetName)))
+			if (MorphTargetNameUsed.IsEmpty())
+			{
+				MorphTargetNameUsed = TEXT("Unnamed");
+				UE_ASSET_LOG(LogSkeletalMeshLODImporterData, Warning, InSkeletalMesh, TEXT("Morph target found with no name. Renamed to 'Unnamed'."));  
+			}
+
+			// Has this name been used already?
+			FString MorphTargetNameBase(MorphTargetNameUsed);
+			for(int32 Suffix = 1; CheckDuplicateSet.Contains(MorphTargetNameUsed); Suffix++)
+			{
+				MorphTargetNameUsed = FString::Printf(TEXT("%s_%d"), *MorphTargetNameBase, Suffix);
+			}
+			if (MorphTargetNameUsed != MorphTargetNameBase)
+			{
+				UE_ASSET_LOG(LogSkeletalMeshLODImporterData, Warning, InSkeletalMesh, TEXT("Duplicate morph target '%s' found, renamed to '%s'."),
+					*MorphTargetNameBase, *MorphTargetNameUsed);  
+			}
+			
+			if (ensure(MeshAttributes.RegisterMorphTargetAttribute(*MorphTargetNameUsed)))
 			{
-				ValidMorphTargets.Add(MorphTargetName);
+				ValidMorphTargets.Add(MorphTargetIndex, MorphTargetNameUsed);
+				CheckDuplicateSet.Add(MorphTargetNameUsed);
 			}
 		}
 	}
@@ -1850,37 +1956,68 @@ bool FSkeletalMeshImportData::GetMeshDescription(FMeshDescription& OutMeshDescri
 	TSet<FString> ValidSkinWeights;
 	if (ensure(AlternateInfluenceProfileNames.Num() == AlternateInfluences.Num()))
 	{
-		for (int32 AlternativeInfluenceIndex = 0; AlternativeInfluenceIndex < AlternateInfluences.Num(); AlternativeInfluenceIndex++)
+		for (int32 AlternateInfluenceIndex = 0; AlternateInfluenceIndex < AlternateInfluences.Num(); AlternateInfluenceIndex++)
 		{
-			const FSkeletalMeshImportData& AlternativeInfluenceMesh = AlternateInfluences[AlternativeInfluenceIndex];
+			const FString AlternateInfluenceProfileName(AlternateInfluenceProfileNames[AlternateInfluenceIndex]);
+			const FSkeletalMeshImportData& AlternateInfluenceMesh = AlternateInfluences[AlternateInfluenceIndex];
 			
 			// We minimally need the same point count.
 			// Q: Should we match on topology as well?
-			if (!ensure(AlternativeInfluenceMesh.Points.Num() == Points.Num()))
+			if (AlternateInfluenceMesh.Points.Num() != Points.Num())
 			{
+				if (InSkeletalMesh)
+				{
+					UE_ASSET_LOG(LogSkeletalMeshLODImporterData, Warning, InSkeletalMesh, TEXT("Alternate influence profile '%s' point count is different from base mesh. Profile will be dropped since it cannot be matched to the base mesh."),
+						*AlternateInfluenceProfileName);  
+				}
 				continue;
 			}
 
-			// The set of bones on the alternative influence mesh must be a fully included in the set of the bones on the
-			// base mesh. Otherwise we cannot map an influence across.
-			TSet<FString> BaseBones, AltBones;
+			// The set of bound bones must completely include the bones in the base mesh. Otherwise we cannot map the influences over to
+			// the base skeleton.
+			TSet<FString> BaseBones;
+			TSet<int32> MissingBoneIndices;
 			for (int32 BoneIndex = 0; BoneIndex < RefBonesBinary.Num(); BoneIndex++)
 			{
 				BaseBones.Add(RefBonesBinary[BoneIndex].Name);
 			}
-			for (int32 BoneIndex = 0; BoneIndex < AlternativeInfluenceMesh.RefBonesBinary.Num(); BoneIndex++)
+			for (int32 BoneIndex = 0; BoneIndex < AlternateInfluenceMesh.RefBonesBinary.Num(); BoneIndex++)
 			{
-				AltBones.Add(AlternativeInfluenceMesh.RefBonesBinary[BoneIndex].Name);
+				if (!BaseBones.Contains(AlternateInfluenceMesh.RefBonesBinary[BoneIndex].Name))
+				{
+					MissingBoneIndices.Add(BoneIndex);
+				}
 			}
 
-			if (!ensure(BaseBones.Includes(AltBones)))
+			TSet<int32> InfluencesBoundToMissingBones; 
+			for (const SkeletalMeshImportData::FRawBoneInfluence& Influence: AlternateInfluenceMesh.Influences)
 			{
+				if (MissingBoneIndices.Contains(Influence.BoneIndex))
+				{
+					InfluencesBoundToMissingBones.Add(Influence.BoneIndex);
+				}
+			}
+
+			if (!InfluencesBoundToMissingBones.IsEmpty())
+			{
+				TArray<FString> MissingBoneNames;
+				for (const int32 BoneIndex: InfluencesBoundToMissingBones)
+				{
+					MissingBoneNames.Add(FString::Printf(TEXT("'%s'"), *AlternateInfluenceMesh.RefBonesBinary[BoneIndex].Name));
+				}
+				MissingBoneNames.Sort();
+				
+				if (InSkeletalMesh)
+				{
+					UE_ASSET_LOG(LogSkeletalMeshLODImporterData, Warning, InSkeletalMesh, TEXT("Alternate influence profile '%s' binds to one or more bones (%s) that do not exist on base mesh's skeleton. Profile will be dropped since it may otherwise introduce visual errors."),
+						*AlternateInfluenceProfileName, *FString::Join(MissingBoneNames, TEXT(", ")));  
+				}
 				continue;
 			}
 
-			if (ensure(MeshAttributes.RegisterSkinWeightAttribute(*AlternateInfluenceProfileNames[AlternativeInfluenceIndex])))
+			if (ensure(MeshAttributes.RegisterSkinWeightAttribute(*AlternateInfluenceProfileName)))
 			{
-				ValidSkinWeights.Add(AlternateInfluenceProfileNames[AlternativeInfluenceIndex]);
+				ValidSkinWeights.Add(AlternateInfluenceProfileNames[AlternateInfluenceIndex]);
 			}
 
 		}
@@ -1952,9 +2089,14 @@ bool FSkeletalMeshImportData::GetMeshDescription(FMeshDescription& OutMeshDescri
 		const FVertexID VertexID = OutMeshDescription.CreateVertex();
 		VertexIDMap.Add(VertexID);
 		VertexPositions.Set(VertexID, Points[Idx]);
+
+		if (ImportPointIndex.IsValid())
+		{
+			ImportPointIndex.Set(VertexID, PointToRawMap[Idx]);
+		}
 	}
 
-	CopySkinWeightsToAttribute(Influences, VertexIDMap, nullptr, VertexSkinWeights);
+	CopySkinWeightsToAttribute(InSkeletalMesh, NAME_None, Influences, VertexIDMap, nullptr, VertexSkinWeights);
 
 	// Set Bone Attributes
 	for (int Idx = 0; Idx < RefBonesBinary.Num(); ++Idx)
@@ -2075,20 +2217,23 @@ bool FSkeletalMeshImportData::GetMeshDescription(FMeshDescription& OutMeshDescri
 	// Convert morph targets
 	for (int32 MorphTargetIndex = 0; MorphTargetIndex < MorphTargets.Num(); MorphTargetIndex++)
 	{
-		FString MorphTargetName(MorphTargetNames[MorphTargetIndex]);
-		if (!ValidMorphTargets.Contains(MorphTargetName))
+		if (!ValidMorphTargets.Contains(MorphTargetIndex))
 		{
 			continue;
 		}
 
+		FString MorphTargetName(ValidMorphTargets[MorphTargetIndex]);
 		const TArray<FVector3f>& MorphPoints = MorphTargets[MorphTargetIndex].Points;
-		const TSet<uint32>& ModifiedPoints = MorphTargetModifiedPoints[MorphTargetIndex]; 
+		const TSet<uint32>& ModifiedPoints = MorphTargetModifiedPoints[MorphTargetIndex];
+
 		FMorphTargetVertexAttributesRef MorphTargetRef = MeshAttributes.GetVertexMorphTarget(*MorphTargetName);
 
 		int32 ModifiedPointIndex = 0;
+		// This relies on the dubious assumptions that TSet values are retained in insertion order.
 		for (uint32 PointIndex: ModifiedPoints)
 		{
 			const FVertexID VertexID = VertexIDMap[PointIndex];
+			
 			MorphTargetRef.SetPositionDelta(VertexID, MorphPoints[ModifiedPointIndex] - Points[PointIndex]);
 			ModifiedPointIndex++;
 		}
@@ -2103,7 +2248,7 @@ bool FSkeletalMeshImportData::GetMeshDescription(FMeshDescription& OutMeshDescri
 			continue;
 		}
 
-		CopySkinWeightsToMeshDescription(FName(AlternateInfluenceProfileName), AlternateInfluences[AlternateInfluenceIndex], VertexIDMap, OutMeshDescription);
+		CopySkinWeightsToMeshDescription(InSkeletalMesh, FName(AlternateInfluenceProfileName), AlternateInfluences[AlternateInfluenceIndex], VertexIDMap, OutMeshDescription);
 	}
 	
 
@@ -2168,12 +2313,30 @@ bool FSkeletalMeshImportData::GetMeshDescription(FMeshDescription& OutMeshDescri
 	}
 
 	FSkeletalMeshOperations::ConvertSmoothGroupToHardEdges(FaceSmoothingMasks, OutMeshDescription);
+
+	// Convert the MeshInfo data.
+	if (!MeshInfos.IsEmpty())
+	{
+		FSkeletalMeshAttributes::FSourceGeometryPartNameRef NameAttribute = MeshAttributes.GetSourceGeometryPartNames();
+		FSkeletalMeshAttributes::FSourceGeometryPartVertexOffsetAndCountRef VertexAndCountAttribute = MeshAttributes.GetSourceGeometryPartVertexOffsetAndCounts();
+
+		MeshAttributes.SourceGeometryParts().Reserve(MeshInfos.Num());
+		
+		for (const SkeletalMeshImportData::FMeshInfo& Info: MeshInfos)
+		{
+			FSourceGeometryPartID PartID = MeshAttributes.CreateSourceGeometryPart();
+
+			NameAttribute.Set(PartID, Info.Name);
+			VertexAndCountAttribute.Set(PartID, {Info.StartImportedVertex, Info.NumVertices});
+		}
+	}
 	
 	return true;
 }
 
 
 void FSkeletalMeshImportData::CopySkinWeightsToMeshDescription(
+	const USkeletalMesh* InSkeletalMesh,
 	const FName InSkinWeightName, 
 	const FSkeletalMeshImportData& InSkinWeightMesh,
 	const TArray<FVertexID>& InVertexIDMap,
@@ -2199,8 +2362,12 @@ void FSkeletalMeshImportData::CopySkinWeightsToMeshDescription(
 	TMap<int32, int32> AltMeshBoneToBaseBoneMap;
 	for (int32 BoneIndex = 0; BoneIndex < InSkinWeightMesh.RefBonesBinary.Num(); BoneIndex++)
 	{
-		const int32 BaseBoneIndex = BaseBoneToIndexMap[InSkinWeightMesh.RefBonesBinary[BoneIndex].Name];
-		AltMeshBoneToBaseBoneMap.Add(BoneIndex, BaseBoneIndex);
+		// It's ok if a matching bone isn't in the ref map, as long as it's not bound to the skin. This would have been checked
+		// in GetMeshDescription already.
+		if (const int32* BaseBoneIndex = BaseBoneToIndexMap.Find(InSkinWeightMesh.RefBonesBinary[BoneIndex].Name))
+		{
+			AltMeshBoneToBaseBoneMap.Add(BoneIndex, *BaseBoneIndex);
+		}
 	}
 	
 	bool bTopologySame = true;
@@ -2240,13 +2407,13 @@ void FSkeletalMeshImportData::CopySkinWeightsToMeshDescription(
 
 	if (bTopologySame)
 	{
-		CopySkinWeightsToAttribute(InSkinWeightMesh.Influences, InVertexIDMap, &AltMeshBoneToBaseBoneMap, VertexSkinWeights);
+		CopySkinWeightsToAttribute(InSkeletalMesh, InSkinWeightName, InSkinWeightMesh.Influences, InVertexIDMap, &AltMeshBoneToBaseBoneMap, VertexSkinWeights);
 		return;
 	}
 
 	// The topologies don't match, proceed as above.
 	FMeshDescription AlternateInfluenceMesh;
-	InSkinWeightMesh.GetMeshDescription(AlternateInfluenceMesh);
+	InSkinWeightMesh.GetMeshDescription(nullptr, AlternateInfluenceMesh);
 	
 	FSkeletalMeshOperations::CopySkinWeightAttributeFromMesh(
 	AlternateInfluenceMesh, OutMeshDescription, NAME_None, InSkinWeightName, &AltMeshBoneToBaseBoneMap);   
@@ -2355,17 +2522,19 @@ FSkeletalMeshImportData FSkeletalMeshImportData::CreateFromMeshDescription(const
 {
 	FSkeletalMeshImportData SkelMeshImportData;
 	
-	FSkeletalMeshConstAttributes Attributes(InMeshDescription);
-	TVertexAttributesConstRef<FVector3f> VertexPositions = Attributes.GetVertexPositions();
-	FSkinWeightsVertexAttributesConstRef VertexSkinWeights = Attributes.GetVertexSkinWeights();
+	FSkeletalMeshConstAttributes MeshAttributes(InMeshDescription);
+	TVertexAttributesConstRef<FVector3f> VertexPositions = MeshAttributes.GetVertexPositions();
+	FSkinWeightsVertexAttributesConstRef VertexSkinWeights = MeshAttributes.GetVertexSkinWeights();
 
-	TVertexInstanceAttributesConstRef<FVector2f> VertexInstanceUVs = Attributes.GetVertexInstanceUVs();
-	TVertexInstanceAttributesConstRef<FVector3f> VertexInstanceNormals = Attributes.GetVertexInstanceNormals();
-	TVertexInstanceAttributesConstRef<FVector3f> VertexInstanceTangents = Attributes.GetVertexInstanceTangents();
-	TVertexInstanceAttributesConstRef<float> VertexInstanceBiNormalSigns = Attributes.GetVertexInstanceBinormalSigns();
-	TVertexInstanceAttributesConstRef<FVector4f> VertexInstanceColors = Attributes.GetVertexInstanceColors();
+	TVertexAttributesConstRef<int32> ImportPointIndex = InMeshDescription.VertexAttributes().GetAttributesRef<int32>(MeshAttribute::Vertex::ImportPointIndex); 
+		
+	TVertexInstanceAttributesConstRef<FVector2f> VertexInstanceUVs = MeshAttributes.GetVertexInstanceUVs();
+	TVertexInstanceAttributesConstRef<FVector3f> VertexInstanceNormals = MeshAttributes.GetVertexInstanceNormals();
+	TVertexInstanceAttributesConstRef<FVector3f> VertexInstanceTangents = MeshAttributes.GetVertexInstanceTangents();
+	TVertexInstanceAttributesConstRef<float> VertexInstanceBiNormalSigns = MeshAttributes.GetVertexInstanceBinormalSigns();
+	TVertexInstanceAttributesConstRef<FVector4f> VertexInstanceColors = MeshAttributes.GetVertexInstanceColors();
 
-	TPolygonGroupAttributesConstRef<FName> PolygonGroupMaterialSlotNames = Attributes.GetPolygonGroupMaterialSlotNames();
+	TPolygonGroupAttributesConstRef<FName> PolygonGroupMaterialSlotNames = MeshAttributes.GetPolygonGroupMaterialSlotNames();
 
 	//Get the per face smoothing
 	TArray<uint32> FaceSmoothingMasks;
@@ -2388,12 +2557,13 @@ FSkeletalMeshImportData FSkeletalMeshImportData::CreateFromMeshDescription(const
 	//Copy the vertex positions and the influences
 
 	//Reserve the point and influences
-	SkelMeshImportData.Points.AddZeroed(InMeshDescription.Vertices().Num());
+	SkelMeshImportData.Points.SetNumUninitialized(InMeshDescription.Vertices().Num());
 	SkelMeshImportData.Influences.Reserve(InMeshDescription.Vertices().Num() * 4);
+	SkelMeshImportData.PointToRawMap.SetNumUninitialized(InMeshDescription.Vertices().Num());
 	
 	for (FVertexID VertexID : InMeshDescription.Vertices().GetElementIDs())
 	{
-		//We can use GetValue because the Meshdescription was compacted before the copy
+		// We can use GetValue here because the mesh description was compacted before the copy
 		SkelMeshImportData.Points[VertexID.GetValue()] = VertexPositions[VertexID];
 		if (VertexSkinWeights.IsValid())
 		{
@@ -2410,6 +2580,15 @@ FSkeletalMeshImportData FSkeletalMeshImportData::CreateFromMeshDescription(const
 				BoneInfluence.Weight = BoneWeights[InfluenceIndex].GetWeight();
 			}
 		}
+
+		if (ImportPointIndex.IsValid())
+		{
+			SkelMeshImportData.PointToRawMap[VertexID.GetValue()] = ImportPointIndex[VertexID] == INDEX_NONE ? VertexID.GetValue() : ImportPointIndex[VertexID];
+		}
+		else
+		{
+			SkelMeshImportData.PointToRawMap[VertexID.GetValue()] = VertexID.GetValue();
+		}
 	}
 
 	//////////////////////////////////////////////////////////////////////////
@@ -2470,12 +2649,12 @@ FSkeletalMeshImportData FSkeletalMeshImportData::CreateFromMeshDescription(const
 	}
 
 	// Update Bones data
-	const FSkeletalMeshAttributes::FBoneArray& Bones = Attributes.Bones();
+	const FSkeletalMeshAttributes::FBoneArray& Bones = MeshAttributes.Bones();
 	SkelMeshImportData.RefBonesBinary.Reserve(Bones.Num());
 	
-	FSkeletalMeshAttributes::FBoneNameAttributesConstRef BoneNames = Attributes.GetBoneNames();
-	FSkeletalMeshAttributes::FBoneParentIndexAttributesConstRef BoneParentIndices = Attributes.GetBoneParentIndices();
-	FSkeletalMeshAttributes::FBonePoseAttributesConstRef BonePoses = Attributes.GetBonePoses();
+	FSkeletalMeshAttributes::FBoneNameAttributesConstRef BoneNames = MeshAttributes.GetBoneNames();
+	FSkeletalMeshAttributes::FBoneParentIndexAttributesConstRef BoneParentIndices = MeshAttributes.GetBoneParentIndices();
+	FSkeletalMeshAttributes::FBonePoseAttributesConstRef BonePoses = MeshAttributes.GetBonePoses();
 	if (BoneNames.IsValid() && BoneParentIndices.IsValid() && BonePoses.IsValid())
 	{
 		auto GetNumChildren = [&](const int32 ParentIndex)
@@ -2505,28 +2684,27 @@ FSkeletalMeshImportData FSkeletalMeshImportData::CreateFromMeshDescription(const
 	}
 
 	// Copy morph targets.
-	for (const FName MorphTargetName: Attributes.GetMorphTargetNames())
+	for (const FName MorphTargetName: MeshAttributes.GetMorphTargetNames())
 	{
-		FMorphTargetVertexAttributesConstRef MorphTargetAttribute = Attributes.GetVertexMorphTarget(MorphTargetName);
+		FMorphTargetVertexAttributesConstRef MorphTargetAttribute = MeshAttributes.GetVertexMorphTarget(MorphTargetName);
 
 		SkelMeshImportData.MorphTargetNames.Add(MorphTargetName.ToString());
 		FSkeletalMeshImportData& MorphTarget = SkelMeshImportData.MorphTargets.AddDefaulted_GetRef();
-		TSet<uint32> MorphTargetModifiedPoints = SkelMeshImportData.MorphTargetModifiedPoints.AddDefaulted_GetRef();
+		TSet<uint32>& MorphTargetModifiedPoints = SkelMeshImportData.MorphTargetModifiedPoints.AddDefaulted_GetRef();
 
-		MorphTarget.Points.SetNumZeroed(SkelMeshImportData.Points.Num());
 		for (FVertexID VertexID : InMeshDescription.Vertices().GetElementIDs())
 		{
 			FVector3f PositionDelta = MorphTargetAttribute.GetPositionDelta(VertexID);
 			if (!PositionDelta.IsNearlyZero())
 			{
-				MorphTarget.Points[VertexID.GetValue()] = MorphTargetAttribute.GetPositionDelta(VertexID) + SkelMeshImportData.Points[VertexID.GetValue()];
+				MorphTarget.Points.Add(PositionDelta + SkelMeshImportData.Points[VertexID.GetValue()]);
 				MorphTargetModifiedPoints.Add(VertexID.GetValue());
 			}
 		}
 	}
 
 	// Copy alternate influences.
-	for (const FName SkinWeightProfileName: Attributes.GetSkinWeightProfileNames())
+	for (const FName SkinWeightProfileName: MeshAttributes.GetSkinWeightProfileNames())
 	{
 		// The default profile was already handled above.
 		if (SkinWeightProfileName == FSkeletalMeshAttributes::DefaultSkinWeightProfileName)
@@ -2534,7 +2712,7 @@ FSkeletalMeshImportData FSkeletalMeshImportData::CreateFromMeshDescription(const
 			continue;
 		}
 		
-		FSkinWeightsVertexAttributesConstRef SkinWeightsAttribute = Attributes.GetVertexSkinWeights(SkinWeightProfileName);
+		FSkinWeightsVertexAttributesConstRef SkinWeightsAttribute = MeshAttributes.GetVertexSkinWeights(SkinWeightProfileName);
 
 		SkelMeshImportData.AlternateInfluenceProfileNames.Add(SkinWeightProfileName.ToString());
 		FSkeletalMeshImportData& AlternateInfluenceMesh = SkelMeshImportData.AlternateInfluences.AddDefaulted_GetRef();
@@ -2562,9 +2740,7 @@ FSkeletalMeshImportData FSkeletalMeshImportData::CreateFromMeshDescription(const
 		}
 	}
 	
-
 	SkelMeshImportData.CleanUpUnusedMaterials();
-	SkelMeshImportData.SplitVerticesBySmoothingGroups();
 	
 	// Copy any non-reserved float vertex attributes.
 	InMeshDescription.VertexAttributes().ForEachByType<float>(FCreateAndCopyAttributeValues<float>(SkelMeshImportData, InMeshDescription.Vertices()));
@@ -2572,6 +2748,54 @@ FSkeletalMeshImportData FSkeletalMeshImportData::CreateFromMeshDescription(const
 	InMeshDescription.VertexAttributes().ForEachByType<FVector3f>(FCreateAndCopyAttributeValues<FVector3f>(SkelMeshImportData, InMeshDescription.Vertices()));
 	InMeshDescription.VertexAttributes().ForEachByType<FVector4f>(FCreateAndCopyAttributeValues<FVector4f>(SkelMeshImportData, InMeshDescription.Vertices()));
 
+	// Copy MeshInfos back in, if any, and only if they're valid.
+	if (MeshAttributes.HasSourceGeometryParts())
+	{
+		FSkeletalMeshAttributes::FSourceGeometryPartNameConstRef NameAttribute = MeshAttributes.GetSourceGeometryPartNames();
+		FSkeletalMeshAttributes::FSourceGeometryPartVertexOffsetAndCountConstRef VertexAndCountAttribute = MeshAttributes.GetSourceGeometryPartVertexOffsetAndCounts();
+
+		// Ensure that the counts + offsets add up to exactly the vertices we have.
+		TArray<SkeletalMeshImportData::FMeshInfo> MeshInfos;
+		for (FSourceGeometryPartID SourceGeometryPartID: MeshAttributes.SourceGeometryParts().GetElementIDs())
+		{
+			SkeletalMeshImportData::FMeshInfo Info;
+			Info.Name = NameAttribute.Get(SourceGeometryPartID);
+			Info.NumVertices = VertexAndCountAttribute.Get(SourceGeometryPartID)[1];
+			Info.StartImportedVertex = VertexAndCountAttribute.Get(SourceGeometryPartID)[0]; 
+			MeshInfos.Add(Info);
+		}
+
+		if (!MeshInfos.IsEmpty())
+		{
+			MeshInfos.Sort([](const SkeletalMeshImportData::FMeshInfo& A, const SkeletalMeshImportData::FMeshInfo& B)
+			{
+				return A.StartImportedVertex < B.StartImportedVertex;
+			});
+
+			bool bValid = true;
+			int32 VertexIndex = 0;
+			for (int32 Index = 0; Index < MeshInfos.Num(); Index++)
+			{
+				if (VertexIndex != MeshInfos[Index].StartImportedVertex)
+				{
+					bValid = false;
+					break;
+				}
+
+				VertexIndex += MeshInfos[Index].NumVertices;
+			}
+			if (VertexIndex != SkelMeshImportData.Points.Num())
+			{
+				bValid = false;
+			}
+
+			if (bValid)
+			{
+				SkelMeshImportData.MeshInfos = MoveTemp(MeshInfos);
+			}
+		}
+	}
+
 	return SkelMeshImportData;
 }
 
diff --git a/Engine/Source/Runtime/Engine/Private/SkeletalMesh.cpp b/Engine/Source/Runtime/Engine/Private/SkeletalMesh.cpp
--- a/Engine/Source/Runtime/Engine/Private/SkeletalMesh.cpp
+++ b/Engine/Source/Runtime/Engine/Private/SkeletalMesh.cpp
@@ -34,7 +34,6 @@
 #include "RenderUtils.h"
 #include "UObject/UE5MainStreamObjectVersion.h"
 #include "SceneInterface.h"
-#include "UObject/UE5MainStreamObjectVersion.h"
 #include "EngineUtils.h"
 #include "EditorSupportDelegates.h"
 #include "SkeletalMeshLegacyCustomVersions.h"
@@ -93,8 +92,11 @@
 #endif
 #include "Animation/SkinWeightProfileManager.h"
 #include "BoneWeights.h"
+#include "SkeletalMeshAttributes.h"
 #include "Logging/StructuredLog.h"
 
+UE_DISABLE_OPTIMIZATION
+
 #define LOCTEXT_NAMESPACE "SkeltalMesh"
 
 DEFINE_LOG_CATEGORY(LogSkeletalMesh);
@@ -407,6 +409,8 @@ USkeletalMesh::USkeletalMesh(FVTableHelper& Helper)
 
 PRAGMA_DISABLE_DEPRECATION_WARNINGS
 USkeletalMesh::~USkeletalMesh() = default;
+
+
 PRAGMA_ENABLE_DEPRECATION_WARNINGS
 
 FSkeletalMeshRenderData* USkeletalMesh::GetSkeletalMeshRenderData() const
@@ -1580,6 +1584,17 @@ void USkeletalMesh::Serialize( FArchive& Ar )
 		// async compilation and the serialization.
 		FSkinnedAssetCompilingManager::Get().FinishCompilation({ this });
 	}
+
+	// If saving out to disk, ensure that all source models have had their raw mesh bulk data converted to mesh description,
+	// since it won't be reloaded.
+	if (Ar.IsSaving())
+	{
+		for (int32 LODIndex = 0; LODIndex < GetLODNum(); LODIndex++)
+		{
+			GetSourceModel(LODIndex).EnsureRawMeshBulkDataIsConvertedToNew();
+		}
+	}
+	
 #endif
 
 	Super::Serialize(Ar);
@@ -2280,184 +2295,312 @@ void USkeletalMesh::RemoveLegacyClothingSections()
 	}
 }
 
-USkeletalMeshEditorData& USkeletalMesh::GetMeshEditorData() const
+///////////////////////////////////////////////////////////////////////
+//// Source Model API
+
+int32 USkeletalMesh::GetNumSourceModels() const
+{
+	WaitUntilAsyncPropertyReleased(ESkeletalMeshAsyncProperties::SourceModels);
+	PRAGMA_DISABLE_DEPRECATION_WARNINGS
+	return SourceModels.Num();
+	PRAGMA_ENABLE_DEPRECATION_WARNINGS
+}
+
+void USkeletalMesh::SetNumSourceModels(const int32 InNumSourceModels)
+{
+	WaitUntilAsyncPropertyReleased(ESkeletalMeshAsyncProperties::SourceModels);
+	SetNumSourceModelsInternal(InNumSourceModels);
+}
+
+void USkeletalMesh::SetNumSourceModelsInternal(const int32 InNumSourceModels)
 {
 	PRAGMA_DISABLE_DEPRECATION_WARNINGS
-	if (MeshEditorDataObject == nullptr)
+	const int32 OldNumSourceModels = SourceModels.Num();
+
+	// If we're lowering the LOD count, clear out any mesh data stored immediately rather
+	// than wait for GC to do the work.
+	if (OldNumSourceModels > InNumSourceModels)
+	{
+		for (int32 Index = InNumSourceModels; Index < OldNumSourceModels; Index++)
+		{
+			FSkeletalMeshSourceModel& SourceModel = SourceModels[Index];
+			SourceModel.ClearAllMeshData();
+		}
+	}
+
+	SourceModels.SetNum(InNumSourceModels);
+
+	for (int32 Index = OldNumSourceModels; Index < InNumSourceModels; Index++)
 	{
-			//The asset is created in the skeletalmesh package. We keep it private so the user cannot see it in the content browser
-		//RF_Transactional make sure the asset can be transactional if we want to edit it
-		USkeletalMesh* NonConstSkeletalMesh = const_cast<USkeletalMesh*>(this);
-		MeshEditorDataObject = NewObject<USkeletalMeshEditorData>(NonConstSkeletalMesh, NAME_None, RF_Transactional);
+		SourceModels[Index].Initialize(this);
 	}
-	//Make sure we have a valid pointer
-	check(MeshEditorDataObject != nullptr);
-	return *MeshEditorDataObject;
 	PRAGMA_ENABLE_DEPRECATION_WARNINGS
 }
 
-void USkeletalMesh::LoadLODImportedData(const int32 LODIndex, FSkeletalMeshImportData& OutMesh) const
+void USkeletalMesh::RemoveSourceModel(const int32 InLODIndex)
+{
+	WaitUntilAsyncPropertyReleased(ESkeletalMeshAsyncProperties::SourceModels);
+	PRAGMA_DISABLE_DEPRECATION_WARNINGS
+	RemoveSourceModelInternal(InLODIndex);	
+	PRAGMA_ENABLE_DEPRECATION_WARNINGS
+}
+
+void USkeletalMesh::RemoveSourceModelInternal(const int32 InLODIndex)
+{
+	PRAGMA_DISABLE_DEPRECATION_WARNINGS
+	if (ensure(SourceModels.IsValidIndex(InLODIndex)))
+	{
+		FSkeletalMeshSourceModel& SourceModel = SourceModels[InLODIndex];
+
+		SourceModel.ClearAllMeshData();
+		SourceModels.RemoveAt(InLODIndex);
+	}
+	PRAGMA_ENABLE_DEPRECATION_WARNINGS
+}
+
+
+const FSkeletalMeshSourceModel& USkeletalMesh::GetSourceModel(const int32 InLODIndex) const
 {
-	GetMeshEditorData().GetLODImportedData(LODIndex).LoadRawMesh(OutMesh);
+	WaitUntilAsyncPropertyReleased(ESkeletalMeshAsyncProperties::SourceModels);
+	PRAGMA_DISABLE_DEPRECATION_WARNINGS
+	check(SourceModels.IsValidIndex(InLODIndex));
+	return SourceModels[InLODIndex];
+	PRAGMA_ENABLE_DEPRECATION_WARNINGS
 }
 
-void USkeletalMesh::SaveLODImportedData(const int32 LODIndex, FSkeletalMeshImportData& InMesh)
+FSkeletalMeshSourceModel& USkeletalMesh::GetSourceModel(const int32 InLODIndex)
 {
-	FRawSkeletalMeshBulkData& RawSkeletalMeshBulkData = GetMeshEditorData().GetLODImportedData(LODIndex);
-	RawSkeletalMeshBulkData.SaveRawMesh(InMesh);
-	//Update the cache
-	check(GetImportedModel()->LODModels.IsValidIndex(LODIndex));
-	GetImportedModel()->LODModels[LODIndex].RawSkeletalMeshBulkDataID = RawSkeletalMeshBulkData.GetIdString();
-	GetImportedModel()->LODModels[LODIndex].bIsBuildDataAvailable = RawSkeletalMeshBulkData.IsBuildDataAvailable();
-	GetImportedModel()->LODModels[LODIndex].bIsRawSkeletalMeshBulkDataEmpty = RawSkeletalMeshBulkData.IsEmpty();
+	WaitUntilAsyncPropertyReleased(ESkeletalMeshAsyncProperties::SourceModels);
+	PRAGMA_DISABLE_DEPRECATION_WARNINGS
+	check(SourceModels.IsValidIndex(InLODIndex));
+	return SourceModels[InLODIndex];
+	PRAGMA_ENABLE_DEPRECATION_WARNINGS
+}
+
 
-	FSkeletalMeshLODModel& LODModel = GetImportedModel()->LODModels[LODIndex];
-	for (const SkeletalMeshImportData::FMeshInfo& MeshInfo : InMesh.MeshInfos)
+///////////////////////////////////////////////////////////////////////
+//// Mesh Description API
+
+FMeshDescription* USkeletalMesh::GetMeshDescription(const int32 InLODIndex) const
+{
+	if (!ensure(IsValidLODIndex(InLODIndex)))
 	{
-		LODModel.ImportedMeshInfos.AddDefaulted();
-		FSkelMeshImportedMeshInfo& LODMeshInfo = LODModel.ImportedMeshInfos.Last();
-		LODMeshInfo.Name = MeshInfo.Name;
-		LODMeshInfo.NumVertices = MeshInfo.NumVertices;
-		LODMeshInfo.StartImportedVertex = MeshInfo.StartImportedVertex;
+		return nullptr;
 	}
+
+	return GetSourceModel(InLODIndex).GetMeshDescription();
 }
 
-bool USkeletalMesh::IsLODImportedDataBuildAvailable(const int32 LODIndex) const
+
+bool USkeletalMesh::CloneMeshDescription(const int32 InLODIndex, FMeshDescription& OutMeshDescription) const
 {
-	if (!GetImportedModel()->LODModels.IsValidIndex(LODIndex))
+	if (!ensure(IsValidLODIndex(InLODIndex)))
 	{
 		return false;
 	}
-	return GetImportedModel()->LODModels[LODIndex].bIsBuildDataAvailable;
+
+	return GetSourceModel(InLODIndex).CloneMeshDescription(OutMeshDescription);
 }
 
-bool USkeletalMesh::IsLODImportedDataEmpty(const int32 LODIndex) const
+
+bool USkeletalMesh::HasMeshDescription(const int32 InLODIndex) const
 {
-	if (!GetImportedModel()->LODModels.IsValidIndex(LODIndex))
+	if (!ensure(IsValidLODIndex(InLODIndex)))
 	{
 		return false;
 	}
-	return GetImportedModel()->LODModels[LODIndex].bIsRawSkeletalMeshBulkDataEmpty;
+
+	return GetSourceModel(InLODIndex).HasMeshDescription();
 }
 
-void USkeletalMesh::GetLODImportedDataVersions(const int32 LODIndex, ESkeletalMeshGeoImportVersions& OutGeoImportVersion, ESkeletalMeshSkinningImportVersions& OutSkinningImportVersion) const
+
+FMeshDescription* USkeletalMesh::CreateMeshDescription(const int32 InLODIndex)
 {
-	const FRawSkeletalMeshBulkData& RawSkeletalMeshBulkData = GetMeshEditorData().GetLODImportedData(LODIndex);
-	OutGeoImportVersion = RawSkeletalMeshBulkData.GeoImportVersion;
-	OutSkinningImportVersion = RawSkeletalMeshBulkData.SkinningImportVersion;
+	if (!ensure(IsValidLODIndex(InLODIndex)))
+	{
+		return nullptr;
+	}
+
+	return GetSourceModel(InLODIndex).CreateMeshDescription();
 }
 
-void USkeletalMesh::SetLODImportedDataVersions(const int32 LODIndex, const ESkeletalMeshGeoImportVersions& InGeoImportVersion, const ESkeletalMeshSkinningImportVersions& InSkinningImportVersion)
+
+FMeshDescription* USkeletalMesh::CreateMeshDescription(const int32 InLODIndex, FMeshDescription&& InMeshDescription)
 {
-	FRawSkeletalMeshBulkData& RawSkeletalMeshBulkData = GetMeshEditorData().GetLODImportedData(LODIndex);
-	RawSkeletalMeshBulkData.GeoImportVersion = InGeoImportVersion;
-	RawSkeletalMeshBulkData.SkinningImportVersion = InSkinningImportVersion;
-	//Update the cache
-	check(GetImportedModel()->LODModels.IsValidIndex(LODIndex));
-	GetImportedModel()->LODModels[LODIndex].RawSkeletalMeshBulkDataID = RawSkeletalMeshBulkData.GetIdString();
-	GetImportedModel()->LODModels[LODIndex].bIsBuildDataAvailable = RawSkeletalMeshBulkData.IsBuildDataAvailable();
-	GetImportedModel()->LODModels[LODIndex].bIsRawSkeletalMeshBulkDataEmpty = RawSkeletalMeshBulkData.IsEmpty();
+	FMeshDescription* MeshDescription = CreateMeshDescription(InLODIndex);
+	if (MeshDescription)
+	{
+		*MeshDescription = MoveTemp(InMeshDescription);
+	}
+
+	return MeshDescription;
 }
 
-void USkeletalMesh::CopyImportedData(int32 SrcLODIndex, USkeletalMesh* SrcSkeletalMesh, int32 DestLODIndex, USkeletalMesh* DestSkeletalMesh)
+
+bool USkeletalMesh::CommitMeshDescription(
+	const int32 InLODIndex,
+	const FCommitMeshDescriptionParams& InParams
+	)
 {
-	check(DestSkeletalMesh->GetImportedModel()->LODModels.IsValidIndex(DestLODIndex));
-	FRawSkeletalMeshBulkData& SrcRawMesh = SrcSkeletalMesh->GetMeshEditorData().GetLODImportedData(SrcLODIndex);
-	FRawSkeletalMeshBulkData& DestRawMesh = DestSkeletalMesh->GetMeshEditorData().GetLODImportedData(DestLODIndex);
-	FSkeletalMeshImportData SrcImportData;
-	SrcRawMesh.LoadRawMesh(SrcImportData);
-	DestRawMesh.SaveRawMesh(SrcImportData);
-	DestRawMesh.GeoImportVersion = SrcRawMesh.GeoImportVersion;
-	DestRawMesh.SkinningImportVersion = SrcRawMesh.SkinningImportVersion;
+	if (!ensure(IsValidLODIndex(InLODIndex)))
+	{
+		return false;
+	}
+
+	FSkeletalMeshSourceModel& SourceModel = GetSourceModel(InLODIndex);
+	const bool bUseHashAsGuid = !InParams.bForceUpdate;
+	SourceModel.CommitMeshDescription(bUseHashAsGuid);
+
+	if (InLODIndex == 0)
+	{
+		SetImportedBounds(SourceModel.GetBoundsFast());
+	}
+
+	if (ensure(GetImportedModel()->LODModels.IsValidIndex(InLODIndex)))
+	{
+		const FMeshDescriptionBulkData& BulkData = SourceModel.MeshDescriptionBulkData->GetBulkData();
+		FSkeletalMeshLODModel& LODModel = GetImportedModel()->LODModels[InLODIndex];
+		
+		LODModel.RawSkeletalMeshBulkDataID = BulkData.GetIdString();
+		LODModel.bIsBuildDataAvailable = !BulkData.IsEmpty();
+		LODModel.bIsRawSkeletalMeshBulkDataEmpty = BulkData.IsEmpty();
+	}
+	
+
+	if (InParams.bMarkPackageDirty)
+	{
+		(void)MarkPackageDirty();
+	}
 	
-	FSkeletalMeshLODModel& DestLODModel = DestSkeletalMesh->GetImportedModel()->LODModels[DestLODIndex];
-	DestLODModel.RawSkeletalMeshBulkDataID = DestRawMesh.GetIdString();
-	DestLODModel.bIsBuildDataAvailable = DestRawMesh.IsBuildDataAvailable();
-	DestLODModel.bIsRawSkeletalMeshBulkDataEmpty = DestRawMesh.IsEmpty();
+	return true;
 }
 
-void USkeletalMesh::ReserveLODImportData(int32 MaxLODIndex)
+
+bool USkeletalMesh::ModifyMeshDescription(const int32 InLODIndex, const bool bInAlwaysMarkPackageDirty)
 {
-	//Getting the LODImportedData will allocate the data to default value.
-	GetMeshEditorData().GetLODImportedData(MaxLODIndex);
+	if (!ensure(IsValidLODIndex(InLODIndex)))
+	{
+		return false;
+	}
+
+	return GetSourceModel(InLODIndex).MeshDescriptionBulkData->Modify(bInAlwaysMarkPackageDirty);
 }
 
-void USkeletalMesh::ForceBulkDataResident(const int32 LODIndex)
+
+void USkeletalMesh::ClearMeshDescription(const int32 InLODIndex)
 {
-	GetMeshEditorData().GetLODImportedData(LODIndex).GetBulkData().ForceBulkDataResident();
+	if (!ensure(IsValidLODIndex(InLODIndex)))
+	{
+		return;
+	}
+
+	return GetSourceModel(InLODIndex).ClearMeshDescription();
 }
 
-void USkeletalMesh::EmptyLODImportData(const int32 LODIndex)
+
+void USkeletalMesh::ClearAllMeshDescriptions()
 {
-	if(!GetImportedModel()->LODModels.IsValidIndex(LODIndex) || !GetMeshEditorData().IsLODImportDataValid(LODIndex))
+	for (int32 LODIndex = 0, LODCount = GetLODNum(); LODIndex < LODCount; LODIndex++)
+	{
+		ClearMeshDescription(LODIndex);
+	}
+}
+
+
+void USkeletalMesh::ClearMeshDescriptionAndBulkData(const int32 InLODIndex)
+{
+	if (!ensure(IsValidLODIndex(InLODIndex)))
 	{
 		return;
 	}
 
-	FRawSkeletalMeshBulkData& RawMesh = GetMeshEditorData().GetLODImportedData(LODIndex);
-	FSkeletalMeshImportData EmptyData;
-	RawMesh.SaveRawMesh(EmptyData);
-	RawMesh.GeoImportVersion = ESkeletalMeshGeoImportVersions::Before_Versionning;
-	RawMesh.SkinningImportVersion = ESkeletalMeshSkinningImportVersions::Before_Versionning;
-	GetImportedModel()->LODModels[LODIndex].RawSkeletalMeshBulkDataID = RawMesh.GetIdString();
-	GetImportedModel()->LODModels[LODIndex].bIsBuildDataAvailable = RawMesh.IsBuildDataAvailable();
-	GetImportedModel()->LODModels[LODIndex].bIsRawSkeletalMeshBulkDataEmpty = RawMesh.IsEmpty();
+	return GetSourceModel(InLODIndex).ClearAllMeshData();
 }
 
-void USkeletalMesh::EmptyAllImportData()
+
+void USkeletalMesh::LoadLODImportedData(const int32 LODIndex, FSkeletalMeshImportData& OutMesh) const
 {
-	const int32 LODNumber = GetLODNum();
-	for(int32 LODIndex = 0; LODIndex < LODNumber; ++LODIndex)
+	if (!ensure(IsValidLODIndex(LODIndex)))
 	{
-		EmptyLODImportData(LODIndex);
+		return;
+	}
+
+	if (const FMeshDescription* MeshDescription = GetMeshDescription(LODIndex))
+	{
+		OutMesh = FSkeletalMeshImportData::CreateFromMeshDescription(*MeshDescription);
 	}
 }
 
-void USkeletalMesh::GetMeshDescription(const int32 InLODIndex, FMeshDescription& OutMeshDescription) const
+void USkeletalMesh::SaveLODImportedData(const int32 LODIndex, const FSkeletalMeshImportData& InMesh)
 {
-	// Check first if we have bulk data available and non-empty.
-	if (IsLODImportedDataBuildAvailable(InLODIndex) && !IsLODImportedDataEmpty(InLODIndex))
+	if (!ensure(IsValidLODIndex(LODIndex)))
 	{
-		FSkeletalMeshImportData SkeletalMeshImportData;
-		LoadLODImportedData(InLODIndex, SkeletalMeshImportData);
-		SkeletalMeshImportData.GetMeshDescription(OutMeshDescription);
+		return;
 	}
-	else
+	
+	FMeshDescription MeshDescription;
+
+	if (InMesh.GetMeshDescription(nullptr, MeshDescription))
 	{
-		// Fall back on the LOD model directly if no bulk data exists. When we commit
-		// the mesh description, we override using the bulk data. This can happen for older
-		// skeletal meshes, from UE 4.24 and earlier.
-		const FSkeletalMeshModel* SkeletalMeshModel = GetImportedModel();
-		if (SkeletalMeshModel && SkeletalMeshModel->LODModels.IsValidIndex(InLODIndex))
-		{
-			SkeletalMeshModel->LODModels[InLODIndex].GetMeshDescription(OutMeshDescription, this);
-		}			
+		CreateMeshDescription(LODIndex, MoveTemp(MeshDescription));
+		CommitMeshDescription(LODIndex);
 	}
 }
 
-void USkeletalMesh::CommitMeshDescription(
-	const int32 InLODIndex, const FMeshDescription& InMeshDescription, const FCommitMeshDescriptionParams& InParams)
-{	
-	FSkeletalMeshImportData SkeletalMeshImportData = FSkeletalMeshImportData::CreateFromMeshDescription(InMeshDescription);
+bool USkeletalMesh::IsLODImportedDataBuildAvailable(const int32 LODIndex) const
+{
+	return HasMeshDescription(LODIndex);
+}
+
+bool USkeletalMesh::IsLODImportedDataEmpty(const int32 LODIndex) const
+{
+	return !HasMeshDescription(LODIndex);
+}
+
+void USkeletalMesh::GetLODImportedDataVersions(const int32 LODIndex, ESkeletalMeshGeoImportVersions& OutGeoImportVersion, ESkeletalMeshSkinningImportVersions& OutSkinningImportVersion) const
+{
+	OutGeoImportVersion = ESkeletalMeshGeoImportVersions::SkeletalMeshBuildRefactor;
+	OutSkinningImportVersion = ESkeletalMeshSkinningImportVersions::SkeletalMeshBuildRefactor;
+}
+
+void USkeletalMesh::SetLODImportedDataVersions(const int32 LODIndex, const ESkeletalMeshGeoImportVersions& InGeoImportVersion, const ESkeletalMeshSkinningImportVersions& InSkinningImportVersion)
+{
+}
 
-	if (InParams.bModify)
+void USkeletalMesh::CopyImportedData(int32 SrcLODIndex, USkeletalMesh* SrcSkeletalMesh, int32 DestLODIndex, USkeletalMesh* DestSkeletalMesh)
+{
+	FMeshDescription MeshDescription;
+	if (SrcSkeletalMesh->CloneMeshDescription(SrcLODIndex, MeshDescription))
 	{
-		SetFlags(RF_Transactional);
-		Modify();
-		GetMeshEditorData().Modify();
+		DestSkeletalMesh->CreateMeshDescription(DestLODIndex, MoveTemp(MeshDescription));
+		DestSkeletalMesh->CommitMeshDescription(DestLODIndex);
 	}
+}
 
-	SaveLODImportedData(InLODIndex, SkeletalMeshImportData);
-	// Make sure the mesh builder knows it's the latest variety, so that the render data gets properly rebuilt.
-	SetLODImportedDataVersions(InLODIndex, ESkeletalMeshGeoImportVersions::LatestVersion, ESkeletalMeshSkinningImportVersions::LatestVersion);
-	SetUseLegacyMeshDerivedDataKey(false);
+void USkeletalMesh::ReserveLODImportData(int32 MaxLODIndex)
+{
+}
+
+void USkeletalMesh::ForceBulkDataResident(const int32 LODIndex)
+{	
+}
+
+void USkeletalMesh::EmptyLODImportData(const int32 LODIndex)
+{
+	ClearMeshDescriptionAndBulkData(LODIndex);
+}
+
+void USkeletalMesh::EmptyAllImportData()
+{
+	for(int32 LODIndex = 0, LODCount = GetLODNum(); LODIndex < LODCount; ++LODIndex)
+	{
+		ClearMeshDescriptionAndBulkData(LODIndex);
+	}
 }
 
+
 void USkeletalMesh::CreateUserSectionsDataForLegacyAssets()
 {
-	//We want to avoid changing the ddc if we load an old asset.
-	//This bool should be put to false at the end of the postload, if there is another posteditchange call after a new ddc will be created
-	SetUseLegacyMeshDerivedDataKey(true);
 	//Fill up the Section ChunkedParentSectionIndex and OriginalDataSectionIndex
 	//We also want to create the UserSectionsData structure so the user can change the section data
 	for (int32 LodIndex = 0; LodIndex < GetLODInfoArray().Num(); LodIndex++)
@@ -2475,7 +2618,7 @@ void USkeletalMesh::CreateUserSectionsDataForLegacyAssets()
 			&& GetImportedModel()->OriginalReductionSourceMeshData_DEPRECATED.IsValidIndex(LodIndex)
 			&& !(GetImportedModel()->OriginalReductionSourceMeshData_DEPRECATED[LodIndex]->IsEmpty());
 
-		if (bIsLODReductionActive && !ThisLODInfo->bHasBeenSimplified && IsLODImportedDataEmpty(LodIndex))
+		if (bIsLODReductionActive && !ThisLODInfo->bHasBeenSimplified && !HasMeshDescription(LodIndex))
 		{
 			if (LodIndex > ThisLODInfo->ReductionSettings.BaseLOD)
 			{
@@ -2736,103 +2879,6 @@ bool USkeletalMesh::IsAsyncTaskComplete() const
 	return AsyncTask == nullptr || AsyncTask->IsWorkDone();
 }
 
-void USkeletalMesh::PostLoadEnsureImportDataExist()
-{
-	//If we have a LODModel with no import data and the LOD model have at least one section using more bone then any platform max GPU bone count. We will recreate the import data to allow the asset to be build and chunk properly.
-	const int32 MinimumPerPlatformMaxGPUSkinBones = FGPUBaseSkinVertexFactory::GetMinimumPerPlatformMaxGPUSkinBonesValue();
-	//This flag is set from the max gpu skin bones, if true ImportLODData will be created
-	bool bNeedToCreateImportData = false;
-
-	//The following two flags are acting together when bNeedToCreateImportData == false.
-	//When we found a LOD that is reduce inline, we want to convert the deprecated reduction source data
-	//and create the import data from it.
-
-	//If we found a LOD doing inline reduction this flag will be set to trtue
-	bool bContainInlineReduction = false;
-	//If we found a LOD doing inline reduction but with baked skinweight data this flag will be true.
-	//If this flag is true we will not create the imported data, except if bNeedToCreateData is true
-	//In case bNeedToCreateData is true, user will have to re-import the skinning data to comply with
-	//the per platform max gpu bone.
-	bool bCannotConvertSkinWeightProfile = false;
-
-	for (int32 LodIndex = 0; LodIndex < GetLODNum(); LodIndex++)
-	{
-		const FSkeletalMeshLODModel* LODModel = &(GetImportedModel()->LODModels[LodIndex]);
-		const FSkeletalMeshLODInfo* ThisLODInfo = GetLODInfo(LodIndex);
-		check(ThisLODInfo);
-		const bool bRawDataEmpty = IsLODImportedDataEmpty(LodIndex);
-		const bool bRawBuildDataAvailable = IsLODImportedDataBuildAvailable(LodIndex);
-		if (!bRawDataEmpty && bRawBuildDataAvailable)
-		{
-			continue;
-		}
-		
-		//We create the import data only if the mesh was simplified and its doing inline reduction
-		//Non simplified mesh without import data should not be converted because the simplification will kick in and the asset will be simplified (asset vertex and face count will change)
-		const bool bReductionActive = IsReductionActive(LodIndex);
-		const bool bInlineReduction = (ThisLODInfo->ReductionSettings.BaseLOD == LodIndex);
-		if (bReductionActive && !bInlineReduction)
-		{
-			//Generated LOD (not inline) do not need imported data
-			continue;
-		}
-		else if (bInlineReduction && (bReductionActive || ThisLODInfo->bHasBeenSimplified))
-		{
-			//If we cannot convert alternate skin weight data, we do not want to allow 
-			if (GetImportedModel()->OriginalReductionSourceMeshData_DEPRECATED.IsValidIndex(LodIndex))
-			{
-				//Old inline reduced assets do not have the original imported skin weight data, it must be re-import if we want to convert it to import data
-				bool bUseSkinWeightProfile = LODModel->SkinWeightProfiles.Num() > 0;
-				FSkeletalMeshLODModel ReductionLODModel;
-				TMap<FString, TArray<FMorphTargetDelta>> ReductionLODMorphTargetData;
-				//Swap the LODModel pointer to the one we store when reducing, so we add the true import data
-				GetImportedModel()->OriginalReductionSourceMeshData_DEPRECATED[LodIndex]->LoadReductionData(ReductionLODModel, ReductionLODMorphTargetData, this);
-
-				if (bUseSkinWeightProfile && ReductionLODModel.SkinWeightProfiles.IsEmpty())
-				{
-					//If we cannot convert a particular LOD, prevent Creation of import data and exit the loop
-					bCannotConvertSkinWeightProfile = true;
-				}
-			}
-			bContainInlineReduction = true;
-		}
-		//See if the LODModel data use more bones then the chunking allow
-		int32 MaxBoneperSection = 0;
-		for(const FSkelMeshSection& Section : LODModel->Sections)
-		{
-			MaxBoneperSection = FMath::Max(MaxBoneperSection, Section.BoneMap.Num());
-		}
-		//If we use more bone then the minimum maxGPUSkinbone, we need to re-create de import data to be able to build the asset
-		if (MaxBoneperSection > MinimumPerPlatformMaxGPUSkinBones)
-		{
-			bNeedToCreateImportData = true;
-			break;
-		}
-	}
-
-	if (bNeedToCreateImportData || (bContainInlineReduction && !bCannotConvertSkinWeightProfile))
-	{
-		if (!ensure(IsInGameThread()))
-		{
-			UE_ASSET_LOG(LogSkeletalMesh, Error, this, TEXT("USkeletalMesh::PostLoadEnsureImportDataExist() must be call on the game thread."));
-			return;
-		}
-		if (bCannotConvertSkinWeightProfile)
-		{
-			UE_ASSET_LOG(LogSkeletalMesh, Warning, this, TEXT("USkeletalMesh::PostLoadEnsureImportDataExist() Cannot convert old alternate skin weight profiles data. Re-import all alternate profiles for this skeletal mesh."));
-		}
-		IMeshUtilities& MeshUtilities = FModuleManager::Get().LoadModuleChecked<IMeshUtilities>("MeshUtilities");
-		//We create the import data for all LOD that do not have import data except for the generated LODs.
-		{
-			FScopedSkeletalMeshPostEditChange ScopedPostEditChange(this, false, false);
-			MeshUtilities.CreateImportDataFromLODModel(this);
-		}
-#if WITH_EDITORONLY_DATA
-		//If the import data is existing we want to turn use legacy derive data key to false
-		SetUseLegacyMeshDerivedDataKey(false);
-#endif
-	}
-}
 
 void USkeletalMesh::PostLoadVerifyAndFixBadTangent()
 {
@@ -2840,7 +2886,7 @@ void USkeletalMesh::PostLoadVerifyAndFixBadTangent()
 	bool bFoundBadTangents = false;
 	for (int32 LodIndex = 0; LodIndex < GetLODNum(); LodIndex++)
 	{
-		if (!IsLODImportedDataEmpty(LodIndex))
+		if (HasMeshDescription(LodIndex))
 		{
 			//No need to verify skeletalmesh that have valid imported data, the tangents will always exist in this case
 			continue;
@@ -3037,15 +3083,6 @@ void USkeletalMesh::BeginPostLoadInternal(FSkinnedAssetPostLoadContext& Context)
 		MeshClothingAsset->ConditionalPostLoad();
 	}
 
-	// Make sure the mesh editor data object is a sub object of the skeletalmesh, rename it to change the owner to be the skeletalmesh.
-	if (IsMeshEditorDataValid() && GetMeshEditorData().GetOuter() != this)
-	{
-		//Post load call so no need to: dirty, redirect, transact or reset the loader.
-		const TCHAR* NewName = nullptr;
-		GetMeshEditorData().Rename(NewName, this, REN_ForceNoResetLoaders | REN_DoNotDirty | REN_DontCreateRedirectors | REN_NonTransactional);
-		GetMeshEditorData().SetFlags(RF_Transactional);
-	}
-
 	if (!GetOutermost()->bIsCookedForEditor)
 	{
 		TArray<FSkeletalMeshLODInfo>& LODInfoArray = GetLODInfoArray();
@@ -3066,7 +3103,6 @@ void USkeletalMesh::BeginPostLoadInternal(FSkinnedAssetPostLoadContext& Context)
 		for (int32 LodIndex = 0; LodIndex < TotalLODNum; LodIndex++)
 		{
 			FSkeletalMeshLODInfo& ThisLODInfo = LODInfoArray[LodIndex];
-			FSkeletalMeshLODModel& ThisLODModel = GetImportedModel()->LODModels[LodIndex];
 
 			if (ThisLODInfo.ReductionSettings.BonesToRemove_DEPRECATED.Num() > 0)
 			{
@@ -3117,6 +3153,13 @@ void USkeletalMesh::BeginPostLoadInternal(FSkinnedAssetPostLoadContext& Context)
 			GetLODSettings()->SetLODSettingsToMesh(this);
 		}
 
+		if (GetLinkerCustomVersion(FFortniteMainBranchObjectVersion::GUID) < FFortniteMainBranchObjectVersion::MeshDescriptionForSkeletalMesh)
+		{
+			// Ensure we have source model storage that matches the number of LODs defined on this mesh.
+			SetNumSourceModels(GetLODNum());
+		}
+		
+
 		if (GetLinkerUEVersion() < VER_UE4_SORT_ACTIVE_BONE_INDICES)
 		{
 			for (int32 LodIndex = 0; LodIndex < LODInfoArray.Num(); LodIndex++)
@@ -3139,16 +3182,16 @@ void USkeletalMesh::BeginPostLoadInternal(FSkinnedAssetPostLoadContext& Context)
 
 		if (GetLinkerCustomVersion(FEditorObjectVersion::GUID) < FEditorObjectVersion::SkeletalMeshMoveEditorSourceDataToPrivateAsset)
 		{
-			ReserveLODImportData(GetImportedModel()->LODModels.Num() - 1);
 			for (int32 LODIndex = 0; LODIndex < GetImportedModel()->LODModels.Num(); ++LODIndex)
 			{
 				FSkeletalMeshLODModel& ThisLODModel = GetImportedModel()->LODModels[LODIndex];
+				
 				//We can have partial data if the asset was save after the split workflow implementation
 				//Use the deprecated member to retrieve this data
 				if (GetLinkerCustomVersion(FFortniteMainBranchObjectVersion::GUID) >= FFortniteMainBranchObjectVersion::NewSkeletalMeshImporterWorkflow)
 				{
 					//Get the deprecated data 
-					FRawSkeletalMeshBulkData* RawSkeletalMeshBulkData_DEPRECATED = nullptr;
+					FRawSkeletalMeshBulkData* RawSkeletalMeshBulkData_DEPRECATED;
 					{
 						PRAGMA_DISABLE_DEPRECATION_WARNINGS
 						RawSkeletalMeshBulkData_DEPRECATED = &ThisLODModel.GetRawSkeletalMeshBulkData_DEPRECATED();
@@ -3159,33 +3202,69 @@ void USkeletalMesh::BeginPostLoadInternal(FSkinnedAssetPostLoadContext& Context)
 					{
 						FSkeletalMeshImportData SerializeMeshData;
 						RawSkeletalMeshBulkData_DEPRECATED->LoadRawMesh(SerializeMeshData);
-						SaveLODImportedData(LODIndex, SerializeMeshData);
+
+						// Some older versions of the bulk data did not store the morph targets, but they're available on the skeletal
+						// mesh itself. Try to back-fill from the skeletal mesh.
+						if (SerializeMeshData.MorphTargets.IsEmpty() && !GetMorphTargets().IsEmpty())
+						{
+							SerializeMeshData.SetMorphTargets(GetMorphTargets(), LODIndex, ThisLODModel.MeshToImportVertexMap);
+						}
+
+						FMeshDescription MeshDescription;
+						if (SerializeMeshData.GetMeshDescription(this, MeshDescription))
+						{
+							CreateMeshDescription(LODIndex, MoveTemp(MeshDescription));
+							CommitMeshDescription(LODIndex);
+						}
 					}
-					//Get the FRawSkeletalMeshBulkData to set the geo and skinning version
-					FRawSkeletalMeshBulkData& RawSkeletalMeshBulkData = GetMeshEditorData().GetLODImportedData(LODIndex);
-					RawSkeletalMeshBulkData.GeoImportVersion = RawSkeletalMeshBulkData_DEPRECATED->GeoImportVersion;
-					RawSkeletalMeshBulkData.SkinningImportVersion = RawSkeletalMeshBulkData_DEPRECATED->SkinningImportVersion;
 					//Empty the DEPRECATED member
 					FSkeletalMeshImportData EmptyMeshData;
 					RawSkeletalMeshBulkData_DEPRECATED->SaveRawMesh(EmptyMeshData);
 					RawSkeletalMeshBulkData_DEPRECATED->EmptyBulkData();
 				}
-				//Set the cache data into the LODModel
-				FRawSkeletalMeshBulkData& RawSkeletalMeshBulkData = GetMeshEditorData().GetLODImportedData(LODIndex);
-				ThisLODModel.bIsRawSkeletalMeshBulkDataEmpty = RawSkeletalMeshBulkData.IsEmpty();
-				ThisLODModel.bIsBuildDataAvailable = RawSkeletalMeshBulkData.IsBuildDataAvailable();
-				ThisLODModel.RawSkeletalMeshBulkDataID = RawSkeletalMeshBulkData.GetIdString();
 			}
 		}
 
+		if (GetLinkerCustomVersion(FFortniteMainBranchObjectVersion::GUID) < FFortniteMainBranchObjectVersion::MeshDescriptionForSkeletalMesh)
+		{
+			// Transfer all bulk data from the deprecated private bulk storage objects to the new source model structure.
+			// When the user asks for mesh description from the bulk data, the raw mesh bulk data will be unpacked, converted and
+			// discarded.
+			PRAGMA_DISABLE_DEPRECATION_WARNINGS
+			USkeletalMeshEditorData* ImportData = MeshEditorDataObject_DEPRECATED;
+			PRAGMA_ENABLE_DEPRECATION_WARNINGS
+
+			if (ImportData)
+			{
+				for (int32 LODIndex = 0, LODNum = GetLODNum(); LODIndex < LODNum; LODIndex++)
+				{
+					if (LODIndex < ImportData->RawSkeletalMeshBulkDatas.Num() &&
+						!ImportData->RawSkeletalMeshBulkDatas[LODIndex]->IsEmpty() &&
+						ImportData->RawSkeletalMeshBulkDatas[LODIndex]->IsBuildDataAvailable())
+					{
+						FSkeletalMeshSourceModel& SourceModel = GetSourceModel(LODIndex);
+						SourceModel.RawMeshBulkData = ImportData->RawSkeletalMeshBulkDatas[LODIndex];
+						// When we do on-demand conversion of the raw mesh bulk data, we need to know which
+						// LOD we came from so that we can grab the appropriate reconstruction data from the 
+						// correct FSkeletalMeshModel::LODModels variant.
+						SourceModel.RawMeshBulkDataLODIndex = LODIndex; 
+					}
+				}
+			}
+		}
+
+		// If we didn't get any meshes from the bulk data, then try to recover them from the LODModel listings.
+		constexpr bool bInResetReductionAsNeeded = true;
+			
+		IMeshUtilities& MeshUtilities = FModuleManager::Get().LoadModuleChecked<IMeshUtilities>("MeshUtilities");
+		MeshUtilities.CreateImportDataFromLODModel(this, bInResetReductionAsNeeded);
+		
 		if (GetLinkerCustomVersion(FEditorObjectVersion::GUID) < FEditorObjectVersion::SkeletalMeshBuildRefactor)
 		{
 			CreateUserSectionsDataForLegacyAssets();
 		}
 
 		ValidateAllLodMaterialIndexes();
-
-		PostLoadEnsureImportDataExist();
 	}
 
 #endif // #if WITH_EDITOR
@@ -3367,15 +3446,6 @@ void USkeletalMesh::FinishPostLoadInternal(FSkinnedAssetPostLoadContext& Context
 #if !WITH_EDITOR
 	RebuildSocketMap();
 #endif // !WITH_EDITOR
-#if WITH_EDITORONLY_DATA
-	// Only call the setter if the value needs adjustment since we're
-	// not supposed to modify this value during async build.
-	if (GetUseLegacyMeshDerivedDataKey())
-	{
-		//Next postedit change will use the new ddc key scheme
-		SetUseLegacyMeshDerivedDataKey(false);
-	}
-#endif
 
 #if WITH_EDITORONLY_DATA
 	FPerPlatformInt PerPlatformData = GetMinLod();
@@ -3833,34 +3903,21 @@ bool USkeletalMesh::RemoveMorphTargets(TConstArrayView<FName> InMorphTargetNames
 			Modify();
 			MorphTarget->Modify();
 
-			if (!IsLODImportedDataEmpty(0) && IsLODImportedDataBuildAvailable(0))
+			constexpr int32 LODIndex = 0;
+
+			if (HasMeshDescription(LODIndex))
 			{
 				//Remove the morph target from the raw import data
-				FSkeletalMeshImportData SkelMeshImportData;
-				LoadLODImportedData(0, SkelMeshImportData);
-				int32 ToDeleteIndex = INDEX_NONE;
-				for (int32 MorphTargetIndex = 0; MorphTargetIndex < SkelMeshImportData.MorphTargetNames.Num(); ++MorphTargetIndex)
-				{
-					if (SkelMeshImportData.MorphTargetNames[MorphTargetIndex].Equals(MorphTargetName.ToString()))
-					{
-						ToDeleteIndex = MorphTargetIndex;
-						break;
-					}
-				}
+				FMeshDescription* MeshDescription = GetMeshDescription(LODIndex);
+				FSkeletalMeshAttributes MeshAttributes(*MeshDescription);
 
-				if (ToDeleteIndex != INDEX_NONE)
+				if (MeshAttributes.GetMorphTargetNames().Contains(MorphTargetName))
 				{
-					SkelMeshImportData.MorphTargetNames.RemoveAt(ToDeleteIndex);
-					SkelMeshImportData.MorphTargetModifiedPoints.RemoveAt(ToDeleteIndex);
-					SkelMeshImportData.MorphTargets.RemoveAt(ToDeleteIndex);
-					SaveLODImportedData(0, SkelMeshImportData);
+					ModifyMeshDescription(LODIndex);
+					MeshAttributes.UnregisterMorphTargetAttribute(MorphTargetName);
+					CommitMeshDescription(LODIndex);
 				}
 			}
-			else
-			{
-				//If we deal with an old asset (pre 4.24) and we do not have some valid import data, we need to dirty the ddc key so it wont take the ddc with the old morph target we just delete
-				InvalidateDeriveDataCacheGUID();
-			}
 
 			UnregisterMorphTarget(MorphTarget);
 
@@ -3901,25 +3958,25 @@ bool USkeletalMesh::RenameMorphTarget(FName InOldName, FName InNewName)
 	Modify();
 	MorphTarget->Modify();
 
-	if (!IsLODImportedDataEmpty(0) && IsLODImportedDataBuildAvailable(0))
+	constexpr int32 LODIndex = 0;
+	if (HasMeshDescription(LODIndex))
 	{
-		// Rename the morph target in the raw import data
-		FSkeletalMeshImportData SkelMeshImportData;
-		LoadLODImportedData(0, SkelMeshImportData);
-		for (int32 MorphTargetIndex = 0; MorphTargetIndex < SkelMeshImportData.MorphTargetNames.Num(); ++MorphTargetIndex)
+		FMeshDescription* MeshDescription = GetMeshDescription(LODIndex);
+		FSkeletalMeshAttributes MeshAttributes(*MeshDescription);
+
+		if (MeshAttributes.GetMorphTargetNames().Contains(InOldName))
 		{
-			if (SkelMeshImportData.MorphTargetNames[MorphTargetIndex].Equals(InOldName.ToString()))
+			ModifyMeshDescription(LODIndex);
+			if (MeshAttributes.RegisterMorphTargetAttribute(InNewName))
 			{
-				SkelMeshImportData.MorphTargetNames[MorphTargetIndex] = InNewName.ToString();
-				SaveLODImportedData(0, SkelMeshImportData);
-				break;
+				const FMorphTargetVertexAttributesConstRef SourceMorphRef{ MeshAttributes.GetVertexMorphTarget(InOldName) };
+				FMorphTargetVertexAttributesRef TargetMorphRef{ MeshAttributes.GetVertexMorphTarget(InNewName) };
+
+				TargetMorphRef.Copy(SourceMorphRef);
+				MeshAttributes.UnregisterMorphTargetAttribute(InOldName);
 			}
 		}
 	}
-	else
-	{
-		InvalidateDeriveDataCacheGUID();
-	}
 
 	// Rename the morph target itself
 	MorphTarget->Rename(*InNewName.ToString(), nullptr, REN_ForceNoResetLoaders | REN_DontCreateRedirectors);
@@ -4483,7 +4540,7 @@ namespace InternalSkeletalMeshHelper
 				|| LODInfoEntry == nullptr
 				|| LODInfoEntry->LODMaterialMap.Num() == 0 //If there is no LODMaterialMap we have nothing to backup
 				|| SkeletalMesh->IsReductionActive(LODIndex) //Reduction will manage the LODMaterialMap, avoid backup restore
-				|| !SkeletalMesh->IsLODImportedDataBuildAvailable(LODIndex)) //Legacy asset are not build, avoid backup restore
+				|| !SkeletalMesh->HasMeshDescription(LODIndex)) //Legacy asset are not build, avoid backup restore
 			{
 				continue;
 			}
@@ -4659,42 +4716,21 @@ FString USkeletalMesh::BuildDerivedDataKey(const ITargetPlatform* TargetPlatform
 {
 	FString KeySuffix(TEXT(""));
 
-	if (GetUseLegacyMeshDerivedDataKey())
-	{
-		//Old asset will have the same LOD settings for bUseFullPrecisionUVs. We can use the LOD 0
-		const FSkeletalMeshLODInfo* BaseLODInfo = GetLODInfo(0);
-		bool bUseFullPrecisionUVs = BaseLODInfo ? BaseLODInfo->BuildSettings.bUseFullPrecisionUVs : false;
-		KeySuffix += GetImportedModel()->GetIdString();
-		KeySuffix += bUseFullPrecisionUVs ? "1" : "0";
-
-		//Dummy call to update the FSkeletalMeshLODModel::BuildStringID members.
-		GetImportedModel()->GetLODModelIdString();
-
-		//Dummy call to update the LODInfo::BuildGUID members.
-		{
-			FString TmpString;
-			SerializeLODInfoForDDC(this, TmpString);
-		}
-	}
-	else
-	{
-		FString TmpPartialKeySuffix;
-		//Synchronize the user data that are part of the key
-		GetImportedModel()->SyncronizeLODUserSectionsData();
-		TmpPartialKeySuffix = GetImportedModel()->GetIdString();
-		KeySuffix += TmpPartialKeySuffix;
-		TmpPartialKeySuffix = GetImportedModel()->GetLODModelIdString();
-		KeySuffix += TmpPartialKeySuffix;
-
-		//Add the max gpu bone per section
-		const int32 MaxGPUSkinBones = FGPUBaseSkinVertexFactory::GetMaxGPUSkinBones(TargetPlatform);
-		KeySuffix += FString::FromInt(MaxGPUSkinBones);
-
-		TmpPartialKeySuffix = TEXT("");
-		SerializeLODInfoForDDC(this, TmpPartialKeySuffix);
-		KeySuffix += TmpPartialKeySuffix;
-	}
-
+	FString TmpPartialKeySuffix;
+	//Synchronize the user data that are part of the key
+	GetImportedModel()->SyncronizeLODUserSectionsData();
+	TmpPartialKeySuffix = GetImportedModel()->GetIdString();
+	KeySuffix += TmpPartialKeySuffix;
+	TmpPartialKeySuffix = GetImportedModel()->GetLODModelIdString();
+	KeySuffix += TmpPartialKeySuffix;
+
+	//Add the max gpu bone per section
+	const int32 MaxGPUSkinBones = FGPUBaseSkinVertexFactory::GetMaxGPUSkinBones(TargetPlatform);
+	KeySuffix += FString::FromInt(MaxGPUSkinBones);
+
+	TmpPartialKeySuffix = TEXT("");
+	SerializeLODInfoForDDC(this, TmpPartialKeySuffix);
+	KeySuffix += TmpPartialKeySuffix;
 	KeySuffix += GetHasVertexColors() ? "1" : "0";
 	KeySuffix += GetVertexColorGuid().ToString(EGuidFormats::Digits);
 
@@ -5310,14 +5346,22 @@ FSkeletalMeshLODInfo& USkeletalMesh::AddLODInfo()
 		// if this is the first LOD, then just use default setting of the struct
 	}
 
+#if WITH_EDITOR
+	// Add source models to match
+	SetNumSourceModels(LODInfoArray.Num());
+#endif
+
 	return LODInfoArray[NewIndex];
 }
 
 void USkeletalMesh::AddLODInfo(const FSkeletalMeshLODInfo& NewLODInfo) 
 {
-	WaitUntilAsyncPropertyReleased(ESkeletalMeshAsyncProperties::LODInfo);
+	WaitUntilAsyncPropertyReleased(ESkeletalMeshAsyncProperties::LODInfo|ESkeletalMeshAsyncProperties::SourceModels);
 	PRAGMA_DISABLE_DEPRECATION_WARNINGS
 	LODInfo.Add(NewLODInfo);
+#if WITH_EDITOR
+	SetNumSourceModelsInternal(LODInfo.Num());
+#endif
 	PRAGMA_ENABLE_DEPRECATION_WARNINGS
 }
 
@@ -5327,10 +5371,7 @@ void USkeletalMesh::RemoveLODInfo(int32 Index)
 	if (LODInfoArray.IsValidIndex(Index))
 	{
 #if WITH_EDITOR
-		if (IsMeshEditorDataValid())
-		{
-			GetMeshEditorDataObject()->RemoveLODImportedData(Index);
-		}
+		RemoveSourceModel(Index);
 		
 		if (GetImportedModel()->InlineReductionCacheDatas.IsValidIndex(Index))
 		{
@@ -5344,6 +5385,9 @@ void USkeletalMesh::RemoveLODInfo(int32 Index)
 void USkeletalMesh::ResetLODInfo()
 {
 	GetLODInfoArray().Reset();
+#if WITH_EDITOR
+	SetNumSourceModels(0);
+#endif
 }
 
 #if WITH_EDITOR
@@ -5416,17 +5460,13 @@ void USkeletalMesh::BuildLODModel(const ITargetPlatform* TargetPlatform, int32 L
 	//We want to avoid building a LOD if the LOD was generated from a previous LODIndex.
 	const bool bIsGeneratedLodNotInline = (LODInfoPtr->bHasBeenSimplified && IsReductionActive(LODIndex) && GetReductionSettings(LODIndex).BaseLOD < LODIndex);
 
-	//Make sure the LOD have all the data needed to be build
-	const bool bRawDataEmpty = IsLODImportedDataEmpty(LODIndex);
-	const bool bRawBuildDataAvailable = IsLODImportedDataBuildAvailable(LODIndex);
-
 	//Build the source model before the render data, if we are a purely generated LOD we do not need to be build
 	IMeshBuilderModule& MeshBuilderModule = IMeshBuilderModule::GetForPlatform(TargetPlatform);
-	if (!bIsGeneratedLodNotInline && !bRawDataEmpty && bRawBuildDataAvailable)
+	if (!bIsGeneratedLodNotInline && HasMeshDescription(LODIndex))
 	{
 		LODInfoPtr->bHasBeenSimplified = false;
-		const bool bRegenDepLODs = true;
-		FSkeletalMeshBuildParameters BuildParameters(this, TargetPlatform, LODIndex, bRegenDepLODs);
+		constexpr bool bRegenDepLODs = true;
+		const FSkeletalMeshBuildParameters BuildParameters(this, TargetPlatform, LODIndex, bRegenDepLODs);
 		MeshBuilderModule.BuildSkeletalMesh(BuildParameters);
 	}
 	else
@@ -5636,7 +5676,7 @@ void USkeletalMesh::WaitUntilAsyncPropertyReleased(ESkeletalMeshAsyncProperties
 
 FString USkeletalMesh::GetAsyncPropertyName(uint64 Property) const
 {
-	return StaticEnum<ESkeletalMeshAsyncProperties>()->GetNameByValue(Property).ToString();
+	return StaticEnum<ESkeletalMeshAsyncProperties>()->GetValueOrBitfieldAsString(Property);
 }
 
 int32 USkeletalMesh::GetPostProcessAnimBPLODThreshold() const
@@ -7266,5 +7306,6 @@ void GetSkeletalMeshRefTangentBasis(const USkeletalMesh* Mesh, const FSkeletalMe
 	GetRefTangentBasisTyped(Mesh, Section, LODData.StaticVertexBuffers.StaticMeshVertexBuffer, SkinWeightVertexBuffer, VertIndexInChunk, OutTangentX, OutTangentY, OutTangentZ);
 }
 
-#undef LOCTEXT_NAMESPACE
+UE_ENABLE_OPTIMIZATION
 
+#undef LOCTEXT_NAMESPACE
diff --git a/Engine/Source/Runtime/Engine/Private/SkeletalMeshLODModel.cpp b/Engine/Source/Runtime/Engine/Private/SkeletalMeshLODModel.cpp
--- a/Engine/Source/Runtime/Engine/Private/SkeletalMeshLODModel.cpp
+++ b/Engine/Source/Runtime/Engine/Private/SkeletalMeshLODModel.cpp
@@ -18,16 +18,6 @@
 #include "SkeletalMeshAttributes.h"
 #include "ReferenceSkeleton.h"
 
-/*-----------------------------------------------------------------------------
-FSkelMeshImportedMeshInfo
------------------------------------------------------------------------------*/
-FArchive& operator<<(FArchive& Ar, FSkelMeshImportedMeshInfo& MeshInfo)
-{
-	Ar << MeshInfo.Name;
-	Ar << MeshInfo.NumVertices;
-	Ar << MeshInfo.StartImportedVertex;
-	return Ar;
-}
 
 /*-----------------------------------------------------------------------------
 FSoftSkinVertex
@@ -991,18 +981,6 @@ void FSkeletalMeshLODModel::DeclareCustomVersions(FArchive& Ar)
 	FSkelMeshSection::DeclareCustomVersions(Ar);
 }
 
-int32 FSkeletalMeshLODModel::FindMeshInfoIndex(FName Name) const
-{
-	for (int32 Index = 0; Index < ImportedMeshInfos.Num(); ++Index)
-	{
-		if (ImportedMeshInfos[Index].Name == Name)
-		{
-			return Index;
-		}
-	}
-	return INDEX_NONE;
-}
-
 void FSkeletalMeshLODModel::GetSectionFromVertexIndex(int32 InVertIndex, int32& OutSectionIndex, int32& OutVertIndex) const
 {
 	OutSectionIndex = 0;
@@ -1366,7 +1344,7 @@ void FSkeletalMeshLODModel::GetMeshDescription(FMeshDescription& MeshDescription
 
 	MeshDescription.Empty();
 	
-	FSkeletalMeshAttributes MeshAttributes(MeshDescription);			
+	FSkeletalMeshAttributes MeshAttributes(MeshDescription);	
 	
 	// Register extra attributes for us.
 	MeshAttributes.Register();
@@ -1468,7 +1446,7 @@ void FSkeletalMeshLODModel::GetMeshDescription(FMeshDescription& MeshDescription
 				VertexInstanceBinormalSigns.Set(VertexInstanceID, FMatrix44f(
 					SourceVertex.TangentX.GetSafeNormal(),
 					SourceVertex.TangentY.GetSafeNormal(),
-					(FVector3f)(SourceVertex.TangentZ.GetSafeNormal()),
+					FVector3f(SourceVertex.TangentZ.GetSafeNormal()),
 					FVector3f::ZeroVector).Determinant() < 0.0f ? -1.0f : +1.0f);
 
 				for (int32 UVIndex = 0; UVIndex < static_cast<int32>(NumTexCoords); UVIndex++)
diff --git a/Engine/Source/Runtime/Engine/Private/SkeletalMeshRenderData.cpp b/Engine/Source/Runtime/Engine/Private/SkeletalMeshRenderData.cpp
--- a/Engine/Source/Runtime/Engine/Private/SkeletalMeshRenderData.cpp
+++ b/Engine/Source/Runtime/Engine/Private/SkeletalMeshRenderData.cpp
@@ -260,37 +260,28 @@ void FSkeletalMeshRenderData::Cache(const ITargetPlatform* TargetPlatform, USkin
 			//Helper structure to change the morph targets
 			TUniquePtr<FFinishBuildMorphTargetData> FinishBuildMorphTargetData;
 
-			//With skeletal mesh build refactor we serialize the LODModel data into the DDC
-			//We need to store those so we do not have to rerun the reduction to make them up to date
-			//with the serialize renderdata. This allow to use DDC when changing the reduction settings.
-			//The old workflow has to reduce the LODModel before getting the render data DDC.
-			if (!Owner->GetUseLegacyMeshDerivedDataKey())
-			{
-				FSkeletalMeshModel* SkelMeshModel = Owner->GetImportedModel();
-				check(SkelMeshModel);
+			FSkeletalMeshModel* SkelMeshModel = Owner->GetImportedModel();
+			check(SkelMeshModel);
 
-				//Get the morph target data, we put it in the compilation context to apply them in the game thread before the InitResources
-				{
-					if (Owner->GetMorphTargets().Num() > 0)
-					{
-						FinishBuildMorphTargetData = Owner->GetMorphTargets()[0]->CreateFinishBuildMorphTargetData();
-					}
-					else
-					{
-						// Create and initialize the FinishBuildInternalData, use the class default object to call the virtual function
-						FinishBuildMorphTargetData = UMorphTarget::StaticClass()->GetDefaultObject<UMorphTarget>()->CreateFinishBuildMorphTargetData();
-					}
-					check(FinishBuildMorphTargetData);
-					FinishBuildMorphTargetData->LoadFromMemoryArchive(Ar);
-				}
+			//Get the morph target data, we put it in the compilation context to apply them in the game thread before the InitResources
+			if (Owner->GetMorphTargets().Num() > 0)
+			{
+				FinishBuildMorphTargetData = Owner->GetMorphTargets()[0]->CreateFinishBuildMorphTargetData();
+			}
+			else
+			{
+				// Create and initialize the FinishBuildInternalData, use the class default object to call the virtual function
+				FinishBuildMorphTargetData = UMorphTarget::StaticClass()->GetDefaultObject<UMorphTarget>()->CreateFinishBuildMorphTargetData();
+			}
+			check(FinishBuildMorphTargetData);
+			FinishBuildMorphTargetData->LoadFromMemoryArchive(Ar);
 
-				//Serialize the LODModel sections since they are dependent on the reduction
-				for (int32 LODIndex = 0; LODIndex < SkelMeshModel->LODModels.Num(); LODIndex++)
-				{
-					FSkeletalMeshLODModel* LODModel = &(SkelMeshModel->LODModels[LODIndex]);
-					SerializeLodModelDdcData(LODModel, Ar);
-					LODModel->SyncronizeUserSectionsDataArray();
-				}
+			//Serialize the LODModel sections since they are dependent on the reduction
+			for (int32 LODIndex = 0; LODIndex < SkelMeshModel->LODModels.Num(); LODIndex++)
+			{
+				FSkeletalMeshLODModel* LODModel = &(SkelMeshModel->LODModels[LODIndex]);
+				SerializeLodModelDdcData(LODModel, Ar);
+				LODModel->SyncronizeUserSectionsDataArray();
 			}
 
 			Serialize(Ar, Owner);
@@ -368,24 +359,19 @@ void FSkeletalMeshRenderData::Cache(const ITargetPlatform* TargetPlatform, USkin
 
 			FLargeMemoryWriter Ar(0, /*bIsPersistent=*/ true);
 
-			//If we load an old asset we want to be sure the serialize ddc will be the same has before the skeletalmesh build refactor
-			//So we do not serialize the LODModel sections.
-			if (!Owner->GetUseLegacyMeshDerivedDataKey())
+			int32 MorphTargetNumber = Owner->GetMorphTargets().Num();
+			Ar << MorphTargetNumber;
+			for (int32 MorphTargetIndex = 0; MorphTargetIndex < MorphTargetNumber; ++MorphTargetIndex)
 			{
-				int32 MorphTargetNumber = Owner->GetMorphTargets().Num();
-				Ar << MorphTargetNumber;
-				for (int32 MorphTargetIndex = 0; MorphTargetIndex < MorphTargetNumber; ++MorphTargetIndex)
-				{
-					Owner->GetMorphTargets()[MorphTargetIndex]->SerializeMemoryArchive(Ar);
-				}
-				//No need to serialize the morph target mapping since we will rebuild the mapping when loading a ddc
+				Owner->GetMorphTargets()[MorphTargetIndex]->SerializeMemoryArchive(Ar);
+			}
+			//No need to serialize the morph target mapping since we will rebuild the mapping when loading a ddc
 
-				//Serialize the LODModel sections since they are dependent on the reduction
-				for (int32 LODIndex = 0; LODIndex < SkelMeshModel->LODModels.Num(); LODIndex++)
-				{
-					FSkeletalMeshLODModel* LODModel = &(SkelMeshModel->LODModels[LODIndex]);
-					SerializeLodModelDdcData(LODModel, Ar);
-				}
+			//Serialize the LODModel sections since they are dependent on the reduction
+			for (int32 LODIndex = 0; LODIndex < SkelMeshModel->LODModels.Num(); LODIndex++)
+			{
+				FSkeletalMeshLODModel* LODModel = &(SkelMeshModel->LODModels[LODIndex]);
+				SerializeLodModelDdcData(LODModel, Ar);
 			}
 
 			IMeshBuilderModule& MeshBuilderModule = IMeshBuilderModule::GetForPlatform(TargetPlatform);
diff --git a/Engine/Source/Runtime/Engine/Public/Rendering/SkeletalMeshLODImporterData.h b/Engine/Source/Runtime/Engine/Public/Rendering/SkeletalMeshLODImporterData.h
--- a/Engine/Source/Runtime/Engine/Public/Rendering/SkeletalMeshLODImporterData.h
+++ b/Engine/Source/Runtime/Engine/Public/Rendering/SkeletalMeshLODImporterData.h
@@ -47,8 +47,8 @@ namespace SkeletalMeshImportData
 	struct FMeshInfo
 	{
 		FName Name;	// The name of the mesh.
-		int32 NumVertices;	// The number of imported (dcc) vertices that are part of this mesh. This is a value of 8 for a cube. So NOT the number of render vertices.
-		int32 StartImportedVertex;	// The first index of imported (dcc) vertices in the mesh. So this NOT an index into the render vertex buffer. In range of 0..7 for a cube.
+		int32 NumVertices = 0;	// The number of imported (dcc) vertices that are part of this mesh. This is a value of 8 for a cube. So NOT the number of render vertices.
+		int32 StartImportedVertex = 0;	// The first index of imported (dcc) vertices in the mesh. So this NOT an index into the render vertex buffer. In range of 0..7 for a cube.
 	};
 
 	struct FMeshWedge
@@ -418,20 +418,27 @@ public:
 	 */
 	ENGINE_API void ComputeSmoothGroupFromNormals();
 
+	/*
+	 * Add morph target data from UMorphTarget in case there was none on the mesh itself.
+	 */
+	ENGINE_API void SetMorphTargets(const TArray<TObjectPtr<UMorphTarget>>& InMorphTargets, int32 InLODIndex, const TArray<int32>& InVertexMap);
+
 	/**
-	 * Returns a mesh description from the import data
+	 * Returns a mesh description from the import data. If logging on failures is required, pass in a pointer 
+	 * to the owning skeletal mesh. Otherwise leave as a \c nullptr. 
 	 */
-	ENGINE_API bool GetMeshDescription(FMeshDescription& OutMeshDescription) const;
+	ENGINE_API bool GetMeshDescription(const USkeletalMesh* InSkeletalMesh, FMeshDescription& OutMeshDescription) const;
 
 	/**
 	 * @note MeshDescription always contains color, normal and tangent data by default. Therefore, while iterating
-	 * over the vertices, we check if at least one normal/tangent vector is not a zero vector and the color is
-	 * not white, then we set the corresponding bHasNormals/bHasTangent/bHasVertexColors flags to true.
+	 * over the vertices, we check if at least one normal/tangent vector is not a zero vector and the vertex color is
+	 * not pure white, then we set the corresponding bHasNormals/bHasTangent/bHasVertexColors flags to true.
 	 */
 	static ENGINE_API FSkeletalMeshImportData CreateFromMeshDescription(const FMeshDescription& InMeshDescription);
 
 private:
 	void CopySkinWeightsToMeshDescription(
+		const USkeletalMesh* InSkeletalMesh,
 		const FName InSkinWeightName,
 		const FSkeletalMeshImportData& InSkinWeightMesh,
 		const TArray<FVertexID>& InVertexIDMap,
diff --git a/Engine/Source/Runtime/Engine/Public/Rendering/SkeletalMeshLODModel.h b/Engine/Source/Runtime/Engine/Public/Rendering/SkeletalMeshLODModel.h
--- a/Engine/Source/Runtime/Engine/Public/Rendering/SkeletalMeshLODModel.h
+++ b/Engine/Source/Runtime/Engine/Public/Rendering/SkeletalMeshLODModel.h
@@ -322,23 +322,6 @@ struct FSkelMeshSourceSectionUserData
 	friend FArchive& operator<<(FArchive& Ar, FSkelMeshSourceSectionUserData& S);
 };
 
-/**
- * The information for a given mesh as it was imported and appears in the dcc.
- * This gives us some extra information about the mesh as it appeared inside the source asset.
- * There can be multiple of those per skeletal mesh. The skeletal mesh is some merged big mesh separated in sections.
- * However sometimes we want to know some information about the actual individual meshes as they appeared in the dcc.
- * This class provides some of that information.
- */
-struct FSkelMeshImportedMeshInfo
-{
-	FName Name;	// The name of the mesh.
-	int32 NumVertices;	// The number of imported (dcc) vertices that are part of this mesh. This is a value of 8 for a cube. So NOT the number of render vertices.
-	int32 StartImportedVertex;	// The first index of imported (dcc) vertices in the mesh. So this NOT an index into the render vertex buffer. In range of 0..7 for a cube.
-
-	// Serialization.
-	friend FArchive& operator<<(FArchive& Ar, FSkelMeshImportedMeshInfo& MeshInfo);
-};
-
 /**
 * All data to define a certain LOD model for a skeletal mesh.
 */
@@ -348,9 +331,6 @@ public:
 	/** Sections. */
 	TArray<FSkelMeshSection> Sections;
 
-	/** The information about individual meshes. This can be empty if we did choose not to store this information. */
-	TArray<FSkelMeshImportedMeshInfo> ImportedMeshInfos;
-
 	/*
 	 * When user change section data in the UI, we store it here to be able to regenerate the changes
 	 * Note: the key (int32) is the original imported section data, because of BONE chunk the size of
@@ -447,10 +427,30 @@ public:
 		BulkDataReadMutex = BackupBulkDataReadMutex;
 	}
 
-	/** Find a mesh info by a mesh name. Returns INDEX_NONE when not found. */
-	int32 FindMeshInfoIndex(FName Name) const;
-
 private:
+	friend struct FSkeletalMeshSourceModel;
+	
+	struct FSkelMeshImportedMeshInfo
+	{
+		FName Name;	// The name of the mesh.
+		int32 NumVertices;	// The number of imported (dcc) vertices that are part of this mesh. This is a value of 8 for a cube. So NOT the number of render vertices.
+		int32 StartImportedVertex;	// The first index of imported (dcc) vertices in the mesh. So this NOT an index into the render vertex buffer. In range of 0..7 for a cube.
+
+		void Serialize(FArchive& Ar)
+		{
+			Ar << Name;
+			Ar << NumVertices;
+			Ar << StartImportedVertex;
+		}
+	};
+	friend FArchive& operator<<(FArchive& Ar, FSkelMeshImportedMeshInfo& Info)
+	{
+		Info.Serialize(Ar);
+		return Ar;
+	}
+
+	TArray<FSkelMeshImportedMeshInfo> ImportedMeshInfos;
+
 	/** Editor only data: array of the original point (wedge) indices for each of the vertices in a FSkeletalMeshLODModel */
 	TArray<uint32>				RawPointIndices2;
 
diff --git a/Engine/Source/Runtime/MeshConversion/Private/DynamicMeshToMeshDescription.cpp b/Engine/Source/Runtime/MeshConversion/Private/DynamicMeshToMeshDescription.cpp
--- a/Engine/Source/Runtime/MeshConversion/Private/DynamicMeshToMeshDescription.cpp
+++ b/Engine/Source/Runtime/MeshConversion/Private/DynamicMeshToMeshDescription.cpp
@@ -711,14 +711,22 @@ void FDynamicMeshToMeshDescription::Convert_NoSharedInstances(const FDynamicMesh
 		FSkeletalMeshAttributes MeshOutAttributes(MeshOut);
 		MeshOutAttributes.Register();
 		MeshOutAttributes.RegisterColorAttribute();
+		
+		// If there was an attribute to map from the mesh description index to the imported mesh index, remove that
+		// now, since we don't carry this mapping over through the dynamic mesh.
+		MeshOutAttributes.UnregisterImportPointIndexAttribute();
 	}
 
 	if (bHasBones)
 	{
 		FSkeletalMeshAttributes MeshOutAttributes(MeshOut);
 		checkSlow(MeshOutAttributes.HasBones()); // FSkeletalMeshAttributes::Register() must have been called
+
+		
 		
 		const int32 NumBones = MeshIn->Attributes()->GetNumBones();
+
+		MeshOutAttributes.BoneAttributes();
 		MeshOutAttributes.ReserveNewBones(NumBones);
 
 		for (int Idx = 0; Idx < NumBones; ++Idx)
@@ -744,6 +752,7 @@ void FDynamicMeshToMeshDescription::Convert_NoSharedInstances(const FDynamicMesh
 		}
 	}
 
+
 	// always copy when we are baking new mesh? should this be a config option?
 	bool bCopyInstanceColors = (ColorOverlay != nullptr); 
 
diff --git a/Engine/Source/Runtime/MeshDescription/Private/MeshDescription.cpp b/Engine/Source/Runtime/MeshDescription/Private/MeshDescription.cpp
--- a/Engine/Source/Runtime/MeshDescription/Private/MeshDescription.cpp
+++ b/Engine/Source/Runtime/MeshDescription/Private/MeshDescription.cpp
@@ -603,26 +603,24 @@ void FMeshDescription::RebuildIndexers()
 
 void FMeshDescription::Empty()
 {
-	VertexElements->Reset();
-	VertexInstanceElements->Reset();
-	UVElements->Reset();
-	EdgeElements->Reset();
-	TriangleElements->Reset();
-	PolygonElements->Reset();
-	PolygonGroupElements->Reset();
+	for (TPair<FName, FMeshElementTypeWrapper>& ElementsItem: Elements)
+	{
+		ElementsItem.Value.Get()->Reset();
+	}
 	ResetIndexers();
 }
 
 
 bool FMeshDescription::IsEmpty() const
 {
-	return VertexElements->IsEmpty() &&
-		   VertexInstanceElements->IsEmpty() &&
-		   UVElements->IsEmpty() &&
-		   EdgeElements->IsEmpty() &&
-		   TriangleElements->IsEmpty() &&
-		   PolygonElements->IsEmpty() &&
-		   PolygonGroupElements->IsEmpty();
+	for (const TPair<FName, FMeshElementTypeWrapper>& ElementsItem: Elements)
+	{
+		if (!ElementsItem.Value.Get()->IsEmpty())
+		{
+			return false;
+		}
+	}
+	return true;
 }
 
 bool FMeshDescription::NeedsCompact() const
diff --git a/Engine/Source/Runtime/MeshDescription/Public/MeshAttributeArray.h b/Engine/Source/Runtime/MeshDescription/Public/MeshAttributeArray.h
--- a/Engine/Source/Runtime/MeshDescription/Public/MeshAttributeArray.h
+++ b/Engine/Source/Runtime/MeshDescription/Public/MeshAttributeArray.h
@@ -70,6 +70,7 @@ template <typename... Ts> struct TVariantFromTuple<TTuple<Ts...>> { using Type =
  */
 template <typename T> struct TIsBulkSerializable { static const bool Value = true; };
 template <> struct TIsBulkSerializable<FName> { static const bool Value = false; };
+template <> struct TIsBulkSerializable<FTransform> { static const bool Value = false; };
 
 
 /**
@@ -1279,7 +1280,7 @@ void TMeshAttributesRef<ElementIDType, TArrayView<AttributeType>>::Copy(TMeshAtt
 	check(Extent > 0);
 	check(Src.IsValid());
 	check(Src.Extent == Extent);
-	const TMeshAttributeArrayBase<AttributeType>& SrcArray = static_cast<const BoundedArrayType*>(Src->ArrayPtr)->GetArrayForChannel(SrcChannel);
+	const TMeshAttributeArrayBase<AttributeType>& SrcArray = static_cast<const BoundedArrayType*>(Src.ArrayPtr)->GetArrayForChannel(SrcChannel);
 	TMeshAttributeArrayBase<AttributeType>& DestArray = static_cast<BoundedArrayType*>(ArrayPtr)->GetArrayForChannel(DestChannel);
 	const int32 Num = FMath::Min(SrcArray.Num(), DestArray.Num());
 	for (int32 Index = 0; Index < Num; Index++)
diff --git a/Engine/Source/Runtime/MeshDescription/Public/MeshElementContainer.h b/Engine/Source/Runtime/MeshDescription/Public/MeshElementContainer.h
--- a/Engine/Source/Runtime/MeshDescription/Public/MeshElementContainer.h
+++ b/Engine/Source/Runtime/MeshDescription/Public/MeshElementContainer.h
@@ -318,7 +318,7 @@ public:
 
 			FORCEINLINE ElementIDType operator*() const
 			{
-				return Iterator ? ElementIDType(Iterator.GetIndex()) : INDEX_NONE;
+				return ElementIDType{Iterator ? Iterator.GetIndex() : INDEX_NONE};
 			}
 
 			friend FORCEINLINE bool operator==(const TConstIterator& Lhs, const TConstIterator& Rhs)
diff --git a/Engine/Source/Runtime/MeshDescription/Public/MeshTypes.h b/Engine/Source/Runtime/MeshDescription/Public/MeshTypes.h
--- a/Engine/Source/Runtime/MeshDescription/Public/MeshTypes.h
+++ b/Engine/Source/Runtime/MeshDescription/Public/MeshTypes.h
@@ -289,28 +289,3 @@ struct FPolygonID : public FElementID
 	UE_DEPRECATED(4.26, "Please use INDEX_NONE as an invalid ID.")
 	MESHDESCRIPTION_API static const FPolygonID Invalid;	// @todo mesheditor script: Can we expose these to BP nicely?	Do we even need to?
 };
-
-USTRUCT( BlueprintType )
-struct FBoneID : public FElementID
-{
-	GENERATED_BODY()
-
-	FBoneID()
-	{
-	}
-
-	FBoneID( const FElementID InitElementID )
-		: FElementID( InitElementID.GetValue() )
-	{
-	}
-
-	FBoneID( const int32 InitIDValue )
-		: FElementID( InitIDValue )
-	{
-	}
-
-	FORCEINLINE friend uint32 GetTypeHash( const FBoneID& Other )
-	{
-		return GetTypeHash( Other.IDValue );
-	}
-};
diff --git a/Engine/Source/Runtime/SkeletalMeshDescription/Private/SkeletalMeshAttributes.cpp b/Engine/Source/Runtime/SkeletalMeshDescription/Private/SkeletalMeshAttributes.cpp
--- a/Engine/Source/Runtime/SkeletalMeshDescription/Private/SkeletalMeshAttributes.cpp
+++ b/Engine/Source/Runtime/SkeletalMeshDescription/Private/SkeletalMeshAttributes.cpp
@@ -8,6 +8,7 @@ namespace MeshAttribute
 	namespace Vertex
 	{
 		const FName SkinWeights("SkinWeights");
+		const FName ImportPointIndex("ImportPointIndex");
 	}
 	
 	namespace Bone
@@ -17,12 +18,19 @@ namespace MeshAttribute
 		const FName Pose("Pose");
 		const FName Color("Color");
 	}
+
+	namespace SourceGeometryPart
+	{
+		const FName Name("Name");
+		const FName VertexOffsetAndCount("VertexOffsetAndCount");
+	}
 }
 
 
 FName FSkeletalMeshAttributesShared::DefaultSkinWeightProfileName("Default");
 
 FName FSkeletalMeshAttributesShared::BonesElementName("BonesElementName");
+FName FSkeletalMeshAttributesShared::SourceGeometryPartElementName("SourceGeometryPartElementName");
 
 static FString MorphTargetAttributeNamePrefix("Morph-");
 
@@ -37,15 +45,18 @@ static FString SkinWeightAttributeNamePrefix()
 // FSkeletalMeshAttributes
 //
 
-FSkeletalMeshAttributes::FSkeletalMeshAttributes(FMeshDescription& InMeshDescription)
-:
-FStaticMeshAttributes(InMeshDescription),
-FSkeletalMeshAttributesShared(InMeshDescription)
+FSkeletalMeshAttributes::FSkeletalMeshAttributes(FMeshDescription& InMeshDescription) :
+	FStaticMeshAttributes(InMeshDescription),
+	FSkeletalMeshAttributesShared(InMeshDescription)
 {
 	if (MeshDescription.GetElements().Contains(BonesElementName))
 	{
 		BoneElements = MeshDescription.GetElements()[BonesElementName].Get();
 	}
+	if (MeshDescription.GetElements().Contains(SourceGeometryPartElementName))
+	{
+		SourceGeometryPartElements = MeshDescription.GetElements()[SourceGeometryPartElementName].Get();
+	}
 }
 
 void FSkeletalMeshAttributes::Register(bool bKeepExistingAttribute)
@@ -57,7 +68,7 @@ void FSkeletalMeshAttributes::Register(bool bKeepExistingAttribute)
 
 	if (MeshDescription.GetElements().Contains(BonesElementName) == false)
 	{
-		BoneElements = MeshDescription.GetElements().Emplace(BonesElementName).Get();
+		BoneElementsShared = BoneElements = MeshDescription.GetElements().Emplace(BonesElementName).Get();
 	}
 
 	BoneAttributes().RegisterAttribute<FName>(MeshAttribute::Bone::Name, 1, NAME_None, EMeshAttributeFlags::Mandatory);
@@ -77,6 +88,18 @@ void FSkeletalMeshAttributes::RegisterColorAttribute()
 }
 
 
+bool FSkeletalMeshAttributes::RegisterImportPointIndexAttribute()
+{
+	return MeshDescription.VertexAttributes().RegisterAttribute<int32>(MeshAttribute::Vertex::ImportPointIndex, 1, INDEX_NONE).IsValid();	
+}
+
+
+void FSkeletalMeshAttributes::UnregisterImportPointIndexAttribute()
+{
+	return MeshDescription.VertexAttributes().UnregisterAttribute(MeshAttribute::Vertex::ImportPointIndex);	
+}
+
+
 bool FSkeletalMeshAttributes::RegisterSkinWeightAttribute(const FName InProfileName)
 {
 	if (!IsValidSkinWeightProfileName(InProfileName))
@@ -140,6 +163,29 @@ bool FSkeletalMeshAttributes::RegisterMorphTargetAttribute(const FName InMorphTa
 	return MeshDescription.VertexAttributes().RegisterAttribute<FVector3f[2]>(AttributeName, 1, FVector3f::ZeroVector, EMeshAttributeFlags::None).IsValid();
 }
 
+bool FSkeletalMeshAttributes::UnregisterMorphTargetAttribute(const FName InMorphTargetName)
+{
+	if (InMorphTargetName.IsNone())
+	{
+		return false;
+	}
+
+	const FName AttributeName = CreateMorphTargetAttributeName(InMorphTargetName);
+	if (!ensure(AttributeName.IsValid()))
+	{
+		return false;
+	}
+
+	// Attribute not there?
+	if (!MeshDescription.VertexAttributes().HasAttribute(AttributeName))
+	{
+		return false;
+	}
+
+	MeshDescription.VertexAttributes().UnregisterAttribute(AttributeName);
+	return true;
+}
+
 FMorphTargetVertexAttributesRef FSkeletalMeshAttributes::GetVertexMorphTarget(const FName InMorphTargetName)
 {
 	return MeshDescription.VertexAttributes().GetAttributesRef<TArrayView<FVector3f>>(CreateMorphTargetAttributeName(InMorphTargetName));
@@ -205,20 +251,89 @@ FSkeletalMeshAttributes::FBoneColorAttributesRef FSkeletalMeshAttributes::GetBon
 	return BoneAttributes().GetAttributesRef<FVector4f>(MeshAttribute::Bone::Color);
 }
 
+void FSkeletalMeshAttributes::RegisterSourceGeometryPartsAttributes()
+{
+	if (MeshDescription.GetElements().Contains(SourceGeometryPartElementName) == false)
+	{
+		SourceGeometryPartElementsShared = SourceGeometryPartElements = MeshDescription.GetElements().Emplace(SourceGeometryPartElementName).Get();
+	}
+
+	SourceGeometryPartAttributes().RegisterAttribute<FName>(MeshAttribute::SourceGeometryPart::Name, 1, NAME_None, EMeshAttributeFlags::Mandatory);	
+	SourceGeometryPartAttributes().RegisterAttribute<int32[2]>(MeshAttribute::SourceGeometryPart::VertexOffsetAndCount, 1, {0}, EMeshAttributeFlags::Mandatory);	
+}
+
+FSkeletalMeshAttributesShared::FSourceGeometryPartArray& FSkeletalMeshAttributes::SourceGeometryParts()
+{
+	if (!ensure(HasSourceGeometryParts()))
+	{
+		static FSourceGeometryPartArray Empty;
+		return Empty;
+	}
+	return static_cast<FSourceGeometryPartArray&>(SourceGeometryPartElements->Get());
+}
+
+TAttributesSet<FSourceGeometryPartID>& FSkeletalMeshAttributes::SourceGeometryPartAttributes()
+{
+	if (!ensure(HasSourceGeometryParts()))
+	{
+		static TAttributesSet<FSourceGeometryPartID> Empty;
+		return Empty;
+	}
+	return SourceGeometryParts().GetAttributes();
+}
+
+FSourceGeometryPartID FSkeletalMeshAttributes::CreateSourceGeometryPart()
+{
+	if (!ensure(HasSourceGeometryParts()))
+	{
+		return {};
+	}
+	
+	return SourceGeometryParts().Add();
+}
+
+void FSkeletalMeshAttributes::DeleteSourceGeometryPart(FSourceGeometryPartID InSourceGeometryPartID)
+{
+	if (ensure(HasSourceGeometryParts()))
+	{
+		return SourceGeometryParts().Remove(InSourceGeometryPartID);
+	}
+}
+
+FSkeletalMeshAttributesShared::FSourceGeometryPartNameRef FSkeletalMeshAttributes::GetSourceGeometryPartNames()
+{
+	if (!ensure(HasSourceGeometryParts()))
+	{
+		return FSourceGeometryPartNameRef{};
+	}
+	return SourceGeometryPartAttributes().GetAttributesRef<FName>(MeshAttribute::SourceGeometryPart::Name);
+}
+
+FSkeletalMeshAttributesShared::FSourceGeometryPartVertexOffsetAndCountRef FSkeletalMeshAttributes::GetSourceGeometryPartVertexOffsetAndCounts()
+{
+	if (!ensure(HasSourceGeometryParts()))
+	{
+		return FSourceGeometryPartVertexOffsetAndCountRef{};
+	}
+	return SourceGeometryPartAttributes().GetAttributesRef<TArrayView<int32>>(MeshAttribute::SourceGeometryPart::VertexOffsetAndCount);
+}
 
 
 //
 // FSkeletalMeshAttributesShared
 //
 
-FSkeletalMeshAttributesShared::FSkeletalMeshAttributesShared(const FMeshDescription& InMeshDescription) 
-:
-MeshDescriptionShared(InMeshDescription)
+FSkeletalMeshAttributesShared::FSkeletalMeshAttributesShared(const FMeshDescription& InMeshDescription) :
+	MeshDescriptionShared(InMeshDescription)
 {
 	if (MeshDescriptionShared.GetElements().Contains(BonesElementName))
 	{
 		BoneElementsShared = MeshDescriptionShared.GetElements()[BonesElementName].Get();
 	}
+	if (MeshDescriptionShared.GetElements().Contains(SourceGeometryPartElementName))
+	{
+		SourceGeometryPartElementsShared = MeshDescriptionShared.GetElements()[SourceGeometryPartElementName].Get();
+	}
 }
 
 FSkinWeightsVertexAttributesConstRef FSkeletalMeshAttributesShared::GetVertexSkinWeights(const FName InProfileName) const
@@ -232,10 +347,8 @@ FSkinWeightsVertexAttributesConstRef FSkeletalMeshAttributesShared::GetVertexSki
 	{
 		return MeshDescriptionShared.VertexAttributes().GetAttributesRef<TArrayAttribute<int32>>(InAttributeName);
 	}
-	else
-	{
-		return {};
-	}
+	
+	return {};
 }
 
 TArray<FName> FSkeletalMeshAttributesShared::GetSkinWeightProfileNames() const
@@ -421,4 +534,56 @@ FSkeletalMeshAttributesShared::FBonePoseAttributesConstRef FSkeletalMeshAttribut
 FSkeletalMeshAttributesShared::FBoneColorAttributesConstRef FSkeletalMeshAttributesShared::GetBoneColors() const
 {
 	return BoneAttributes().GetAttributesRef<FVector4f>(MeshAttribute::Bone::Color);
-}
\ No newline at end of file
+}
+
+const FSkeletalMeshAttributesShared::FSourceGeometryPartArray& FSkeletalMeshAttributesShared::SourceGeometryParts() const
+{
+	if (!ensure(HasSourceGeometryParts()))
+	{
+		static FSourceGeometryPartArray Empty;
+		return Empty;
+	}
+	return static_cast<const FSourceGeometryPartArray&>(SourceGeometryPartElementsShared->Get());
+}
+
+const TAttributesSet<FSourceGeometryPartID>& FSkeletalMeshAttributesShared::SourceGeometryPartAttributes() const
+{
+	if (!ensure(HasSourceGeometryParts()))
+	{
+		static TAttributesSet<FSourceGeometryPartID> Empty;
+		return Empty;
+	}
+	return SourceGeometryParts().GetAttributes();
+}
+
+int32 FSkeletalMeshAttributesShared::GetNumSourceGeometryParts() const
+{
+	return HasSourceGeometryParts() ? SourceGeometryPartElementsShared->Get().Num() : 0;
+}
+
+bool FSkeletalMeshAttributesShared::IsSourceGeometryPartValid(const FSourceGeometryPartID InSourceGeometryPartID) const
+{
+	if (!ensure(HasSourceGeometryParts()))
+	{
+		return false;
+	}
+	return SourceGeometryParts().IsValid(InSourceGeometryPartID);
+}
+
+FSkeletalMeshAttributesShared::FSourceGeometryPartNameConstRef FSkeletalMeshAttributesShared::GetSourceGeometryPartNames() const
+{
+	if (!ensure(HasSourceGeometryParts()))
+	{
+		return FSourceGeometryPartNameConstRef{};
+	}
+	return SourceGeometryPartAttributes().GetAttributesRef<FName>(MeshAttribute::SourceGeometryPart::Name);
+}
+
+FSkeletalMeshAttributesShared::FSourceGeometryPartVertexOffsetAndCountConstRef FSkeletalMeshAttributesShared::GetSourceGeometryPartVertexOffsetAndCounts() const
+{
+	if (!ensure(HasSourceGeometryParts()))
+	{
+		return FSourceGeometryPartVertexOffsetAndCountConstRef{};
+	}
+	return SourceGeometryPartAttributes().GetAttributesRef<TArrayView<int32>>(MeshAttribute::SourceGeometryPart::VertexOffsetAndCount);
+}
