Skeletal Mesh: Build fix.

#rnx

diff --git a/Engine/Source/Runtime/Engine/Private/SkeletalMesh.cpp b/Engine/Source/Runtime/Engine/Private/SkeletalMesh.cpp
--- a/Engine/Source/Runtime/Engine/Private/SkeletalMesh.cpp
+++ b/Engine/Source/Runtime/Engine/Private/SkeletalMesh.cpp
@@ -95,8 +95,6 @@
 #include "SkeletalMeshAttributes.h"
 #include "Logging/StructuredLog.h"
 
-UE_DISABLE_OPTIMIZATION
-
 #define LOCTEXT_NAMESPACE "SkeltalMesh"
 
 DEFINE_LOG_CATEGORY(LogSkeletalMesh);
@@ -7306,6 +7304,4 @@ void GetSkeletalMeshRefTangentBasis(const USkeletalMesh* Mesh, const FSkeletalMe
 	GetRefTangentBasisTyped(Mesh, Section, LODData.StaticVertexBuffers.StaticMeshVertexBuffer, SkinWeightVertexBuffer, VertIndexInChunk, OutTangentX, OutTangentY, OutTangentZ);
 }
 
-UE_ENABLE_OPTIMIZATION
-
 #undef LOCTEXT_NAMESPACE
diff --git a/Engine/Source/Runtime/SkeletalMeshDescription/Public/MorphTargetAttributesRef.h b/Engine/Source/Runtime/SkeletalMeshDescription/Public/MorphTargetAttributesRef.h
--- a/Engine/Source/Runtime/SkeletalMeshDescription/Public/MorphTargetAttributesRef.h
+++ b/Engine/Source/Runtime/SkeletalMeshDescription/Public/MorphTargetAttributesRef.h
@@ -4,6 +4,8 @@
 
 #include "MeshDescription.h"
 
+class FMorphTargetVertexAttributesConstRef;
+
 class FMorphTargetVertexAttributesRef
 {
 	friend class FSkeletalMeshAttributes;
