PR #6821: Fix to prevent anim montages from continuing to play animations past the end of the the last section.

diff --git a/Engine/Source/Runtime/Engine/Private/Animation/AnimMontage.cpp b/Engine/Source/Runtime/Engine/Private/Animation/AnimMontage.cpp
--- a/Engine/Source/Runtime/Engine/Private/Animation/AnimMontage.cpp
+++ b/Engine/Source/Runtime/Engine/Private/Animation/AnimMontage.cpp
@@ -2563,6 +2563,11 @@ void FAnimMontageInstance::Advance(float DeltaTime, struct FRootMotionMovementPa
 						else
 						{
 							// If there is no next section and we've reached the end of this one, exit
+
+							// Stop playing to prevent playing animation data past the end of the current section
+							// We already called Stop above if needed, like if bEnableAutoBlendOut is true
+							bPlaying = false;
+
 							break;
 						}
 					}
