Avoid log message being treated as an error

diff --git a/Engine/Source/Runtime/Engine/Private/Animation/AnimCompress.cpp b/Engine/Source/Runtime/Engine/Private/Animation/AnimCompress.cpp
--- a/Engine/Source/Runtime/Engine/Private/Animation/AnimCompress.cpp
+++ b/Engine/Source/Runtime/Engine/Private/Animation/AnimCompress.cpp
@@ -118,11 +118,11 @@ FCompressionMemorySummary::~FCompressionMemorySummary()
 		FText Message;
 		if (!bPerformedCompression)
 		{
-			Message = FText::Format(NSLOCTEXT("Engine", "CompressionMemorySummaryNoChange", "Fetched compressed data for {NumberOfAnimations} Animation(s)\n\nRaw: {TotalRaw} - Compressed: {TotalAfterCompressed}\nSaving: {TotalAfterSaving} ({NewCompressionRatio})\n\nEnd Effector Translation Added By Compression:\n Average: {AverageError} Max:\n{WorstBoneError}\n\nMax Average Animation Error:\n{WorstAnimationError}"), Args);
+			Message = FText::Format(NSLOCTEXT("Engine", "CompressionMemorySummaryNoChange", "Fetched compressed data for {NumberOfAnimations} Animation(s)\n\nRaw: {TotalRaw} - Compressed: {TotalAfterCompressed}\nSaving: {TotalAfterSaving} ({NewCompressionRatio})\n\nEnd Effector Translation Added By Compression:\n Average: {AverageError} Max:\n{WorstBoneError}\n\nWorst Animation Error Seen:\n{WorstAnimationError}"), Args);
 		}
 		else
 		{
-			Message = FText::Format(NSLOCTEXT("Engine", "CompressionMemorySummary", "Compressed {NumberOfAnimations} Animation(s)\n\nPre Compression:\n\nRaw: {TotalRaw} - Compressed: {TotalBeforeCompressed}\nSaving: {TotalBeforeSaving} ({OldCompressionRatio})\n\nPost Compression:\n\nRaw: {TotalRaw} - Compressed: {TotalAfterCompressed}\nSaving: {TotalAfterSaving} ({NewCompressionRatio})\n\nTotal Compression Time: {TotalTimeSpentCompressingPretty} (Seconds: {TotalTimeSpentCompressingRawSeconds})\n\nEnd Effector Translation Added By Compression:\n Average: {AverageError} Max:\n{WorstBoneError}\n\nMax Average Animation Error:\n{WorstAnimationError}"), Args);
+			Message = FText::Format(NSLOCTEXT("Engine", "CompressionMemorySummary", "Compressed {NumberOfAnimations} Animation(s)\n\nPre Compression:\n\nRaw: {TotalRaw} - Compressed: {TotalBeforeCompressed}\nSaving: {TotalBeforeSaving} ({OldCompressionRatio})\n\nPost Compression:\n\nRaw: {TotalRaw} - Compressed: {TotalAfterCompressed}\nSaving: {TotalAfterSaving} ({NewCompressionRatio})\n\nTotal Compression Time: {TotalTimeSpentCompressingPretty} (Seconds: {TotalTimeSpentCompressingRawSeconds})\n\nEnd Effector Translation Added By Compression:\n Average: {AverageError} Max:\n{WorstBoneError}\n\nWorst Animation Error Seen:\n{WorstAnimationError}"), Args);
 		}
 
 		UE_LOG(LogAnimationCompression, Display, TEXT("Top 10 Worst Bone Errors:"));
