PR #6071: Adding AnimDrawDebugCapsule to AnimInstanceProxy for debug drawing purposes

diff --git a/Engine/Source/Runtime/Engine/Private/Animation/AnimInstanceProxy.cpp b/Engine/Source/Runtime/Engine/Private/Animation/AnimInstanceProxy.cpp
--- a/Engine/Source/Runtime/Engine/Private/Animation/AnimInstanceProxy.cpp
+++ b/Engine/Source/Runtime/Engine/Private/Animation/AnimInstanceProxy.cpp
@@ -681,6 +681,7 @@ void FAnimInstanceProxy::PostUpdate(UAnimInstance* InAnimInstance) const
 				break;
 			}
 			case EDrawDebugItemType::Cone: DrawDebugCone(InAnimInstance->GetSkelMeshComponent()->GetWorld(), DebugItem.Center, DebugItem.Direction.Value, DebugItem.Length, DebugItem.AngleWidth, DebugItem.AngleHeight, DebugItem.Segments, DebugItem.Color, DebugItem.bPersistentLines, DebugItem.LifeTime, DebugItem.DepthPriority, DebugItem.Thickness); break;
+			case EDrawDebugItemType::Capsule: DrawDebugCapsule(InAnimInstance->GetSkelMeshComponent()->GetWorld(), DebugItem.Center, DebugItem.Size, DebugItem.Radius, DebugItem.Rotation.Quaternion(), DebugItem.Color, DebugItem.bPersistentLines, DebugItem.LifeTime, 0, DebugItem.Thickness); break;
 		}
 	}
 #endif
@@ -2306,6 +2307,22 @@ void FAnimInstanceProxy::AnimDrawDebugCone(const FVector& Center, float Length,
 	QueuedDrawDebugItems.Add(DrawDebugItem);
 }
 
+void FAnimInstanceProxy::AnimDrawDebugCapsule(const FVector& Center, float HalfHeight, float Radius, const FRotator& Rotation, const FColor& Color, bool bPersistentLines /*= false*/, float LifeTime /*= -1.f*/, float Thickness /*= 0.f*/)
+{
+	FQueuedDrawDebugItem DrawDebugItem;
+
+	DrawDebugItem.ItemType = EDrawDebugItemType::Capsule;
+	DrawDebugItem.Center = Center;
+	DrawDebugItem.Size = HalfHeight;
+	DrawDebugItem.Radius = Radius;
+	DrawDebugItem.Rotation = Rotation;
+	DrawDebugItem.Color = Color;
+	DrawDebugItem.bPersistentLines = bPersistentLines;
+	DrawDebugItem.LifeTime = LifeTime;
+	DrawDebugItem.Thickness = Thickness;
+
+	QueuedDrawDebugItems.Add(DrawDebugItem);
+}
 
 #endif // ENABLE_ANIM_DRAW_DEBUG
 
diff --git a/Engine/Source/Runtime/Engine/Public/Animation/AnimInstanceProxy.h b/Engine/Source/Runtime/Engine/Public/Animation/AnimInstanceProxy.h
--- a/Engine/Source/Runtime/Engine/Public/Animation/AnimInstanceProxy.h
+++ b/Engine/Source/Runtime/Engine/Public/Animation/AnimInstanceProxy.h
@@ -74,7 +74,8 @@ namespace EDrawDebugItemType
 		Point,
 		Circle,
 		Cone,
-		InWorldMessage
+		InWorldMessage,
+		Capsule,
 	};
 }
 
@@ -458,6 +459,7 @@ public:
 	ENGINE_API void AnimDrawDebugPoint(const FVector& Loc, float Size, const FColor& Color, bool bPersistentLines = false, float LifeTime = -1.f, ESceneDepthPriorityGroup DepthPriority = SDPG_World);
 	ENGINE_API void AnimDrawDebugCircle(const FVector& Center, float Radius, int32 Segments, const FColor& Color, const FVector& UpVector = FVector::UpVector, bool bPersistentLines = false, float LifeTime = -1.f, ESceneDepthPriorityGroup DepthPriority = SDPG_World, float Thickness = 0.f);
 	ENGINE_API void AnimDrawDebugCone(const FVector& Center, float Radius, const FVector& Direction, float AngleWidth, float AngleHeight, int32 Segments, const FColor& Color, bool bPersistentLines = false, float LifeTime = -1.f, ESceneDepthPriorityGroup DepthPriority = SDPG_World, float Thickness = 0.f);
+	ENGINE_API void AnimDrawDebugCapsule(const FVector& Center, float HalfHeight, float Radius, const FRotator& Rotation, const FColor& Color, bool bPersistentLines = false, float LifeTime = -1.f, float Thickness = 0.f);
 #else
 	void AnimDrawDebugOnScreenMessage(const FString& DebugMessage, const FColor& Color, const FVector2D& TextScale = FVector2D::UnitVector, ESceneDepthPriorityGroup DepthPriority = SDPG_World) {}
 	void AnimDrawDebugInWorldMessage(const FString& DebugMessage, const FVector& TextLocation, const FColor& Color, float TextScale) {}
@@ -469,6 +471,7 @@ public:
 	void AnimDrawDebugPoint(const FVector& Loc, float Size, const FColor& Color, bool bPersistentLines = false, float LifeTime = -1.f, ESceneDepthPriorityGroup DepthPriority = SDPG_World) {}
 	void AnimDrawDebugCircle(const FVector& Center, float Radius, int32 Segments, const FColor& Color, const FVector& UpVector = FVector::UpVector, bool bPersistentLines = false, float LifeTime=-1.f, ESceneDepthPriorityGroup DepthPriority = SDPG_World, float Thickness = 0.f) {}
 	void AnimDrawDebugCone(const FVector& Center, float Radius, const FVector& Direction, float AngleWidth, float AngleHeight, int32 Segments, const FColor& Color, bool bPersistentLines = false, float LifeTime = -1.f, ESceneDepthPriorityGroup DepthPriority = SDPG_World, float Thickness = 0.f) {}
+	ENGINE_API void AnimDrawDebugCapsule(const FVector& Center, float HalfHeight, float Radius, const FRotator& Rotation, const FColor& Color, bool bPersistentLines = false, float LifeTime = -1.f, float Thickness = 0.f) {}
 #endif // ENABLE_ANIM_DRAW_DEBUG
 
 #if ENABLE_ANIM_LOGGING
