Fix follower component lagging behind lead component. Dynamic data updates for skinned mesh components happen in parallel, i.e. follower component can update before lead component. As follower component relies on lead component's current revision number, updating the revision number in dynamic data updates create race condition. So change back to updating current frame revision number immediately when animation updates, and add a previous revision number variable to help calculate current revision number.

diff --git a/Engine/Source/Runtime/Engine/Classes/Components/SkinnedMeshComponent.h b/Engine/Source/Runtime/Engine/Classes/Components/SkinnedMeshComponent.h
--- a/Engine/Source/Runtime/Engine/Classes/Components/SkinnedMeshComponent.h
+++ b/Engine/Source/Runtime/Engine/Classes/Components/SkinnedMeshComponent.h
@@ -297,7 +297,8 @@ public:
 	ENGINE_API void ForceMotionVector();
 
 private:
-	ENGINE_API EPreviousBoneTransformUpdateMode UpdateBoneTransformRevisionNumber();
+	ENGINE_API void UpdateBoneTransformRevisionNumber();
+	ENGINE_API EPreviousBoneTransformUpdateMode GetPreviousBoneTransformUpdateMode();
 
 	enum class EBoneTransformUpdateMethod
 	{
@@ -329,6 +330,9 @@ protected:
 	/** The index for the ComponentSpaceTransforms buffer we can currently read from */
 	int32 CurrentReadComponentTransforms;
 
+	/** Cache previous bone transform revision number to help compute CurrentBoneTransformRevisionNumber */
+	uint32 PreviousBoneTransformRevisionNumber;
+
 	/** current bone transform revision number */
 	uint32 CurrentBoneTransformRevisionNumber;
 
diff --git a/Engine/Source/Runtime/Engine/Private/Components/SkinnedMeshComponent.cpp b/Engine/Source/Runtime/Engine/Private/Components/SkinnedMeshComponent.cpp
--- a/Engine/Source/Runtime/Engine/Private/Components/SkinnedMeshComponent.cpp
+++ b/Engine/Source/Runtime/Engine/Private/Components/SkinnedMeshComponent.cpp
@@ -435,6 +435,7 @@ USkinnedMeshComponent::USkinnedMeshComponent(const FObjectInitializer& ObjectIni
 	bSyncAttachParentLOD = true;
 	bIgnoreLeaderPoseComponentLOD = false;
 
+	PreviousBoneTransformRevisionNumber = 0;
 	CurrentBoneTransformRevisionNumber = 0;
 
 	ExternalInterpolationAlpha = 0.0f;
@@ -867,21 +868,28 @@ void USkinnedMeshComponent::RefreshExternalMorphTargetWeights(bool bZeroOldWeigh
 	}
 }
 
-EPreviousBoneTransformUpdateMode USkinnedMeshComponent::UpdateBoneTransformRevisionNumber()
+void USkinnedMeshComponent::UpdateBoneTransformRevisionNumber()
+{
+	if (BoneTransformUpdateMethodQueue.Last() == EBoneTransformUpdateMethod::ClearMotionVector)
+	{
+		// Last entry is ClearMotionVector, increment revision number by 2 which allows current bone buffer to be bound to previous shader slot to cancel out velocity
+		// See FGPUBaseSkinVertexFactory::FShaderDataType::GetBoneBufferInternal() for how it's used.
+		CurrentBoneTransformRevisionNumber = PreviousBoneTransformRevisionNumber + 2;
+	}
+	else
+	{
+		CurrentBoneTransformRevisionNumber = PreviousBoneTransformRevisionNumber + 1;
+	}
+}
+
+EPreviousBoneTransformUpdateMode USkinnedMeshComponent::GetPreviousBoneTransformUpdateMode()
 {
 	EPreviousBoneTransformUpdateMode PreviousUpdateMode = EPreviousBoneTransformUpdateMode::None;
 
 	if (BoneTransformUpdateMethodQueue.Num() > 0)
 	{
-		if (BoneTransformUpdateMethodQueue.Last() == EBoneTransformUpdateMethod::ClearMotionVector)
-		{
-			// Last entry is ClearMotionVector, increment revision number by 2 which allows current bone buffer to be bound to previous shader slot to cancel out velocity
-			// See FGPUBaseSkinVertexFactory::FShaderDataType::GetBoneBufferInternal() for how it's used.
-			CurrentBoneTransformRevisionNumber += 2;
-		}
-		else
+		if (BoneTransformUpdateMethodQueue.Last() != EBoneTransformUpdateMethod::ClearMotionVector)
 		{
-			CurrentBoneTransformRevisionNumber++;
 			if (BoneTransformUpdateMethodQueue.Num() > 1)
 			{
 				// When there are multiple entries in the queue, previous bone transforms may or may not need updating, so for simplicity make sure it updates.
@@ -908,8 +916,6 @@ EPreviousBoneTransformUpdateMode USkinnedMeshComponent::UpdateBoneTransformRevis
 				}
 			}
 		}
-
-		BoneTransformUpdateMethodQueue.Reset();
 	}
 
 	return PreviousUpdateMode;
@@ -994,8 +1000,8 @@ void USkinnedMeshComponent::CreateRenderState_Concurrent(FRegisterComponentConte
 
 	if (GetSkinnedAsset())
 	{
-		// Update revision number before mesh object update. We need to update regardless if MeshObject exists.
-		UpdateBoneTransformRevisionNumber();
+		BoneTransformUpdateMethodQueue.Reset();
+		PreviousBoneTransformRevisionNumber = CurrentBoneTransformRevisionNumber;
 
 		// Update dynamic data
 		if(MeshObject)
@@ -1119,8 +1125,10 @@ void USkinnedMeshComponent::SendRenderDynamicData_Concurrent()
 	}
 #endif
 
-	// Update revision number before mesh object update. We need to update regardless if MeshObject exists.
-	EPreviousBoneTransformUpdateMode PreviousBoneTransformUpdateMode = UpdateBoneTransformRevisionNumber();
+	EPreviousBoneTransformUpdateMode PreviousBoneTransformUpdateMode = GetPreviousBoneTransformUpdateMode();
+	// CurrentBoneTransformRevisionNumber and PreviousBoneTransformUpdateMode are up-to-date at this point, it is safe to reset BoneTransformUpdateMethodQueue
+	BoneTransformUpdateMethodQueue.Reset();
+	PreviousBoneTransformRevisionNumber = CurrentBoneTransformRevisionNumber;
 
 	// if we have not updated the transforms then no need to send them to the rendering thread
 	// @todo GIsEditor used to be bUpdateSkelWhenNotRendered. Look into it further to find out why it doesn't update animations in the AnimSetViewer, when a level is loaded in UED (like POC_Cover.gear).
@@ -1190,6 +1198,8 @@ void USkinnedMeshComponent::ClearMotionVector()
 		// if you have situation where you want to clear the bone velocity (that causes temporal AA or motion blur)
 		// use this function to clear it
 		BoneTransformUpdateMethodQueue.Add(EBoneTransformUpdateMethod::ClearMotionVector);
+		// Similar to FlipEditableSpaceBases(), update revision number before reaching dynamic update to make sure dependency between lead and follower components are preserved.
+		UpdateBoneTransformRevisionNumber();
 		// Make sure an update happens
 		bForceMeshObjectUpdate = true;
 		MarkRenderDynamicDataDirty();
@@ -1208,6 +1218,8 @@ void USkinnedMeshComponent::ForceMotionVector()
 		}
 
 		BoneTransformUpdateMethodQueue.Add(EBoneTransformUpdateMethod::ForceMotionVector);
+		// Similar to FlipEditableSpaceBases(), update revision number before reaching dynamic update to make sure dependency between lead and follower components are preserved.
+		UpdateBoneTransformRevisionNumber();
 		// Make sure an update happens
 		bForceMeshObjectUpdate = true;
 		MarkRenderDynamicDataDirty();
@@ -3960,6 +3972,9 @@ void USkinnedMeshComponent::FlipEditableSpaceBases()
 		}
 
 		BoneTransformUpdateMethodQueue.Add(EBoneTransformUpdateMethod::AnimationUpdate);
+		// Bone revision number needs to be updated immediately, because dynamic updates on components are run in parallel later,
+		// for a follower component it relies on its lead component to be up-to-date, so updating the lead component revision number here guarantees it.
+		UpdateBoneTransformRevisionNumber();
 	}
 }
 
