Skeletal Mesh: Ensure we only iterate over existing source models when converting legacy bulk data on save. Also ensure that afterwards that LOD count and source model count matches.

#rnx

diff --git a/Engine/Source/Runtime/Engine/Private/SkeletalMesh.cpp b/Engine/Source/Runtime/Engine/Private/SkeletalMesh.cpp
--- a/Engine/Source/Runtime/Engine/Private/SkeletalMesh.cpp
+++ b/Engine/Source/Runtime/Engine/Private/SkeletalMesh.cpp
@@ -1583,14 +1583,17 @@ void USkeletalMesh::Serialize( FArchive& Ar )
 		FSkinnedAssetCompilingManager::Get().FinishCompilation({ this });
 	}
 
-	// If saving out to disk, ensure that all source models have had their raw mesh bulk data converted to mesh description,
-	// since it won't be reloaded.
 	if (Ar.IsSaving())
 	{
-		for (int32 LODIndex = 0; LODIndex < GetLODNum(); LODIndex++)
+		// If saving out to disk, ensure that all source models have had their raw mesh bulk data converted to mesh description,
+		// since the bulk data won't be reloaded.
+		for (int32 LODIndex = 0; LODIndex < GetNumSourceModels(); LODIndex++)
 		{
 			GetSourceModel(LODIndex).EnsureRawMeshBulkDataIsConvertedToNew();
 		}
+		
+		// Ensure source models and LODs match.
+		SetNumSourceModels(GetLODNum());
 	}
 	
 #endif
