Add a console command to not draw the green bounding box on the current debug target ("r.Debug.DrawCurrentDebugTargetBoundingBox")

diff --git a/Engine/Source/Runtime/Engine/Private/HUD.cpp b/Engine/Source/Runtime/Engine/Private/HUD.cpp
--- a/Engine/Source/Runtime/Engine/Private/HUD.cpp
+++ b/Engine/Source/Runtime/Engine/Private/HUD.cpp
@@ -56,6 +56,13 @@ TAutoConsoleVariable<int32> GMaxDebugTextStringsPerActorCVar(
 	128,
 	TEXT("The maximum number of debug strings that can be attached to a given actor (<=0 : no limit)"));
 
+#if ENABLE_DRAW_DEBUG
+TAutoConsoleVariable<int32> GDrawCurrentDebugTargetBoundingBox(
+	TEXT("r.Debug.DrawCurrentDebugTargetBoundingBox"),
+	1,
+	TEXT("Draw the bounding box of the currently selected debug target (Default: 1)"));
+#endif // ENABLE_DRAW_DEBUG
+
 AHUD::AHUD(const FObjectInitializer& ObjectInitializer)
 	: Super(ObjectInitializer)
 {
@@ -440,6 +447,7 @@ void AHUD::ShowDebugInfo(float& YL, float& YPos)
 		{
 			// Draw box around Actor being debugged.
 #if ENABLE_DRAW_DEBUG
+			if (GDrawCurrentDebugTargetBoundingBox.GetValueOnGameThread() > 0)
 			{
 				FVector BoundsOrigin, BoundsExtent;
 				ShowDebugTargetActor->GetActorBounds(true, BoundsOrigin, BoundsExtent);
