Robomerge conflict resolve on behalf of Yoan.Stamant

[FActorInstanceHandle] added lazy resolution to allow non game threads to create a handle from a valid weak component pointer. The handle will be resolved when accessed from the game thread.
- MakeActorHandleToResolve can be called from any thread with an initialized WeakObjectPtr. In that case the handle is mark as 'needs resolving'
- CVar "IA.ValidateAccessFromGameThread"  to enable access validation (disabled for now since they are still some cases to address before enabling by default)

diff --git a/Engine/Source/Runtime/Engine/Classes/Engine/ActorInstanceHandle.h b/Engine/Source/Runtime/Engine/Classes/Engine/ActorInstanceHandle.h
--- a/Engine/Source/Runtime/Engine/Classes/Engine/ActorInstanceHandle.h
+++ b/Engine/Source/Runtime/Engine/Classes/Engine/ActorInstanceHandle.h
@@ -13,7 +13,21 @@ class ULevel;
 
 using FActorInstanceManagerInterface = TWeakInterfacePtr<IActorInstanceManagerInterface>;
 
-// Handle to a unique object. This may specify a full weigh actor or it may only specify the actor instance that represents the same object.
+/**
+ * Handle to a unique object. This may specify a full weigh actor or it may only specify the actor instance that represents the same object.
+ * 
+ * @note The handle has game thread constraints related to UObjects and should be used carefully from other threads.
+ * 
+ *	Can only be used on the game thread
+ *	-	all constructors 
+ *	-	all getters (GetXYZ, FetchActor, IsActorValid, DoesRepresent) *	
+ *	-	comparison operators against live AActor pointer
+ *	
+ *	Can be used on any thread
+ *	-	MakeActorHandleToResolve to create a handle that will be lazily resolved on the game thread
+ *		since it only stores a weak object ptr without any access to the live object
+ *	-	handle validity and comparison operators against another handle (i.e. IsValid(), operator==|!=(const FActorInstanceHandle& Other))
+ */
 USTRUCT(BlueprintType)
 struct FActorInstanceHandle
 {
@@ -33,6 +47,12 @@ struct FActorInstanceHandle
 	 */
 	static FActorInstanceHandle MakeDehydratedActorHandle(UObject& Manager, int32 InInstanceIndex);
 
+	/** 
+	 * A path dedicated to creation of handles from any threads.
+	 * This path marks the handle as need resolving once it gets accessed from the game thread.
+	 */
+	static FActorInstanceHandle MakeActorHandleToResolve(const TWeakObjectPtr<UPrimitiveComponent>& WeakComponent, int32 CollisionInstanceIndex);
+
 	ENGINE_API bool IsValid() const;
 
 	ENGINE_API bool DoesRepresentClass(const UClass* OtherClass) const;
@@ -61,7 +81,7 @@ struct FActorInstanceHandle
 	template <typename T>
 	T* FetchActor() const;
 
-	AActor* GetCachedActor() const { return Actor.Get(); }
+	ENGINE_API AActor* GetCachedActor() const;
 	ENGINE_API void SetCachedActor(AActor* InActor) const;
 
 	/* Returns the index used internally by the manager */
@@ -96,21 +116,42 @@ private:
 	 * helper functions that let us treat the actor pointer as a UObject in templated functions
 	 * these do NOT fetch the actor so they will return nullptr if we don't have a full actor representation
 	 */
-	ENGINE_API UObject* GetActorAsUObject();
-	ENGINE_API const UObject* GetActorAsUObject() const;
+	UObject* GetActorAsUObject();
+	const UObject* GetActorAsUObject() const;
 
 	/** Returns true if Actor is not null and not pending kill */
-	ENGINE_API bool IsActorValid() const;
+	bool IsActorValid() const;
 
-	/** this is cached here for convenience */
+	void ResolveHandle() const;
+	
+	/**
+	 * Weak UObject pointer used for two purposes:
+	 *  - a resolved handle uses it to store the AActor
+	 *  - a handle to be resolved uses it to store the UPrimitiveComponent provided by MakeActorHandleToResolve 
+	 */
 	UPROPERTY()
-	mutable TWeakObjectPtr<AActor> Actor;
+	mutable TWeakObjectPtr<UObject> ReferenceObject;
 
 	/** Identifies the actor instance manager to use */
 	FActorInstanceManagerInterface ManagerInterface;
 
 	/** Identifies the instance within the manager */
 	int32 InstanceIndex = INDEX_NONE;
+
+	/**
+	 * Enum to keep track of the resolution status of the handle.
+	 * It is only possible to safely resolve the handle on the game thread so
+	 * other threads should use MakeActorHandleToResolve to safely create one. 
+	 */
+	enum class EResolutionStatus : uint8
+	{
+		Invalid,
+		Resolved, /* ManagerInterface and InstanceIndex are set or an Actor pointer is stored in ReferenceObject */
+		NeedsResolving /* Component pointer is stored in ReferenceObject and InstanceIndex might hold a CollisionInstanceIndex */
+	};
+
+	/** Indicates if the handle is resolved, needs to be resolved (extract Actor|ManagerInterface) */
+	EResolutionStatus ResolutionStatus = EResolutionStatus::Invalid;
 };
 
 template<typename T>
diff --git a/Engine/Source/Runtime/Engine/Private/Collision/CollisionConversions.cpp b/Engine/Source/Runtime/Engine/Private/Collision/CollisionConversions.cpp
--- a/Engine/Source/Runtime/Engine/Private/Collision/CollisionConversions.cpp
+++ b/Engine/Source/Runtime/Engine/Private/Collision/CollisionConversions.cpp
@@ -133,7 +133,7 @@ static void SetHitResultFromShapeAndFaceIndex(const FPhysicsShape& Shape,  const
 	CHAOS_CHECK(ShapeIndex < (int32)TNumericLimits<uint8>::Max()); // I could just write < 256, but this makes it more clear *why*
 	OutResult.ElementIndex = (uint8)ShapeIndex;
 	
-	UPrimitiveComponent* OwningComponent = nullptr;
+	TWeakObjectPtr<UPrimitiveComponent> OwningComponent;
 	if(const FBodyInstance* BodyInst = GetUserData(Actor))
 	{
 		BodyInst = FPhysicsInterface::ShapeToOriginalBodyInstance(BodyInst, &Shape);
@@ -146,7 +146,7 @@ static void SetHitResultFromShapeAndFaceIndex(const FPhysicsShape& Shape,  const
 			OutResult.BoneName = BodySetup->BoneName;
 		}
 
-		OwningComponent = BodyInst->OwnerComponent.Get();
+		OwningComponent = BodyInst->OwnerComponent;
 	}
 	else
 	{
@@ -182,10 +182,13 @@ static void SetHitResultFromShapeAndFaceIndex(const FPhysicsShape& Shape,  const
 	OutResult.PhysMaterial = nullptr;
 
 	// Grab actor/component
-	if( OwningComponent )
+	if (!OwningComponent.IsExplicitlyNull())
 	{
 		OutResult.Component = OwningComponent;
-		OutResult.HitObjectHandle = FActorInstanceHandle(OwningComponent, OutResult.Item);
+
+		// Create a handle that needs to be resolved later when accessed in game thread,
+		// since resolving FActorInstanceHandle may require to access related UObjects (e.g. owner actor), which is not safe outside game thread.
+		OutResult.HitObjectHandle = FActorInstanceHandle::MakeActorHandleToResolve(OwningComponent, OutResult.Item);
 
 		if (bReturnPhysMat)
 		{
@@ -602,7 +605,9 @@ void ConvertQueryOverlap(const FPhysicsShape& Shape, const FPhysicsActor& Actor,
         BodyInst = FPhysicsInterface::ShapeToOriginalBodyInstance(BodyInst, &Shape);
 		if (const UPrimitiveComponent* OwnerComponent = BodyInst->OwnerComponent.Get())
 		{
-			OutOverlap.OverlapObjectHandle = FActorInstanceHandle(OwnerComponent->GetOwner(), OwnerComponent, BodyInst->InstanceBodyIndex);
+			// Create a handle that needs to be resolved later when accessed in game thread,
+			// since resolving FActorInstanceHandle may require to access related UObjects (e.g. owner actor), which is not safe outside game thread.
+			OutOverlap.OverlapObjectHandle = FActorInstanceHandle::MakeActorHandleToResolve(BodyInst->OwnerComponent, BodyInst->InstanceBodyIndex);
 			OutOverlap.Component = BodyInst->OwnerComponent; // Copying weak pointer is faster than assigning raw pointer.
 			OutOverlap.ItemIndex = OwnerComponent->bMultiBodyOverlap ? BodyInst->InstanceBodyIndex : INDEX_NONE;
 		}
@@ -616,8 +621,10 @@ void ConvertQueryOverlap(const FPhysicsShape& Shape, const FPhysicsActor& Actor,
 
 		if(PossibleOwner)
 		{
+			// Create a handle that needs to be resolved later when accessed in game thread,
+			// since resolving FActorInstanceHandle may require to access related UObjects (e.g. owner actor), which is not safe outside game thread.
+			OutOverlap.OverlapObjectHandle = FActorInstanceHandle::MakeActorHandleToResolve(PossibleOwner, INDEX_NONE);
 			OutOverlap.Component = PossibleOwner;
-			OutOverlap.OverlapObjectHandle = FActorInstanceHandle(OutOverlap.Component->GetOwner());
 			OutOverlap.ItemIndex = INDEX_NONE;
 		}
 		else
diff --git a/Engine/Source/Runtime/Engine/Private/Engine/ActorInstanceHandle.cpp b/Engine/Source/Runtime/Engine/Private/Engine/ActorInstanceHandle.cpp
--- a/Engine/Source/Runtime/Engine/Private/Engine/ActorInstanceHandle.cpp
+++ b/Engine/Source/Runtime/Engine/Private/Engine/ActorInstanceHandle.cpp
@@ -1,8 +1,6 @@
 // Copyright Epic Games, Inc. All Rights Reserved.
 
 #include "Engine/ActorInstanceHandle.h"
-#include "PhysicsPublic.h"
-#include "Engine/World.h"
 #include "Engine/ActorInstanceManagerInterface.h"
 #include "GameFramework/Actor.h"
 #include "Components/PrimitiveComponent.h"
@@ -10,16 +8,39 @@
 
 #include UE_INLINE_GENERATED_CPP_BY_NAME(ActorInstanceHandle)
 
+namespace UE::FActorInstanceHandle::Private
+{
+bool bValidateAccessFromGameThread = false;
+static FAutoConsoleVariableRef CVarValidateAccessFromGameThread(
+	TEXT("IA.ValidateAccessFromGameThread"),
+	bValidateAccessFromGameThread,
+	TEXT("If set errors will get reported when trying to resolve or access the handle from non game threads."),
+	ECVF_Default);
+} // UE::FActorInstanceHandle::Private
+
+#define VALIDATE_ACCESS() \
+	if (UE::FActorInstanceHandle::Private::bValidateAccessFromGameThread) \
+	{ \
+		ensureAlwaysMsgf(IsInGameThread() || IsInParallelGameThread(), TEXT("%hs can only be called from the game thread"), __FUNCTION__); \
+	}
+
+
+//-----------------------------------------------------------------------------
+// FActorInstanceHandleInternalHelper
+//-----------------------------------------------------------------------------
 struct FActorInstanceHandleInternalHelper
 {
-	inline static void SetUpAsInterface(FActorInstanceHandle& InstanceHandle, IActorInstanceManagerInterface& InManagerInterface, const UPrimitiveComponent* RelevantComponent, int32 CollisionInstanceIndex)
+	static void SetUpAsInterface(FActorInstanceHandle& InstanceHandle, IActorInstanceManagerInterface& InManagerInterface, const UPrimitiveComponent* RelevantComponent, const int32 CollisionInstanceIndex)
 	{
+		checkf(InstanceHandle.ResolutionStatus == FActorInstanceHandle::EResolutionStatus::Invalid, TEXT("Expected to be called on a newly constructed handle."));
 		InstanceHandle.InstanceIndex = InManagerInterface.ConvertCollisionIndexToInstanceIndex(CollisionInstanceIndex, RelevantComponent);
-		InstanceHandle.Actor = InManagerInterface.FindActor(InstanceHandle);
+		InstanceHandle.ReferenceObject = InManagerInterface.FindActor(InstanceHandle);
+		InstanceHandle.ResolutionStatus = FActorInstanceHandle::EResolutionStatus::Resolved;
 	}
 
-	inline static void SetUpWithActor(FActorInstanceHandle& InstanceHandle, AActor* InActor, const UPrimitiveComponent* RelevantComponent, int32 CollisionInstanceIndex)
+	static void SetUpWithActor(FActorInstanceHandle& InstanceHandle, AActor* InActor, const UPrimitiveComponent* RelevantComponent, const int32 CollisionInstanceIndex)
 	{
+		checkf(InstanceHandle.ResolutionStatus == FActorInstanceHandle::EResolutionStatus::Invalid, TEXT("Expected to be called on a newly constructed handle."));
 		InstanceHandle.ManagerInterface = FActorInstanceManagerInterface(InActor);
 		if (IActorInstanceManagerInterface* ManagerInterfacePtr = InstanceHandle.ManagerInterface.Get())
 		{
@@ -27,8 +48,9 @@ struct FActorInstanceHandleInternalHelper
 		}
 		else
 		{
-			InstanceHandle.Actor = InActor;
+			InstanceHandle.ReferenceObject = InActor;
 		}
+		InstanceHandle.ResolutionStatus = FActorInstanceHandle::EResolutionStatus::Resolved;
 	}
 };
 
@@ -36,12 +58,16 @@ struct FActorInstanceHandleInternalHelper
 // FActorInstanceHandle
 //-----------------------------------------------------------------------------
 FActorInstanceHandle::FActorInstanceHandle(AActor* InActor)
-	: Actor(InActor)
+	: ReferenceObject(InActor)
+	, ResolutionStatus(EResolutionStatus::Resolved)
 {
+	VALIDATE_ACCESS();
 }
 
-FActorInstanceHandle::FActorInstanceHandle(const UPrimitiveComponent* RelevantComponent, int32 CollisionInstanceIndex)
+FActorInstanceHandle::FActorInstanceHandle(const UPrimitiveComponent* RelevantComponent, const int32 CollisionInstanceIndex)
 {
+	VALIDATE_ACCESS();
+
 	if (UNLIKELY(!ensureMsgf(RelevantComponent, TEXT("Calling FActorInstanceHandle(UPrimitiveComponent, int32) constructor is pointless with RelevantComponent == nullptr"))))
 	{
 		return;
@@ -53,8 +79,10 @@ FActorInstanceHandle::FActorInstanceHandle(const UPrimitiveComponent* RelevantCo
 	}
 }
 
-FActorInstanceHandle::FActorInstanceHandle(AActor* InActor, const UPrimitiveComponent* RelevantComponent, int32 CollisionInstanceIndex)
+FActorInstanceHandle::FActorInstanceHandle(AActor* InActor, const UPrimitiveComponent* RelevantComponent, const int32 CollisionInstanceIndex)
 {
+	VALIDATE_ACCESS();
+
 	if (LIKELY(InActor))
 	{
 		FActorInstanceHandleInternalHelper::SetUpWithActor(*this, InActor, RelevantComponent, CollisionInstanceIndex);
@@ -63,15 +91,13 @@ FActorInstanceHandle::FActorInstanceHandle(AActor* InActor, const UPrimitiveComp
 	{
 		*this = FActorInstanceHandle(RelevantComponent, CollisionInstanceIndex);
 	}
-	else
-	{
-		Actor = InActor;
-	}
 }
 
-FActorInstanceHandle::FActorInstanceHandle(FActorInstanceManagerInterface InManagerInterface, int32 CollisionInstanceIndex)
+FActorInstanceHandle::FActorInstanceHandle(const FActorInstanceManagerInterface InManagerInterface, const int32 CollisionInstanceIndex)
 	: ManagerInterface(InManagerInterface)
 {
+	VALIDATE_ACCESS();
+
 	if (IActorInstanceManagerInterface* ManagerInterfacePtr = ManagerInterface.Get())
 	{
 		FActorInstanceHandleInternalHelper::SetUpAsInterface(*this, *ManagerInterfacePtr, /*RelevantComponent=*/nullptr, CollisionInstanceIndex);
@@ -80,23 +106,72 @@ FActorInstanceHandle::FActorInstanceHandle(FActorInstanceManagerInterface InMana
 
 FActorInstanceHandle::FActorInstanceHandle(const FActorInstanceHandle& Other)
 {
-	Actor = Other.Actor;
-
+	ReferenceObject = Other.ReferenceObject;
 	ManagerInterface = Other.ManagerInterface;
 	InstanceIndex = Other.InstanceIndex;
+	ResolutionStatus = Other.ResolutionStatus;
 }
 
-FActorInstanceHandle FActorInstanceHandle::MakeDehydratedActorHandle(UObject& Manager, int32 InInstanceIndex)
+FActorInstanceHandle FActorInstanceHandle::MakeDehydratedActorHandle(UObject& Manager, const int32 InInstanceIndex)
 {
 	FActorInstanceHandle ReturnHandle;
 	ReturnHandle.ManagerInterface = FActorInstanceManagerInterface(&Manager);
 	ReturnHandle.InstanceIndex = InInstanceIndex;
+	ReturnHandle.ResolutionStatus = EResolutionStatus::Resolved;
+
+	return ReturnHandle;
+}
+
+FActorInstanceHandle FActorInstanceHandle::MakeActorHandleToResolve(const TWeakObjectPtr<UPrimitiveComponent>& WeakComponent, const int32 CollisionInstanceIndex)
+{
+	ensureMsgf(!WeakComponent.IsExplicitlyNull(), TEXT("Provided weak pointer must be initialized."));
+	FActorInstanceHandle ReturnHandle;
+	ReturnHandle.ReferenceObject = WeakComponent;
+	ReturnHandle.InstanceIndex = CollisionInstanceIndex;
+	ReturnHandle.ResolutionStatus = EResolutionStatus::NeedsResolving;
 
 	return ReturnHandle;
 }
 
+void FActorInstanceHandle::ResolveHandle() const
+{
+	if (ResolutionStatus == EResolutionStatus::NeedsResolving)
+	{
+		VALIDATE_ACCESS();
+
+		FActorInstanceHandle* MutableThis = const_cast<FActorInstanceHandle*>(this);
+
+		if (const UPrimitiveComponent* Component = CastChecked<UPrimitiveComponent>(ReferenceObject.Get(), ECastCheckedType::NullAllowed))
+		{
+			if (AActor* OwnerActor = Component->GetOwner())
+			{
+				// Reset resolution status to 'Invalid' since 'SetUpWithActor' requires handle to be in that status
+				MutableThis->ResolutionStatus = EResolutionStatus::Invalid;
+
+				// Resolving by using valid Actor, Component and collision instance index.
+				FActorInstanceHandleInternalHelper::SetUpWithActor(*MutableThis, OwnerActor, Component, InstanceIndex);
+				ensure(ResolutionStatus == EResolutionStatus::Resolved);
+			}
+		}
+
+		// Reset the handle if we were unable to resolve it (e.g. component no longer valid or missing owner actor)
+		if (ResolutionStatus == EResolutionStatus::NeedsResolving)
+		{
+			// OwnerActor is not valid so we need to reset the handle
+			new (MutableThis)(FActorInstanceHandle);
+			checkf(ResolutionStatus == EResolutionStatus::Invalid, TEXT("Default constructor must mark the resolution as invalid."));
+		}
+	}
+}
+
 bool FActorInstanceHandle::IsValid() const
 {
+	// A handle properly setup from another thread that needs resolving is considered valid.
+	if (ResolutionStatus == EResolutionStatus::NeedsResolving)
+	{
+		return true;
+	}
+
 	return (ManagerInterface.IsValid() && InstanceIndex != INDEX_NONE) || IsActorValid();
 }
 
@@ -111,9 +186,10 @@ bool FActorInstanceHandle::DoesRepresentClass(const UClass* OtherClass) const
 
 UClass* FActorInstanceHandle::GetRepresentedClass() const
 {
+	// Calling IsActorValid will resolve the handle if necessary
 	if (IsActorValid())
 	{
-		return Actor->GetClass();
+		return ReferenceObject->GetClass();
 	}
 
 	IActorInstanceManagerInterface* ManagerInterfacePtr = ManagerInterface.Get();
@@ -124,7 +200,7 @@ UClass* FActorInstanceHandle::GetRepresentedClass() const
 
 ULevel* FActorInstanceHandle::GetLevel() const
 {
-	if (IsActorValid())
+	if (const AActor* Actor = GetCachedActor())
 	{
 		return Actor->GetLevel();
 	}
@@ -137,7 +213,7 @@ ULevel* FActorInstanceHandle::GetLevel() const
 
 FVector FActorInstanceHandle::GetLocation() const
 {
-	if (IsActorValid())
+	if (const AActor* Actor = GetCachedActor())
 	{
 		return Actor->GetActorLocation();
 	}
@@ -150,7 +226,7 @@ FVector FActorInstanceHandle::GetLocation() const
 
 FRotator FActorInstanceHandle::GetRotation() const
 {
-	if (IsActorValid())
+	if (const AActor* Actor = GetCachedActor())
 	{
 		return Actor->GetActorRotation();
 	}
@@ -163,7 +239,7 @@ FRotator FActorInstanceHandle::GetRotation() const
 
 FTransform FActorInstanceHandle::GetTransform() const
 {
-	if (IsActorValid())
+	if (const AActor* Actor = GetCachedActor())
 	{
 		return Actor->GetActorTransform();
 	}
@@ -176,7 +252,7 @@ FTransform FActorInstanceHandle::GetTransform() const
 
 FName FActorInstanceHandle::GetFName() const
 {
-	if (IsActorValid())
+	if (const AActor* Actor = GetCachedActor())
 	{
 		return Actor->GetFName();
 	}
@@ -186,14 +262,14 @@ FName FActorInstanceHandle::GetFName() const
 
 FString FActorInstanceHandle::GetName() const
 {
-	if (IsActorValid())
+	if (const AActor* Actor = GetCachedActor())
 	{
 		return Actor->GetName();
 	}
 
 	if (ManagerInterface.IsValid())
 	{
-		return FString::Printf(TEXT("%s:d"), *GetNameSafe(ManagerInterface.GetObject()), InstanceIndex);
+		return FString::Printf(TEXT("%s:%d"), *GetNameSafe(ManagerInterface.GetObject()), InstanceIndex);
 	}
 
 	return TEXT("Invalid");
@@ -201,9 +277,9 @@ FString FActorInstanceHandle::GetName() const
 
 AActor* FActorInstanceHandle::GetManagingActor() const
 {
-	if (IsActorValid())
+	if (AActor* Actor = GetCachedActor())
 	{
-		return Actor.Get();
+		return Actor;
 	}
 
 	return Cast<AActor>(ManagerInterface.GetObject());
@@ -211,7 +287,7 @@ AActor* FActorInstanceHandle::GetManagingActor() const
 
 USceneComponent* FActorInstanceHandle::GetRootComponent() const
 {
-	if (IsActorValid())
+	if (const AActor* Actor = GetCachedActor())
 	{
 		return Actor->GetRootComponent();
 	}
@@ -239,42 +315,52 @@ AActor* FActorInstanceHandle::FetchActor() const
 
 UObject* FActorInstanceHandle::GetActorAsUObject()
 {
-	// 
-	return Cast<UObject>(Actor.Get());
+	return GetCachedActor();
 }
 
 const UObject* FActorInstanceHandle::GetActorAsUObject() const
 {
-	if (IsActorValid())
-	{
-		return Cast<UObject>(Actor.Get());
-	}
-
-	return nullptr;
+	return GetCachedActor();
 }
 
 bool FActorInstanceHandle::IsActorValid() const
 {
-	return Actor.IsValid();
+	return GetCachedActor() != nullptr;
+}
+
+AActor* FActorInstanceHandle::GetCachedActor() const
+{
+	// Make sure handle is resolved before getting the actor.
+	ResolveHandle();
+
+	return Cast<AActor>(ReferenceObject.Get());
 }
 
 void FActorInstanceHandle::SetCachedActor(AActor* InActor) const
 {
-	check(Actor.IsValid() == false);
-	Actor = InActor;
+	check(IsActorValid() == false);
+	ReferenceObject = InActor;
 }
 
 FActorInstanceHandle& FActorInstanceHandle::operator=(AActor* OtherActor)
 {
-	Actor = OtherActor;
-	ManagerInterface.Reset();
-	InstanceIndex = INDEX_NONE;
-
+	new (this)FActorInstanceHandle(OtherActor);
 	return *this;
 }
 
 bool FActorInstanceHandle::operator==(const FActorInstanceHandle& Other) const
 {
+	// Handles that needs resolving can be compared using the reference object and index.
+	if (ResolutionStatus == EResolutionStatus::NeedsResolving
+		&& Other.ResolutionStatus == EResolutionStatus::NeedsResolving)
+	{
+		return (ReferenceObject.HasSameIndexAndSerialNumber(Other.ReferenceObject)) && (InstanceIndex == Other.InstanceIndex);
+	}
+
+	// Both handles need to be resolved to perform a valid comparison.
+	ResolveHandle();
+	Other.ResolveHandle();
+
 	// try to compare managers and indices first if we have them
 	if (ManagerInterface.IsValid() && Other.ManagerInterface.IsValid() && InstanceIndex != INDEX_NONE && Other.InstanceIndex != INDEX_NONE)
 	{
@@ -295,8 +381,10 @@ bool FActorInstanceHandle::operator!=(const FActorInstanceHandle& Other) const
 
 bool FActorInstanceHandle::operator==(const AActor* OtherActor) const
 {
+	VALIDATE_ACCESS();
+
 	// if we have an actor, compare the two actors
-	if (AActor* AsActor = Actor.Get())
+	if (const AActor* AsActor = GetCachedActor())
 	{
 		return AsActor == OtherActor;
 	}
@@ -319,7 +407,7 @@ bool FActorInstanceHandle::operator!=(const AActor* OtherActor) const
 uint32 GetTypeHash(const FActorInstanceHandle& Handle)
 {
 	uint32 Hash = 0;
-	if (AActor* Actor = Handle.Actor.Get())
+	if (const AActor* Actor = Handle.GetCachedActor())
 	{
 		FCrc::StrCrc32(*(Actor->GetPathName()), Hash);
 	}
@@ -337,7 +425,7 @@ FArchive& operator<<(FArchive& Ar, FActorInstanceHandle& Handle)
 	Ar.UsingCustomVersion(FFortniteValkyrieBranchObjectVersion::GUID);
 	if (Ar.CustomVer(FFortniteValkyrieBranchObjectVersion::GUID) < FFortniteValkyrieBranchObjectVersion::ActorInstanceHandleSwitchedToInterfaces)
 	{
-		Ar << Handle.Actor;
+		Ar << Handle.ReferenceObject;
 		TWeakObjectPtr<AActor> Manager;
 		Ar << Manager;
 		ensureMsgf(Ar.IsLoading(), TEXT("We expect this piece of code to be running only while loading data."));
@@ -346,7 +434,7 @@ FArchive& operator<<(FArchive& Ar, FActorInstanceHandle& Handle)
 		return Ar;
 	}
 	
-	Ar << Handle.Actor;
+	Ar << Handle.ReferenceObject;
 	Ar << Handle.InstanceIndex;
 
 	TSoftObjectPtr<UObject> SoftObject(Handle.ManagerInterface.GetObject());
