Fix unity compile and typo in GetPrevWorldToLocal

diff --git a/Engine/Source/Runtime/Core/Private/UObject/DevObjectVersion.cpp b/Engine/Source/Runtime/Core/Private/UObject/DevObjectVersion.cpp
--- a/Engine/Source/Runtime/Core/Private/UObject/DevObjectVersion.cpp
+++ b/Engine/Source/Runtime/Core/Private/UObject/DevObjectVersion.cpp
@@ -129,7 +129,7 @@ FDevSystemGuids::FDevSystemGuids()
 	, POSESEARCHDB_DERIVEDDATA_VER(0x389117E4, 0x807A4CC0, 0x9F37C2E6, 0xD808A78D)
 	, SkeletalMeshDerivedDataVersion(0x9B5F4544, 0x76D7481C, 0x9AD3F614, 0xA6C07904)
 	, STATICMESH_DERIVEDDATA_VER(0x5B7A05A7, 0x7371440C, 0xA314621A, 0x9B1E08BC)
-	, MaterialTranslationDDCVersion(0x81B267E7, 0x0A8645C0, 0x8F3FBAF4, 0xE9DB53A7)
+	, MaterialTranslationDDCVersion(0xBBD65080, 0x4D314CE7, 0xA44ADF9B, 0xB4F3F5F9)
 {
 }
 
diff --git a/Engine/Source/Runtime/Core/Private/UObject/FortniteMainBranchObjectVersion.cpp b/Engine/Source/Runtime/Core/Private/UObject/FortniteMainBranchObjectVersion.cpp
--- a/Engine/Source/Runtime/Core/Private/UObject/FortniteMainBranchObjectVersion.cpp
+++ b/Engine/Source/Runtime/Core/Private/UObject/FortniteMainBranchObjectVersion.cpp
@@ -14,6 +14,6 @@ TMap<FGuid, FGuid> FFortniteMainBranchObjectVersion::GetSystemGuids()
 	SystemGuids.Add(DevGuids.Niagara_LatestScriptCompileVersion, FGuid("CE33E29C6268469E8E18F1D96A55DFDA"));
 	SystemGuids.Add(DevGuids.SkeletalMeshDerivedDataVersion, FGuid("41EFB33F1B72B1F49F9AA7D295871F36"));
 	SystemGuids.Add(DevGuids.STATICMESH_DERIVEDDATA_VER, FGuid("543843B103794E0E9BA4BE80FB602F79"));
-	SystemGuids.Add(DevGuids.MaterialTranslationDDCVersion, FGuid("81B267E70A8645C08F3FBAF4E9DB53A7"));
+	SystemGuids.Add(DevGuids.MaterialTranslationDDCVersion, FGuid("BBD650804D314CE7A44ADF9BB4F3F5F9"));
 	return SystemGuids;
 }
diff --git a/Engine/Source/Runtime/Engine/Private/Materials/HLSLMaterialTranslator.cpp b/Engine/Source/Runtime/Engine/Private/Materials/HLSLMaterialTranslator.cpp
--- a/Engine/Source/Runtime/Engine/Private/Materials/HLSLMaterialTranslator.cpp
+++ b/Engine/Source/Runtime/Engine/Private/Materials/HLSLMaterialTranslator.cpp
@@ -10153,8 +10153,8 @@ int32 FHLSLMaterialTranslator::TransformBase(EMaterialCommonBasis SourceCoordBas
 
 				// TODO: inconsistent with TransformLocal<TO>World with instancing
 				AddLWCFuncUsage(ELWCFunctionKind::MultiplyVectorMatrix);
-				CodeStr = LWCMultiplyMatrix(TEXT("<A>"), TEXT("Get<PREVIOUS>WorldToLocal(Parameters)"), AWComponent, true);
-				CodeDerivStr = LWCMultiplyMatrix(TEXT("<A>"), TEXT("Get<PREVIOUS>WorldToLocal(Parameters)"), 0, true);
+				CodeStr = LWCMultiplyMatrix(TEXT("<A>"), TEXT("Get<PREV>WorldToLocal(Parameters)"), AWComponent, true);
+				CodeDerivStr = LWCMultiplyMatrix(TEXT("<A>"), TEXT("Get<PREV>WorldToLocal(Parameters)"), 0, true);
 			}
 			else if (DestCoordBasis == MCB_TranslatedWorld)
 			{
diff --git a/Engine/Source/Runtime/Engine/Private/Materials/MaterialHLSLEmitter.cpp b/Engine/Source/Runtime/Engine/Private/Materials/MaterialHLSLEmitter.cpp
--- a/Engine/Source/Runtime/Engine/Private/Materials/MaterialHLSLEmitter.cpp
+++ b/Engine/Source/Runtime/Engine/Private/Materials/MaterialHLSLEmitter.cpp
@@ -622,8 +622,8 @@ static void GetMaterialEnvironment(EShaderPlatform InPlatform,
 	}
 
 	//static TConsoleVariableData<int32>* CVarLWCEnabled = IConsoleManager::Get().FindTConsoleVariableDataInt(TEXT("r.MaterialEditor.LWCEnabled"));
-	static IConsoleVariable* CVarLWCEnabled = IConsoleManager::Get().FindConsoleVariable(TEXT("r.MaterialEditor.LWCEnabled"));
-	OutEnvironment.SetDefine(TEXT("MATERIAL_LWC_ENABLED"), CVarLWCEnabled->GetInt() ? TEXT("1") : TEXT("0"));
+	static IConsoleVariable* CVarLWCIsEnabled = IConsoleManager::Get().FindConsoleVariable(TEXT("r.MaterialEditor.LWCEnabled"));
+	OutEnvironment.SetDefine(TEXT("MATERIAL_LWC_ENABLED"), CVarLWCIsEnabled->GetInt() ? TEXT("1") : TEXT("0"));
 	OutEnvironment.SetDefine(TEXT("WSVECTOR_IS_TILEOFFSET"), TEXT("1"));
 	OutEnvironment.SetDefine(TEXT("WSVECTOR_IS_DOUBLEFLOAT"), TEXT("0"));
 
