Add option to wait for decal proxy creation when PSOs are still compiling

diff --git a/Engine/Source/Runtime/Engine/Classes/Components/DecalComponent.h b/Engine/Source/Runtime/Engine/Classes/Components/DecalComponent.h
--- a/Engine/Source/Runtime/Engine/Classes/Components/DecalComponent.h
+++ b/Engine/Source/Runtime/Engine/Classes/Components/DecalComponent.h
@@ -7,6 +7,8 @@
 #include "UObject/ObjectMacros.h"
 #include "Engine/EngineTypes.h"
 #include "Components/SceneComponent.h"
+#include "PSOPrecache.h"
+
 #include "DecalComponent.generated.h"
 
 class FDeferredDecalProxy;
@@ -122,6 +124,11 @@ public:
 	UFUNCTION(BlueprintCallable, Category="Rendering|Components|Decal")
 	ENGINE_API virtual class UMaterialInstanceDynamic* CreateDynamicMaterialInstance();
 
+#if UE_WITH_PSO_PRECACHING
+protected:
+	/** Graph event used to track all the PSO precache events - used for delayed proxy creation */
+	FGraphEventRef PSOPrecacheCompileEvent;
+#endif
 
 public:
 	/** The decal proxy. */
@@ -175,11 +182,13 @@ public:
 	ENGINE_API virtual void DestroyRenderState_Concurrent() override;
 	ENGINE_API virtual void SendRenderTransform_Concurrent() override;
 	ENGINE_API virtual const UObject* AdditionalStatObject() const override;
+	ENGINE_API virtual void PrecachePSOs() override;
 	//~ End UActorComponent Interface
 	
 	//~ Begin UObject Interface. 
 	ENGINE_API virtual void Serialize(FArchive& Ar) override;
 	ENGINE_API virtual bool IsPostLoadThreadSafe() const override;
+	ENGINE_API virtual void PostLoad() override;
 	//~ End UObject Interface
 
 	//~ Begin USceneComponent Interface
diff --git a/Engine/Source/Runtime/Engine/Private/Components/DecalComponent.cpp b/Engine/Source/Runtime/Engine/Private/Components/DecalComponent.cpp
--- a/Engine/Source/Runtime/Engine/Private/Components/DecalComponent.cpp
+++ b/Engine/Source/Runtime/Engine/Private/Components/DecalComponent.cpp
@@ -12,7 +12,9 @@
 #include "TimerManager.h"
 #include "SceneManagement.h"
 #include "SceneView.h"
+#include "LocalVertexFactory.h"
 #include "Materials/MaterialInstanceDynamic.h"
+#include "MarkActorRenderStateDirtyTask.h"
 
 #include UE_INLINE_GENERATED_CPP_BY_NAME(DecalComponent)
 
@@ -298,9 +300,48 @@ void UDecalComponent::SetDecalMaterial(class UMaterialInterface* NewDecalMateria
 {
 	DecalMaterial = NewDecalMaterial;
 
+	PrecachePSOs();
+
 	MarkRenderStateDirty();	
 }
 
+void UDecalComponent::PostLoad()
+{
+	Super::PostLoad();
+
+	PrecachePSOs();
+}
+
+void UDecalComponent::PrecachePSOs()
+{
+#if UE_WITH_PSO_PRECACHING
+	if (!FApp::CanEverRender() || !IsComponentPSOPrecachingEnabled())
+	{
+		return;
+	}
+
+	// clear the current request data
+	PSOPrecacheCompileEvent = nullptr;
+
+	if (DecalMaterial)
+	{
+		FPSOPrecacheParams PSOPrecacheParams;		
+		FPSOPrecacheVertexFactoryDataList VertexFactoryDataList;		
+		VertexFactoryDataList.Add(FPSOPrecacheVertexFactoryData(&FLocalVertexFactory::StaticType));
+
+		// Immediately create at high priority and thus doesn't need boosting anymore
+		TArray<FMaterialPSOPrecacheRequestID> MaterialPSOPrecacheRequestIDs;
+		FGraphEventArray GraphEvents = DecalMaterial->PrecachePSOs(VertexFactoryDataList, PSOPrecacheParams, EPSOPrecachePriority::High, MaterialPSOPrecacheRequestIDs);
+
+		// Request recreate of the render state when the PSO compilation is ready (if we want to delay proxy creation)
+		if (GraphEvents.Num() > 0 && GetPSOPrecacheProxyCreationStrategy() != EPSOPrecacheProxyCreationStrategy::AlwaysCreate)
+		{
+			PSOPrecacheCompileEvent = TGraphTask<FMarkActorRenderStateDirtyTask>::CreateTask(&GraphEvents).ConstructAndDispatchWhenReady(this);
+		}
+	}
+#endif
+}
+
 void UDecalComponent::PushSelectionToProxy()
 {
 	MarkRenderStateDirty();	
@@ -339,6 +380,15 @@ void UDecalComponent::GetUsedMaterials( TArray<UMaterialInterface*>& OutMaterial
 FDeferredDecalProxy* UDecalComponent::CreateSceneProxy()
 {
 	LLM_SCOPE(ELLMTag::SceneRender);
+
+#if UE_WITH_PSO_PRECACHING
+	if (PSOPrecacheCompileEvent && !PSOPrecacheCompileEvent->IsComplete() && GetPSOPrecacheProxyCreationStrategy() == EPSOPrecacheProxyCreationStrategy::DelayUntilPSOPrecached)
+	{
+		return nullptr;
+	}
+	PSOPrecacheCompileEvent = nullptr;
+#endif // UE_WITH_PSO_PRECACHING
+
 	return new FDeferredDecalProxy(this);
 }
 
diff --git a/Engine/Source/Runtime/Engine/Private/Components/PrimitiveComponent.cpp b/Engine/Source/Runtime/Engine/Private/Components/PrimitiveComponent.cpp
--- a/Engine/Source/Runtime/Engine/Private/Components/PrimitiveComponent.cpp
+++ b/Engine/Source/Runtime/Engine/Private/Components/PrimitiveComponent.cpp
@@ -50,6 +50,7 @@
 #include "PSOPrecacheMaterial.h"
 #include "MaterialCachedData.h"
 #include "MaterialShared.h"
+#include "MarkActorRenderStateDirtyTask.h"
 
 #if WITH_EDITOR
 #include "Engine/LODActor.h"
@@ -4639,32 +4640,6 @@ void UPrimitiveComponent::SetupPrecachePSOParams(FPSOPrecacheParams& Params)
 	}
 }
 
-class FMarkRenderStateDirtyTask
-{
-public:
-	explicit FMarkRenderStateDirtyTask(UPrimitiveComponent* InPrimitiveComponent)
-		: PrimitiveComponent(InPrimitiveComponent)
-	{
-	}
-
-	void DoTask(ENamedThreads::Type CurrentThread, const FGraphEventRef& MyCompletionGraphEvent)
-	{
-		if (PrimitiveComponent.IsValid())
-		{
-			PrimitiveComponent->MarkRenderStateDirty();
-			PrimitiveComponent = nullptr;
-		}
-	}
-
-public:
-
-	TWeakObjectPtr<UPrimitiveComponent> PrimitiveComponent;
-
-	static ESubsequentsMode::Type	GetSubsequentsMode() { return ESubsequentsMode::TrackSubsequents; }
-	ENamedThreads::Type				GetDesiredThread() { return ENamedThreads::GameThread; }
-	FORCEINLINE TStatId				GetStatId() const { return TStatId(); }
-};
-
 void UPrimitiveComponent::PrecachePSOs()
 {
 #if UE_WITH_PSO_PRECACHING
@@ -4698,7 +4673,7 @@ void UPrimitiveComponent::RequestRecreateRenderStateWhenPSOPrecacheFinished(cons
 	// schedule a task to mark the render state dirty when all PSOs are compiled so the proxy gets recreated.
 	if (UsePSOPrecacheRenderProxyDelay() && GetPSOPrecacheProxyCreationStrategy() != EPSOPrecacheProxyCreationStrategy::AlwaysCreate && !PSOPrecacheCompileEvents.IsEmpty())
 	{
-		PSOPrecacheCompileEvent = TGraphTask<FMarkRenderStateDirtyTask>::CreateTask(&PSOPrecacheCompileEvents).ConstructAndDispatchWhenReady(this);
+		PSOPrecacheCompileEvent = TGraphTask<FMarkActorRenderStateDirtyTask>::CreateTask(&PSOPrecacheCompileEvents).ConstructAndDispatchWhenReady(this);
 	}
 
 	bPSOPrecacheCalled = true;
diff --git a/Engine/Source/Runtime/Engine/Public/MarkActorRenderStateDirtyTask.h b/Engine/Source/Runtime/Engine/Public/MarkActorRenderStateDirtyTask.h
new file mode 100644
--- /dev/null
+++ b/Engine/Source/Runtime/Engine/Public/MarkActorRenderStateDirtyTask.h
@@ -0,0 +1,31 @@
+// Copyright Epic Games, Inc. All Rights Reserved.
+
+#pragma once
+
+#include "CoreMinimal.h"
+#include "Components/ActorComponent.h"
+
+class FMarkActorRenderStateDirtyTask
+{
+public:
+	explicit FMarkActorRenderStateDirtyTask(UActorComponent* InActorComponent)
+		: ActorComponent(InActorComponent)
+	{
+	}
+
+	void DoTask(ENamedThreads::Type CurrentThread, const FGraphEventRef& MyCompletionGraphEvent)
+	{
+		if (UActorComponent* AC = ActorComponent.Get())
+		{
+			AC->MarkRenderStateDirty();
+		}
+	}
+
+public:
+
+	TWeakObjectPtr<UActorComponent> ActorComponent;
+
+	static ESubsequentsMode::Type	GetSubsequentsMode() { return ESubsequentsMode::TrackSubsequents; }
+	ENamedThreads::Type				GetDesiredThread() { return ENamedThreads::GameThread; }
+	FORCEINLINE TStatId				GetStatId() const { return TStatId(); }
+};
\ No newline at end of file
diff --git a/Engine/Source/Runtime/Renderer/Private/RendererScene.cpp b/Engine/Source/Runtime/Renderer/Private/RendererScene.cpp
--- a/Engine/Source/Runtime/Renderer/Private/RendererScene.cpp
+++ b/Engine/Source/Runtime/Renderer/Private/RendererScene.cpp
@@ -2945,6 +2945,12 @@ void FScene::AddDecal(UDecalComponent* Component)
 		// Create the decals's scene proxy.
 		Component->SceneProxy = Component->CreateSceneProxy();
 
+		// If there is no scene proxy then don't add to scene
+		if (!Component->SceneProxy)
+		{
+			return;
+		}
+
 		INC_DWORD_STAT(STAT_SceneDecals);
 
 		// Send a command to the rendering thread to add the light to the scene.
