New material HLSL translator fixes:
- Material property output activeness is now deferred. Previously, it was checked during tree generation. It needs to be deferred because the tree is shared between a base material and its instances. An inactive property in base can be active in an instance.
- For FExpressionSetStructField, do not generate a previous frame field expression if the field isn't requested.
- Implemented GenerateHLSLExpression for UMaterialExpressionVolumetricCloudEmptySpaceSkippingInput and UMaterialExpressionVolumetricCloudEmptySpaceSkippingOutput.
- Fallback to preview values (if allowed) for material function inputs connected to invalid reroute nodes.
- FExpressionTextureSample now falls back to hardware derivatives if analytic derivatives fail to prepare.
- Added virtual function GetPreviewExpression to FExpression. Some expressions don't have numeric values and cannot be previewed directly. Those expressions should override this virtual function to return numeric preview expressions that represent their preview values. Also implemented the virtual function for FExpressionParameter.
- Added SeenStaticParameterValues to FEmitContext. For some reasons, users can create static parameter nodes with the same name but different default values. Use the first default value seen to keep things consistent.
- FMaterialResource::HasVertexPositionOffsetConnected() now queries cached expression data instead of the base material because the answer can change in instances.

diff --git a/Engine/Source/Editor/MaterialEditor/Private/MaterialEditor.cpp b/Engine/Source/Editor/MaterialEditor/Private/MaterialEditor.cpp
--- a/Engine/Source/Editor/MaterialEditor/Private/MaterialEditor.cpp
+++ b/Engine/Source/Editor/MaterialEditor/Private/MaterialEditor.cpp
@@ -230,7 +230,7 @@ FMatExpressionPreview::FMatExpressionPreview(UMaterialExpression* InExpression)
 		CachedHLSLTree.Reset(LocalTree);
 
 		FMaterialCachedExpressionData* LocalCachedData = new FMaterialCachedExpressionData();
-		LocalCachedData->UpdateForCachedHLSLTree(*LocalTree, nullptr);
+		LocalCachedData->UpdateForCachedHLSLTree(*LocalTree, nullptr, BaseMaterial);
 		CachedExpressionData.Reset(LocalCachedData);
 	}
 	else
diff --git a/Engine/Source/Editor/UnrealEd/Private/MaterialGraph.cpp b/Engine/Source/Editor/UnrealEd/Private/MaterialGraph.cpp
--- a/Engine/Source/Editor/UnrealEd/Private/MaterialGraph.cpp
+++ b/Engine/Source/Editor/UnrealEd/Private/MaterialGraph.cpp
@@ -173,6 +173,7 @@ void UMaterialGraph::UpdatePinTypes()
 	FMemStackBase Allocator;
 	const FMaterialCachedHLSLTree& CachedTree = Material->GetCachedHLSLTree();
 	FEmitContext EmitContext(Allocator, FTargetParameters(), NullErrorHandler, CachedTree.GetTypeRegistry());
+	EmitContext.MaterialInterface = Material;
 
 	Material::FEmitData& EmitMaterialData = EmitContext.AcquireData<Material::FEmitData>();
 
diff --git a/Engine/Source/Runtime/Engine/Private/HLSLTree/HLSLTree.cpp b/Engine/Source/Runtime/Engine/Private/HLSLTree/HLSLTree.cpp
--- a/Engine/Source/Runtime/Engine/Private/HLSLTree/HLSLTree.cpp
+++ b/Engine/Source/Runtime/Engine/Private/HLSLTree/HLSLTree.cpp
@@ -3,6 +3,7 @@
 
 #include "HLSLTree/HLSLTree.h"
 #include "HLSLTree/HLSLTreeEmit.h"
+#include "HLSLTree/HLSLTreeCommon.h"
 #include "Misc/MemStackUtility.h"
 #include "Shader/Preshader.h"
 #include "Shader/PreshaderTypes.h"
@@ -284,7 +285,15 @@ bool FExpressionLocalPHI::PrepareValue(FEmitContext& Context, FEmitScope& Scope,
 		}
 		if (NumValidTypes < NumLiveScopes)
 		{
-			return Context.Error(TEXT("Failed to compute all types for LocalPHI"));
+			if (Chain.Last().Type == ELocalPHIChainType::Ddx || Chain.Last().Type == ELocalPHIChainType::Ddy)
+			{
+				// Don't treat failing analytic derivatives as errors. Just fall back to HW ones
+				return false;
+			}
+			else
+			{
+				return Context.Error(TEXT("Failed to compute all types for LocalPHI"));
+			}
 		}
 
 		if (CurrentType != InitialType)
@@ -1238,6 +1247,11 @@ const FExpression* FExpression::ComputePreviousFrame(FTree& Tree, const FRequest
 	return nullptr;
 }
 
+const FExpression* FExpression::GetPreviewExpression(FTree& Tree) const
+{
+	return this;
+}
+
 void FExpression::EmitValueShader(FEmitContext& Context, FEmitScope& Scope, const FRequestedType& RequestedType, FEmitValueShaderResult& OutResult) const
 {
 	check(false);
@@ -1547,12 +1561,6 @@ bool FTree::Finalize()
 				{
 					const FExpressionDerivatives Derivatives = GetAnalyticDerivatives(LocalValue);
 					LocalValue = (ChainType == ELocalPHIChainType::Ddx) ? Derivatives.ExpressionDdx : Derivatives.ExpressionDdy;
-					
-					// TODO: Revisit. Can we fallback to hardware derivatives? Handle fallback in FExpression::ComputeAnalyticDerivatives better?
-					if (!LocalValue)
-					{
-						LocalValue = NewConstant(0.f);
-					}
 				}
 				else
 				{
@@ -1746,6 +1754,12 @@ const FExpression* FTree::GetPreviousFrame(const FExpression* InExpression, cons
 	return Result;
 }
 
+const FExpression* FTree::GetPreview(const FExpression* InExpression)
+{
+	// Cache result?
+	return InExpression ? InExpression->GetPreviewExpression(*this) : nullptr;
+}
+
 FScope* FTree::NewScope(FScope& Scope)
 {
 	FScope* NewScope = NewNode<FScope>();
diff --git a/Engine/Source/Runtime/Engine/Private/HLSLTree/HLSLTreeCommon.cpp b/Engine/Source/Runtime/Engine/Private/HLSLTree/HLSLTreeCommon.cpp
--- a/Engine/Source/Runtime/Engine/Private/HLSLTree/HLSLTreeCommon.cpp
+++ b/Engine/Source/Runtime/Engine/Private/HLSLTree/HLSLTreeCommon.cpp
@@ -9,6 +9,7 @@
 #include "MaterialShared.h"
 #include "MaterialHLSLTree.h"
 #include "MaterialCachedData.h"
+#include "Materials/MaterialInterface.h"
 #include "DataDrivenShaderPlatformInfo.h"
 
 namespace UE::HLSLTree
@@ -276,8 +277,8 @@ void FExpressionSetStructField::ComputeAnalyticDerivatives(FTree& Tree, FExpress
 		const Shader::FStructField* DerivativeField = DerivativeStructType->FindFieldByName(Field->Name);
 		check(DerivativeField);
 
-		OutResult.ExpressionDdx = Tree.NewExpression<FExpressionSetStructField>(DerivativeStructType, DerivativeField, StructDerivatives.ExpressionDdx, FieldDerivatives.ExpressionDdx);
-		OutResult.ExpressionDdy = Tree.NewExpression<FExpressionSetStructField>(DerivativeStructType, DerivativeField, StructDerivatives.ExpressionDdy, FieldDerivatives.ExpressionDdy);
+		OutResult.ExpressionDdx = Tree.NewExpression<FExpressionSetStructField>(DerivativeStructType, DerivativeField, StructDerivatives.ExpressionDdx, FieldDerivatives.ExpressionDdx, TestMaterialProperty);
+		OutResult.ExpressionDdy = Tree.NewExpression<FExpressionSetStructField>(DerivativeStructType, DerivativeField, StructDerivatives.ExpressionDdy, FieldDerivatives.ExpressionDdy, TestMaterialProperty);
 	}
 }
 
@@ -287,17 +288,30 @@ const FExpression* FExpressionSetStructField::ComputePreviousFrame(FTree& Tree,
 	const FExpression* PrevStructExpression = Tree.GetPreviousFrame(StructExpression, RequestedStructType);
 
 	const FRequestedType RequestedFieldType = MakeRequestedFieldType(Tree.ActiveStructFieldStack, RequestedType);
+	if (RequestedFieldType.IsEmpty())
+	{
+		// Ignore the field subtree if not requested
+		return PrevStructExpression;
+	}
+
 	const FExpression* PrevFieldExpression;
 	{
 		FScopedActiveStructField ScopedActiveField(Tree.ActiveStructFieldStack, Field);
 		PrevFieldExpression = Tree.GetPreviousFrame(FieldExpression, RequestedFieldType);
 	}
 
-	return Tree.NewExpression<FExpressionSetStructField>(StructType, Field, PrevStructExpression, PrevFieldExpression);
+	return Tree.NewExpression<FExpressionSetStructField>(StructType, Field, PrevStructExpression, PrevFieldExpression, TestMaterialProperty);
 }
 
 bool FExpressionSetStructField::PrepareValue(FEmitContext& Context, FEmitScope& Scope, const FRequestedType& RequestedType, FPrepareValueResult& OutResult) const
 {
+	if (!IsStructFieldActive(Context))
+	{
+		// The struct field isn't active so don't touch it and just forward
+		const FPreparedType& PreparedType = Context.PrepareExpression(StructExpression, Scope, RequestedType);
+		return OutResult.SetType(Context, RequestedType, PreparedType);
+	}
+
 	FRequestedType RequestedStructType = MakeRequestedStructType(Context.ActiveStructFieldStack, RequestedType);
 
 	const FPreparedType StructPreparedType = Context.PrepareExpression(StructExpression, Scope, RequestedStructType);
@@ -311,7 +325,7 @@ bool FExpressionSetStructField::PrepareValue(FEmitContext& Context, FEmitScope&
 	{
 		ResultType = StructType;
 	}
-	
+
 	const FRequestedType RequestedFieldType = MakeRequestedFieldType(Context.ActiveStructFieldStack, RequestedType);
 	if (!RequestedFieldType.IsEmpty())
 	{
@@ -332,6 +346,12 @@ bool FExpressionSetStructField::PrepareValue(FEmitContext& Context, FEmitScope&
 
 void FExpressionSetStructField::EmitValueShader(FEmitContext& Context, FEmitScope& Scope, const FRequestedType& RequestedType, FEmitValueShaderResult& OutResult) const
 {
+	if (!IsStructFieldActive(Context))
+	{
+		StructExpression->EmitValueShader(Context, Scope, RequestedType, OutResult);
+		return;
+	}
+
 	FRequestedType RequestedStructType = MakeRequestedStructType(Context.ActiveStructFieldStack, RequestedType);
 	const EExpressionEvaluation StructEvaluation = Context.GetEvaluation(StructExpression, Scope, RequestedStructType);
 	check(StructEvaluation != EExpressionEvaluation::None);
@@ -375,6 +395,12 @@ void FExpressionSetStructField::EmitValueShader(FEmitContext& Context, FEmitScop
 
 void FExpressionSetStructField::EmitValuePreshader(FEmitContext& Context, FEmitScope& Scope, const FRequestedType& RequestedType, FEmitValuePreshaderResult& OutResult) const
 {
+	if (!IsStructFieldActive(Context))
+	{
+		StructExpression->EmitValuePreshader(Context, Scope, RequestedType, OutResult);
+		return;
+	}
+	
 	FRequestedType RequestedStructType = MakeRequestedStructType(Context.ActiveStructFieldStack, RequestedType);
 	const EExpressionEvaluation StructEvaluation = Context.GetEvaluation(StructExpression, Scope, RequestedStructType);
 
@@ -443,6 +469,27 @@ FRequestedType FExpressionSetStructField::MakeRequestedFieldType(const FActiveSt
 	}
 }
 
+bool FExpressionSetStructField::IsStructFieldActive(FEmitContext& Context) const
+{
+	if (TestMaterialProperty != MP_MAX)
+	{
+		if (Context.MaterialInterface)
+		{
+			return Context.MaterialInterface->IsPropertyActive(TestMaterialProperty);
+		}
+		else if (Context.Material && Context.Material->GetMaterialInterface())
+		{
+			return Context.Material->GetMaterialInterface()->IsPropertyActive(TestMaterialProperty);
+		}
+		else
+		{
+			checkNoEntry();
+			return false;
+		}
+	}
+	return true;
+}
+
 void FExpressionSelect::ComputeAnalyticDerivatives(FTree& Tree, FExpressionDerivatives& OutResult) const
 {
 	const FExpressionDerivatives TrueDerivatives = Tree.GetAnalyticDerivatives(TrueExpression);
@@ -465,6 +512,11 @@ const FExpression* FExpressionSelect::ComputePreviousFrame(FTree& Tree, const FR
 bool FExpressionSelect::PrepareValue(FEmitContext& Context, FEmitScope& Scope, const FRequestedType& RequestedType, FPrepareValueResult& OutResult) const
 {
 	const FPreparedType& ConditionType = Context.PrepareExpression(ConditionExpression, Scope, Shader::EValueType::Bool1);
+	if (ConditionType.IsVoid())
+	{
+		return false;
+	}
+
 	const EExpressionEvaluation ConditionEvaluation = ConditionType.GetEvaluation(Scope, Shader::EValueType::Bool1);
 	if (IsConstantEvaluation(ConditionEvaluation))
 	{
diff --git a/Engine/Source/Runtime/Engine/Private/HLSLTree/HLSLTreeEmit.cpp b/Engine/Source/Runtime/Engine/Private/HLSLTree/HLSLTreeEmit.cpp
--- a/Engine/Source/Runtime/Engine/Private/HLSLTree/HLSLTreeEmit.cpp
+++ b/Engine/Source/Runtime/Engine/Private/HLSLTree/HLSLTreeEmit.cpp
@@ -470,8 +470,7 @@ FPreparedType FEmitContext::PrepareExpression(const FExpression* InExpression, F
 		FXxHash64Builder Hasher;
 		Hasher.Update(&InExpression, sizeof(InExpression));
 		Hasher.Update(&bMarkLiveValues, sizeof(bMarkLiveValues));
-		const bool bRequestedStruct = RequestedType.IsStruct();
-		Hasher.Update(&bRequestedStruct, sizeof(bRequestedStruct));
+		Hasher.Update(&RequestedType.Type.StructType, sizeof(RequestedType.Type.StructType));
 		Hash = Hasher.Finalize();
 	}
 
diff --git a/Engine/Source/Runtime/Engine/Private/Materials/Material.cpp b/Engine/Source/Runtime/Engine/Private/Materials/Material.cpp
--- a/Engine/Source/Runtime/Engine/Private/Materials/Material.cpp
+++ b/Engine/Source/Runtime/Engine/Private/Materials/Material.cpp
@@ -2073,7 +2073,7 @@ void UMaterial::UpdateCachedExpressionData()
 
 		LocalCachedTree = new FMaterialCachedHLSLTree();
 		LocalCachedTree->GenerateTree(this, nullptr, nullptr);
-		LocalCachedExpressionData->UpdateForCachedHLSLTree(*LocalCachedTree, nullptr);
+		LocalCachedExpressionData->UpdateForCachedHLSLTree(*LocalCachedTree, nullptr, this);
 	}
 	else
 	{
diff --git a/Engine/Source/Runtime/Engine/Private/Materials/MaterialCachedData.cpp b/Engine/Source/Runtime/Engine/Private/Materials/MaterialCachedData.cpp
--- a/Engine/Source/Runtime/Engine/Private/Materials/MaterialCachedData.cpp
+++ b/Engine/Source/Runtime/Engine/Private/Materials/MaterialCachedData.cpp
@@ -969,7 +969,7 @@ void PrepareHLSLTree(UE::HLSLTree::FEmitContext& EmitContext,
 
 }
 
-void FMaterialCachedExpressionData::UpdateForCachedHLSLTree(const FMaterialCachedHLSLTree& CachedTree, const FStaticParameterSet* StaticParameters)
+void FMaterialCachedExpressionData::UpdateForCachedHLSLTree(const FMaterialCachedHLSLTree& CachedTree, const FStaticParameterSet* StaticParameters, const UMaterialInterface* TargetMaterial)
 {
 	using namespace UE::HLSLTree;
 
@@ -978,6 +978,7 @@ void FMaterialCachedExpressionData::UpdateForCachedHLSLTree(const FMaterialCache
 
 	FMemStackBase Allocator;
 	FEmitContext EmitContext(Allocator, FTargetParameters(), NullErrorHandler, CachedTree.GetTypeRegistry());
+	EmitContext.MaterialInterface = TargetMaterial;
 
 	Material::FEmitData& EmitMaterialData = EmitContext.AcquireData<Material::FEmitData>();
 	EmitMaterialData.CachedExpressionData = this;
diff --git a/Engine/Source/Runtime/Engine/Private/Materials/MaterialExpressionHLSL.cpp b/Engine/Source/Runtime/Engine/Private/Materials/MaterialExpressionHLSL.cpp
--- a/Engine/Source/Runtime/Engine/Private/Materials/MaterialExpressionHLSL.cpp
+++ b/Engine/Source/Runtime/Engine/Private/Materials/MaterialExpressionHLSL.cpp
@@ -921,11 +921,8 @@ bool UMaterialExpressionStaticSwitchParameter::GenerateHLSLExpression(FMaterialH
 
 	const FExpression* ExpressionA = A.AcquireHLSLExpression(Generator, Scope);
 	const FExpression* ExpressionB = B.AcquireHLSLExpression(Generator, Scope);
-	if (!ExpressionA || !ExpressionB)
-	{
-		return false;
-	}
 
+	// FExpressionSelect can handle null inputs so not doing null check here
 	// No reason to generate a dynamic branch here, since the condition is always a static switch parameter
 	const FExpression* ExpressionSwitch = Generator.GenerateMaterialParameter(ParameterName, ParameterMeta);
 	OutExpression = Generator.GetTree().NewExpression<FExpressionSelect>(ExpressionSwitch, ExpressionA, ExpressionB);
@@ -4119,12 +4116,31 @@ UE::Shader::EValueType UMaterialExpressionBentNormalCustomOutput::GetCustomOutpu
 
 bool UMaterialExpressionVolumetricCloudEmptySpaceSkippingInput::GenerateHLSLExpression(FMaterialHLSLGenerator& Generator, UE::HLSLTree::FScope& Scope, int32 OutputIndex, UE::HLSLTree::FExpression const*& OutExpression) const
 {
-	return Generator.Error(TEXT("UMaterialExpressionVolumetricCloudEmptySpaceSkippingInput::GenerateHLSLExpression unimplemented"));
+	using namespace UE::HLSLTree;
+
+	if (OutputIndex == 0)
+	{
+		OutExpression = Generator.GetTree().NewExpression<FExpressionInlineCustomHLSL>(UE::Shader::EValueType::Float3, TEXT("MaterialExpressionCloudEmptySpaceSkippingSphereCenterWorldPosition(Parameters)"));
+		return true;
+	}
+	else if (OutputIndex == 1)
+	{
+		OutExpression = Generator.GetTree().NewExpression<FExpressionInlineCustomHLSL>(UE::Shader::EValueType::Float1, TEXT("MaterialExpressionCloudEmptySpaceSkippingSphereRadius(Parameters)"));
+		return true;
+	}
+
+	return Generator.Error(TEXT("Invalid input parameter"));
 }
 
 bool UMaterialExpressionVolumetricCloudEmptySpaceSkippingOutput::GenerateHLSLExpression(FMaterialHLSLGenerator& Generator, UE::HLSLTree::FScope& Scope, int32 OutputIndex, UE::HLSLTree::FExpression const*& OutExpression) const
 {
-	return Generator.Error(TEXT("UMaterialExpressionVolumetricCloudEmptySpaceSkippingOutput::GenerateHLSLExpression unimplemented"));
+	if (OutputIndex == 0)
+	{
+		OutExpression = ContainsMatter.AcquireHLSLExpressionOrConstant(Generator, Scope, 1.0f);
+		return true;
+	}
+
+	return Generator.Error(TEXT("Invalid output"));
 }
 
 UE::Shader::EValueType UMaterialExpressionVolumetricCloudEmptySpaceSkippingOutput::GetCustomOutputType(int32 OutputIndex) const
diff --git a/Engine/Source/Runtime/Engine/Private/Materials/MaterialHLSLGenerator.cpp b/Engine/Source/Runtime/Engine/Private/Materials/MaterialHLSLGenerator.cpp
--- a/Engine/Source/Runtime/Engine/Private/Materials/MaterialHLSLGenerator.cpp
+++ b/Engine/Source/Runtime/Engine/Private/Materials/MaterialHLSLGenerator.cpp
@@ -177,27 +177,25 @@ static const UE::HLSLTree::FExpression* CompileMaterialInput(FMaterialHLSLGenera
 	using namespace UE::HLSLTree;
 
 	const FExpression* Expression = nullptr;
-	if (Material->IsPropertyActive(InputProperty))
+	// Not checking Material->IsPropertyActive here because the tree is shared with material instances
+	FMaterialInputDescription InputDescription;
+	if (Material->GetExpressionInputDescription(InputProperty, InputDescription))
 	{
-		FMaterialInputDescription InputDescription;
-		if (Material->GetExpressionInputDescription(InputProperty, InputDescription))
+		UE::Shader::FValue DefaultValue = UE::Shader::Cast(FMaterialAttributeDefinitionMap::GetDefaultValue(InputProperty), InputDescription.Type);
+		if (InputDescription.bUseConstant)
 		{
-			UE::Shader::FValue DefaultValue = UE::Shader::Cast(FMaterialAttributeDefinitionMap::GetDefaultValue(InputProperty), InputDescription.Type);
-			if (InputDescription.bUseConstant)
+			if (InputDescription.ConstantValue != DefaultValue)
 			{
-				if (InputDescription.ConstantValue != DefaultValue)
-				{
-					Expression = Generator.NewConstant(InputDescription.ConstantValue);
-				}
+				Expression = Generator.NewConstant(InputDescription.ConstantValue);
 			}
-			else
+		}
+		else
+		{
+			check(InputDescription.Input);
+			Expression = InputDescription.Input->TryAcquireHLSLExpression(Generator, Scope, (int32)InputProperty);
+			if (Expression)
 			{
-				check(InputDescription.Input);
-				Expression = InputDescription.Input->TryAcquireHLSLExpression(Generator, Scope, (int32)InputProperty);
-				if (Expression)
-				{
-					Expression = Generator.GetTree().NewExpression<FExpressionDefaultValue>(Expression, DefaultValue);
-				}
+				Expression = Generator.GetTree().NewExpression<FExpressionDefaultValue>(Expression, DefaultValue);
 			}
 		}
 	}
@@ -268,7 +266,7 @@ bool FMaterialHLSLGenerator::GenerateResult(UE::HLSLTree::FScope& Scope)
 						const FString& WPOName = FMaterialAttributeDefinitionMap::GetAttributeName(MP_WorldPositionOffset);
 						const FStructField* WPOField = CachedTree.GetMaterialAttributesType()->FindFieldByName(*WPOName);
 
-						FRequestedType PrevRequestedType(CachedTree.GetMaterialAttributesType());
+						FRequestedType PrevRequestedType(CachedTree.GetMaterialAttributesType(), false);
 						PrevRequestedType.SetFieldRequested(WPOField);
 
 						PrevAttributesExpression = GetTree().GetPreviousFrame(PrevAttributesExpression, PrevRequestedType);
@@ -309,13 +307,14 @@ bool FMaterialHLSLGenerator::GenerateResult(UE::HLSLTree::FScope& Scope)
 								CachedTree.GetMaterialAttributesType(),
 								AttributeField,
 								AttributesExpression,
-								InputExpression);
+								InputExpression,
+								Property);
 
 							if (Property == MP_WorldPositionOffset)
 							{
 								const FExpression* PrevWPOExpression = GetTree().GetPreviousFrame(InputExpression, EValueType::Float3);
 								ensure(PrevWPOExpression);
-								AttributesExpression = GetTree().NewExpression<FExpressionSetStructField>(CachedTree.GetMaterialAttributesType(), PrevWPOField, AttributesExpression, PrevWPOExpression);
+								AttributesExpression = GetTree().NewExpression<FExpressionSetStructField>(CachedTree.GetMaterialAttributesType(), PrevWPOField, AttributesExpression, PrevWPOExpression, Property);
 							}
 						}
 					}
@@ -349,7 +348,7 @@ bool FMaterialHLSLGenerator::GenerateResult(UE::HLSLTree::FScope& Scope)
 				PreviewExpressionResult = AcquireExpression(Scope, INDEX_NONE, PreviewExpression, OutputIndex, FSwizzleParameters());
 			}
 
-			const FExpression* ExpressionEmissive = PreviewExpressionResult;
+			const FExpression* ExpressionEmissive = GetTree().GetPreview(PreviewExpressionResult);
 			if (ExpressionEmissive)
 			{
 				const FString& EmissiveColorName = FMaterialAttributeDefinitionMap::GetAttributeName(MP_EmissiveColor);
@@ -597,8 +596,11 @@ const UE::HLSLTree::FExpression* FMaterialHLSLGenerator::AcquireFunctionInputExp
 			{
 				InputIndex = Index;
 				FConnectedInput& ConnectedInput = FunctionEntry->ConnectedInputs[Index];
-				if (!ConnectedInput.Expression && ConnectedInput.Input)
+				// Use GetTracedInput to detect invalid reroute node and attempt to use preview as input if allowed
+				if (!ConnectedInput.Expression && ConnectedInput.Input && ConnectedInput.Input->GetTracedInput().Expression)
 				{
+					// Not using GetTracedInput here because we want the HLSL expression for the reroute node instead of
+					// the expression of the redirected node
 					ConnectedInput.Expression = ConnectedInput.Input->TryAcquireHLSLExpression(*this, *ConnectedInput.Scope);
 				}
 				InputExpression = ConnectedInput.Expression;
@@ -709,6 +711,7 @@ const UE::HLSLTree::FExpression* FMaterialHLSLGenerator::GenerateMaterialParamet
 {
 	using namespace UE::Shader;
 
+	const FMaterialParameterInfo& ParameterInfo = GetParameterInfo(InParameterName);
 	FMaterialParameterMetadata ParameterMeta(InParameterMeta);
 	FMaterialParameterMetadata OverrideParameterMeta;
 	if (GetParameterOverrideValueForCurrentFunction(InParameterMeta.Value.Type, InParameterName, OverrideParameterMeta))
@@ -720,7 +723,7 @@ const UE::HLSLTree::FExpression* FMaterialHLSLGenerator::GenerateMaterialParamet
 		ParameterMeta.ScalarCurve = OverrideParameterMeta.ScalarCurve;
 	}
 
-	return GetTree().NewExpression<UE::HLSLTree::Material::FExpressionParameter>(GetParameterInfo(InParameterName), ParameterMeta, InSamplerType, InExternalTextureGuid);
+	return GetTree().NewExpression<UE::HLSLTree::Material::FExpressionParameter>(ParameterInfo, ParameterMeta, InSamplerType, InExternalTextureGuid);
 }
 
 const UE::HLSLTree::FExpression* FMaterialHLSLGenerator::GenerateFunctionCall(UE::HLSLTree::FScope& Scope,
diff --git a/Engine/Source/Runtime/Engine/Private/Materials/MaterialHLSLTree.cpp b/Engine/Source/Runtime/Engine/Private/Materials/MaterialHLSLTree.cpp
--- a/Engine/Source/Runtime/Engine/Private/Materials/MaterialHLSLTree.cpp
+++ b/Engine/Source/Runtime/Engine/Private/Materials/MaterialHLSLTree.cpp
@@ -597,6 +597,34 @@ void FExpressionParameter::ComputeAnalyticDerivatives(FTree& Tree, FExpressionDe
 	}
 }
 
+const FExpression* FExpressionParameter::GetPreviewExpression(FTree& Tree) const
+{
+	const EMaterialParameterType ParameterType = ParameterMeta.Value.Type;
+	switch (ParameterType)
+	{
+	case EMaterialParameterType::Scalar:
+	case EMaterialParameterType::Vector:
+	case EMaterialParameterType::DoubleVector:
+		return this;
+	case EMaterialParameterType::Texture:
+	{
+		const FExpression* TexCoordsExpression = Tree.NewExpression<FExpressionExternalInput>(MakeInputTexCoord(0));
+		return Tree.NewExpression<FExpressionTextureSample>(
+			this,
+			TexCoordsExpression,
+			nullptr /* InMipValueExpression */,
+			nullptr /* InAutomaticMipBiasExpression */,
+			FExpressionDerivatives(),
+			SSM_FromTextureAsset,
+			TMVM_None);
+	}
+	default:
+		// Not implemented yet
+		checkNoEntry();
+		return this;
+	}
+}
+
 bool FExpressionParameter::PrepareValue(FEmitContext& Context, FEmitScope& Scope, const FRequestedType& RequestedType, FPrepareValueResult& OutResult) const
 {
 	FEmitData& EmitData = Context.FindData<FEmitData>();
@@ -621,6 +649,14 @@ bool FExpressionParameter::PrepareValue(FEmitContext& Context, FEmitScope& Scope
 	EExpressionEvaluation Evaluation = EExpressionEvaluation::Shader;
 	if (IsStaticMaterialParameter(ParameterType))
 	{
+		if (EmitData.CachedExpressionData)
+		{
+			check(!ParameterInfo.Name.IsNone());
+			// We are preparing the tree for cached expression data update. Need to make sure
+			// static parameter values are consistent with later translations. Otherwise, some
+			// sub-trees may be evaluated without being prepared first
+			Context.SeenStaticParameterValues.FindOrAdd(ParameterInfo, ParameterMeta.Value);
+		}
 		Evaluation = EExpressionEvaluation::Constant;
 	}
 	else if (ParameterType == EMaterialParameterType::Scalar ||
@@ -656,6 +692,7 @@ void FExpressionParameter::EmitValuePreshader(FEmitContext& Context, FEmitScope&
 	if (IsStaticMaterialParameter(ParameterType))
 	{
 		Shader::FValue Value = DefaultValue;
+		bool bFoundOverride = false;
 		if (EmitMaterialData.StaticParameters)
 		{
 			switch (ParameterType)
@@ -666,6 +703,7 @@ void FExpressionParameter::EmitValuePreshader(FEmitContext& Context, FEmitScope&
 					if (Parameter.ParameterInfo == ParameterInfo)
 					{
 						Value = Parameter.Value;
+						bFoundOverride = true;
 						break;
 					}
 				}
@@ -676,6 +714,7 @@ void FExpressionParameter::EmitValuePreshader(FEmitContext& Context, FEmitScope&
 					if (Parameter.ParameterInfo == ParameterInfo)
 					{
 						Value = Shader::FValue(Parameter.R, Parameter.G, Parameter.B, Parameter.A);
+						bFoundOverride = true;
 						break;
 					}
 				}
@@ -685,6 +724,13 @@ void FExpressionParameter::EmitValuePreshader(FEmitContext& Context, FEmitScope&
 				break;
 			}
 		}
+		
+		if (!bFoundOverride && EmitMaterialData.CachedExpressionData)
+		{
+			// For some reasons, users can create static switch parameter nodes with the same name
+			// but different default values. Use the first default value seen to keep things consistent
+			Value = Context.SeenStaticParameterValues.FindChecked(ParameterInfo).AsShaderValue();
+		}
 		OutResult.Preshader.WriteOpcode(Shader::EPreshaderOpcode::Constant).Write(Value);
 	}
 	else
@@ -1029,10 +1075,10 @@ Shader::EValueType GetTexCoordType(EMaterialValueType TextureType)
 	}
 }
 
-ETextureMipValueMode GetMipValueMode(FEmitContext& Context, const FExpressionTextureSample* Expression)
+ETextureMipValueMode GetMipValueMode(FEmitContext& Context, const FExpressionTextureSample* Expression, bool bDerivsPrepFailed = false)
 {
 	const ETextureMipValueMode MipValueMode = Expression->MipValueMode;
-	const bool bUseAnalyticDerivatives = Context.bUseAnalyticDerivatives && (MipValueMode != TMVM_MipLevel) && Expression->TexCoordDerivatives.IsValid();
+	const bool bUseAnalyticDerivatives = Context.bUseAnalyticDerivatives && MipValueMode != TMVM_MipLevel && !bDerivsPrepFailed && Expression->TexCoordDerivatives.IsValid();
 	if (Context.ShaderFrequency != SF_Pixel)
 	{
 		// TODO - should we allow TMVM_Derivative in non-PS?
@@ -1086,11 +1132,24 @@ bool FExpressionTextureSample::PrepareValue(FEmitContext& Context, FEmitScope& S
 		}
 	}
 
-	const ETextureMipValueMode LocalMipValueMode = Private::GetMipValueMode(Context, this);
+	bool bDerivsPrepFailed = false;
+ RecomputeMipValueMode:
+	const ETextureMipValueMode LocalMipValueMode = Private::GetMipValueMode(Context, this, bDerivsPrepFailed);
+
 	if (LocalMipValueMode == TMVM_Derivative)
 	{
-		Context.PrepareExpression(TexCoordDerivatives.ExpressionDdx, Scope, RequestedTexCoordType);
-		Context.PrepareExpression(TexCoordDerivatives.ExpressionDdy, Scope, RequestedTexCoordType);
+		const FPreparedType& DdxPreparedType = Context.PrepareExpression(TexCoordDerivatives.ExpressionDdx, Scope, RequestedTexCoordType);
+		if (DdxPreparedType.IsVoid() && !bDerivsPrepFailed)
+		{
+			bDerivsPrepFailed = true;
+			goto RecomputeMipValueMode;
+		}
+		const FPreparedType& DdyPreparedType = Context.PrepareExpression(TexCoordDerivatives.ExpressionDdy, Scope, RequestedTexCoordType);
+		if (DdyPreparedType.IsVoid() && !bDerivsPrepFailed)
+		{
+			bDerivsPrepFailed = true;
+			goto RecomputeMipValueMode;
+		}
 	}
 	else if (LocalMipValueMode == TMVM_MipLevel || LocalMipValueMode == TMVM_MipBias)
 	{
@@ -3095,6 +3154,7 @@ bool FExpressionNaniteReplaceFunction::PrepareValue(FEmitContext& Context, FEmit
 	{
 		return false;
 	}
+	Shader::FType ResultType = DefaultType.Type;
 
 	// skip preparing if platform doesn't support Nanite
 	if (Context.TargetParameters.IsGenericTarget() || FDataDrivenShaderPlatformInfo::GetSupportsNanite(Context.TargetParameters.ShaderPlatform))
@@ -3104,24 +3164,30 @@ bool FExpressionNaniteReplaceFunction::PrepareValue(FEmitContext& Context, FEmit
 		{
 			return false;
 		}
+		ResultType = Shader::CombineTypes(ResultType, NaniteType.Type);
 	}
 
-	return OutResult.SetType(Context, RequestedType, EExpressionEvaluation::Shader, RequestedType.Type);
+	if (ResultType.IsVoid())
+	{
+		return false;
+	}
+
+	return OutResult.SetType(Context, RequestedType, EExpressionEvaluation::Shader, ResultType);
 }
 
 void FExpressionNaniteReplaceFunction::EmitValueShader(FEmitContext& Context, FEmitScope& Scope, const FRequestedType& RequestedType, FEmitValueShaderResult& OutResult) const
 {
-	const Shader::FType LocalType = Context.GetResultType(this, RequestedType);
-	FEmitShaderExpression* DefaultValue = DefaultExpression->GetValueShader(Context, Scope, RequestedType, LocalType);
+	const Shader::FType ResultType = Context.GetResultType(this, RequestedType);
+	FEmitShaderExpression* DefaultValue = DefaultExpression->GetValueShader(Context, Scope, RequestedType, ResultType);
 
 	if (FDataDrivenShaderPlatformInfo::GetSupportsNanite(Context.TargetParameters.ShaderPlatform))
 	{
-		FEmitShaderExpression* NaniteValue = NaniteExpression->GetValueShader(Context, Scope, RequestedType, LocalType);
-		OutResult.Code = Context.EmitExpression(Scope, LocalType, TEXT("(GetNaniteReplaceState() ? % : %)"), NaniteValue, DefaultValue);
+		FEmitShaderExpression* NaniteValue = NaniteExpression->GetValueShader(Context, Scope, RequestedType, ResultType);
+		OutResult.Code = Context.EmitExpression(Scope, ResultType, TEXT("(GetNaniteReplaceState() ? % : %)"), NaniteValue, DefaultValue);
 	}
 	else
 	{
-		OutResult.Code = Context.EmitExpression(Scope, LocalType, TEXT("%"), DefaultValue);
+		OutResult.Code = Context.EmitExpression(Scope, ResultType, TEXT("%"), DefaultValue);
 	}
 }
 
diff --git a/Engine/Source/Runtime/Engine/Private/Materials/MaterialInstanceConstant.cpp b/Engine/Source/Runtime/Engine/Private/Materials/MaterialInstanceConstant.cpp
--- a/Engine/Source/Runtime/Engine/Private/Materials/MaterialInstanceConstant.cpp
+++ b/Engine/Source/Runtime/Engine/Private/Materials/MaterialInstanceConstant.cpp
@@ -281,7 +281,7 @@ void UMaterialInstanceConstant::UpdateCachedData()
 
 			check(!LocalCachedExpressionData);
 			LocalCachedExpressionData.Reset(new FMaterialCachedExpressionData());
-			LocalCachedExpressionData->UpdateForCachedHLSLTree(GetCachedHLSLTree(), &OverriddenStaticParameters);
+			LocalCachedExpressionData->UpdateForCachedHLSLTree(GetCachedHLSLTree(), &OverriddenStaticParameters, this);
 		}
 
 		CachedExpressionData = MoveTemp(LocalCachedExpressionData);
diff --git a/Engine/Source/Runtime/Engine/Private/Materials/MaterialShared.cpp b/Engine/Source/Runtime/Engine/Private/Materials/MaterialShared.cpp
--- a/Engine/Source/Runtime/Engine/Private/Materials/MaterialShared.cpp
+++ b/Engine/Source/Runtime/Engine/Private/Materials/MaterialShared.cpp
@@ -1730,7 +1730,7 @@ bool FMaterialResource::IsUsedWithEditorCompositing() const { return Material->b
 bool FMaterialResource::IsDeferredDecal() const { return Material->MaterialDomain == MD_DeferredDecal; }
 bool FMaterialResource::IsVolumetricPrimitive() const { return Material->MaterialDomain == MD_Volume; }
 bool FMaterialResource::IsSpecialEngineMaterial() const { return Material->bUsedAsSpecialEngineMaterial; }
-bool FMaterialResource::HasVertexPositionOffsetConnected() const { return Material->HasVertexPositionOffsetConnected(); }
+bool FMaterialResource::HasVertexPositionOffsetConnected() const { return GetCachedExpressionData().IsPropertyConnected(MP_WorldPositionOffset); }
 bool FMaterialResource::HasPixelDepthOffsetConnected() const { return Material->HasPixelDepthOffsetConnected(); }
 #if WITH_EDITOR
 bool FMaterialResource::HasMaterialAttributesConnected() const { return (Material->bUseMaterialAttributes && Material->GetEditorOnlyData()->MaterialAttributes.IsConnected()); }
diff --git a/Engine/Source/Runtime/Engine/Public/HLSLTree/HLSLTree.h b/Engine/Source/Runtime/Engine/Public/HLSLTree/HLSLTree.h
--- a/Engine/Source/Runtime/Engine/Public/HLSLTree/HLSLTree.h
+++ b/Engine/Source/Runtime/Engine/Public/HLSLTree/HLSLTree.h
@@ -471,6 +471,9 @@ protected:
 	/** Creates a new expression representing this expression on the previous frame.  By default returns nullptr, which means the previous frame is the same as the current frame */
 	ENGINE_API virtual const FExpression* ComputePreviousFrame(FTree& Tree, const FRequestedType& RequestedType) const;
 
+	/** Non-numeric expressions need to override and return numeric preview expressions. Returns this expression itself by default */
+	ENGINE_API virtual const FExpression* GetPreviewExpression(FTree& Tree) const;
+
 	/**
 	 * Computes a FPreparedType for this expression, given a FRequestedType.  Will be called multiple times, with potentially different requested types, if the FExpression is used multiple times.
 	 * In this case, all the prepared types will be merged together if possible (or generate an error otherwise).
@@ -509,6 +512,7 @@ private:
 	friend class FExpressionPreviousFrameSwitch;
 	friend class FExpressionFunctionCall;
 	friend class FExpressionOperation;
+	friend class FExpressionSetStructField;
 };
 
 /**
@@ -626,6 +630,7 @@ public:
 
 	ENGINE_API FExpressionDerivatives GetAnalyticDerivatives(const FExpression* InExpression);
 	const FExpression* GetPreviousFrame(const FExpression* InExpression, const FRequestedType& RequestedType);
+	const FExpression* GetPreview(const FExpression* InExpression);
 
 	FScope* NewScope(FScope& Scope);
 	FScope* NewOwnedScope(FStatement& Owner);
diff --git a/Engine/Source/Runtime/Engine/Public/HLSLTree/HLSLTreeCommon.h b/Engine/Source/Runtime/Engine/Public/HLSLTree/HLSLTreeCommon.h
--- a/Engine/Source/Runtime/Engine/Public/HLSLTree/HLSLTreeCommon.h
+++ b/Engine/Source/Runtime/Engine/Public/HLSLTree/HLSLTreeCommon.h
@@ -4,6 +4,7 @@
 #if WITH_EDITOR
 
 #include "HLSLTree/HLSLTree.h"
+#include "SceneTypes.h"
 #if UE_ENABLE_INCLUDE_ORDER_DEPRECATED_IN_5_2
 #include "RHIDefinitions.h"
 #endif
@@ -118,11 +119,12 @@ public:
 class FExpressionSetStructField : public FExpression
 {
 public:
-	FExpressionSetStructField(const Shader::FStructType* InStructType, const Shader::FStructField* InField, const FExpression* InStructExpression, const FExpression* InFieldExpression)
+	FExpressionSetStructField(const Shader::FStructType* InStructType, const Shader::FStructField* InField, const FExpression* InStructExpression, const FExpression* InFieldExpression, EMaterialProperty InTestMaterialProperty = MP_MAX)
 		: StructType(InStructType)
 		, Field(InField)
 		, StructExpression(InStructExpression)
 		, FieldExpression(InFieldExpression)
+		, TestMaterialProperty(InTestMaterialProperty)
 	{
 		check(InStructType);
 		check(InField);
@@ -132,6 +134,8 @@ public:
 	const Shader::FStructField* Field;
 	const FExpression* StructExpression;
 	const FExpression* FieldExpression;
+	// If not MP_MAX, the material property to test IsPropertyActive on a given material (instance)
+	EMaterialProperty TestMaterialProperty;
 
 	virtual void ComputeAnalyticDerivatives(FTree& Tree, FExpressionDerivatives& OutResult) const override;
 	virtual const FExpression* ComputePreviousFrame(FTree& Tree, const FRequestedType& RequestedType) const override;
@@ -142,6 +146,8 @@ public:
 private:
 	FRequestedType MakeRequestedStructType(const FActiveStructFieldStack& ActiveFieldStack, const FRequestedType& RequestedType) const;
 	FRequestedType MakeRequestedFieldType(const FActiveStructFieldStack& ActiveFieldStack, const FRequestedType& RequestedType) const;
+
+	bool IsStructFieldActive(FEmitContext& Context) const;
 };
 
 class FExpressionSelect : public FExpression
diff --git a/Engine/Source/Runtime/Engine/Public/HLSLTree/HLSLTreeEmit.h b/Engine/Source/Runtime/Engine/Public/HLSLTree/HLSLTreeEmit.h
--- a/Engine/Source/Runtime/Engine/Public/HLSLTree/HLSLTreeEmit.h
+++ b/Engine/Source/Runtime/Engine/Public/HLSLTree/HLSLTreeEmit.h
@@ -596,6 +596,7 @@ public:
 	TMap<const FScope*, FEmitScope*> EmitScopeMap;
 	TMap<FXxHash64, FPrepareValueResult*> PrepareValueMap;
 	TMap<FXxHash64, FRequestedType*> RequestedTypeTracker;
+	TMap<FMaterialParameterInfo, FMaterialParameterValue> SeenStaticParameterValues;
 	TMap<const FExpression*, FEmitScope*> PrepareLocalPHIMap;
 	TMap<FXxHash64, FEmitShaderExpression*> EmitLocalPHIMap;
 	TMap<FXxHash64, FEmitShaderExpression*> EmitExpressionMap;
@@ -620,6 +621,7 @@ public:
 	 * Would also need some generic interface between the preshaders generated here, and the preshaders stored in FMaterialCompilationOutput
 	 */
 	const FMaterial* Material = nullptr;
+	const UMaterialInterface* MaterialInterface = nullptr;
 	FMaterialCompilationOutput* MaterialCompilationOutput = nullptr;
 	uint32 UniformPreshaderOffset = 0u;
 	uint32 CurrentBoolUniformOffset = ~0u;
diff --git a/Engine/Source/Runtime/Engine/Public/MaterialCachedData.h b/Engine/Source/Runtime/Engine/Public/MaterialCachedData.h
--- a/Engine/Source/Runtime/Engine/Public/MaterialCachedData.h
+++ b/Engine/Source/Runtime/Engine/Public/MaterialCachedData.h
@@ -193,7 +193,7 @@ struct FMaterialCachedExpressionData
 	void UpdateForLayerFunctions(const FMaterialCachedExpressionContext& Context, const FMaterialLayersFunctions& LayerFunctions);
 	void AnalyzeMaterial(UMaterial& Material);
 
-	ENGINE_API void UpdateForCachedHLSLTree(const FMaterialCachedHLSLTree& CachedTree, const FStaticParameterSet* StaticParameters);
+	ENGINE_API void UpdateForCachedHLSLTree(const FMaterialCachedHLSLTree& CachedTree, const FStaticParameterSet* StaticParameters, const UMaterialInterface* TargetMaterial);
 	void Validate(const UMaterialInterface& Material);
 
 	/** Adds a parameter. If this returns false, a parameter with identical name has already been added but it was set to a different value. */
diff --git a/Engine/Source/Runtime/Engine/Public/MaterialHLSLTree.h b/Engine/Source/Runtime/Engine/Public/MaterialHLSLTree.h
--- a/Engine/Source/Runtime/Engine/Public/MaterialHLSLTree.h
+++ b/Engine/Source/Runtime/Engine/Public/MaterialHLSLTree.h
@@ -291,6 +291,7 @@ public:
 	EMaterialSamplerType TextureSamplerType;
 
 	virtual void ComputeAnalyticDerivatives(FTree& Tree, FExpressionDerivatives& OutResult) const override;
+	virtual const FExpression* GetPreviewExpression(FTree& Tree) const override;
 	virtual bool PrepareValue(FEmitContext& Context, FEmitScope& Scope, const FRequestedType& RequestedType, FPrepareValueResult& OutResult) const override;
 	virtual void EmitValuePreshader(FEmitContext& Context, FEmitScope& Scope, const FRequestedType& RequestedType, FEmitValuePreshaderResult& OutResult) const override;
 	virtual void EmitValueShader(FEmitContext& Context, FEmitScope& Scope, const FRequestedType& RequestedType, FEmitValueShaderResult& OutResult) const override;
