Fix CIS Issue 418008: PVS static analysis

diff --git a/Engine/Source/Runtime/Engine/Private/Engine/ActorInstanceHandle.cpp b/Engine/Source/Runtime/Engine/Private/Engine/ActorInstanceHandle.cpp
--- a/Engine/Source/Runtime/Engine/Private/Engine/ActorInstanceHandle.cpp
+++ b/Engine/Source/Runtime/Engine/Private/Engine/ActorInstanceHandle.cpp
@@ -159,7 +159,6 @@ void FActorInstanceHandle::ResolveHandle() const
 		{
 			// OwnerActor is not valid so we need to reset the handle
 			new (MutableThis)(FActorInstanceHandle);
-			checkf(ResolutionStatus == EResolutionStatus::Invalid, TEXT("Default constructor must mark the resolution as invalid."));
 		}
 	}
 }
