PR #11061: PostProcessAnimBPLODThreshold implementation for Skeletal Mesh Component

diff --git a/Engine/Source/Runtime/Engine/Classes/Components/SkeletalMeshComponent.h b/Engine/Source/Runtime/Engine/Classes/Components/SkeletalMeshComponent.h
--- a/Engine/Source/Runtime/Engine/Classes/Components/SkeletalMeshComponent.h
+++ b/Engine/Source/Runtime/Engine/Classes/Components/SkeletalMeshComponent.h
@@ -552,7 +552,17 @@ public:
 	/** Used to scale speed of all animations on this skeletal mesh. */
 	UPROPERTY(EditAnywhere, AdvancedDisplay, BlueprintReadWrite, Category=Animation)
 	float GlobalAnimRateScale;
-	
+
+private:
+	/*
+	 * Max LOD level that post-process AnimBPs are evaluated.
+	 * For example if you have the threshold set to 2, it will evaluate until including LOD 2 (based on 0 index). In case the LOD level gets set to 3, it will stop evaluating the post-process AnimBP.
+	 * Setting it to -1 will always evaluate it and disable LODing.
+	 */
+	UPROPERTY(EditAnywhere, AdvancedDisplay, BlueprintReadWrite, Category=Animation, meta= (DisplayName = "Post-Process AnimBP LOD Threshold", AllowPrivateAccess = "true"))
+	int32 PostProcessAnimBPLODThreshold = INDEX_NONE;
+
+public:
 	/** If we are running physics, should we update non-simulated bones based on the animation bone positions. */
 	UPROPERTY(EditAnywhere, AdvancedDisplay, BlueprintReadWrite, Category=SkeletalMesh)
 	TEnumAsByte<EKinematicBonesUpdateToPhysics::Type> KinematicBonesUpdateType;
@@ -1954,6 +1964,7 @@ public:
 	FOnBoneTransformsFinalized OnBoneTransformsFinalized;
 
 	// Conditions used to gate when post process events happen
+	ENGINE_API bool ShouldEvaluatePostProcessAnimBP() const;
 	ENGINE_API bool ShouldUpdatePostProcessInstance() const;
 	ENGINE_API bool ShouldPostUpdatePostProcessInstance() const;
 	ENGINE_API bool ShouldEvaluatePostProcessInstance() const;
diff --git a/Engine/Source/Runtime/Engine/Private/Components/SkeletalMeshComponent.cpp b/Engine/Source/Runtime/Engine/Private/Components/SkeletalMeshComponent.cpp
--- a/Engine/Source/Runtime/Engine/Private/Components/SkeletalMeshComponent.cpp
+++ b/Engine/Source/Runtime/Engine/Private/Components/SkeletalMeshComponent.cpp
@@ -4331,27 +4331,40 @@ void USkeletalMeshComponent::FinalizeBoneTransform()
 	TRACE_SKELETAL_MESH_COMPONENT(this);
 }
 
-bool USkeletalMeshComponent::ShouldUpdatePostProcessInstance() const
+bool USkeletalMeshComponent::ShouldEvaluatePostProcessAnimBP() const
 {
+	const int32 LODLevel = GetPredictedLODLevel();
+	if (!(PostProcessAnimBPLODThreshold == INDEX_NONE || LODLevel <= PostProcessAnimBPLODThreshold))
+	{
+		return false;
+	}
+
 	if (USkeletalMesh* SkelMesh = GetSkeletalMeshAsset())
 	{
-		if (!SkelMesh->ShouldEvaluatePostProcessAnimBP(GetPredictedLODLevel()))
+		if (!SkelMesh->ShouldEvaluatePostProcessAnimBP(LODLevel))
 		{
 			return false;
 		}
 	}
 
+	return true;
+}
+
+bool USkeletalMeshComponent::ShouldUpdatePostProcessInstance() const
+{
+	if (!ShouldEvaluatePostProcessAnimBP())
+	{
+		return false;
+	}
+
 	return PostProcessAnimInstance && !bDisablePostProcessBlueprint;
 }
 
 bool USkeletalMeshComponent::ShouldPostUpdatePostProcessInstance() const
 {
-	if (USkeletalMesh* SkelMesh = GetSkeletalMeshAsset())
+	if (!ShouldEvaluatePostProcessAnimBP())
 	{
-		if (!SkelMesh->ShouldEvaluatePostProcessAnimBP(GetPredictedLODLevel()))
-		{
-			return false;
-		}
+		return false;
 	}
 
 	return PostProcessAnimInstance && PostProcessAnimInstance->NeedsUpdate() && !bDisablePostProcessBlueprint;
@@ -4359,12 +4372,9 @@ bool USkeletalMeshComponent::ShouldPostUpdatePostProcessInstance() const
 
 bool USkeletalMeshComponent::ShouldEvaluatePostProcessInstance() const
 {
-	if (USkeletalMesh* SkelMesh = GetSkeletalMeshAsset())
+	if (!ShouldEvaluatePostProcessAnimBP())
 	{
-		if (!SkelMesh->ShouldEvaluatePostProcessAnimBP(GetPredictedLODLevel()))
-		{
-			return false;
-		}
+		return false;
 	}
 
 	return PostProcessAnimInstance && !bDisablePostProcessBlueprint;
