Fix recording/display of pose attribute debug info

Attributes were being recorded for the preview mesh at post-update time, rather than post-evaluation (which is what generates the debug data)

diff --git a/Engine/Source/Runtime/Engine/Private/Animation/AnimInstanceProxy.cpp b/Engine/Source/Runtime/Engine/Private/Animation/AnimInstanceProxy.cpp
--- a/Engine/Source/Runtime/Engine/Private/Animation/AnimInstanceProxy.cpp
+++ b/Engine/Source/Runtime/Engine/Private/Animation/AnimInstanceProxy.cpp
@@ -603,7 +603,6 @@ void FAnimInstanceProxy::PostUpdate(UAnimInstance* InAnimInstance) const
 	if (FAnimBlueprintDebugData* DebugData = GetAnimBlueprintDebugData())
 	{
 		DebugData->RecordNodeVisitArray(UpdatedNodesThisFrame);
-		DebugData->RecordNodeAttributeMaps(NodeInputAttributesThisFrame, NodeOutputAttributesThisFrame);
 		DebugData->RecordNodeSyncsArray(NodeSyncsThisFrame);
 	}
 #endif
@@ -703,6 +702,13 @@ void FAnimInstanceProxy::PostEvaluate(UAnimInstance* InAnimInstance)
 {
 	DECLARE_SCOPE_HIERARCHICAL_COUNTER_FUNC()
 
+#if WITH_EDITORONLY_DATA
+	if (FAnimBlueprintDebugData* DebugData = GetAnimBlueprintDebugData())
+	{
+		DebugData->RecordNodeAttributeMaps(NodeInputAttributesThisFrame, NodeOutputAttributesThisFrame);
+	}
+#endif
+	
 	ClearObjects();
 
 #if ENABLE_ANIM_LOGGING
