Add MSAA sample count support to canvas render targets.

We restrain the implementation UCanvasRenderTarget2D to minimize the scope of what needs to be resolved: only the canvas render passes and not scene captures, etc.

diff --git a/Engine/Source/Runtime/Engine/Classes/Engine/CanvasRenderTarget2D.h b/Engine/Source/Runtime/Engine/Classes/Engine/CanvasRenderTarget2D.h
--- a/Engine/Source/Runtime/Engine/Classes/Engine/CanvasRenderTarget2D.h
+++ b/Engine/Source/Runtime/Engine/Classes/Engine/CanvasRenderTarget2D.h
@@ -91,13 +91,30 @@ public:
 		bShouldClearRenderTargetOnReceiveUpdate = bInShouldClearRenderTargetOnReceiveUpdate;
 	}
 
+	UFUNCTION(BlueprintGetter)
+	virtual ETextureRenderTargetSampleCount GetSampleCount() const override;
+
+	UFUNCTION(BlueprintSetter)
+	void SetSampleCount(ETextureRenderTargetSampleCount InSampleCount);
+
+#if WITH_EDITOR
+	virtual bool CanEditChange(const FProperty* InProperty) const override;
+	virtual void PostEditChangeProperty(FPropertyChangedEvent& PropertyChangedEvent) override;
+#endif // WITH_EDITOR
+
 protected:
 
 	ENGINE_API void RepaintCanvas();
 
+	void OnSampleCountChanged();
+
 	/* The world this render target will be used with */
 	UPROPERTY()
 	TWeakObjectPtr<UWorld> World;
+
+	/** MSAA sample count. */
+	UPROPERTY(EditAnywhere, BlueprintGetter = GetSampleCount, BlueprintSetter = SetSampleCount, Category = TextureRenderTarget2D, AssetRegistrySearchable, AdvancedDisplay, meta = (AllowPrivateAccess = "true"))
+	ETextureRenderTargetSampleCount SampleCount;
 	
 	// If true, clear the render target to green whenever OnReceiveUpdate() is called.  (Defaults to true.)
 	// If false, the render target will retain whatever values it had, allowing the user to update only areas that
diff --git a/Engine/Source/Runtime/Engine/Classes/Engine/TextureRenderTarget2D.h b/Engine/Source/Runtime/Engine/Classes/Engine/TextureRenderTarget2D.h
--- a/Engine/Source/Runtime/Engine/Classes/Engine/TextureRenderTarget2D.h
+++ b/Engine/Source/Runtime/Engine/Classes/Engine/TextureRenderTarget2D.h
@@ -66,6 +66,32 @@ inline EPixelFormat GetPixelFormatFromRenderTargetFormat(ETextureRenderTargetFor
 	return PF_Unknown;
 }
 
+UENUM(BlueprintType)
+enum class ETextureRenderTargetSampleCount : uint8
+{
+	RTSC_1 UMETA(DisplayName = "MSAAx1"),
+	RTSC_2 UMETA(DisplayName = "MSAAx2"),
+	RTSC_4 UMETA(DisplayName = "MSAAx4"),
+	RTSC_8 UMETA(DisplayName = "MSAAx8"),
+
+	RTSC_MAX,
+};
+
+inline int32 GetNumFromRenderTargetSampleCount(ETextureRenderTargetSampleCount InSampleCount)
+{
+	switch (InSampleCount)
+	{
+	case ETextureRenderTargetSampleCount::RTSC_1: return 1 << 0;
+	case ETextureRenderTargetSampleCount::RTSC_2: return 1 << 1;
+	case ETextureRenderTargetSampleCount::RTSC_4: return 1 << 2;
+	case ETextureRenderTargetSampleCount::RTSC_8: return 1 << 3;
+
+	default:
+		ensureMsgf(false, TEXT("Unhandled ETextureRenderTargetSampleCount entry %u"), (uint32)InSampleCount);
+		return 1;
+	}
+}
+
 /**
  * TextureRenderTarget2D
  *
@@ -235,6 +261,8 @@ PRAGMA_ENABLE_DEPRECATION_WARNINGS
 		return NumMips;
 	}
 
+	virtual ETextureRenderTargetSampleCount GetSampleCount() const;
+
 private:
 	int32	NumMips;
 };
diff --git a/Engine/Source/Runtime/Engine/Private/CanvasRenderTarget2D.cpp b/Engine/Source/Runtime/Engine/Private/CanvasRenderTarget2D.cpp
--- a/Engine/Source/Runtime/Engine/Private/CanvasRenderTarget2D.cpp
+++ b/Engine/Source/Runtime/Engine/Private/CanvasRenderTarget2D.cpp
@@ -14,8 +14,9 @@
 #include UE_INLINE_GENERATED_CPP_BY_NAME(CanvasRenderTarget2D)
 
 UCanvasRenderTarget2D::UCanvasRenderTarget2D( const FObjectInitializer& ObjectInitializer )
-	: Super(ObjectInitializer),
-	  World( nullptr )
+	: Super(ObjectInitializer)
+	, World( nullptr )
+	, SampleCount(ETextureRenderTargetSampleCount::RTSC_1)
 {
 	bNeedsTwoCopies = false;
 	bShouldClearRenderTargetOnReceiveUpdate = true;
@@ -56,6 +57,56 @@ void UCanvasRenderTarget2D::FastUpdateResource()
 	}
 }
 
+ETextureRenderTargetSampleCount UCanvasRenderTarget2D::GetSampleCount() const
+{
+	return SampleCount;
+}
+
+void UCanvasRenderTarget2D::SetSampleCount(ETextureRenderTargetSampleCount InSampleCount)
+{
+	if (SampleCount != InSampleCount)
+	{
+		SampleCount = InSampleCount;
+
+		OnSampleCountChanged();
+	}
+}
+
+#if WITH_EDITOR
+bool UCanvasRenderTarget2D::CanEditChange(const FProperty* InProperty) const
+{
+	if (!Super::CanEditChange(InProperty))
+	{
+		return false;
+	}
+
+	const FName& Name = InProperty->GetFName();
+
+	if (Name == GET_MEMBER_NAME_CHECKED(ThisClass, SampleCount))
+	{
+		return !bAutoGenerateMips;
+	}
+	else if (Name == GET_MEMBER_NAME_CHECKED(ThisClass, bAutoGenerateMips))
+	{
+		return SampleCount == ETextureRenderTargetSampleCount::RTSC_1;
+	}
+
+	return true;
+}
+
+void UCanvasRenderTarget2D::PostEditChangeProperty(FPropertyChangedEvent& PropertyChangedEvent)
+{
+	const FName PropertyName = (PropertyChangedEvent.Property != nullptr) ? PropertyChangedEvent.Property->GetFName() : NAME_None;
+
+	if (PropertyName == GET_MEMBER_NAME_CHECKED(ThisClass, SampleCount))
+	{
+		OnSampleCountChanged();
+	}
+
+	Super::PostEditChangeProperty(PropertyChangedEvent);
+}
+#endif
+
 void UCanvasRenderTarget2D::RepaintCanvas()
 {
 	// Create or find the canvas object to use to render onto the texture.  Multiple canvas render target textures can share the same canvas.
@@ -113,6 +164,18 @@ void UCanvasRenderTarget2D::RepaintCanvas()
 	UpdateResourceImmediate(false);
 }
 
+void UCanvasRenderTarget2D::OnSampleCountChanged()
+{
+	// UAVs are not supported for sample counts higher than one.
+	bCanCreateUAV = (SampleCount == ETextureRenderTargetSampleCount::RTSC_1);
+
+	if (bAutoGenerateMips && SampleCount != ETextureRenderTargetSampleCount::RTSC_1)
+	{
+		UE_LOG(LogTexture, Warning, TEXT("Auto-generate mips is not supported with sample counts above 1, disabling."));
+		bAutoGenerateMips = false;
+	}
+}
+
 UCanvasRenderTarget2D* UCanvasRenderTarget2D::CreateCanvasRenderTarget2D(UObject* WorldContextObject, TSubclassOf<UCanvasRenderTarget2D> CanvasRenderTarget2DClass, int32 Width, int32 Height)
 {
 	if ((Width > 0) && (Height > 0) && (CanvasRenderTarget2DClass != NULL))
diff --git a/Engine/Source/Runtime/Engine/Private/KismetRenderingLibrary.cpp b/Engine/Source/Runtime/Engine/Private/KismetRenderingLibrary.cpp
--- a/Engine/Source/Runtime/Engine/Private/KismetRenderingLibrary.cpp
+++ b/Engine/Source/Runtime/Engine/Private/KismetRenderingLibrary.cpp
@@ -831,12 +831,16 @@ void UKismetRenderingLibrary::EndDrawCanvasToRenderTarget(UObject* WorldContextO
 		{
 			FTextureRenderTargetResource* RenderTargetResource = Context.RenderTarget->GameThread_GetRenderTargetResource();
 
-			ENQUEUE_RENDER_COMMAND(CanvasRenderTargetResolveCommand)(
-				[RenderTargetResource](FRHICommandListImmediate& RHICmdList)
-				{
-					TransitionAndCopyTexture(RHICmdList, RenderTargetResource->GetRenderTargetTexture(), RenderTargetResource->TextureRHI, {});
-				}
-			);
+			// Note: If multisampled, it should have already been resolved by ~FCanvasRenderThreadScope()
+			if (!RenderTargetResource->GetRenderTargetTexture()->GetDesc().IsMultisample())
+			{
+				ENQUEUE_RENDER_COMMAND(CanvasRenderTargetResolveCommand)(
+					[RenderTargetResource](FRHICommandListImmediate& RHICmdList)
+					{
+						TransitionAndCopyTexture(RHICmdList, RenderTargetResource->GetRenderTargetTexture(), RenderTargetResource->TextureRHI, {});
+					}
+				);
+			}
 
 #if WANTS_DRAW_MESH_EVENTS
 			STOP_DRAW_EVENT_GAMETHREAD(*Context.DrawEvent);
diff --git a/Engine/Source/Runtime/Engine/Private/TextureRenderTarget2D.cpp b/Engine/Source/Runtime/Engine/Private/TextureRenderTarget2D.cpp
--- a/Engine/Source/Runtime/Engine/Private/TextureRenderTarget2D.cpp
+++ b/Engine/Source/Runtime/Engine/Private/TextureRenderTarget2D.cpp
@@ -356,6 +356,12 @@ FString UTextureRenderTarget2D::GetDesc()
 	return FString::Printf( TEXT("Render to Texture %dx%d[%s]"), SizeX, SizeY, GPixelFormats[GetFormat()].Name );
 }
 
+ETextureRenderTargetSampleCount UTextureRenderTarget2D::GetSampleCount() const
+{
+	// Note: MSAA is currently only supported in UCanvasRenderTarget2D
+	return ETextureRenderTargetSampleCount::RTSC_1;
+}
+
 UTexture2D* UTextureRenderTarget2D::ConstructTexture2D(UObject* InOuter, const FString& InNewTextureName, EObjectFlags InObjectFlags, uint32 InFlags, TArray<uint8>* InAlphaOverride)
 {
 	UTexture2D* Result = nullptr;
@@ -520,12 +526,14 @@ void FTextureRenderTarget2DResource::InitRHI(FRHICommandListBase& RHICmdList)
 
 		FString ResourceName = Owner->GetName();
 		ETextureCreateFlags TexCreateFlags = GetCreateFlags();
+		const int32 NumSamples = GetNumFromRenderTargetSampleCount(Owner->GetSampleCount());
 
 		FRHITextureCreateDesc Desc =
 			FRHITextureCreateDesc::Create2D(*ResourceName)
 			.SetExtent(Owner->SizeX, Owner->SizeY)
 			.SetFormat(Format)
 			.SetNumMips(Owner->GetNumMips())
+			.SetNumSamples(NumSamples)
 			.SetFlags(TexCreateFlags)
 			.SetInitialState(ERHIAccess::SRVMask)
 			.SetClearValue(FClearValueBinding(ClearColor))
@@ -534,7 +542,20 @@ void FTextureRenderTarget2DResource::InitRHI(FRHICommandListBase& RHICmdList)
 
 		TextureRHI = RenderTargetTextureRHI = RHICreateTexture(Desc);
 
-		if (Owner->bNeedsTwoCopies)
+		if (NumSamples > 1)
+		{
+			ensureMsgf(!Owner->bNeedsTwoCopies, TEXT("bNeedsTwoCopies flag is ignored when with sample counts higher than 1."));
+
+			ETextureCreateFlags ResolveTexCreateFlags = TexCreateFlags;
+			EnumRemoveFlags(ResolveTexCreateFlags, ETextureCreateFlags::RenderTargetable);
+			EnumAddFlags(ResolveTexCreateFlags, ETextureCreateFlags::ShaderResource | ETextureCreateFlags::ResolveTargetable);
+
+			Desc.SetFlags(ResolveTexCreateFlags);
+			Desc.SetNumSamples(1);
+
+			TextureRHI = RHICreateTexture(Desc);
+		}
+		else if (Owner->bNeedsTwoCopies)
 		{
 			Desc.SetFlags(TexCreateFlags | ETextureCreateFlags::ShaderResource);
 
diff --git a/Engine/Source/Runtime/Engine/Private/UserInterface/Canvas.cpp b/Engine/Source/Runtime/Engine/Private/UserInterface/Canvas.cpp
--- a/Engine/Source/Runtime/Engine/Private/UserInterface/Canvas.cpp
+++ b/Engine/Source/Runtime/Engine/Private/UserInterface/Canvas.cpp
@@ -563,6 +563,20 @@ FCanvasRenderThreadScope::~FCanvasRenderThreadScope()
 			(*RenderCommandArray)[Index](RenderContext);
 		}
 
+		// MSAA resolve as the last optional render pass
+		if (RenderTarget->GetRenderTargetTexture()->GetDesc().IsMultisample())
+		{
+			FRenderTargetParameters* PassParameters = GraphBuilder.AllocParameters<FRenderTargetParameters>();
+			FRDGTextureRef ResolveTexture = RegisterExternalTexture(GraphBuilder, RenderTarget->GetShaderResourceTexture(), TEXT("MSAAResolveCanvasTexture"));
+			PassParameters->RenderTargets[0] = FRenderTargetBinding(
+				RenderTarget->GetRenderTargetTexture(GraphBuilder),
+				ResolveTexture,
+				ERenderTargetLoadAction::ELoad
+			);
+
+			GraphBuilder.AddPass(RDG_EVENT_NAME("MSAAResolveCanvas"), PassParameters, ERDGPassFlags::Raster, [](FRHICommandList&) {});
+		}
+
 		GraphBuilder.Execute();
 		delete RenderCommandArray;
 	});
diff --git a/Engine/Source/Runtime/RHI/Public/RHIResources.h b/Engine/Source/Runtime/RHI/Public/RHIResources.h
--- a/Engine/Source/Runtime/RHI/Public/RHIResources.h
+++ b/Engine/Source/Runtime/RHI/Public/RHIResources.h
@@ -4265,7 +4265,7 @@ struct FRHIRenderPassInfo
 	// Color, no depth, optional resolve, optional mip, optional array slice
 	explicit FRHIRenderPassInfo(FRHITexture* ColorRT, ERenderTargetActions ColorAction, FRHITexture* ResolveRT = nullptr, uint8 InMipIndex = 0, int32 InArraySlice = -1)
 	{
-		check(!ResolveRT || ResolveRT->IsMultisampled());
+		check(!(ResolveRT && ResolveRT->IsMultisampled()));
 		check(ColorRT);
 		ColorRenderTargets[0].RenderTarget = ColorRT;
 		ColorRenderTargets[0].ResolveTarget = ResolveRT;
