Ensure that bones missing from an anim sequence still follow their parent when the sequence is a mesh-space additive just like they would if the bones were present in the sequence but not animated

diff --git a/Engine/Plugins/Experimental/AnimNext/Source/AnimNext/Private/DecompressionTools.cpp b/Engine/Plugins/Experimental/AnimNext/Source/AnimNext/Private/DecompressionTools.cpp
--- a/Engine/Plugins/Experimental/AnimNext/Source/AnimNext/Private/DecompressionTools.cpp
+++ b/Engine/Plugins/Experimental/AnimNext/Source/AnimNext/Private/DecompressionTools.cpp
@@ -222,7 +222,8 @@ void FDecompressionTools::GetBonePose(const UAnimSequence* AnimSequence, const F
 				, AnimSequence->GetSkeleton()->GetRefLocalPoses()
 				, AnimSequence->CompressedData.CompressedTrackToSkeletonMapTable
 				, AnimSequence->GetSkeleton()
-				, AnimSequence->IsValidAdditive());
+				, AnimSequence->IsValidAdditive()
+				, AnimSequence->GetAdditiveAnimType());
 
 			DecompressPose(OutAnimationPoseData, AnimSequence->CompressedData, ExtractionContext, DecompContext, AnimSequence->GetRetargetTransforms(), RootMotionReset);
 		}
diff --git a/Engine/Source/Runtime/Engine/Private/Animation/AnimCompressionTypes.cpp b/Engine/Source/Runtime/Engine/Private/Animation/AnimCompressionTypes.cpp
--- a/Engine/Source/Runtime/Engine/Private/Animation/AnimCompressionTypes.cpp
+++ b/Engine/Source/Runtime/Engine/Private/Animation/AnimCompressionTypes.cpp
@@ -767,6 +767,7 @@ FCompressibleAnimData::FCompressibleAnimData(class UAnimSequence* InSeq, const b
 	, BoneCompressionSettings(InSeq->BoneCompressionSettings)
 	, Interpolation(InSeq->Interpolation)
 	, SequenceLength(InSeq->GetDataModelInterface()->GetPlayLength())
+	, AdditiveType(InSeq->GetAdditiveAnimType())
 	, bIsValidAdditive(InSeq->IsValidAdditive())
 #if WITH_EDITORONLY_DATA
 	, ErrorThresholdScale(InSeq->CompressionErrorThresholdScale)
@@ -1009,6 +1010,7 @@ FCompressibleAnimData::FCompressibleAnimData(UAnimBoneCompressionSettings* InBon
 	, NumberOfFrames(InNumberOfKeys)
 	PRAGMA_ENABLE_DEPRECATION_WARNINGS
 	, NumberOfKeys(InNumberOfKeys)
+	, AdditiveType(AAT_None)
 	, bIsValidAdditive(false)
 	, ErrorThresholdScale(1.f)
 	, TargetPlatform(InTargetPlatform)
@@ -1029,6 +1031,7 @@ PRAGMA_DISABLE_DEPRECATION_WARNINGS
 , NumberOfFrames(0)
 PRAGMA_ENABLE_DEPRECATION_WARNINGS
 , NumberOfKeys(0)
+, AdditiveType(AAT_None)
 , bIsValidAdditive(false)
 , ErrorThresholdScale(1.f)
 , TargetPlatform(nullptr)
diff --git a/Engine/Source/Runtime/Engine/Private/Animation/AnimSequence.cpp b/Engine/Source/Runtime/Engine/Private/Animation/AnimSequence.cpp
--- a/Engine/Source/Runtime/Engine/Private/Animation/AnimSequence.cpp
+++ b/Engine/Source/Runtime/Engine/Private/Animation/AnimSequence.cpp
@@ -1270,7 +1270,7 @@ void UAnimSequence::GetBoneTransform(FTransform& OutAtom, FSkeletonPoseBoneIndex
 
 	if ( bEvaluateCompressedData && IsCompressedDataValid())
 	{
-		FAnimSequenceDecompressionContext DecompContext(PlatformTargetFrameRate.Default, PlatformTargetFrameRate.Default.AsFrameTime(GetPlayLength()).RoundToFrame().Value, Interpolation, GetRetargetTransformsSourceName(), *CompressedData.CompressedDataStructure, GetSkeleton()->GetRefLocalPoses(), CompressedData.CompressedTrackToSkeletonMapTable, GetSkeleton(), IsValidAdditive());
+		FAnimSequenceDecompressionContext DecompContext(PlatformTargetFrameRate.Default, PlatformTargetFrameRate.Default.AsFrameTime(GetPlayLength()).RoundToFrame().Value, Interpolation, GetRetargetTransformsSourceName(), *CompressedData.CompressedDataStructure, GetSkeleton()->GetRefLocalPoses(), CompressedData.CompressedTrackToSkeletonMapTable, GetSkeleton(), IsValidAdditive(), AdditiveAnimType);
 		DecompContext.Seek(Time);
 		if (CompressedData.BoneCompressionCodec)
 		{
@@ -1695,7 +1695,7 @@ void UAnimSequence::GetBonePose(FAnimationPoseData& OutAnimationPoseData, const
 	if (NumTracks != 0)
 	{
 		// Evaluate compressed bone data
-		FAnimSequenceDecompressionContext DecompContext(PlatformTargetFrameRate.Default, PlatformTargetFrameRate.Default.AsFrameTime(GetPlayLength()).RoundToFrame().Value, Interpolation, GetRetargetTransformsSourceName(), *CompressedData.CompressedDataStructure, GetSkeleton()->GetRefLocalPoses(), CompressedData.CompressedTrackToSkeletonMapTable, GetSkeleton(), IsValidAdditive());
+		FAnimSequenceDecompressionContext DecompContext(PlatformTargetFrameRate.Default, PlatformTargetFrameRate.Default.AsFrameTime(GetPlayLength()).RoundToFrame().Value, Interpolation, GetRetargetTransformsSourceName(), *CompressedData.CompressedDataStructure, GetSkeleton()->GetRefLocalPoses(), CompressedData.CompressedTrackToSkeletonMapTable, GetSkeleton(), IsValidAdditive(), AdditiveAnimType);
 		UE::Anim::Decompression::DecompressPose(OutPose, CompressedData, ExtractionContext, DecompContext, GetRetargetTransforms(), RootMotionReset);
 	}
 
diff --git a/Engine/Source/Runtime/Engine/Private/Animation/AnimStreamable.cpp b/Engine/Source/Runtime/Engine/Private/Animation/AnimStreamable.cpp
--- a/Engine/Source/Runtime/Engine/Private/Animation/AnimStreamable.cpp
+++ b/Engine/Source/Runtime/Engine/Private/Animation/AnimStreamable.cpp
@@ -268,6 +268,7 @@ void UAnimStreamable::GetAnimationPose(FAnimationPoseData& OutAnimationPoseData,
 	//const bool bUseRawDataForPoseExtraction = bForceUseRawData || UseRawDataForPoseExtraction(RequiredBones);
 
 	const bool bIsBakedAdditive = false;//!bUseRawDataForPoseExtraction && IsValidAdditive();
+	const EAdditiveAnimationType AdditiveType = AAT_None;
 
 	const USkeleton* MySkeleton = GetSkeleton();
 	if (!MySkeleton)
@@ -403,7 +404,7 @@ void UAnimStreamable::GetAnimationPose(FAnimationPoseData& OutAnimationPoseData,
 		ChunkExtractionContext.PoseCurves = ExtractionContext.PoseCurves;
 
 		//FPlatformMisc::LowLevelOutputDebugStringf(TEXT("Playing Streaming Anim %s Time: %.2f Chunk:%i\n"), *GetName(), ExtractionContext.CurrentTime, ChunkIndex);
-		FAnimSequenceDecompressionContext Context(SamplingFrameRate, CurrentChunk.NumFrames, Interpolation, GetFName(), *CompressedData->CompressedDataStructure.Get(), GetSkeleton()->GetRefLocalPoses(), CompressedData->CompressedTrackToSkeletonMapTable, GetSkeleton(), bIsBakedAdditive);
+		FAnimSequenceDecompressionContext Context(SamplingFrameRate, CurrentChunk.NumFrames, Interpolation, GetFName(), *CompressedData->CompressedDataStructure.Get(), GetSkeleton()->GetRefLocalPoses(), CompressedData->CompressedTrackToSkeletonMapTable, GetSkeleton(), bIsBakedAdditive, AdditiveType);
 		UE::Anim::Decompression::DecompressPose(OutPose, *CompressedData, ChunkExtractionContext, Context, RetargetSource, RootMotionReset);
 	}
 }
diff --git a/Engine/Source/Runtime/Engine/Private/Animation/AnimationDecompression.cpp b/Engine/Source/Runtime/Engine/Private/Animation/AnimationDecompression.cpp
--- a/Engine/Source/Runtime/Engine/Private/Animation/AnimationDecompression.cpp
+++ b/Engine/Source/Runtime/Engine/Private/Animation/AnimationDecompression.cpp
@@ -23,6 +23,9 @@ struct FGetBonePoseScratchArea : public TThreadSingleton<FGetBonePoseScratchArea
 	BoneTrackArray AnimScaleRetargetingPairs;
 	BoneTrackArray AnimRelativeRetargetingPairs;
 	BoneTrackArray OrientAndScaleRetargetingPairs;
+
+	// A bit set that specifies whether a compact bone index has its rotation animated by the sequence or not
+	TBitArray<> AnimatedCompactRotations;
 };
 
 void DecompressPose(FCompactPose& OutPose,
@@ -43,6 +46,7 @@ void DecompressPose(FCompactPose& OutPose,
 	const TArray<FTransform>& RetargetTransforms,
 	const FRootMotionReset& RootMotionReset)
 {
+	const int32 NumCompactBones = OutPose.GetNumBones();
 	const FBoneContainer& RequiredBones = OutPose.GetBoneContainer();
 	const int32 NumTracks = CompressedData.CompressedTrackToSkeletonMapTable.Num();
 
@@ -63,6 +67,13 @@ void DecompressPose(FCompactPose& OutPose,
 	AnimRelativeRetargetingPairs.Reset();
 	OrientAndScaleRetargetingPairs.Reset();
 
+	const bool bIsMeshSpaceAdditive = DecompressionContext.GetAdditiveType() == AAT_RotationOffsetMeshSpace;
+	TBitArray<>& AnimatedCompactRotations = ScratchArea.AnimatedCompactRotations;
+	if (bIsMeshSpaceAdditive)
+	{
+		AnimatedCompactRotations.Init(false, NumCompactBones);
+	}
+
 	// Optimization: assuming first index is root bone. That should always be the case in Skeletons.
 	checkSlow((RequiredBones.GetMeshPoseIndexFromSkeletonPoseIndex(FSkeletonPoseBoneIndex(0)) == FMeshPoseBoneIndex(0)));
 	// this is not guaranteed for AnimSequences though... If Root is not animated, Track will not exist.
@@ -86,6 +97,11 @@ void DecompressPose(FCompactPose& OutPose,
 				{
 					RotationScalePairs.Add(BoneTrackPair(CompactPoseBoneIndex, TrackIndex));
 
+					if (bIsMeshSpaceAdditive)
+					{
+						AnimatedCompactRotations[CompactPoseBoneIndex] = true;
+					}
+
 					// Skip extracting translation component for EBoneTranslationRetargetingMode::Skeleton.
 					switch (TargetSkeleton->GetBoneTranslationRetargetingMode(TargetSkeletonBoneIndex, RequiredBones.GetDisableRetargeting()))
 					{
@@ -177,7 +193,7 @@ void DecompressPose(FCompactPose& OutPose,
 	{
 		if (DecompressionContext.IsAdditiveAnimation())
 		{
-			for (FCompactPoseBoneIndex BoneIndex(bFirstTrackIsRootBone ? 1 : 0); BoneIndex < OutPose.GetNumBones(); ++BoneIndex)
+			for (FCompactPoseBoneIndex BoneIndex(bFirstTrackIsRootBone ? 1 : 0); BoneIndex < NumCompactBones; ++BoneIndex)
 			{
 				const int32 TargetSkeletonBoneIndex = RequiredBones.GetSkeletonIndex(BoneIndex);
 				OutPose[BoneIndex].SetRotation(SkeletonRemapping.RetargetAdditiveRotationToTargetSkeleton(TargetSkeletonBoneIndex, OutPose[BoneIndex].GetRotation()));
@@ -189,7 +205,7 @@ void DecompressPose(FCompactPose& OutPose,
 		}
 		else
 		{
-			for (FCompactPoseBoneIndex BoneIndex(bFirstTrackIsRootBone ? 1 : 0); BoneIndex < OutPose.GetNumBones(); ++BoneIndex)
+			for (FCompactPoseBoneIndex BoneIndex(bFirstTrackIsRootBone ? 1 : 0); BoneIndex < NumCompactBones; ++BoneIndex)
 			{
 				const int32 TargetSkeletonBoneIndex = RequiredBones.GetSkeletonIndex(BoneIndex);
 				OutPose[BoneIndex].SetRotation(SkeletonRemapping.RetargetBoneRotationToTargetSkeleton(TargetSkeletonBoneIndex, OutPose[BoneIndex].GetRotation()));
@@ -269,7 +285,7 @@ void DecompressPose(FCompactPose& OutPose,
 		const TArray<int32>& CompactPoseIndexToOrientAndScaleIndex = RetargetSourceCachedData.CompactPoseIndexToOrientAndScaleIndex;
 
 		// If we have any cached retargeting data.
-		if ((OrientAndScaleDataArray.Num() > 0) && (CompactPoseIndexToOrientAndScaleIndex.Num() == RequiredBones.GetCompactPoseNumBones()))
+		if (OrientAndScaleDataArray.Num() > 0 && CompactPoseIndexToOrientAndScaleIndex.Num() == NumCompactBones)
 		{
 			for (int32 Index = 0; Index < NumBonesToOrientAndScaleRetarget; Index++)
 			{
@@ -292,5 +308,27 @@ void DecompressPose(FCompactPose& OutPose,
 			}
 		}
 	}
+
+	if (bIsMeshSpaceAdditive)
+	{
+		// When an animation is a mesh-space additive, bones that aren't animated will end up with some non-identity
+		// delta relative to the base used to create the additive. This is because the delta is calculated in mesh-space
+		// unlike regular additive animations where bones that aren't animated has an identity delta. For rotations,
+		// this mesh-space delta will be the parent bone rotation.
+		// However, if a bone isn't animated in the sequence but present on the target skeleton, we have no data for it
+		// and the output pose will contain an identity delta which isn't what we want. As such, bones missing from
+		// the sequence have their rotation set to their parent.
+
+		// If the first track is the root, we skip it since it has no parent (its delta value is fine as the identity)
+		for (FCompactPoseBoneIndex CompactBoneIndex(bFirstTrackIsRootBone ? 1 : 0); CompactBoneIndex < NumCompactBones; ++CompactBoneIndex)
+		{
+			if (!AnimatedCompactRotations[CompactBoneIndex.GetInt()])
+			{
+				// This bone wasn't animated in the sequence, fix it up
+				const FCompactPoseBoneIndex CompactParentIndex = RequiredBones.GetParentBoneIndex(CompactBoneIndex);
+				OutPose[CompactBoneIndex].SetRotation(OutPose[CompactParentIndex].GetRotation());
+			}
+		}
+	}
 }
 }
diff --git a/Engine/Source/Runtime/Engine/Private/Animation/AnimationUtils.cpp b/Engine/Source/Runtime/Engine/Private/Animation/AnimationUtils.cpp
--- a/Engine/Source/Runtime/Engine/Private/Animation/AnimationUtils.cpp
+++ b/Engine/Source/Runtime/Engine/Private/Animation/AnimationUtils.cpp
@@ -283,7 +283,7 @@ void FAnimationUtils::ComputeCompressionError(const FCompressibleAnimData& Compr
 		const FTransform EndEffectorDummyBone(FQuat::Identity, FVector(END_EFFECTOR_DUMMY_BONE_LENGTH));
 
 		FAnimSequenceDecompressionContext DecompContext(CompressibleAnimData.SampledFrameRate, CompressibleAnimData.GetNumberOfFrames(), CompressibleAnimData.Interpolation, CompressibleAnimData.AnimFName, *CompressedData.AnimData, RefPose,
-			CompressibleAnimData.TrackToSkeletonMapTable, nullptr, CompressibleAnimData.bIsValidAdditive);
+			CompressibleAnimData.TrackToSkeletonMapTable, nullptr, CompressibleAnimData.bIsValidAdditive, CompressibleAnimData.AdditiveType);
 
 		const TArray<FBoneData>& BoneData = CompressibleAnimData.BoneData;
 
@@ -1116,7 +1116,7 @@ void FAnimationUtils::ExtractTransformFromCompressionData(const FCompressibleAni
 		// Build our read-only version from the mutable source
 		FUECompressedAnimData AnimData(AnimDataMutable);
 
-		FAnimSequenceDecompressionContext DecompContext(CompressibleAnimData.SampledFrameRate, CompressibleAnimData.GetNumberOfFrames(), CompressibleAnimData.Interpolation, CompressibleAnimData.AnimFName, AnimData, CompressibleAnimData.RefLocalPoses, CompressibleAnimData.TrackToSkeletonMapTable, nullptr, CompressibleAnimData.bIsValidAdditive);
+		FAnimSequenceDecompressionContext DecompContext(CompressibleAnimData.SampledFrameRate, CompressibleAnimData.GetNumberOfFrames(), CompressibleAnimData.Interpolation, CompressibleAnimData.AnimFName, AnimData, CompressibleAnimData.RefLocalPoses, CompressibleAnimData.TrackToSkeletonMapTable, nullptr, CompressibleAnimData.bIsValidAdditive, CompressibleAnimData.AdditiveType);
 		DecompContext.Seek(Time);
 		CompressedAnimData.Codec->DecompressBone(DecompContext, TrackIndex, OutBoneTransform);
 		return;
diff --git a/Engine/Source/Runtime/Engine/Private/Animation/IspcTestAnimEncodingConstantKeyLerp.cpp b/Engine/Source/Runtime/Engine/Private/Animation/IspcTestAnimEncodingConstantKeyLerp.cpp
--- a/Engine/Source/Runtime/Engine/Private/Animation/IspcTestAnimEncodingConstantKeyLerp.cpp
+++ b/Engine/Source/Runtime/Engine/Private/Animation/IspcTestAnimEncodingConstantKeyLerp.cpp
@@ -69,7 +69,8 @@ bool FIspcTestAnimEncodingConstantKeyLerpGetPoseRotations::RunTest(const FString
 			ISPCTransforms,
 			SkeletonMap,
 			nullptr,
-			false);
+			false,
+			AAT_None);
 		KeyLerp.GetPoseRotations(ISPCAtoms, Bones, DecompressContext);
 	}
 
@@ -84,7 +85,8 @@ bool FIspcTestAnimEncodingConstantKeyLerpGetPoseRotations::RunTest(const FString
 			CPPTransforms,
 			SkeletonMap,
 			nullptr,
-			false);
+			false,
+			AAT_None);
 		KeyLerp.GetPoseRotations(CPPAtoms, Bones, DecompressContext);
 	}
 
@@ -161,7 +163,8 @@ bool FIspcTestAnimEncodingConstantKeyLerpGetPoseTranslations::RunTest(const FStr
 			ISPCTransforms,
 			SkeletonMap,
 			nullptr,
-			false);
+			false,
+			AAT_None);
 		KeyLerp.GetPoseTranslations(ISPCAtoms, Bones, DecompressContext);
 	}
 
@@ -176,7 +179,8 @@ bool FIspcTestAnimEncodingConstantKeyLerpGetPoseTranslations::RunTest(const FStr
 			CPPTransforms,
 			SkeletonMap,
 			nullptr,
-			false);
+			false,
+			AAT_None);
 		KeyLerp.GetPoseTranslations(CPPAtoms, Bones, DecompressContext);
 	}
 
@@ -253,7 +257,8 @@ bool FIspcTestAnimEncodingConstantKeyLerpGetPoseScales::RunTest(const FString& P
 			ISPCTransforms,
 			SkeletonMap,
 			nullptr,
-			false);
+			false,
+			AAT_None);
 		KeyLerp.GetPoseScales(ISPCAtoms, Bones, DecompressContext);
 	}
 
@@ -268,7 +273,8 @@ bool FIspcTestAnimEncodingConstantKeyLerpGetPoseScales::RunTest(const FString& P
 			CPPTransforms,
 			SkeletonMap,
 			nullptr,
-			false);
+			false,
+			AAT_None);
 		KeyLerp.GetPoseScales(CPPAtoms, Bones, DecompressContext);
 	}
 
diff --git a/Engine/Source/Runtime/Engine/Private/Animation/IspcTestAnimEncodingPerTrackCompression.cpp b/Engine/Source/Runtime/Engine/Private/Animation/IspcTestAnimEncodingPerTrackCompression.cpp
--- a/Engine/Source/Runtime/Engine/Private/Animation/IspcTestAnimEncodingPerTrackCompression.cpp
+++ b/Engine/Source/Runtime/Engine/Private/Animation/IspcTestAnimEncodingPerTrackCompression.cpp
@@ -77,7 +77,8 @@ bool FIspcTestAnimEncodingPerTrackCompressionGetPoseRotations::RunTest(const FSt
 			ISPCTransforms,
 			SkeletonMap,
 			nullptr,
-			false);
+			false,
+			AAT_None);
 		KeyLerp.GetPoseRotations(ISPCAtoms, Bones, DecompressContext);
 	}
 
@@ -92,7 +93,8 @@ bool FIspcTestAnimEncodingPerTrackCompressionGetPoseRotations::RunTest(const FSt
 			CPPTransforms,
 			SkeletonMap,
 			nullptr,
-			false);
+			false,
+			AAT_None);
 		KeyLerp.GetPoseRotations(CPPAtoms, Bones, DecompressContext);
 	}
 
@@ -177,7 +179,8 @@ bool FIspcTestAnimEncodingPerTrackCompressionGetPoseTranslations::RunTest(const
 			ISPCTransforms,
 			SkeletonMap,
 			nullptr,
-			false);
+			false,
+			AAT_None);
 		KeyLerp.GetPoseTranslations(ISPCAtoms, Bones, DecompressContext);
 	}
 
@@ -192,7 +195,8 @@ bool FIspcTestAnimEncodingPerTrackCompressionGetPoseTranslations::RunTest(const
 			CPPTransforms,
 			SkeletonMap,
 			nullptr,
-			false);
+			false,
+			AAT_None);
 		KeyLerp.GetPoseTranslations(CPPAtoms, Bones, DecompressContext);
 	}
 
@@ -277,7 +281,8 @@ bool FIspcTestAnimEncodingPerTrackCompressionGetPoseScales::RunTest(const FStrin
 			ISPCTransforms,
 			SkeletonMap,
 			nullptr,
-			false);
+			false,
+			AAT_None);
 		KeyLerp.GetPoseScales(ISPCAtoms, Bones, DecompressContext);
 	}
 
@@ -292,7 +297,8 @@ bool FIspcTestAnimEncodingPerTrackCompressionGetPoseScales::RunTest(const FStrin
 			CPPTransforms,
 			SkeletonMap,
 			nullptr,
-			false);
+			false,
+			AAT_None);
 		KeyLerp.GetPoseScales(CPPAtoms, Bones, DecompressContext);
 	}
 
diff --git a/Engine/Source/Runtime/Engine/Private/Animation/IspcTestAnimEncodingVariableKeyLerp.cpp b/Engine/Source/Runtime/Engine/Private/Animation/IspcTestAnimEncodingVariableKeyLerp.cpp
--- a/Engine/Source/Runtime/Engine/Private/Animation/IspcTestAnimEncodingVariableKeyLerp.cpp
+++ b/Engine/Source/Runtime/Engine/Private/Animation/IspcTestAnimEncodingVariableKeyLerp.cpp
@@ -69,7 +69,8 @@ bool FIspcTestAnimEncodingVariableKeyLerpGetPoseRotations::RunTest(const FString
 			ISPCTransforms,
 			SkeletonMap,
 			nullptr,
-			false);
+			false,
+			AAT_None);
 		KeyLerp.GetPoseRotations(ISPCAtoms, Bones, DecompressContext);
 	}
 
@@ -84,7 +85,8 @@ bool FIspcTestAnimEncodingVariableKeyLerpGetPoseRotations::RunTest(const FString
 			CPPTransforms,
 			SkeletonMap,
 			nullptr,
-			false);
+			false,
+			AAT_None);
 		KeyLerp.GetPoseRotations(CPPAtoms, Bones, DecompressContext);
 	}
 
@@ -161,7 +163,8 @@ bool FIspcTestAnimEncodingVariableKeyLerpGetPoseTranslations::RunTest(const FStr
 			ISPCTransforms,
 			SkeletonMap,
 			nullptr,
-			false);
+			false,
+			AAT_None);
 		KeyLerp.GetPoseTranslations(ISPCAtoms, Bones, DecompressContext);
 	}
 
@@ -176,7 +179,8 @@ bool FIspcTestAnimEncodingVariableKeyLerpGetPoseTranslations::RunTest(const FStr
 			CPPTransforms,
 			SkeletonMap,
 			nullptr,
-			false);
+			false,
+			AAT_None);
 		KeyLerp.GetPoseTranslations(CPPAtoms, Bones, DecompressContext);
 	}
 
@@ -253,7 +257,8 @@ bool FIspcTestAnimEncodingVariableKeyLerpGetPoseScales::RunTest(const FString& P
 			ISPCTransforms,
 			SkeletonMap,
 			nullptr,
-			false);
+			false,
+			AAT_None);
 		KeyLerp.GetPoseScales(ISPCAtoms, Bones, DecompressContext);
 	}
 
@@ -268,7 +273,8 @@ bool FIspcTestAnimEncodingVariableKeyLerpGetPoseScales::RunTest(const FString& P
 			CPPTransforms,
 			SkeletonMap,
 			nullptr,
-			false);
+			false,
+			AAT_None);
 		KeyLerp.GetPoseScales(CPPAtoms, Bones, DecompressContext);
 	}
 
diff --git a/Engine/Source/Runtime/Engine/Public/Animation/AnimCompressionTypes.h b/Engine/Source/Runtime/Engine/Public/Animation/AnimCompressionTypes.h
--- a/Engine/Source/Runtime/Engine/Public/Animation/AnimCompressionTypes.h
+++ b/Engine/Source/Runtime/Engine/Public/Animation/AnimCompressionTypes.h
@@ -243,6 +243,7 @@ public:
 
 	int32 GetNumberOfFrames() const { return FMath::Max(NumberOfKeys - 1, 1); }
 
+	EAdditiveAnimationType AdditiveType = AAT_None;
 	bool bIsValidAdditive;
 
 	float ErrorThresholdScale;
diff --git a/Engine/Source/Runtime/Engine/Public/Animation/AnimSequenceDecompressionContext.h b/Engine/Source/Runtime/Engine/Public/Animation/AnimSequenceDecompressionContext.h
--- a/Engine/Source/Runtime/Engine/Public/Animation/AnimSequenceDecompressionContext.h
+++ b/Engine/Source/Runtime/Engine/Public/Animation/AnimSequenceDecompressionContext.h
@@ -18,12 +18,7 @@ struct FAnimSequenceDecompressionContext
 		PRAGMA_DISABLE_DEPRECATION_WARNINGS
 		SequenceLength(SequenceLength_),
 		PRAGMA_ENABLE_DEPRECATION_WARNINGS
-		Interpolation(Interpolation_), AnimName(AnimName_), CompressedAnimData(CompressedAnimData_),
-		PRAGMA_DISABLE_DEPRECATION_WARNINGS
-		Time(0.f), RelativePos(0.f),
-		PRAGMA_ENABLE_DEPRECATION_WARNINGS
-		SourceSkeleton(nullptr),
-		bAdditiveAnimation(false)
+		Interpolation(Interpolation_), AnimName(AnimName_), CompressedAnimData(CompressedAnimData_)
 	{
 	}
 
@@ -35,13 +30,11 @@ struct FAnimSequenceDecompressionContext
 		SequenceLength(SequenceLength_),
 		PRAGMA_ENABLE_DEPRECATION_WARNINGS
 		Interpolation(Interpolation_), AnimName(AnimName_), CompressedAnimData(CompressedAnimData_),
-		RefPoses(RefPoses_.GetData(), RefPoses_.Num()), TrackToSkeletonMap(TrackToSkeletonMap_.GetData(), TrackToSkeletonMap_.Num()),
-		PRAGMA_DISABLE_DEPRECATION_WARNINGS
-		Time(0.f), RelativePos(0.f)
-		PRAGMA_ENABLE_DEPRECATION_WARNINGS
+		RefPoses(RefPoses_.GetData(), RefPoses_.Num()), TrackToSkeletonMap(TrackToSkeletonMap_.GetData(), TrackToSkeletonMap_.Num())
 	{
 	}
 
+	UE_DEPRECATED(5.3, "FAnimSequenceDecompressionContext signature has been deprecated, use different one with additive type")
 	FAnimSequenceDecompressionContext(const FFrameRate& InSamplingRate, const int32 InNumberOfFrames, EAnimInterpolationType Interpolation_, const FName& AnimName_, const ICompressedAnimData& CompressedAnimData_, const TArray<FTransform>& RefPoses_, const TArray<FTrackToSkeletonMap>& TrackToSkeletonMap_, const USkeleton* InSourceSkeleton, bool bIsBakedAdditive)
 		:
 		PRAGMA_DISABLE_DEPRECATION_WARNINGS
@@ -49,43 +42,52 @@ struct FAnimSequenceDecompressionContext
 		PRAGMA_ENABLE_DEPRECATION_WARNINGS
 		Interpolation(Interpolation_), AnimName(AnimName_), CompressedAnimData(CompressedAnimData_),
 		RefPoses(RefPoses_.GetData(), RefPoses_.Num()), TrackToSkeletonMap(TrackToSkeletonMap_.GetData(), TrackToSkeletonMap_.Num()),
+		SamplingRate(InSamplingRate),
+		SamplingTime(0),
+		NumberOfFrames(InNumberOfFrames),
+		SourceSkeleton(InSourceSkeleton),
+		bAdditiveAnimation(bIsBakedAdditive)
+	{}
+
+	FAnimSequenceDecompressionContext(
+		const FFrameRate& InSamplingRate,
+		const int32 InNumberOfFrames,
+		EAnimInterpolationType InInterpolation,
+		const FName& InAnimName,
+		const ICompressedAnimData& InCompressedAnimData,
+		const TArray<FTransform>& InRefPoses,
+		const TArray<FTrackToSkeletonMap>& InTrackToSkeletonMap,
+		const USkeleton* InSourceSkeleton,
+		bool bIsBakedAdditive,
+		EAdditiveAnimationType InAdditiveType)
+		:
 		PRAGMA_DISABLE_DEPRECATION_WARNINGS
-		Time(0.f), RelativePos(0.f),
+		SequenceLength(static_cast<float>(InSamplingRate.AsSeconds(InNumberOfFrames))),
 		PRAGMA_ENABLE_DEPRECATION_WARNINGS
+		Interpolation(InInterpolation), AnimName(InAnimName), CompressedAnimData(InCompressedAnimData),
+		RefPoses(InRefPoses.GetData(), InRefPoses.Num()), TrackToSkeletonMap(InTrackToSkeletonMap.GetData(), InTrackToSkeletonMap.Num()),
 		SamplingRate(InSamplingRate),
 		SamplingTime(0),
 		NumberOfFrames(InNumberOfFrames),
 		SourceSkeleton(InSourceSkeleton),
+		AdditiveType(InAdditiveType),
 		bAdditiveAnimation(bIsBakedAdditive)
 	{}
 
 	// Anim info
 	UE_DEPRECATED(5.1, "Direct access to SequenceLength has been deprecated use GetPlayableLength instead")
-	float SequenceLength;
-	EAnimInterpolationType Interpolation;
+	float SequenceLength = 0.0f;
+	EAnimInterpolationType Interpolation = EAnimInterpolationType::Linear;
 	FName AnimName;
 
 	const ICompressedAnimData& CompressedAnimData;
 
-private:
-	// Immutable reference/bind pose from the source skeleton, maps a bone index to its reference transform.
-	// The same skeleton reference pose must be provided during compression and decompression. This is the
-	// source skeleton reference pose (the one assigned to an anim sequence). Re-targeting is taken into account
-	// elsewhere. Additionally, if the source skeleton reference pose isn't used, the TrackToSkeletonMap indices
-	// may not match the provided pose. This could cause us to read incorrect bone transforms or possibly out
-	// of bounds.
-	TArrayView<const FTransform> RefPoses;
-
-	// Immutable array that maps from a track index to a bone index. This is built during compression using
-	// the source skeleton (the one assigned to an anim sequence).
-	TArrayView<const FTrackToSkeletonMap> TrackToSkeletonMap;
-
 public:
 	UE_DEPRECATED(5.1, "Direct access to Time has been deprecated use GetEvaluationTime or GetInterpolatedEvaluationTime instead")
-	float Time;
+	float Time = 0.0f;
 	
 	UE_DEPRECATED(5.1, "Direct access to RelativePos has been deprecated use GetRelativePosition instead")
-	float RelativePos;
+	float RelativePos = 0.0f;
 
 	double GetPlayableLength() const
 	{
@@ -123,7 +125,8 @@ public:
 	}
 
 	const USkeleton* GetSourceSkeleton() const { return SourceSkeleton; }
-	bool IsAdditiveAnimation() const { return bAdditiveAnimation; }	
+	bool IsAdditiveAnimation() const { return bAdditiveAnimation; }
+	EAdditiveAnimationType GetAdditiveType() const { return AdditiveType; }
 
 	const TArrayView<const FTransform>& GetRefLocalPoses() const
 	{
@@ -136,10 +139,23 @@ public:
 	}
 
 protected:
+	// Immutable reference/bind pose from the source skeleton, maps a bone index to its reference transform.
+	// The same skeleton reference pose must be provided during compression and decompression. This is the
+	// source skeleton reference pose (the one assigned to an anim sequence). Re-targeting is taken into account
+	// elsewhere. Additionally, if the source skeleton reference pose isn't used, the TrackToSkeletonMap indices
+	// may not match the provided pose. This could cause us to read incorrect bone transforms or possibly out
+	// of bounds.
+	TArrayView<const FTransform> RefPoses;
+
+	// Immutable array that maps from a track index to a bone index. This is built during compression using
+	// the source skeleton (the one assigned to an anim sequence).
+	TArrayView<const FTrackToSkeletonMap> TrackToSkeletonMap;
+
 	FFrameRate SamplingRate;
 	FFrameTime SamplingTime;
-	int32 NumberOfFrames;
-	const USkeleton* SourceSkeleton;
-	bool bAdditiveAnimation;
+	int32 NumberOfFrames = 0;
+	const USkeleton* SourceSkeleton = nullptr;
+	EAdditiveAnimationType AdditiveType = AAT_None;
+	bool bAdditiveAnimation = false;
 };
 
