Suppress PVS-Studio static analysis warning false positive

diff --git a/Engine/Source/Runtime/Engine/Private/HLSLTree/HLSLTreeEmit.cpp b/Engine/Source/Runtime/Engine/Private/HLSLTree/HLSLTreeEmit.cpp
--- a/Engine/Source/Runtime/Engine/Private/HLSLTree/HLSLTreeEmit.cpp
+++ b/Engine/Source/Runtime/Engine/Private/HLSLTree/HLSLTreeEmit.cpp
@@ -470,7 +470,7 @@ FPreparedType FEmitContext::PrepareExpression(const FExpression* InExpression, F
 		FXxHash64Builder Hasher;
 		Hasher.Update(&InExpression, sizeof(InExpression));
 		Hasher.Update(&bMarkLiveValues, sizeof(bMarkLiveValues));
-		Hasher.Update(&RequestedType.Type.StructType, sizeof(RequestedType.Type.StructType));
+		Hasher.Update(&RequestedType.Type.StructType, sizeof(RequestedType.Type.StructType)); //-V568
 		Hash = Hasher.Finalize();
 	}
 
