Exposed FindMetaDataByClass function to blueprint

diff --git a/Engine/Source/Runtime/Engine/Classes/Animation/AnimationAsset.h b/Engine/Source/Runtime/Engine/Classes/Animation/AnimationAsset.h
--- a/Engine/Source/Runtime/Engine/Classes/Animation/AnimationAsset.h
+++ b/Engine/Source/Runtime/Engine/Classes/Animation/AnimationAsset.h
@@ -1077,6 +1077,7 @@ public:
 	const TArray<UAnimMetaData*>& GetMetaData() const { return MetaData; }
 	
 	/** Returns the first metadata of the specified class */
+	UFUNCTION(BlueprintCallable, Category = "Animation")
 	ENGINE_API UAnimMetaData* FindMetaDataByClass(const TSubclassOf<UAnimMetaData> MetaDataClass) const;
 
 	/** Templatized version of FindMetaDataByClass that handles casting for you */
