World Partition: removed dead code.

diff --git a/Engine/Source/Runtime/Engine/Private/WorldPartition/WorldPartition.cpp b/Engine/Source/Runtime/Engine/Private/WorldPartition/WorldPartition.cpp
--- a/Engine/Source/Runtime/Engine/Private/WorldPartition/WorldPartition.cpp
+++ b/Engine/Source/Runtime/Engine/Private/WorldPartition/WorldPartition.cpp
@@ -602,7 +602,6 @@ void UWorldPartition::Initialize(UWorld* InWorld, const FTransform& InTransform)
 	check(RuntimeHash);
 	RuntimeHash->SetFlags(RF_Transactional);
 
-	TArray<FGuid> ForceLoadedActorGuids;
 	if (bIsEditor || bIsGame || bIsPIEWorldTravel || bIsDedicatedServer)
 	{
 		FName ContainerPackageName = UActorDescContainerInstance::GetContainerPackageNameFromWorld(OuterWorld);
@@ -621,12 +620,6 @@ void UWorldPartition::Initialize(UWorld* InWorld, const FTransform& InTransform)
 	DataLayerManager->Initialize();
 
 #if WITH_EDITOR
-	if (ForceLoadedActors && ForceLoadedActorGuids.Num() > 0)
-	{
-		TRACE_CPUPROFILER_EVENT_SCOPE(UActorDescContainer::ForceLoadedActors);
-		ForceLoadedActors->AddActors(ForceLoadedActorGuids);
-	}
-	
 	if (bIsEditor)
 	{
 		// Apply level transform on actors already part of the level
