Clean up unnecessary deprecated pragma warning

diff --git a/Engine/Source/Runtime/Engine/Classes/Engine/SkeletalMesh.h b/Engine/Source/Runtime/Engine/Classes/Engine/SkeletalMesh.h
--- a/Engine/Source/Runtime/Engine/Classes/Engine/SkeletalMesh.h
+++ b/Engine/Source/Runtime/Engine/Classes/Engine/SkeletalMesh.h
@@ -957,10 +957,8 @@ public:
 	{
 #if WITH_EDITORONLY_DATA
 		WaitUntilAsyncPropertyReleased(ESkeletalMeshAsyncProperties::MinLod);
-		PRAGMA_DISABLE_DEPRECATION_WARNINGS
 		MinQualityLevelLOD.PerQuality = QualityLevelProperty::ConvertQualtiyLevelData(QualityLevelMinimumLODs);
 		MinQualityLevelLOD.Default = Default >=0 ? Default : MinQualityLevelLOD.Default;
-		PRAGMA_ENABLE_DEPRECATION_WARNINGS
 #endif
 	}
 
@@ -969,32 +967,24 @@ public:
 	{
 #if WITH_EDITORONLY_DATA
 		WaitUntilAsyncPropertyReleased(ESkeletalMeshAsyncProperties::MinLod, ESkinnedAssetAsyncPropertyLockType::ReadOnly);
-		PRAGMA_DISABLE_DEPRECATION_WARNINGS
 		QualityLevelMinimumLODs = QualityLevelProperty::ConvertQualtiyLevelData(MinQualityLevelLOD.PerQuality);
 		Default = MinQualityLevelLOD.Default;
-		PRAGMA_ENABLE_DEPRECATION_WARNINGS
 #endif
 	}
 
 	static FName GetQualityLevelMinLodMemberName()
 	{
-		PRAGMA_DISABLE_DEPRECATION_WARNINGS
 		return GET_MEMBER_NAME_CHECKED(USkeletalMesh, MinQualityLevelLOD);
-		PRAGMA_ENABLE_DEPRECATION_WARNINGS
 	}
 
 	const FPerQualityLevelInt& GetQualityLevelMinLod() const
 	{
-		PRAGMA_DISABLE_DEPRECATION_WARNINGS
-			return MinQualityLevelLOD;
-		PRAGMA_ENABLE_DEPRECATION_WARNINGS
+		return MinQualityLevelLOD;
 	}
 
 	void SetQualityLevelMinLod(FPerQualityLevelInt InMinLod)
 	{
-		PRAGMA_DISABLE_DEPRECATION_WARNINGS
 			MinQualityLevelLOD = MoveTemp(InMinLod);
-		PRAGMA_ENABLE_DEPRECATION_WARNINGS
 	}
 
 	/** Check the QualitLevel property is enabled for MinLod. USkinnedAsset interface. */
diff --git a/Engine/Source/Runtime/Engine/Classes/Engine/StaticMesh.h b/Engine/Source/Runtime/Engine/Classes/Engine/StaticMesh.h
--- a/Engine/Source/Runtime/Engine/Classes/Engine/StaticMesh.h
+++ b/Engine/Source/Runtime/Engine/Classes/Engine/StaticMesh.h
@@ -733,9 +733,7 @@ public:
 
 	static FName GetQualityLevelMinLODMemberName()
 	{
-		PRAGMA_DISABLE_DEPRECATION_WARNINGS
 		return GET_MEMBER_NAME_CHECKED(UStaticMesh, MinQualityLevelLOD);
-		PRAGMA_ENABLE_DEPRECATION_WARNINGS
 	}
 
 	const FPerQualityLevelInt& GetQualityLevelMinLOD() const
@@ -777,10 +775,8 @@ public:
 	{
 #if WITH_EDITORONLY_DATA
 		WaitUntilAsyncPropertyReleased(EStaticMeshAsyncProperties::MinLOD);
-		PRAGMA_DISABLE_DEPRECATION_WARNINGS
 		MinQualityLevelLOD.PerQuality = QualityLevelProperty::ConvertQualtiyLevelData(QualityLevelMinimumLODs);
 		MinQualityLevelLOD.Default = Default >= 0 ? Default : MinQualityLevelLOD.Default;
-		PRAGMA_ENABLE_DEPRECATION_WARNINGS
 #endif
 	}
 
@@ -789,10 +785,8 @@ public:
 	{
 #if WITH_EDITORONLY_DATA
 		WaitUntilAsyncPropertyReleased(EStaticMeshAsyncProperties::MinLOD);
-		PRAGMA_DISABLE_DEPRECATION_WARNINGS
 		QualityLevelMinimumLODs = QualityLevelProperty::ConvertQualtiyLevelData(MinQualityLevelLOD.PerQuality);
 		Default = MinQualityLevelLOD.Default;
-		PRAGMA_ENABLE_DEPRECATION_WARNINGS
 #endif
 	}
 
