PR #10861: GetLinkedAnimLayerInstanceByClass now checks for Child class.

diff --git a/Engine/Source/Runtime/Engine/Classes/Animation/AnimInstance.h b/Engine/Source/Runtime/Engine/Classes/Animation/AnimInstance.h
--- a/Engine/Source/Runtime/Engine/Classes/Animation/AnimInstance.h
+++ b/Engine/Source/Runtime/Engine/Classes/Animation/AnimInstance.h
@@ -844,9 +844,9 @@ public:
 	UE_DEPRECATED(4.24, "Function renamed, please use GetLinkedAnimLayerInstanceByClass")
 	UAnimInstance* GetLayerSubInstanceByClass(TSubclassOf<UAnimInstance> InClass) const { return GetLinkedAnimLayerInstanceByClass(InClass); }
 
-	/** Gets the first layer linked instance corresponding to the specified class */
+	/** Gets the first layer linked instance corresponding to the specified class, optionally if bCheckForChildClass is true, it will check IsChildOf on InClass. */
 	UFUNCTION(BlueprintPure, Category = "Animation|Linked Anim Graphs")
-	ENGINE_API UAnimInstance* GetLinkedAnimLayerInstanceByClass(TSubclassOf<UAnimInstance> InClass) const;
+	ENGINE_API UAnimInstance* GetLinkedAnimLayerInstanceByClass(TSubclassOf<UAnimInstance> InClass, bool bCheckForChildClass = false) const;
 
 	/** Sets up initial layer groupings */
 	ENGINE_API void InitializeGroupedLayers(bool bInDeferSubGraphInitialization);
diff --git a/Engine/Source/Runtime/Engine/Private/Animation/AnimInstance.cpp b/Engine/Source/Runtime/Engine/Private/Animation/AnimInstance.cpp
--- a/Engine/Source/Runtime/Engine/Private/Animation/AnimInstance.cpp
+++ b/Engine/Source/Runtime/Engine/Private/Animation/AnimInstance.cpp
@@ -3626,8 +3626,7 @@ UAnimInstance* UAnimInstance::GetLinkedAnimLayerInstanceByGroupAndClass(FName In
 	return nullptr;
 }
 
-
-UAnimInstance* UAnimInstance::GetLinkedAnimLayerInstanceByClass(TSubclassOf<UAnimInstance> InClass) const
+UAnimInstance* UAnimInstance::GetLinkedAnimLayerInstanceByClass(TSubclassOf<UAnimInstance> InClass, bool bCheckForChildClass) const
 {
 	if (IAnimClassInterface* AnimBlueprintClass = IAnimClassInterface::GetFromClass(GetClass()))
 	{
@@ -3635,9 +3634,13 @@ UAnimInstance* UAnimInstance::GetLinkedAnimLayerInstanceByClass(TSubclassOf<UAni
 		{
 			const FAnimNode_LinkedAnimLayer* Layer = LayerNodeProperty->ContainerPtrToValuePtr<FAnimNode_LinkedAnimLayer>(this);
 			UAnimInstance* TargetInstance = Layer->GetTargetInstance<UAnimInstance>();
-			if (TargetInstance && TargetInstance->GetClass() == InClass.Get())
+			
+			if (TargetInstance && TargetInstance->GetClass())
 			{
-				return TargetInstance;
+				if ((bCheckForChildClass && TargetInstance->GetClass()->IsChildOf(InClass)) || (!bCheckForChildClass && TargetInstance->GetClass() == InClass.Get()))
+				{
+					return TargetInstance;
+				}
 			}
 		}
 	}
