Fixed misconverted shrinking parameter.

diff --git a/Engine/Source/Runtime/Engine/Private/Texture.cpp b/Engine/Source/Runtime/Engine/Private/Texture.cpp
--- a/Engine/Source/Runtime/Engine/Private/Texture.cpp
+++ b/Engine/Source/Runtime/Engine/Private/Texture.cpp
@@ -4400,7 +4400,7 @@ void FTextureSource::InitLayeredImpl(
 	check( NewNumMips <= GetFullMipCount(SizeX,SizeY,GetVolumeSizeZ()) );
 
 	Format = NewLayerFormat[0];
-	LayerFormat.SetNum(NewNumLayers, EAllowShrinking::No);
+	LayerFormat.SetNum(NewNumLayers, EAllowShrinking::Yes);
 	for (int i = 0; i < NewNumLayers; ++i)
 	{
 		LayerFormat[i] = NewLayerFormat[i];
