Implemented support for SparseVolumeTextures falling back to coarser mip levels at page level. This is the first step in moving towards streaming at page granularity instead of mip level granularity. The main problem to solve here is to update the descendands of every streamed in/out page so they point to the proper fallback mip data. In this CL, this is done on the CPU using a new always loaded octree representing the sparse page table topology as well as bitsets to keep track of changes in page GPU residency. In the future, we could maybe move this step to the GPU: For every valid mip0 page a CS thread would walk up the mip chain to find the finest resident mip level to fall back to and then write it to all the relevant pages along the path to the parent.

#rnx

diff --git a/Engine/Shaders/Private/SparseVolumeTexture/SparseVolumeTextureCommon.ush b/Engine/Shaders/Private/SparseVolumeTexture/SparseVolumeTextureCommon.ush
--- a/Engine/Shaders/Private/SparseVolumeTexture/SparseVolumeTextureCommon.ush
+++ b/Engine/Shaders/Private/SparseVolumeTexture/SparseVolumeTextureCommon.ush
@@ -66,14 +66,18 @@ FSparseVolumeTextureUniforms SparseVolumeTextureUnpackUniforms(const uint4 Packe
 	return Result;
 }
 
-float3 SparseVolumeTextureGetVoxelCoord(const uint PackedPhysicalTileCoord, const float3 PageTableCoord)
+float3 SparseVolumeTextureGetVoxelCoord(Texture3D<uint> PageTable, FSparseVolumeTextureUniforms Uniforms, float3 PageTableCoord, int MipLevel)
 {
-	const int3 PhysicalTileCoord = int3(
-		PackedPhysicalTileCoord & 0xFF,
-		(PackedPhysicalTileCoord >> 8) & 0xFF,
-		(PackedPhysicalTileCoord >> 16) & 0xFF);
+	// Load from page table texture
+	const uint PackedPhysicalTileCoord = PageTable.Load(int4(floor(PageTableCoord), MipLevel)).x;
+	const int3 PhysicalTileCoord = int3(PackedPhysicalTileCoord & 0xFFu, (PackedPhysicalTileCoord >> 8u) & 0xFFu, (PackedPhysicalTileCoord >> 16u) & 0xFFu);
+	const uint TileMipLevel = PackedPhysicalTileCoord >> 24u;
+
+	// Scale TileCoord by 1 / exp2(ActualMipLevel - DesiredMipLevel) to get correct UVs when sampling from a higher mip as fallback
+	const float TileRcpMipLevelFactor = rcp(float(1u << (TileMipLevel - uint(MipLevel))));
+	const float3 FracTileCoord = frac(PageTableCoord * TileRcpMipLevelFactor);
 
-	const float3 VoxelCoord = float3(PhysicalTileCoord) * float(SPARSE_VOLUME_TILE_RES_PADDED) + (frac(PageTableCoord) * float(SPARSE_VOLUME_TILE_RES) + float(SPARSE_VOLUME_TILE_BORDER));
+	const float3 VoxelCoord = float3(PhysicalTileCoord) * float(SPARSE_VOLUME_TILE_RES_PADDED) + (FracTileCoord * float(SPARSE_VOLUME_TILE_RES) + float(SPARSE_VOLUME_TILE_BORDER));
 	return VoxelCoord;
 }
 
@@ -81,23 +85,21 @@ float3 SparseVolumeTextureSamplePageTable(Texture3D<uint> PageTable, FSparseVolu
 {
 	// Apply address mode to UVW and clamp mip level to resident levels
 	UVW = SparseVolumeTextureApplyAddressMode(UVW, AddressU, AddressV, AddressW);
-	MipLevel = clamp(MipLevel, Uniforms.LowestResidentMipLevel, Uniforms.HighestMipLevel);
+	MipLevel = clamp(MipLevel, 0, Uniforms.HighestMipLevel);
 
 	const float RcpMipLevelFactor = rcp(float(1u << (uint)MipLevel));
 	const float3 VolumePageCoord = UVW * Uniforms.VolumePageResolution;
 	const float3 MipPageTableOffset = floor(Uniforms.PageTableOffset * RcpMipLevelFactor);
 	const float3 PageTableCoord = VolumePageCoord * RcpMipLevelFactor - MipPageTableOffset;
 
-	const uint PackedPhysicalTileCoord = PageTable.Load(int4(floor(PageTableCoord), MipLevel)).x;
-
-	const float3 VoxelCoord = SparseVolumeTextureGetVoxelCoord(PackedPhysicalTileCoord, PageTableCoord);
+	const float3 VoxelCoord = SparseVolumeTextureGetVoxelCoord(PageTable, Uniforms, PageTableCoord, MipLevel);
 	const float3 VoxelUVW = VoxelCoord * Uniforms.TileDataTexelSize;
 	return VoxelUVW;
 }
 
 int3 SparseVolumeTextureLoadPageTable(Texture3D<uint> PageTable, FSparseVolumeTextureUniforms Uniforms, int3 TexelCoord, int MipLevel = 0)
 {
-	if (MipLevel < Uniforms.LowestResidentMipLevel || MipLevel > Uniforms.HighestMipLevel)
+	if (MipLevel < 0 || MipLevel > Uniforms.HighestMipLevel)
 	{
 		return 0; // Point to null tile
 	}
@@ -106,9 +108,7 @@ int3 SparseVolumeTextureLoadPageTable(Texture3D<uint> PageTable, FSparseVolumeTe
 	const float3 MipPageTableOffset = floor(Uniforms.PageTableOffset * RcpMipLevelFactor);
 	const float3 PageTableCoord = VolumeMipPageCoord - MipPageTableOffset;
 
-	const uint PackedPhysicalTileCoord = PageTable.Load(int4(floor(PageTableCoord), MipLevel)).x;
-
-	const float3 VoxelCoord = SparseVolumeTextureGetVoxelCoord(PackedPhysicalTileCoord, PageTableCoord);
+	const float3 VoxelCoord = SparseVolumeTextureGetVoxelCoord(PageTable, Uniforms, PageTableCoord, MipLevel);
 	return int3(VoxelCoord);
 }
 
diff --git a/Engine/Source/Runtime/Engine/Classes/SparseVolumeTexture/SparseVolumeTexture.h b/Engine/Source/Runtime/Engine/Classes/SparseVolumeTexture/SparseVolumeTexture.h
--- a/Engine/Source/Runtime/Engine/Classes/SparseVolumeTexture/SparseVolumeTexture.h
+++ b/Engine/Source/Runtime/Engine/Classes/SparseVolumeTexture/SparseVolumeTexture.h
@@ -53,6 +53,35 @@ struct FHeader
 	ENGINE_API bool Validate(bool bPrintToLog);
 };
 
+// Linear struct of arrays octree describing the mipped SVT topology
+struct FPageTopology
+{
+	struct FMip
+	{
+		uint32 PageOffset; // Offset in Pages array where the pages for this mip begin
+		uint32 PageCount; // Number of pages in this mip
+	};
+
+	struct FPage
+	{
+		uint32 PackedPageTableCoord; // 11|11|10
+		uint32 TileIndex; // Index of the tile this page points to. Indexes into compressed tile data and is relative to each mip level (starts at 0 for every mip)
+		uint32 ParentIndex;
+	};
+
+	struct FInteriorPageData
+	{
+		uint32 ChildIndices[8];
+	};
+
+	TArray<FMip> MipInfo;
+	TArray<FPage> Pages; // One entry for every page in the sparse topology. This single array is shared across all mip levels
+	TArray<FInteriorPageData> InteriorPageData; // One entry for every interior page/node (all mips except mip0)
+
+	void Reset();
+	void Serialize(FArchive& Ar);
+};
+
 // Describes a mip level of a SVT frame in terms of the sizes and offsets of the data in the built bulk data.
 // Each mip level consists of 4 buffer sections:
 // 
@@ -104,6 +133,8 @@ public:
 	// Data for all streamable mip levels
 	FByteBulkData StreamableMipLevels;
 
+	FPageTopology Topology;
+
 	// These are used for logging and retrieving StreamableMipLevels from DDC in FStreamingManager
 #if WITH_EDITORONLY_DATA
 	FString ResourceName;
@@ -142,6 +173,14 @@ private:
 	void BeginRebuildBulkDataFromCache(const UObject* Owner);
 	void EndRebuildBulkDataFromCache();
 #endif
+
+	// Stores page table into BulkData as two consecutive arrays of packed page coordinates and linear indices into the physical tiles array.
+	// Returns number of written/non-zero page table entries
+	static int32 CompressPageTable(const TArray<uint32>& PageTable, const FIntVector3& Resolution, TArray<uint8>& BulkData, TMap<uint32, uint32>& CoordToIndexMap);
+
+	static void CompressTiles(int32 NumTiles, const TArray64<uint8>& PhysicalTileDataA, const TArray64<uint8>& PhysicalTileDataB, const TArrayView<EPixelFormat>& Formats, const TArrayView<FVector4f>& FallbackValues, TArray<uint8>& BulkData, FMipLevelStreamingInfo& MipStreamingInfo);
+	
+	static FPageTopology BuildTopology(const FTextureData& InTextureData, const TArray<uint8>& InRootData, const TArray<uint8>& InStreamableBulkData, const TArray<TMap<uint32, uint32>, TInlineAllocator<16>>& InCoordToPageIndexInMipData, const TArray<FMipLevelStreamingInfo>& InMipLevelStreamingInfo);
 };
 
 // Encapsulates RHI resources needed to render a SparseVolumeTexture.
@@ -182,6 +221,9 @@ private:
 }
 
 FArchive& operator<<(FArchive& Ar, UE::SVT::FHeader& Header);
+FArchive& operator<<(FArchive& Ar, UE::SVT::FPageTopology::FMip& Mip);
+FArchive& operator<<(FArchive& Ar, UE::SVT::FPageTopology::FPage& Page);
+FArchive& operator<<(FArchive& Ar, UE::SVT::FPageTopology::FInteriorPageData& Children);
 
 enum ESparseVolumeTextureShaderUniform
 {
diff --git a/Engine/Source/Runtime/Engine/Private/SparseVolumeTexture/SparseVolumeTexture.cpp b/Engine/Source/Runtime/Engine/Private/SparseVolumeTexture/SparseVolumeTexture.cpp
--- a/Engine/Source/Runtime/Engine/Private/SparseVolumeTexture/SparseVolumeTexture.cpp
+++ b/Engine/Source/Runtime/Engine/Private/SparseVolumeTexture/SparseVolumeTexture.cpp
@@ -83,6 +83,30 @@ FArchive& operator<<(FArchive& Ar, UE::SVT::FHeader& Header)
 	return Ar;
 }
 
+FArchive& operator<<(FArchive& Ar, UE::SVT::FPageTopology::FMip& Mip)
+{
+	Ar << Mip.PageOffset;
+	Ar << Mip.PageCount;
+	return Ar;
+}
+
+FArchive& operator<<(FArchive& Ar, UE::SVT::FPageTopology::FPage& Page)
+{
+	Ar << Page.PackedPageTableCoord;
+	Ar << Page.TileIndex;
+	Ar << Page.ParentIndex;
+	return Ar;
+}
+
+FArchive& operator<<(FArchive& Ar, UE::SVT::FPageTopology::FInteriorPageData& InteriorPageData)
+{
+	for (uint32& ChildIndex : InteriorPageData.ChildIndices)
+	{
+		Ar << ChildIndex;
+	}
+	return Ar;
+}
+
 ////////////////////////////////////////////////////////////////////////////////////////////////
 
 namespace UE
@@ -105,10 +129,15 @@ static int32 ComputeNumMipLevels(const FIntVector3& InVirtualVolumeMin, const FI
 
 static const FString& GetDerivedDataVersion()
 {
-	static FString CachedVersionString = TEXT("381AE2A9-A903-4C8F-8486-891E24D6FC72");	// Bump this if you want to ignore all cached data so far.
+	static FString CachedVersionString = TEXT("A1A4EFBF-6596-418D-A1C1-11FC07BA0B7F");	// Bump this if you want to ignore all cached data so far.
 	return CachedVersionString;
 }
 
+uint32 PackX11Y11Z10(uint32 X, uint32 Y, uint32 Z)
+{
+	return (X & 0x7FFu) | ((Y & 0x7FFu) << 11u) | ((Z & 0x3FFu) << 22u);
+}
+
 FHeader::FHeader(const FIntVector3& AABBMin, const FIntVector3& AABBMax, EPixelFormat FormatA, EPixelFormat FormatB, const FVector4f& FallbackValueA, const FVector4f& FallbackValueB)
 {
 	VirtualVolumeAABBMin = AABBMin;
@@ -188,6 +217,22 @@ bool FHeader::Validate(bool bPrintToLog)
 
 ////////////////////////////////////////////////////////////////////////////////////////////////
 
+void FPageTopology::Reset()
+{
+	MipInfo.Reset();
+	Pages.Reset();
+	InteriorPageData.Reset();
+}
+
+void FPageTopology::Serialize(FArchive& Ar)
+{
+	Ar << MipInfo;
+	Ar << Pages;
+	Ar << InteriorPageData;
+}
+
+////////////////////////////////////////////////////////////////////////////////////////////////
+
 void FResources::Serialize(FArchive& Ar, UObject* Owner, bool bCooked)
 {
 	// Note: this is all derived data, native versioning is not needed, but be sure to bump GetDerivedDataVersion() when modifying!
@@ -212,6 +257,8 @@ void FResources::Serialize(FArchive& Ar, UObject* Owner, bool bCooked)
 		Ar << MipLevelStreamingInfo;
 		Ar << RootData;
 
+		Topology.Serialize(Ar);
+
 		// StreamableMipLevels is only serialized in cooked builds and when caching to DDC failed in editor builds.
 		// If the data was successfully cached to DDC, we just query it from DDC on the next run or recreate it if that failed.
 		if (StoredResourceFlags & EResourceFlag_StreamingDataInDDC)
@@ -328,190 +375,21 @@ bool FResources::Build(USparseVolumeTextureFrame* Owner, UE::Serialization::FEdi
 		Header = DerivedTextureData.Header;
 		RootData.Reset();
 		MipLevelStreamingInfo.SetNumZeroed(NumMipLevels);
+		Topology.Reset();
 		ResourceFlags = 0;
 		ResourceName.Reset();
 		DDCKeyHash.Reset();
 		DDCRawHash.Reset();
 		DDCRebuildState.store(EDDCRebuildState::Initial);
 
-		// Stores page table into BulkData as two consecutive arrays of packed page coordinates and linear indices into the physical tiles array.
-		// Returns number of written/non-zero page table entries
-		auto CompressPageTable = [](const TArray<uint32>& PageTable, const FIntVector3& Resolution, TArray<uint8>& BulkData) -> int32
-		{
-			TRACE_CPUPROFILER_EVENT_SCOPE(SVT::FResources::Build::CompressPageTable);
-
-			int32 NumNonZeroEntries = 0;
-			for (uint32 Entry : PageTable)
-			{
-				if (Entry)
-				{
-					++NumNonZeroEntries;
-				}
-			}
-
-			const int32 BaseOffset = BulkData.Num();
-			BulkData.SetNum(BulkData.Num() + NumNonZeroEntries * 2 * sizeof(uint32));
-			uint32* PackedCoords = reinterpret_cast<uint32*>(BulkData.GetData() + BaseOffset);
-			uint32* PageEntries = PackedCoords + NumNonZeroEntries;
-			
-			int32 NumWrittenEntries = 0;
-			for (int32 Z = 0; Z < Resolution.Z; ++Z)
-			{
-				for (int32 Y = 0; Y < Resolution.Y; ++Y)
-				{
-					for (int32 X = 0; X < Resolution.X; ++X)
-					{
-						const int32 LinearCoord = (Z * Resolution.Y * Resolution.X) + (Y * Resolution.X) + X;
-						const uint32 Entry = PageTable[LinearCoord];
-						if (Entry)
-						{
-							const uint32 Packed = (X & 0x7FFu) | ((Y & 0x7FFu) << 11u) | ((Z & 0x3FFu) << 22u);
-							PackedCoords[NumWrittenEntries] = Packed;
-							PageEntries[NumWrittenEntries] = Entry - 1;
-							++NumWrittenEntries;
-						}
-					}
-				}
-			}
-
-			return NumNonZeroEntries;
-		};
-
-		auto CompressTiles = [](int32 NumTiles, const TArray64<uint8>& PhysicalTileDataA, const TArray64<uint8>& PhysicalTileDataB, const TArrayView<EPixelFormat>& Formats, const TArrayView<FVector4f>& FallbackValues, TArray<uint8>& BulkData, FMipLevelStreamingInfo& MipStreamingInfo)
-		{
-			TRACE_CPUPROFILER_EVENT_SCOPE(SVT::FResources::Build::CompressTiles);
-
-			const int64 FormatSize[] = { GPixelFormats[Formats[0]].BlockBytes, GPixelFormats[Formats[1]].BlockBytes };
-			uint8 NullTileValuesU8[2][sizeof(float) * 4] = {};
-			int32 NumTextures = 0;
-			for (int32 i = 0; i < 2; ++i)
-			{
-				if (Formats[i] != PF_Unknown)
-				{
-					++NumTextures;
-					SVT::WriteVoxel(0, NullTileValuesU8[i], Formats[i], FallbackValues[i]);
-				}
-			}
-
-			const int32 OccupancySizePerTexture = NumTiles * SVT::NumOccupancyWordsPerPaddedTile * sizeof(uint32);
-			const int32 OccupancySize = NumTextures * OccupancySizePerTexture;
-			const int32 TileDataOffsetsSize = NumTextures * NumTiles * sizeof(uint32);
-			BulkData.Reserve(BulkData.Num() + OccupancySize + TileDataOffsetsSize);
-
-			MipStreamingInfo.OccupancyBitsOffset[0] = BulkData.Num() - MipStreamingInfo.BulkOffset;
-			MipStreamingInfo.OccupancyBitsSize[0] = Formats[0] != PF_Unknown ? OccupancySizePerTexture : 0;
-			BulkData.SetNum(BulkData.Num() + MipStreamingInfo.OccupancyBitsSize[0]);
-			
-			MipStreamingInfo.OccupancyBitsOffset[1] = BulkData.Num() - MipStreamingInfo.BulkOffset;
-			MipStreamingInfo.OccupancyBitsSize[1] = Formats[1] != PF_Unknown ? OccupancySizePerTexture : 0;
-			BulkData.SetNum(BulkData.Num() + MipStreamingInfo.OccupancyBitsSize[1]);
-
-			MipStreamingInfo.TileDataOffsetsOffset[0] = BulkData.Num() - MipStreamingInfo.BulkOffset;
-			MipStreamingInfo.TileDataOffsetsSize[0] = Formats[0] != PF_Unknown ? NumTiles * sizeof(uint32) : 0;
-			BulkData.SetNum(BulkData.Num() + MipStreamingInfo.TileDataOffsetsSize[0]);
-
-			MipStreamingInfo.TileDataOffsetsOffset[1] = BulkData.Num() - MipStreamingInfo.BulkOffset;
-			MipStreamingInfo.TileDataOffsetsSize[1] = Formats[1] != PF_Unknown ? NumTiles * sizeof(uint32) : 0;
-			BulkData.SetNum(BulkData.Num() + MipStreamingInfo.TileDataOffsetsSize[1]);
-
-			// Zero the bitmasks
-			FMemory::Memzero(BulkData.GetData() + MipStreamingInfo.BulkOffset + MipStreamingInfo.OccupancyBitsOffset[0], MipStreamingInfo.OccupancyBitsSize[0] + MipStreamingInfo.OccupancyBitsSize[1]);
-
-			uint32* OccupancyBits[2];
-			OccupancyBits[0] = reinterpret_cast<uint32*>(BulkData.GetData() + MipStreamingInfo.BulkOffset + MipStreamingInfo.OccupancyBitsOffset[0]);
-			OccupancyBits[1] = reinterpret_cast<uint32*>(BulkData.GetData() + MipStreamingInfo.BulkOffset + MipStreamingInfo.OccupancyBitsOffset[1]);
-
-			uint32* PrefixSums[2];
-			PrefixSums[0] = reinterpret_cast<uint32*>(BulkData.GetData() + MipStreamingInfo.BulkOffset + MipStreamingInfo.TileDataOffsetsOffset[0]);
-			PrefixSums[1] = reinterpret_cast<uint32*>(BulkData.GetData() + MipStreamingInfo.BulkOffset + MipStreamingInfo.TileDataOffsetsOffset[1]);
-
-			// Compute occupancy bitmasks and count number of non-fallback voxels per tile
-			ParallelFor(NumTiles, [&](int32 TileIndex)
-				{
-					for (int32 AttributesIdx = 0; AttributesIdx < 2; ++AttributesIdx)
-					{
-						if (Formats[AttributesIdx] != PF_Unknown)
-						{
-							PrefixSums[AttributesIdx][TileIndex] = 0;
-							for (int64 VoxelIndex = 0; VoxelIndex < SVT::NumVoxelsPerPaddedTile; ++VoxelIndex)
-							{
-								const uint8* Src = (AttributesIdx == 0 ? PhysicalTileDataA.GetData() : PhysicalTileDataB.GetData()) + FormatSize[AttributesIdx] * (TileIndex * SVT::NumVoxelsPerPaddedTile + VoxelIndex);
-								bool bIsFallbackValue = FMemory::Memcmp(Src, NullTileValuesU8[AttributesIdx], FormatSize[AttributesIdx]) == 0;
-								if (!bIsFallbackValue)
-								{
-									const int64 WordIndex = TileIndex * SVT::NumOccupancyWordsPerPaddedTile + (VoxelIndex / 32);
-									OccupancyBits[AttributesIdx][WordIndex] |= 1u << (static_cast<uint32>(VoxelIndex) % 32u);
-									PrefixSums[AttributesIdx][TileIndex]++;
-								}
-							}
-						}
-					}
-				});
-
-			// Compute actual per-tile voxel data offsets (prefix sum)
-			uint32 Sums[2] = {};
-			for (int32 AttributesIdx = 0; AttributesIdx < 2; ++AttributesIdx)
-			{
-				if (Formats[AttributesIdx] != PF_Unknown)
-				{
-					for (int32 TileIndex = 0; TileIndex < NumTiles; ++TileIndex)
-					{
-						uint32 Tmp = PrefixSums[AttributesIdx][TileIndex];
-						PrefixSums[AttributesIdx][TileIndex] = Sums[AttributesIdx];
-						Sums[AttributesIdx] += Tmp;
-					}
-				}
-			}
-
-			// Allocate compacted tile data memory
-			MipStreamingInfo.TileDataOffset[0] = BulkData.Num() - MipStreamingInfo.BulkOffset;
-			MipStreamingInfo.TileDataSize[0] = Sums[0] * FormatSize[0];
-
-			MipStreamingInfo.TileDataOffset[1] = MipStreamingInfo.TileDataOffset[0] + MipStreamingInfo.TileDataSize[0];
-			MipStreamingInfo.TileDataSize[1] = Sums[1] * FormatSize[1];
-
-			BulkData.SetNum(BulkData.Num() + MipStreamingInfo.TileDataSize[0] + MipStreamingInfo.TileDataSize[1]);
-
-			// All above pointers into the BulkData are now invalid!
-
-			uint8* RelativeBulkDataPtr = BulkData.GetData() + MipStreamingInfo.BulkOffset;
-
-			uint8* DstTileData[2];
-			DstTileData[0] = RelativeBulkDataPtr + MipStreamingInfo.TileDataOffset[0];
-			DstTileData[1] = RelativeBulkDataPtr + MipStreamingInfo.TileDataOffset[1];
-
-			// Copy voxels to compacted locations
-			ParallelFor(NumTiles, [&](int32 TileIndex)
-				{
-					for (int32 AttributesIdx = 0; AttributesIdx < 2; ++AttributesIdx)
-					{
-						if (Formats[AttributesIdx] != PF_Unknown)
-						{
-							uint32 VoxelDataWriteOffset = reinterpret_cast<uint32*>(RelativeBulkDataPtr + MipStreamingInfo.TileDataOffsetsOffset[AttributesIdx])[TileIndex];
-							const uint32* TileOccupancyBits = reinterpret_cast<uint32*>(RelativeBulkDataPtr + MipStreamingInfo.OccupancyBitsOffset[AttributesIdx] + TileIndex * SVT::NumOccupancyWordsPerPaddedTile * sizeof(uint32));
-
-							for (int64 VoxelIndex = 0; VoxelIndex < SVT::NumVoxelsPerPaddedTile; ++VoxelIndex)
-							{
-								const int64 WordIndex = (VoxelIndex / 32);
-
-								if (TileOccupancyBits[WordIndex] & (1u << (static_cast<uint32>(VoxelIndex) % 32u)))
-								{
-									uint8* Dst = DstTileData[AttributesIdx] + VoxelDataWriteOffset * FormatSize[AttributesIdx];
-									const uint8* Src = (AttributesIdx == 0 ? PhysicalTileDataA.GetData() : PhysicalTileDataB.GetData()) + FormatSize[AttributesIdx] * (TileIndex * SVT::NumVoxelsPerPaddedTile + VoxelIndex);
-									FMemory::Memcpy(Dst, Src, FormatSize[AttributesIdx]);
-									++VoxelDataWriteOffset;
-								}
-							}
-						}
-					}
-				});
-		};
+		TArray<TMap<uint32, uint32>, TInlineAllocator<16>> CoordToPageIndexInMipData;
+		CoordToPageIndexInMipData.SetNum(NumMipLevels);
 
 		TArray<uint8> StreamableBulkData;
-		for (int32 MipLevelIdx = 0; MipLevelIdx < DerivedTextureData.MipMaps.Num(); ++MipLevelIdx)
+		for (int32 MipLevelIdx = NumMipLevels - 1; MipLevelIdx >= 0; --MipLevelIdx)
 		{
 			const FTextureData::FMipMap& Mip = DerivedTextureData.MipMaps[MipLevelIdx];
-			const bool bIsRootMipLevel = (MipLevelIdx == (DerivedTextureData.MipMaps.Num() - 1));
+			const bool bIsRootMipLevel = (MipLevelIdx == (NumMipLevels - 1));
 			TArray<uint8>& BulkData = bIsRootMipLevel ? RootData : StreamableBulkData;
 
 			FIntVector3 MipPageTableResolution = DerivedTextureData.Header.PageTableVolumeResolution >> MipLevelIdx;
@@ -521,7 +399,7 @@ bool FResources::Build(USparseVolumeTextureFrame* Owner, UE::Serialization::FEdi
 			MipStreamingInfo.BulkOffset = BulkData.Num();
 
 			MipStreamingInfo.PageTableOffset = BulkData.Num() - MipStreamingInfo.BulkOffset;
-			const int32 NumNonZeroPageTableEntries = CompressPageTable(Mip.PageTable, MipPageTableResolution, BulkData);
+			const int32 NumNonZeroPageTableEntries = CompressPageTable(Mip.PageTable, MipPageTableResolution, BulkData, CoordToPageIndexInMipData[MipLevelIdx]);
 			MipStreamingInfo.PageTableSize = NumNonZeroPageTableEntries * 2 * sizeof(uint32);
 			
 			CompressTiles(Mip.NumPhysicalTiles, Mip.PhysicalTileDataA, Mip.PhysicalTileDataB, Header.AttributesFormats, Header.FallbackValues, BulkData, MipStreamingInfo);
@@ -541,6 +419,9 @@ bool FResources::Build(USparseVolumeTextureFrame* Owner, UE::Serialization::FEdi
 			StreamableMipLevels.SetBulkDataFlags(BULKDATA_Force_NOT_InlinePayload);
 		}
 
+		// Build page topology
+		Topology = BuildTopology(DerivedTextureData, RootData, StreamableBulkData, CoordToPageIndexInMipData, MipLevelStreamingInfo);
+
 		return true;
 	}
 	return false;
@@ -752,6 +633,285 @@ void FResources::EndRebuildBulkDataFromCache()
 
 #endif // WITH_EDITORONLY_DATA
 
+int32 FResources::CompressPageTable(const TArray<uint32>& PageTable, const FIntVector3& Resolution, TArray<uint8>& BulkData, TMap<uint32, uint32>& CoordToIndexMap)
+{
+	TRACE_CPUPROFILER_EVENT_SCOPE(SVT::FResources::Build::CompressPageTable);
+
+	int32 NumNonZeroEntries = 0;
+	for (uint32 Entry : PageTable)
+	{
+		if (Entry)
+		{
+			++NumNonZeroEntries;
+		}
+	}
+	CoordToIndexMap.Reserve(NumNonZeroEntries);
+
+	const int32 BaseOffset = BulkData.Num();
+	BulkData.SetNum(BulkData.Num() + NumNonZeroEntries * 2 * sizeof(uint32));
+	uint32* PackedCoords = reinterpret_cast<uint32*>(BulkData.GetData() + BaseOffset);
+	uint32* PageEntries = PackedCoords + NumNonZeroEntries;
+
+	int32 NumWrittenEntries = 0;
+	for (int32 Z = 0; Z < Resolution.Z; ++Z)
+	{
+		for (int32 Y = 0; Y < Resolution.Y; ++Y)
+		{
+			for (int32 X = 0; X < Resolution.X; ++X)
+			{
+				const int32 LinearCoord = (Z * Resolution.Y * Resolution.X) + (Y * Resolution.X) + X;
+				const uint32 Entry = PageTable[LinearCoord];
+				if (Entry)
+				{
+					const uint32 Packed = PackX11Y11Z10(X, Y, Z);
+					PackedCoords[NumWrittenEntries] = Packed;
+					PageEntries[NumWrittenEntries] = Entry - 1;
+					CoordToIndexMap.Add(Packed, NumWrittenEntries);
+					++NumWrittenEntries;
+				}
+			}
+		}
+	}
+
+	return NumNonZeroEntries;
+}
+
+void FResources::CompressTiles(int32 NumTiles, const TArray64<uint8>& PhysicalTileDataA, const TArray64<uint8>& PhysicalTileDataB, const TArrayView<EPixelFormat>& Formats, const TArrayView<FVector4f>& FallbackValues, TArray<uint8>& BulkData, FMipLevelStreamingInfo& MipStreamingInfo)
+{
+	TRACE_CPUPROFILER_EVENT_SCOPE(SVT::FResources::Build::CompressTiles);
+
+	const int64 FormatSize[] = { GPixelFormats[Formats[0]].BlockBytes, GPixelFormats[Formats[1]].BlockBytes };
+	uint8 NullTileValuesU8[2][sizeof(float) * 4] = {};
+	int32 NumTextures = 0;
+	for (int32 i = 0; i < 2; ++i)
+	{
+		if (Formats[i] != PF_Unknown)
+		{
+			++NumTextures;
+			SVT::WriteVoxel(0, NullTileValuesU8[i], Formats[i], FallbackValues[i]);
+		}
+	}
+
+	const int32 OccupancySizePerTexture = NumTiles * SVT::NumOccupancyWordsPerPaddedTile * sizeof(uint32);
+	const int32 OccupancySize = NumTextures * OccupancySizePerTexture;
+	const int32 TileDataOffsetsSize = NumTextures * NumTiles * sizeof(uint32);
+	BulkData.Reserve(BulkData.Num() + OccupancySize + TileDataOffsetsSize);
+
+	MipStreamingInfo.OccupancyBitsOffset[0] = BulkData.Num() - MipStreamingInfo.BulkOffset;
+	MipStreamingInfo.OccupancyBitsSize[0] = Formats[0] != PF_Unknown ? OccupancySizePerTexture : 0;
+	BulkData.SetNum(BulkData.Num() + MipStreamingInfo.OccupancyBitsSize[0]);
+
+	MipStreamingInfo.OccupancyBitsOffset[1] = BulkData.Num() - MipStreamingInfo.BulkOffset;
+	MipStreamingInfo.OccupancyBitsSize[1] = Formats[1] != PF_Unknown ? OccupancySizePerTexture : 0;
+	BulkData.SetNum(BulkData.Num() + MipStreamingInfo.OccupancyBitsSize[1]);
+
+	MipStreamingInfo.TileDataOffsetsOffset[0] = BulkData.Num() - MipStreamingInfo.BulkOffset;
+	MipStreamingInfo.TileDataOffsetsSize[0] = Formats[0] != PF_Unknown ? NumTiles * sizeof(uint32) : 0;
+	BulkData.SetNum(BulkData.Num() + MipStreamingInfo.TileDataOffsetsSize[0]);
+
+	MipStreamingInfo.TileDataOffsetsOffset[1] = BulkData.Num() - MipStreamingInfo.BulkOffset;
+	MipStreamingInfo.TileDataOffsetsSize[1] = Formats[1] != PF_Unknown ? NumTiles * sizeof(uint32) : 0;
+	BulkData.SetNum(BulkData.Num() + MipStreamingInfo.TileDataOffsetsSize[1]);
+
+	// Zero the bitmasks
+	FMemory::Memzero(BulkData.GetData() + MipStreamingInfo.BulkOffset + MipStreamingInfo.OccupancyBitsOffset[0], MipStreamingInfo.OccupancyBitsSize[0] + MipStreamingInfo.OccupancyBitsSize[1]);
+
+	uint32* OccupancyBits[2];
+	OccupancyBits[0] = reinterpret_cast<uint32*>(BulkData.GetData() + MipStreamingInfo.BulkOffset + MipStreamingInfo.OccupancyBitsOffset[0]);
+	OccupancyBits[1] = reinterpret_cast<uint32*>(BulkData.GetData() + MipStreamingInfo.BulkOffset + MipStreamingInfo.OccupancyBitsOffset[1]);
+
+	uint32* PrefixSums[2];
+	PrefixSums[0] = reinterpret_cast<uint32*>(BulkData.GetData() + MipStreamingInfo.BulkOffset + MipStreamingInfo.TileDataOffsetsOffset[0]);
+	PrefixSums[1] = reinterpret_cast<uint32*>(BulkData.GetData() + MipStreamingInfo.BulkOffset + MipStreamingInfo.TileDataOffsetsOffset[1]);
+
+	// Compute occupancy bitmasks and count number of non-fallback voxels per tile
+	ParallelFor(NumTiles, [&](int32 TileIndex)
+	{
+		for (int32 AttributesIdx = 0; AttributesIdx < 2; ++AttributesIdx)
+		{
+			if (Formats[AttributesIdx] != PF_Unknown)
+			{
+				PrefixSums[AttributesIdx][TileIndex] = 0;
+				for (int64 VoxelIndex = 0; VoxelIndex < SVT::NumVoxelsPerPaddedTile; ++VoxelIndex)
+				{
+					const uint8* Src = (AttributesIdx == 0 ? PhysicalTileDataA.GetData() : PhysicalTileDataB.GetData()) + FormatSize[AttributesIdx] * (TileIndex * SVT::NumVoxelsPerPaddedTile + VoxelIndex);
+					bool bIsFallbackValue = FMemory::Memcmp(Src, NullTileValuesU8[AttributesIdx], FormatSize[AttributesIdx]) == 0;
+					if (!bIsFallbackValue)
+					{
+						const int64 WordIndex = TileIndex * SVT::NumOccupancyWordsPerPaddedTile + (VoxelIndex / 32);
+						OccupancyBits[AttributesIdx][WordIndex] |= 1u << (static_cast<uint32>(VoxelIndex) % 32u);
+						PrefixSums[AttributesIdx][TileIndex]++;
+					}
+				}
+			}
+		}
+	});
+
+	// Compute actual per-tile voxel data offsets (prefix sum)
+	uint32 Sums[2] = {};
+	for (int32 AttributesIdx = 0; AttributesIdx < 2; ++AttributesIdx)
+	{
+		if (Formats[AttributesIdx] != PF_Unknown)
+		{
+			for (int32 TileIndex = 0; TileIndex < NumTiles; ++TileIndex)
+			{
+				uint32 Tmp = PrefixSums[AttributesIdx][TileIndex];
+				PrefixSums[AttributesIdx][TileIndex] = Sums[AttributesIdx];
+				Sums[AttributesIdx] += Tmp;
+			}
+		}
+	}
+
+	// Allocate compacted tile data memory
+	MipStreamingInfo.TileDataOffset[0] = BulkData.Num() - MipStreamingInfo.BulkOffset;
+	MipStreamingInfo.TileDataSize[0] = Sums[0] * FormatSize[0];
+
+	MipStreamingInfo.TileDataOffset[1] = MipStreamingInfo.TileDataOffset[0] + MipStreamingInfo.TileDataSize[0];
+	MipStreamingInfo.TileDataSize[1] = Sums[1] * FormatSize[1];
+
+	BulkData.SetNum(BulkData.Num() + MipStreamingInfo.TileDataSize[0] + MipStreamingInfo.TileDataSize[1]);
+
+	// All above pointers into the BulkData are now invalid!
+
+	uint8* RelativeBulkDataPtr = BulkData.GetData() + MipStreamingInfo.BulkOffset;
+
+	uint8* DstTileData[2];
+	DstTileData[0] = RelativeBulkDataPtr + MipStreamingInfo.TileDataOffset[0];
+	DstTileData[1] = RelativeBulkDataPtr + MipStreamingInfo.TileDataOffset[1];
+
+	// Copy voxels to compacted locations
+	ParallelFor(NumTiles, [&](int32 TileIndex)
+	{
+		for (int32 AttributesIdx = 0; AttributesIdx < 2; ++AttributesIdx)
+		{
+			if (Formats[AttributesIdx] != PF_Unknown)
+			{
+				uint32 VoxelDataWriteOffset = reinterpret_cast<uint32*>(RelativeBulkDataPtr + MipStreamingInfo.TileDataOffsetsOffset[AttributesIdx])[TileIndex];
+				const uint32* TileOccupancyBits = reinterpret_cast<uint32*>(RelativeBulkDataPtr + MipStreamingInfo.OccupancyBitsOffset[AttributesIdx] + TileIndex * SVT::NumOccupancyWordsPerPaddedTile * sizeof(uint32));
+
+				for (int64 VoxelIndex = 0; VoxelIndex < SVT::NumVoxelsPerPaddedTile; ++VoxelIndex)
+				{
+					const int64 WordIndex = (VoxelIndex / 32);
+
+					if (TileOccupancyBits[WordIndex] & (1u << (static_cast<uint32>(VoxelIndex) % 32u)))
+					{
+						uint8* Dst = DstTileData[AttributesIdx] + VoxelDataWriteOffset * FormatSize[AttributesIdx];
+						const uint8* Src = (AttributesIdx == 0 ? PhysicalTileDataA.GetData() : PhysicalTileDataB.GetData()) + FormatSize[AttributesIdx] * (TileIndex * SVT::NumVoxelsPerPaddedTile + VoxelIndex);
+						FMemory::Memcpy(Dst, Src, FormatSize[AttributesIdx]);
+						++VoxelDataWriteOffset;
+					}
+				}
+			}
+		}
+	});
+}
+
+FPageTopology FResources::BuildTopology(const FTextureData& InTextureData, const TArray<uint8>& InRootData, const TArray<uint8>& InStreamableBulkData, const TArray<TMap<uint32, uint32>, TInlineAllocator<16>>& InCoordToPageIndexInMipData, const TArray<FMipLevelStreamingInfo>& InMipLevelStreamingInfo)
+{
+	FPageTopology Topology;
+
+	const int32 NumMipLevels = InMipLevelStreamingInfo.Num();
+
+	// Reserve memory and initialize the page offsets and counts for each mip level
+	{
+		Topology.MipInfo.SetNum(NumMipLevels);
+		uint32 NumInteriorPages = 0;
+		uint32 NumPages = 0;
+		for (int32 MipLevel = NumMipLevels - 1; MipLevel >= 0; --MipLevel)
+		{
+			// Page table bulk data stores 2 entries per non-zero page table entry
+			const uint32 NumPagesMip = InMipLevelStreamingInfo[MipLevel].PageTableSize / (2 * sizeof(uint32));
+			Topology.MipInfo[MipLevel].PageOffset = NumPages;
+			Topology.MipInfo[MipLevel].PageCount = NumPagesMip;
+			NumPages += NumPagesMip;
+			if (MipLevel > 0)
+			{
+				NumInteriorPages += NumPagesMip;
+			}
+		}
+		Topology.Pages.SetNumUninitialized(NumPages);
+		Topology.InteriorPageData.SetNumUninitialized(NumInteriorPages);
+	}
+
+	// Recursively build the page mip hierarchy. The order of nodes is identical to how they're stored in bulk data.
+	auto BuildHierarchy = [&](const FIntVector3& InPageTableMipResolution, int32 InMipLevel, int32 InPageX, int32 InPageY, int32 InPageZ, uint32 InParentIndex, auto& InRecursionLambda) -> uint32
+	{
+		// Load entry from dense page table. 
+		const int32 LinearPageIndex = InPageZ * (InPageTableMipResolution.X * InPageTableMipResolution.Y) + (InPageY * InPageTableMipResolution.X) + InPageX;
+		const uint32 PageTableEntry = InTextureData.MipMaps[InMipLevel].PageTable[LinearPageIndex];
+		if (!PageTableEntry)
+		{
+			return INDEX_NONE;
+		}
+
+		FPageTopology::FPage Page;
+		Page.PackedPageTableCoord = PackX11Y11Z10(InPageX, InPageY, InPageZ);
+		Page.TileIndex = PageTableEntry - 1;
+		Page.ParentIndex = InParentIndex;
+
+		// Look up index of page in sparse list of non-zero pages
+		const uint32 PageIndex = InCoordToPageIndexInMipData[InMipLevel][Page.PackedPageTableCoord] + Topology.MipInfo[InMipLevel].PageOffset;
+
+		Topology.Pages[PageIndex] = Page;
+
+		// Process any potential children
+		if (InMipLevel > 0)
+		{
+			const int32 ChildMipLevel = InMipLevel - 1;
+			const FIntVector3 ChildPageTableMipResolution = ShiftRightAndMax(InTextureData.Header.PageTableVolumeResolution, ChildMipLevel, 1);
+
+			FPageTopology::FInteriorPageData InteriorPageData;
+			for (uint32 ChildIndex = 0; ChildIndex < 8; ++ChildIndex)
+			{
+				const int32 ChildPageX = InPageX * 2 + int32((ChildIndex >> 0u) & 1u);
+				const int32 ChildPageY = InPageY * 2 + int32((ChildIndex >> 1u) & 1u);
+				const int32 ChildPageZ = InPageZ * 2 + int32((ChildIndex >> 2u) & 1u);
+				const uint32 ChildPageIndex = InRecursionLambda(ChildPageTableMipResolution, ChildMipLevel, ChildPageX, ChildPageY, ChildPageZ, PageIndex, InRecursionLambda);
+				InteriorPageData.ChildIndices[ChildIndex] = ChildPageIndex;
+			}
+
+			Topology.InteriorPageData[PageIndex] = InteriorPageData;
+		}
+
+		return PageIndex;
+	};
+
+	BuildHierarchy(FIntVector3::ZeroValue, NumMipLevels - 1, 0, 0, 0, INDEX_NONE, BuildHierarchy);
+
+	// Validate
+	{
+#if DO_CHECK
+		for (int32 MipLevel = NumMipLevels - 1; MipLevel >= 0; --MipLevel)
+		{
+			const int32 NumPageTableEntries = (InMipLevelStreamingInfo[MipLevel].PageTableSize / (sizeof(uint32) * 2));
+			check(Topology.MipInfo[MipLevel].PageCount == NumPageTableEntries); // PageCount in topology must match number of pages in compressed bulk data
+			const bool bIsHighestMipLevel = MipLevel == (NumMipLevels - 1);
+			check(!bIsHighestMipLevel || NumPageTableEntries <= 1); // Highest mip has a maximum of 1 page
+			const TArray<uint8>& BulkData = bIsHighestMipLevel ? InRootData : InStreamableBulkData;
+			const uint32* PackedCoords = reinterpret_cast<const uint32*>(BulkData.GetData() + InMipLevelStreamingInfo[MipLevel].BulkOffset + InMipLevelStreamingInfo[MipLevel].PageTableOffset);
+			const uint32* PageEntries = PackedCoords + NumPageTableEntries;
+
+			const FIntVector3 PageTableRes = ShiftRightAndMax(InTextureData.Header.PageTableVolumeResolution, MipLevel, 1);
+
+			for (int32 i = 0; i < NumPageTableEntries; ++i)
+			{
+				const int32 PageIndex = Topology.MipInfo[MipLevel].PageOffset + i;
+				check(PackedCoords[i] == Topology.Pages[PageIndex].PackedPageTableCoord); // Topology stores pages in the same order as in the compressed bulk data
+				check(PageEntries[i] == Topology.Pages[PageIndex].TileIndex);
+				const int32 PageX = (PackedCoords[i] >> 0u) & 0x7FFu;
+				const int32 PageY = (PackedCoords[i] >> 11u) & 0x7FFu;
+				const int32 PageZ = (PackedCoords[i] >> 22u) & 0x3FFu;
+				const int32 LinearPageIndex = (PageZ * (PageTableRes.X * PageTableRes.Y)) + (PageY * PageTableRes.X) + PageX;
+				const uint32 PageTableEntry = InTextureData.MipMaps[MipLevel].PageTable[LinearPageIndex];
+				check(PageTableEntry != 0 && PageTableEntry != INDEX_NONE); // Unpacked page coord must map to a valid entry in the dense page table
+			}
+		}
+#endif // DO_CHECK
+	}
+
+	return Topology;
+}
+
 ////////////////////////////////////////////////////////////////////////////////////////////////
 
 void FTextureRenderResources::GetPackedUniforms(FUintVector4& OutPacked0, FUintVector4& OutPacked1) const
@@ -774,7 +934,7 @@ void FTextureRenderResources::GetPackedUniforms(FUintVector4& OutPacked0, FUintV
 	OutPacked0.X = AsUint(VolumePageResolution.X);
 	OutPacked0.Y = AsUint(VolumePageResolution.Y);
 	OutPacked0.Z = AsUint(VolumePageResolution.Z);
-	OutPacked0.W = (PageTableOffset.X & 0x7FFu) | ((PageTableOffset.Y & 0x7FFu) << 11u) | ((PageTableOffset.Z & 0x3FFu) << 22u);
+	OutPacked0.W = PackX11Y11Z10(PageTableOffset.X, PageTableOffset.Y, PageTableOffset.Z);
 	OutPacked1.X = AsUint(TileDataTexelSize.X);
 	OutPacked1.Y = AsUint(TileDataTexelSize.Y);
 	OutPacked1.Z = AsUint(TileDataTexelSize.Z);
diff --git a/Engine/Source/Runtime/Engine/Private/SparseVolumeTexture/SparseVolumeTextureData.cpp b/Engine/Source/Runtime/Engine/Private/SparseVolumeTexture/SparseVolumeTextureData.cpp
--- a/Engine/Source/Runtime/Engine/Private/SparseVolumeTexture/SparseVolumeTextureData.cpp
+++ b/Engine/Source/Runtime/Engine/Private/SparseVolumeTexture/SparseVolumeTextureData.cpp
@@ -40,13 +40,13 @@ namespace UE
 namespace SVT
 {
 
-static bool IsInBounds(const FIntVector3& Point, const FIntVector3& Min, const FIntVector3& Max)
+bool IsInBounds(const FIntVector3& Point, const FIntVector3& Min, const FIntVector3& Max)
 {
 	return Point.X >= Min.X && Point.Y >= Min.Y && Point.Z >= Min.Z
 		&& Point.X < Max.X && Point.Y < Max.Y && Point.Z < Max.Z;
 }
 
-static FIntVector3 ShiftRightAndMax(const FIntVector3& Value, uint32 ShiftBy, int32 MinValue)
+FIntVector3 ShiftRightAndMax(const FIntVector3& Value, uint32 ShiftBy, int32 MinValue)
 {
 	FIntVector3 Result = FIntVector3(
 		FMath::Max(Value.X >> ShiftBy, MinValue),
diff --git a/Engine/Source/Runtime/Engine/Private/SparseVolumeTexture/SparseVolumeTextureStreamingManager.cpp b/Engine/Source/Runtime/Engine/Private/SparseVolumeTexture/SparseVolumeTextureStreamingManager.cpp
--- a/Engine/Source/Runtime/Engine/Private/SparseVolumeTexture/SparseVolumeTextureStreamingManager.cpp
+++ b/Engine/Source/Runtime/Engine/Private/SparseVolumeTexture/SparseVolumeTextureStreamingManager.cpp
@@ -1159,7 +1159,6 @@ void FStreamingManager::BeginAsyncUpdate(FRDGBuilder& GraphBuilder, bool bBlocki
 	AsyncState.NumReadyMipLevels = DetermineReadyMipLevels();
 
 	// Do a first pass over all the mips to be uploaded to compute the upload buffer size requirements.
-	int32 NumPageTableUpdatesTotal = 0;
 	TileDataTexturesToUpdate.Reset();
 	{
 		const int32 StartPendingMipLevelIndex = (NextPendingMipLevelIndex + MaxPendingMipLevels - NumPendingMipLevels) % MaxPendingMipLevels;
@@ -1181,11 +1180,8 @@ void FStreamingManager::BeginAsyncUpdate(FRDGBuilder& GraphBuilder, bool bBlocki
 			SVTInfo->TileDataTexture->NumVoxelsToUploadA += FormatSizeA > 0 ? Resources->MipLevelStreamingInfo[PendingMipLevel.MipLevelIndex].TileDataSize[0] / FormatSizeA : 0;
 			SVTInfo->TileDataTexture->NumVoxelsToUploadB += FormatSizeB > 0 ? Resources->MipLevelStreamingInfo[PendingMipLevel.MipLevelIndex].TileDataSize[1] / FormatSizeB : 0;
 			TileDataTexturesToUpdate.Add(SVTInfo->TileDataTexture.Get());
-			NumPageTableUpdatesTotal += Resources->MipLevelStreamingInfo[PendingMipLevel.MipLevelIndex].PageTableSize / (2 * sizeof(uint32));
 		}
 
-		PageTableUpdater->Init(GraphBuilder, NumPageTableUpdatesTotal, 1);
-
 		for (FTileDataTexture* TileDataTexture : TileDataTexturesToUpdate)
 		{
 			TileDataTexture->TileUploader->Init(GraphBuilder, TileDataTexture->NumTilesToUpload, TileDataTexture->NumVoxelsToUploadA, TileDataTexture->NumVoxelsToUploadB, TileDataTexture->FormatA, TileDataTexture->FormatB);
@@ -1273,7 +1269,8 @@ void FStreamingManager::EndAsyncUpdate(FRDGBuilder& GraphBuilder)
 	SVTsWithInvalidatedStreamingInfoBuffer.Reset();
 	StreamingInfoBufferUpdater->Apply(GraphBuilder);
 
-	PageTableUpdater->Apply(GraphBuilder);
+	// Update page table with newly streamed in/out pages and make sure descendant pages in the hierarchy have correct fallback values
+	PatchPageTable(GraphBuilder);
 
 	check(AsyncState.NumReadyMipLevels <= NumPendingMipLevels);
 	NumPendingMipLevels -= AsyncState.NumReadyMipLevels;
@@ -1350,6 +1347,10 @@ void FStreamingManager::AddInternal(FRDGBuilder& GraphBuilder, FNewSparseVolumeT
 		{
 			FrameInfo.TileAllocations[MipLevel].SetNumZeroed(Resources->MipLevelStreamingInfo[MipLevel].NumPhysicalTiles);
 		}
+		const int32 NumPagesTotal = Resources->Topology.Pages.Num();
+		FrameInfo.ResidentPages.SetNum(NumPagesTotal, false);
+		FrameInfo.ResidentPagesNew.SetNum(NumPagesTotal, false);
+		FrameInfo.PageEntries.SetNum(NumPagesTotal);
 		
 		int32 NumPhysicalTiles = 0;
 		for (const FMipLevelStreamingInfo& MipLevelStreamingInfo : Resources->MipLevelStreamingInfo)
@@ -1464,6 +1465,7 @@ void FStreamingManager::AddInternal(FRDGBuilder& GraphBuilder, FNewSparseVolumeT
 				const uint32 TileCoord = SVTInfo.TileDataTexture->Allocate();
 				check(TileCoord != INDEX_NONE);
 				FrameInfo.TileAllocations.Last()[0] = TileCoord;
+				FrameInfo.PageEntries[Resources->Topology.MipInfo.Last().PageOffset] = TileCoord;
 
 				const int32 NumVoxelsA = FormatSizes[0] > 0 ? RootStreamingInfo->TileDataSize[0] / FormatSizes[0] : 0;
 				const int32 NumVoxelsB = FormatSizes[1] > 0 ? RootStreamingInfo->TileDataSize[1] / FormatSizes[1] : 0;
@@ -1494,8 +1496,13 @@ void FStreamingManager::AddInternal(FRDGBuilder& GraphBuilder, FNewSparseVolumeT
 
 				// Update highest mip (1x1x1) in page table
 				const FUpdateTextureRegion3D UpdateRegion(0, 0, 0, 0, 0, 0, 1, 1, 1);
-				RHIUpdateTexture3D(FrameInfo.PageTableTextureRHIRef, FrameInfo.PageTableTextureRHIRef->GetDesc().NumMips - 1, UpdateRegion, sizeof(uint32), sizeof(uint32), (uint8*)&TileCoord);
+				RHIUpdateTexture3D(FrameInfo.PageTableTextureRHIRef, NumMipLevels - 1, UpdateRegion, sizeof(uint32), sizeof(uint32), (uint8*)&TileCoord);
 			}
+
+			const FPageTopology::FMip& TopologyMipInfo = Resources->Topology.MipInfo[NumMipLevels - 1];
+			FrameInfo.ResidentPagesNew.SetRange(TopologyMipInfo.PageOffset, TopologyMipInfo.PageCount, true);
+
+			InvalidatedSVTFrames.Add(&FrameInfo);
 		}
 
 		RootTileUploader.ResourceUploadTo(GraphBuilder, SVTInfo.TileDataTexture->TileDataTextureARHIRef, SVTInfo.TileDataTexture->TileDataTextureBRHIRef, SVTInfo.FallbackValueA, SVTInfo.FallbackValueB);
@@ -1551,6 +1558,7 @@ void FStreamingManager::RemoveInternal(UStreamableSparseVolumeTexture* SparseVol
 		for (FFrameInfo& FrameInfo : SVTInfo->PerFrameInfo)
 		{
 			FrameInfo.PageTableTextureRHIRef.SafeRelease();
+			InvalidatedSVTFrames.Remove(&FrameInfo);
 		}
 		if (SVTInfo->TileDataTexture)
 		{
@@ -1959,6 +1967,7 @@ void FStreamingManager::StreamOutMipLevel(FStreamingInfo* SVTInfo, FLRUNode* LRU
 	// Update the streaming info buffer data
 	SVTInfo->DirtyStreamingInfoData[FrameIndex] = true;
 	SVTsWithInvalidatedStreamingInfoBuffer.Add(SVTInfo);
+	InvalidatedSVTFrames.Add(&FrameInfo);
 
 	// Unlink
 	LRUNode->Remove();
@@ -1975,6 +1984,9 @@ void FStreamingManager::StreamOutMipLevel(FStreamingInfo* SVTInfo, FLRUNode* LRU
 		SVTInfo->TileDataTexture->Free(TileCoord);
 		TileCoord = 0;
 	}
+
+	const FPageTopology::FMip& TopologyMipInfo = FrameInfo.Resources->Topology.MipInfo[MipLevelIndex];
+	FrameInfo.ResidentPagesNew.SetRange(TopologyMipInfo.PageOffset, TopologyMipInfo.PageCount, false);
 }
 
 int32 FStreamingManager::DetermineReadyMipLevels()
@@ -2157,10 +2169,6 @@ void FStreamingManager::InstallReadyMipLevels()
 
 		FTileUploader::FAddResult TileDataAddResult = SVTInfo->TileDataTexture->TileUploader->Add_GetRef(NumPhysicalTiles, NumVoxelsA, NumVoxelsB);
 
-		uint8* DstPageCoords = nullptr;
-		uint8* DstPageEntries = nullptr;
-		PageTableUpdater->Add_GetRef(FrameInfo.PageTableTextureRHIRef, PendingMipLevel.MipLevelIndex, NumPageTableUpdates, DstPageCoords, DstPageEntries);
-
 		// Tile data
 		{
 			FUploadTask::FTileDataTask TileDataTask = {};
@@ -2187,8 +2195,6 @@ void FStreamingManager::InstallReadyMipLevels()
 		{
 			FUploadTask::FPageTableTask PageTableTask = {};
 			PageTableTask.PendingMipLevel = &PendingMipLevel;
-			PageTableTask.DstPageCoords = DstPageCoords;
-			PageTableTask.DstPageEntries = DstPageEntries;
 			PageTableTask.SrcPageCoords = SrcPtr + MipLevelStreamingInfo.PageTableOffset;
 			PageTableTask.SrcPageEntries = SrcPtr + MipLevelStreamingInfo.PageTableOffset + NumPageTableUpdates * sizeof(uint32);
 			PageTableTask.NumPageTableUpdates = NumPageTableUpdates;
@@ -2216,9 +2222,13 @@ void FStreamingManager::InstallReadyMipLevels()
 		// Update the streaming info buffer data
 		SVTInfo->DirtyStreamingInfoData[PendingMipLevel.FrameIndex] = true;
 		SVTsWithInvalidatedStreamingInfoBuffer.Add(SVTInfo);
+		InvalidatedSVTFrames.Add(&FrameInfo);
 
 		const int32 LRUNodeIndex = PendingMipLevel.FrameIndex * SVTInfo->NumMipLevelsGlobal + PendingMipLevel.MipLevelIndex;
 		SVTInfo->LRUNodes[LRUNodeIndex].PendingMipLevelIndex = INDEX_NONE;
+
+		const FPageTopology::FMip& TopologyMipInfo = FrameInfo.Resources->Topology.MipInfo[PendingMipLevel.MipLevelIndex];
+		FrameInfo.ResidentPagesNew.SetRange(TopologyMipInfo.PageOffset, TopologyMipInfo.PageCount, true);
 	}
 
 	// Do all the memcpy's in parallel
@@ -2234,15 +2244,18 @@ void FStreamingManager::InstallReadyMipLevels()
 				FUploadTask::FPageTableTask& PageTableTask = Task.Union.GetSubtype<FUploadTask::FPageTableTask>();
 				if (PageTableTask.NumPageTableUpdates > 0)
 				{
-					FMemory::Memcpy(PageTableTask.DstPageCoords, PageTableTask.SrcPageCoords, PageTableTask.NumPageTableUpdates * sizeof(uint32));
-
+					const int32 FrameIndex = PageTableTask.PendingMipLevel->FrameIndex;
+					const int32 MipLevelIndex = PageTableTask.PendingMipLevel->MipLevelIndex;
 					FStreamingInfo* SVTInfo = FindStreamingInfo(PageTableTask.PendingMipLevel->SparseVolumeTexture);
-					TArray<uint32>& TileAllocations = SVTInfo->PerFrameInfo[PageTableTask.PendingMipLevel->FrameIndex].TileAllocations[PageTableTask.PendingMipLevel->MipLevelIndex];
+					FFrameInfo& FrameInfo = SVTInfo->PerFrameInfo[FrameIndex];
+					TArray<uint32>& TileAllocations = FrameInfo.TileAllocations[MipLevelIndex];
 					const uint32* SrcEntries = reinterpret_cast<const uint32*>(PageTableTask.SrcPageEntries);
-					uint32* DstEntries = reinterpret_cast<uint32*>(PageTableTask.DstPageEntries);
+					uint32* EntriesForBookKeeping = FrameInfo.PageEntries.GetData() + FrameInfo.Resources->Topology.MipInfo[MipLevelIndex].PageOffset;
+					const uint32 MipLevelEntry = static_cast<uint32>(MipLevelIndex) << 24u;					
 					for (int32 i = 0; i < PageTableTask.NumPageTableUpdates; ++i)
 					{
-						DstEntries[i] = TileAllocations[SrcEntries[i]];
+						const uint32 Entry = TileAllocations[SrcEntries[i]] | MipLevelEntry;
+						EntriesForBookKeeping[i] = Entry;
 					}
 				}
 			}
@@ -2310,6 +2323,146 @@ void FStreamingManager::InstallReadyMipLevels()
 #endif
 }
 
+void FStreamingManager::PatchPageTable(FRDGBuilder& GraphBuilder)
+{
+	int32 NumUpdates = 0;
+
+	// Generate bitsets of invalidated pages for every frame.
+	for (FFrameInfo* FrameInfoPtr : InvalidatedSVTFrames)
+	{
+		FFrameInfo& FrameInfo = *FrameInfoPtr;
+		const FPageTopology& Topology = FrameInfo.Resources->Topology;
+
+		// Get all pages that were streamed in or out in this streaming update.
+		const TBitArray<> ResidentPagesDiff = TBitArray<>::BitwiseXOR(FrameInfo.ResidentPages, FrameInfo.ResidentPagesNew, EBitwiseOperatorFlags::MaxSize);
+		
+		// Initialize InvalidatedPages with the diff. In the following loop, we then find all descendants of the newly streamed in/out pages and also mark them as invalidated.
+		FrameInfo.InvalidatedPages = ResidentPagesDiff;
+		for (TConstSetBitIterator It(ResidentPagesDiff); It; ++It)
+		{
+			const int32 PageIndex = It.GetIndex();
+			check(Topology.Pages.IsValidIndex(PageIndex));
+
+			auto InvalidateDescendants = [&](uint32 InPageIndex, auto& InRecursiveLambda) -> void
+			{
+				// Mip0 pages are leaf-nodes in the topology and don't have children
+				const bool bIsInteriorPage = Topology.InteriorPageData.IsValidIndex(InPageIndex);
+				if (bIsInteriorPage)
+				{
+					for (uint32 ChildPageIndex : Topology.InteriorPageData[InPageIndex].ChildIndices)
+					{
+						// Only process child if it is not already resident or included in the set of invalidated pages
+						if (ChildPageIndex != INDEX_NONE && !FrameInfo.InvalidatedPages[ChildPageIndex] && !FrameInfo.ResidentPages[ChildPageIndex])
+						{
+							FrameInfo.InvalidatedPages[ChildPageIndex] = true;
+							InRecursiveLambda(ChildPageIndex, InRecursiveLambda); // Let the child node walk its own children
+						}
+					}
+				}
+			};
+
+			InvalidateDescendants(PageIndex, InvalidateDescendants);
+		}
+
+		NumUpdates += FrameInfo.InvalidatedPages.CountSetBits();
+	}
+
+	if (NumUpdates > 0)
+	{
+		PageTableUpdater->Init(GraphBuilder, NumUpdates, 0);
+
+		// Generate updates
+		for (FFrameInfo* FrameInfoPtr : InvalidatedSVTFrames)
+		{
+			FFrameInfo& FrameInfo = *FrameInfoPtr;
+			const FPageTopology& Topology = FrameInfo.Resources->Topology;
+			
+			// This set of variables is updated every time we start processing a new mip level
+			bool bEnteredNewMipRange = true;
+			int32 MipLevel = FrameInfo.NumMipLevels - 1;
+			int32 NumUpdatesThisMip = 0;
+			uint32 MipUpdateWriteIndex = 0;
+			uint8* DstCoordsPtr = nullptr;
+			uint8* DstEntryPtr = nullptr;
+
+			// Iterate over all invalidated pages and generate page table updates (packed page write coord and data to write to that coord).
+			for (TConstSetBitIterator It(FrameInfo.InvalidatedPages); It; ++It)
+			{
+				const int32 PageIndex = It.GetIndex();
+				check(Topology.Pages.IsValidIndex(PageIndex));
+
+				auto IsInMipRange = [](const FPageTopology& InTopology, uint32 InIndex, int32 InMipLevel)
+				{
+					return InIndex >= InTopology.MipInfo[InMipLevel].PageOffset && InIndex < (InTopology.MipInfo[InMipLevel].PageOffset + InTopology.MipInfo[InMipLevel].PageCount);
+				};
+
+				// Determine the current mip level. Bits are ordered highest to lowest mip level.
+				while (MipLevel > 0 && !IsInMipRange(Topology, PageIndex, MipLevel))
+				{
+					bEnteredNewMipRange = true;
+					--MipLevel;
+					check(MipLevel >= 0);
+				}
+				check(IsInMipRange(Topology, PageIndex, MipLevel));
+
+				// If we entered a new mip range, get a new set of write pointers from the PageTableUpdater.
+				if (bEnteredNewMipRange)
+				{
+					check(NumUpdatesThisMip == MipUpdateWriteIndex);
+					const uint32 PageOffset = Topology.MipInfo[MipLevel].PageOffset;
+					const uint32 PageCount = Topology.MipInfo[MipLevel].PageCount;
+					NumUpdatesThisMip = FrameInfo.InvalidatedPages.CountSetBits(PageOffset, PageOffset + PageCount);
+					MipUpdateWriteIndex = 0;
+					bEnteredNewMipRange = false;
+
+					PageTableUpdater->Add_GetRef(FrameInfo.PageTableTextureRHIRef, MipLevel, NumUpdatesThisMip, DstCoordsPtr, DstEntryPtr);
+				}
+
+				uint32 PageTableEntry = 0;
+				if (FrameInfo.ResidentPagesNew[PageIndex])
+				{
+					// This page is already resident, so we can simply use its value in PageEntries
+					PageTableEntry = FrameInfo.PageEntries[PageIndex];
+					check(PageTableEntry);
+					PageTableEntry |= MipLevel << 24u;
+				}
+				else
+				{
+					// This page is not resident but needs a fallback value written to it, so we probe the parent pages until we find a resident one
+					uint32 ParentPageIndex = Topology.Pages[PageIndex].ParentIndex;
+					int32 ParentMipLevel = MipLevel + 1;
+					while (ParentPageIndex != INDEX_NONE)
+					{
+						// The parent page is resident in GPU memory, so we can use it's cached value in PageEntries.
+						if (FrameInfo.ResidentPagesNew[ParentPageIndex])
+						{
+							PageTableEntry = FrameInfo.PageEntries[ParentPageIndex];
+							check(PageTableEntry);
+							PageTableEntry |= ParentMipLevel << 24u;
+							FrameInfo.PageEntries[PageIndex] = PageTableEntry; // Don't forget to update the PageEntries value of our current page to reflect that it falls back to a coarser mip.
+							break;
+						}
+						ParentPageIndex = Topology.Pages[ParentPageIndex].ParentIndex;
+						++ParentMipLevel;
+					}
+					check(ParentPageIndex != INDEX_NONE); // If we hit this, then we tried to find the root node's parent. This should never happen as the root node should always be resident.
+				}
+
+				// Write the update to the upload buffer pointers
+				reinterpret_cast<uint32*>(DstCoordsPtr)[MipUpdateWriteIndex] = Topology.Pages[PageIndex].PackedPageTableCoord;
+				reinterpret_cast<uint32*>(DstEntryPtr)[MipUpdateWriteIndex] = PageTableEntry;
+				++MipUpdateWriteIndex;
+			}
+
+			FrameInfo.ResidentPages = FrameInfo.ResidentPagesNew;
+		}
+
+		PageTableUpdater->Apply(GraphBuilder);
+	}
+
+	InvalidatedSVTFrames.Reset();
+}
+
 FStreamingManager::FStreamingInfo* FStreamingManager::FindStreamingInfo(UStreamableSparseVolumeTexture* Key)
 {
 	TUniquePtr<FStreamingInfo>* SVTInfoPtr = StreamingInfo.Find(Key);
diff --git a/Engine/Source/Runtime/Engine/Private/SparseVolumeTexture/SparseVolumeTextureStreamingManager.h b/Engine/Source/Runtime/Engine/Private/SparseVolumeTexture/SparseVolumeTextureStreamingManager.h
--- a/Engine/Source/Runtime/Engine/Private/SparseVolumeTexture/SparseVolumeTextureStreamingManager.h
+++ b/Engine/Source/Runtime/Engine/Private/SparseVolumeTexture/SparseVolumeTextureStreamingManager.h
@@ -185,6 +185,10 @@ private:
 		int32 LowestRequestedMipLevel; // Lowest mip level that should be resident. Can be lower than LowestResidentMipLevel when streaming in new mips. Stream-out is instant, so it should never be higher.
 		int32 LowestResidentMipLevel; // Actually resident on the GPU
 		TArray<TArray<uint32>> TileAllocations; // TileAllocations[MipLevel][PhysicalTileIndex]
+		TArray<uint32> PageEntries; // Flat array of sparse page table. Reflects the current state of the GPU page table
+		TBitArray<> ResidentPages; //  One bit for every (non-zero) page (in the sparse page octree) for all mip levels, starting at the highest mip
+		TBitArray<> ResidentPagesNew; // Reflects changes made during an update of the streaming system
+		TBitArray<> InvalidatedPages; // All pages that need updated page table entries
 		FTextureRHIRef PageTableTextureRHIRef;
 	};
 
@@ -284,8 +288,6 @@ private:
 		struct FPageTableTask
 		{
 			FPendingMipLevel* PendingMipLevel;
-			uint8* DstPageCoords;
-			uint8* DstPageEntries;
 			const uint8* SrcPageCoords;
 			const uint8* SrcPageEntries;
 			int32 NumPageTableUpdates;
@@ -340,6 +342,7 @@ private:
 	TArray<FStreamingRequest> ParentRequestsToAdd;
 	TSet<FTileDataTexture*> TileDataTexturesToUpdate;
 	TSet<FStreamingInfo*> SVTsWithInvalidatedStreamingInfoBuffer; // Changes in the lowest resident mip level cause invalidation
+	TSet<FFrameInfo*> InvalidatedSVTFrames; // Set of SVT frames where pages have been streamed in or out. Used in PatchPageTable().
 	TArray<FStreamingRequest> PrioritizedRequestsHeap;
 	TArray<FStreamingRequest> SelectedRequests;
 	TArray<FPageTableClear> PageTableClears;
@@ -356,6 +359,7 @@ private:
 	void StreamOutMipLevel(FStreamingInfo* SVTInfo, FLRUNode* LRUNode);
 	int32 DetermineReadyMipLevels();
 	void InstallReadyMipLevels();
+	void PatchPageTable(FRDGBuilder& GraphBuilder); // Patches the page table to reflect streamed in/out pages and to ensure non-resident mip levels fall back to coarser mip level tile data
 	FStreamingInfo* FindStreamingInfo(UStreamableSparseVolumeTexture* Key); // Returns nullptr if the key can't be found
 
 #if WITH_EDITORONLY_DATA
diff --git a/Engine/Source/Runtime/Engine/Public/SparseVolumeTexture/SparseVolumeTextureData.h b/Engine/Source/Runtime/Engine/Public/SparseVolumeTexture/SparseVolumeTextureData.h
--- a/Engine/Source/Runtime/Engine/Public/SparseVolumeTexture/SparseVolumeTextureData.h
+++ b/Engine/Source/Runtime/Engine/Public/SparseVolumeTexture/SparseVolumeTextureData.h
@@ -68,6 +68,11 @@ private:
 	void WriteTileDataVoxel(int32 TileIndex, const FIntVector3& TileDataCoord, int32 MipLevel, int32 AttributesIdx, const FVector4f& Value, int32 DstComponent = -1);
 };
 
+// Returns true if Point >= Min && Point < Max for all components
+bool IsInBounds(const FIntVector3& Point, const FIntVector3& Min, const FIntVector3& Max);
+// Computes Max(Value >> ShiftBy, MinValue)
+FIntVector3 ShiftRightAndMax(const FIntVector3& Value, uint32 ShiftBy, int32 MinValue);
+
 }
 }
 
