Truncation Warning: Small fix change from double literal to float notation for FLinearColor

diff --git a/Engine/Source/Runtime/Engine/Classes/Kismet/KismetSystemLibrary.h b/Engine/Source/Runtime/Engine/Classes/Kismet/KismetSystemLibrary.h
--- a/Engine/Source/Runtime/Engine/Classes/Kismet/KismetSystemLibrary.h
+++ b/Engine/Source/Runtime/Engine/Classes/Kismet/KismetSystemLibrary.h
@@ -490,7 +490,7 @@ class UKismetSystemLibrary : public UBlueprintFunctionLibrary
 	 * @param	Key				If a non-empty key is provided, the message will replace any existing on-screen messages with the same key.
 	 */
 	UFUNCTION(BlueprintCallable, meta=(WorldContext="WorldContextObject", CallableWithoutWorldContext, Keywords = "log print", AdvancedDisplay = "2", DevelopmentOnly), Category="Development")
-	static ENGINE_API void PrintString(const UObject* WorldContextObject, const FString& InString = FString(TEXT("Hello")), bool bPrintToScreen = true, bool bPrintToLog = true, FLinearColor TextColor = FLinearColor(0.0, 0.66, 1.0), float Duration = 2.f, const FName Key = NAME_None);
+	static ENGINE_API void PrintString(const UObject* WorldContextObject, const FString& InString = FString(TEXT("Hello")), bool bPrintToScreen = true, bool bPrintToLog = true, FLinearColor TextColor = FLinearColor(0.0f, 0.66f, 1.0f), float Duration = 2.f, const FName Key = NAME_None);
 
 	/**
 	 * Prints text to the log, and optionally, to the screen
@@ -505,7 +505,7 @@ class UKismetSystemLibrary : public UBlueprintFunctionLibrary
 	 * @param	Key				If a non-empty key is provided, the message will replace any existing on-screen messages with the same key.
 	 */
 	UFUNCTION(BlueprintCallable, meta=(WorldContext="WorldContextObject", CallableWithoutWorldContext, Keywords = "log", AdvancedDisplay = "2", DevelopmentOnly), Category="Development")
-	static ENGINE_API void PrintText(const UObject* WorldContextObject, const FText InText = INVTEXT("Hello"), bool bPrintToScreen = true, bool bPrintToLog = true, FLinearColor TextColor = FLinearColor(0.0, 0.66, 1.0), float Duration = 2.f, const FName Key = NAME_None);
+	static ENGINE_API void PrintText(const UObject* WorldContextObject, const FText InText = INVTEXT("Hello"), bool bPrintToScreen = true, bool bPrintToLog = true, FLinearColor TextColor = FLinearColor(0.0f, 0.66f, 1.0f), float Duration = 2.f, const FName Key = NAME_None);
 
 	/**
 	 * Prints a warning string to the log and the screen. Meant to be used as a way to inform the user that they misused the node.
