World Partition - HLOD: Fix SpatialHash HLOD generation not using the same hash settings as the streaming generation.

diff --git a/Engine/Source/Runtime/Engine/Private/WorldPartition/RuntimeSpatialHash/RuntimeSpatialHashHLOD.cpp b/Engine/Source/Runtime/Engine/Private/WorldPartition/RuntimeSpatialHash/RuntimeSpatialHashHLOD.cpp
--- a/Engine/Source/Runtime/Engine/Private/WorldPartition/RuntimeSpatialHash/RuntimeSpatialHashHLOD.cpp
+++ b/Engine/Source/Runtime/Engine/Private/WorldPartition/RuntimeSpatialHash/RuntimeSpatialHashHLOD.cpp
@@ -499,11 +499,14 @@ bool UWorldPartitionRuntimeSpatialHash::SetupHLODActors(const IStreamingGenerati
 
 	const UDataLayerManager* DataLayerManager = WorldPartition->GetDataLayerManager();
 
-	auto GenerateHLODActors = [&GridsHLODActors, &MainActorSetContainer, StreamingGenerationContext, &Params, WorldPartition, DataLayerManager, &Context, SourceControlHelper, &NewActors, &DirtyActors, this]
-	(const FSpatialHashRuntimeGrid& RuntimeGrid, uint32 HLODLevel, const TArray<const IStreamingGenerationContext::FActorSetInstance*>& ActorSetInstances)
+	// Make sure hash settings are up to date
+	FSpatialHashSettings HashSettings = Settings;	
+	HashSettings.UpdateSettings(*this);
+
+	auto GenerateHLODActors = [&](const FSpatialHashRuntimeGrid& RuntimeGrid, uint32 HLODLevel, const TArray<const IStreamingGenerationContext::FActorSetInstance*>& ActorSetInstances)
 	{
 		// Generate HLODs for this grid
-		TArray<FGuid> HLODActors = GenerateHLODActorsForGrid(WorldPartition, StreamingGenerationContext, Params, RuntimeGrid, HLODLevel, Context, SourceControlHelper, ActorSetInstances, NewActors, DirtyActors, Settings);
+		TArray<FGuid> HLODActors = GenerateHLODActorsForGrid(WorldPartition, StreamingGenerationContext, Params, RuntimeGrid, HLODLevel, Context, SourceControlHelper, ActorSetInstances, NewActors, DirtyActors, HashSettings);
 
 		for (const FGuid& HLODActorGuid : HLODActors)
 		{
diff --git a/Engine/Source/Runtime/Engine/Private/WorldPartition/WorldPartitionRuntimeSpatialHash.cpp b/Engine/Source/Runtime/Engine/Private/WorldPartition/WorldPartitionRuntimeSpatialHash.cpp
--- a/Engine/Source/Runtime/Engine/Private/WorldPartition/WorldPartitionRuntimeSpatialHash.cpp
+++ b/Engine/Source/Runtime/Engine/Private/WorldPartition/WorldPartitionRuntimeSpatialHash.cpp
@@ -1191,15 +1191,17 @@ void UWorldPartitionRuntimeSpatialHash::SetDefaultValues()
 	MainGrid.DebugColor = FLinearColor::Gray;
 }
 
-void UWorldPartitionRuntimeSpatialHash::UpdateSettings()
+void FSpatialHashSettings::UpdateSettings(const UWorldPartitionRuntimeSpatialHash& RuntimeSpatialHash)
 {
-	check(StreamingGrids.IsEmpty());
-	Settings.bUseAlignedGridLevels = (UseAlignedGridLevels == EWorldPartitionCVarProjectDefaultOverride::ProjectDefault) ? GRuntimeSpatialHashUseAlignedGridLevels : (UseAlignedGridLevels == EWorldPartitionCVarProjectDefaultOverride::Enabled);
-	Settings.bSnapNonAlignedGridLevelsToLowerLevels = (SnapNonAlignedGridLevelsToLowerLevels == EWorldPartitionCVarProjectDefaultOverride::ProjectDefault) ? GRuntimeSpatialHashSnapNonAlignedGridLevelsToLowerLevels : (SnapNonAlignedGridLevelsToLowerLevels == EWorldPartitionCVarProjectDefaultOverride::Enabled);
-	Settings.bPlaceSmallActorsUsingLocation = (PlaceSmallActorsUsingLocation == EWorldPartitionCVarProjectDefaultOverride::ProjectDefault) ? GRuntimeSpatialHashPlaceSmallActorsUsingLocation : (PlaceSmallActorsUsingLocation == EWorldPartitionCVarProjectDefaultOverride::Enabled);
-	Settings.bPlacePartitionActorsUsingLocation = (PlacePartitionActorsUsingLocation == EWorldPartitionCVarProjectDefaultOverride::ProjectDefault) ? GRuntimeSpatialHashPlacePartitionActorsUsingLocation : (PlacePartitionActorsUsingLocation == EWorldPartitionCVarProjectDefaultOverride::Enabled);
-	UE_LOG(LogWorldPartition, Log, TEXT("UWorldPartitionRuntimeSpatialHash::UpdateSettings : UseAlignedGridLevels = %d, SnapNonAlignedGridLevelsToLowerLevels = %d, PlaceSmallActorsUsingLocation = %d, PlacePartitionActorsUsingLocation = %d"),
-		Settings.bUseAlignedGridLevels, Settings.bSnapNonAlignedGridLevelsToLowerLevels, Settings.bPlaceSmallActorsUsingLocation, Settings.bPlacePartitionActorsUsingLocation);
+#if WITH_EDITORONLY_DATA
+	check(RuntimeSpatialHash.StreamingGrids.IsEmpty());
+	bUseAlignedGridLevels = (RuntimeSpatialHash.UseAlignedGridLevels == EWorldPartitionCVarProjectDefaultOverride::ProjectDefault) ? GRuntimeSpatialHashUseAlignedGridLevels : (RuntimeSpatialHash.UseAlignedGridLevels == EWorldPartitionCVarProjectDefaultOverride::Enabled);
+	bSnapNonAlignedGridLevelsToLowerLevels = (RuntimeSpatialHash.SnapNonAlignedGridLevelsToLowerLevels == EWorldPartitionCVarProjectDefaultOverride::ProjectDefault) ? GRuntimeSpatialHashSnapNonAlignedGridLevelsToLowerLevels : (RuntimeSpatialHash.SnapNonAlignedGridLevelsToLowerLevels == EWorldPartitionCVarProjectDefaultOverride::Enabled);
+	bPlaceSmallActorsUsingLocation = (RuntimeSpatialHash.PlaceSmallActorsUsingLocation == EWorldPartitionCVarProjectDefaultOverride::ProjectDefault) ? GRuntimeSpatialHashPlaceSmallActorsUsingLocation : (RuntimeSpatialHash.PlaceSmallActorsUsingLocation == EWorldPartitionCVarProjectDefaultOverride::Enabled);
+	bPlacePartitionActorsUsingLocation = (RuntimeSpatialHash.PlacePartitionActorsUsingLocation == EWorldPartitionCVarProjectDefaultOverride::ProjectDefault) ? GRuntimeSpatialHashPlacePartitionActorsUsingLocation : (RuntimeSpatialHash.PlacePartitionActorsUsingLocation == EWorldPartitionCVarProjectDefaultOverride::Enabled);
+	UE_LOG(LogWorldPartition, Log, TEXT("FSpatialHashSettings::UpdateSettings : UseAlignedGridLevels = %d, SnapNonAlignedGridLevelsToLowerLevels = %d, PlaceSmallActorsUsingLocation = %d, PlacePartitionActorsUsingLocation = %d"),
+		bUseAlignedGridLevels, bSnapNonAlignedGridLevelsToLowerLevels, bPlaceSmallActorsUsingLocation, bPlacePartitionActorsUsingLocation);
+#endif
 }
 
 bool UWorldPartitionRuntimeSpatialHash::GenerateStreaming(UWorldPartitionStreamingPolicy* StreamingPolicy, const IStreamingGenerationContext* StreamingGenerationContext, TArray<FString>* OutPackagesToGenerate)
@@ -1220,7 +1222,7 @@ bool UWorldPartitionRuntimeSpatialHash::GenerateStreaming(UWorldPartitionStreami
 	StreamingGrids.Empty();
 	
 	// Apply Settings (should no longer change after streaming generation)
-	UpdateSettings();
+	Settings.UpdateSettings(*this);
 
 	// Append grids from ASpatialHashRuntimeGridInfo actors to runtime spatial hash grids
 	TArray<FSpatialHashRuntimeGrid> AllGrids;
diff --git a/Engine/Source/Runtime/Engine/Public/WorldPartition/WorldPartitionRuntimeSpatialHash.h b/Engine/Source/Runtime/Engine/Public/WorldPartition/WorldPartitionRuntimeSpatialHash.h
--- a/Engine/Source/Runtime/Engine/Public/WorldPartition/WorldPartitionRuntimeSpatialHash.h
+++ b/Engine/Source/Runtime/Engine/Public/WorldPartition/WorldPartitionRuntimeSpatialHash.h
@@ -43,6 +43,8 @@ struct FSpatialHashSettings
 
 	friend bool operator==(const FSpatialHashSettings& Lhs, const FSpatialHashSettings& Rhs) = default;
 	friend bool operator!=(const FSpatialHashSettings& Lhs, const FSpatialHashSettings& Rhs) = default;
+
+	void UpdateSettings(const UWorldPartitionRuntimeSpatialHash& RuntimeSpatialHash);
 	
 	UPROPERTY();
 	bool bUseAlignedGridLevels;
@@ -416,11 +418,11 @@ private:
 	ENGINE_API void GetAlwaysLoadedStreamingCells(const FSpatialHashStreamingGrid& StreamingGrid, TSet<const UWorldPartitionRuntimeCell*>& Cells) const;
 	ENGINE_API const TMap<FName, const FSpatialHashStreamingGrid*>& GetNameToGridMapping() const;
 #if WITH_EDITOR
-	ENGINE_API void UpdateSettings();
 	ENGINE_API bool CreateStreamingGrid(const FSpatialHashRuntimeGrid& RuntimeGrid, const FSquare2DGridHelper& PartionedActors, UWorldPartitionStreamingPolicy* StreamingPolicy, TArray<FString>* OutPackagesToGenerate = nullptr);
 #endif
 	ENGINE_API TArray<const FSpatialHashStreamingGrid*> GetFilteredStreamingGrids() const;
 
 	friend class UWorldPartitionSubsystem;
 	friend class UWorldPartitionRuntimeHashSet;
+	friend struct FSpatialHashSettings;
 };
