Fix inability to cook spline meshes by a cooked cooker.

Cooked cooker is a relatively new platform in Unreal that allows producing a cooked editor build and using it to cook new assets (e.g. DLC for the base game). Such cooked cooker build is faster to load than the original editor build that would cook from the source assets, and it can also be more easily deployed to cloud environments.

The problem this fix is solving is as follows:

1) Spline mesh needs collision data from the static mesh it deforms. This collision data is usually derived from its render data.
2) Cooked cooker target platform (see TCookedCookerTargetPlatform) is basically a server build, assets cooked for this platform do not store render data.
3) If a new (uncooked) spline mesh is referencing a static mesh from the base build, it fails to get its collision data because the render data was stripped when producing the cooked version of the base build.

The solution is

- to introduce a new target platform feature that allows to distinguish cooked cooker target platform from other platforms that strip audio-visual data.
- when cooking static meshes for a "cooked cooker" target get a separate copy of the asset's collision data and save it inline in the cooked asset.
- inside the cooked cooker, use this copy as static mesh's collision data instead of trying to derive it from (non-existent at that point) static mesh's render data.

diff --git a/Engine/Source/Developer/CookedEditor/Public/CookedEditorTargetPlatform.h b/Engine/Source/Developer/CookedEditor/Public/CookedEditorTargetPlatform.h
--- a/Engine/Source/Developer/CookedEditor/Public/CookedEditorTargetPlatform.h
+++ b/Engine/Source/Developer/CookedEditor/Public/CookedEditorTargetPlatform.h
@@ -213,6 +213,7 @@ public:
 		case ETargetPlatformFeatures::MemoryMappedAudio:
 			return false;
 
+		case ETargetPlatformFeatures::CanCookPackages:
 		case ETargetPlatformFeatures::Packaging:
 			return true;
 
diff --git a/Engine/Source/Developer/Linux/LinuxTargetPlatform/Public/LinuxTargetPlatform.h b/Engine/Source/Developer/Linux/LinuxTargetPlatform/Public/LinuxTargetPlatform.h
--- a/Engine/Source/Developer/Linux/LinuxTargetPlatform/Public/LinuxTargetPlatform.h
+++ b/Engine/Source/Developer/Linux/LinuxTargetPlatform/Public/LinuxTargetPlatform.h
@@ -198,6 +198,11 @@ public:
 			return true;
 		}
 
+		if (Feature == ETargetPlatformFeatures::CanCookPackages)
+		{
+			return TProperties::HasEditorOnlyData();
+		}
+
 		return TSuper::SupportsFeature(Feature);
 	}
 
diff --git a/Engine/Source/Developer/Mac/MacTargetPlatform/Public/GenericMacTargetPlatform.h b/Engine/Source/Developer/Mac/MacTargetPlatform/Public/GenericMacTargetPlatform.h
--- a/Engine/Source/Developer/Mac/MacTargetPlatform/Public/GenericMacTargetPlatform.h
+++ b/Engine/Source/Developer/Mac/MacTargetPlatform/Public/GenericMacTargetPlatform.h
@@ -105,6 +105,11 @@ public:
 			return (HAS_EDITOR_DATA || !IS_DEDICATED_SERVER);
 		}
 
+		if (Feature == ETargetPlatformFeatures::CanCookPackages)
+		{
+			return HAS_EDITOR_DATA;
+		}
+
 		return TSuper::SupportsFeature(Feature);
 	}
 
diff --git a/Engine/Source/Developer/TargetPlatform/Public/Common/TargetPlatformBase.h b/Engine/Source/Developer/TargetPlatform/Public/Common/TargetPlatformBase.h
--- a/Engine/Source/Developer/TargetPlatform/Public/Common/TargetPlatformBase.h
+++ b/Engine/Source/Developer/TargetPlatform/Public/Common/TargetPlatformBase.h
@@ -472,6 +472,9 @@ public:
 		case ETargetPlatformFeatures::Packaging:
 			return false;
 
+		case ETargetPlatformFeatures::CanCookPackages:
+			return false;
+
 		case ETargetPlatformFeatures::TextureStreaming:
 			return TPlatformProperties::SupportsTextureStreaming();
 		case ETargetPlatformFeatures::MeshLODStreaming:
diff --git a/Engine/Source/Developer/TargetPlatform/Public/Interfaces/ITargetPlatform.h b/Engine/Source/Developer/TargetPlatform/Public/Interfaces/ITargetPlatform.h
--- a/Engine/Source/Developer/TargetPlatform/Public/Interfaces/ITargetPlatform.h
+++ b/Engine/Source/Developer/TargetPlatform/Public/Interfaces/ITargetPlatform.h
@@ -118,7 +118,10 @@ enum class ETargetPlatformFeatures
 	ShowAsPlatformGroup,
 
 	/** Does the platform allow various connection types to be used (ie: wifi and usb) */
-	SupportsMultipleConnectionTypes
+	SupportsMultipleConnectionTypes,
+
+	/** The platform can cook packages (e.g. CookedCooker) */
+	CanCookPackages
 };
 
 enum class EPlatformAuthentication
diff --git a/Engine/Source/Developer/Windows/WindowsTargetPlatform/Public/GenericWindowsTargetPlatform.h b/Engine/Source/Developer/Windows/WindowsTargetPlatform/Public/GenericWindowsTargetPlatform.h
--- a/Engine/Source/Developer/Windows/WindowsTargetPlatform/Public/GenericWindowsTargetPlatform.h
+++ b/Engine/Source/Developer/Windows/WindowsTargetPlatform/Public/GenericWindowsTargetPlatform.h
@@ -165,6 +165,11 @@ public:
 			return TProperties::IsClientOnly();
 		}
 
+		if (Feature == ETargetPlatformFeatures::CanCookPackages)
+		{
+			return TProperties::HasEditorOnlyData();
+		}
+
 		if (Feature == ETargetPlatformFeatures::MobileRendering)
 		{
 			static bool bCachedSupportsMobileRendering = false;
diff --git a/Engine/Source/Runtime/Engine/Private/EngineUtils.cpp b/Engine/Source/Runtime/Engine/Private/EngineUtils.cpp
--- a/Engine/Source/Runtime/Engine/Private/EngineUtils.cpp
+++ b/Engine/Source/Runtime/Engine/Private/EngineUtils.cpp
@@ -595,6 +595,7 @@ FStripDataFlags::FStripDataFlags( class FArchive& Ar, uint8 InClassFlags /*= 0*/
 			// platform's properties.
 			GlobalStripFlags |= Ar.CookingTarget()->HasEditorOnlyData() ? static_cast<uint8>(FStripDataFlags::EStrippedData::None) : static_cast<uint8>(FStripDataFlags::EStrippedData::EditorOnly);
 			GlobalStripFlags |= Ar.CookingTarget()->AllowAudioVisualData() ? static_cast<uint8>(FStripDataFlags::EStrippedData::None) : static_cast<uint8>(FStripDataFlags::EStrippedData::AudioVisual);
+			GlobalStripFlags |= Ar.CookingTarget()->SupportsFeature(ETargetPlatformFeatures::CanCookPackages) ? static_cast<uint8>(FStripDataFlags::EStrippedData::None) : static_cast<uint8>(FStripDataFlags::EStrippedData::NeededForCooking);
 			ClassStripFlags = InClassFlags;
 		}
 		Ar << GlobalStripFlags;
@@ -639,6 +640,7 @@ FStripDataFlags::FStripDataFlags(FStructuredArchive::FSlot Slot, uint8 InClassFl
 			// platform's properties.
 			GlobalStripFlags |= UnderlyingArchive.IsFilterEditorOnly() ? static_cast<uint8>(FStripDataFlags::EStrippedData::EditorOnly) : static_cast<uint8>(FStripDataFlags::EStrippedData::None);
 			GlobalStripFlags |= !UnderlyingArchive.CookingTarget()->AllowAudioVisualData() ? static_cast<uint8>(FStripDataFlags::EStrippedData::AudioVisual) : static_cast<uint8>(FStripDataFlags::EStrippedData::None);
+			GlobalStripFlags |= UnderlyingArchive.CookingTarget()->SupportsFeature(ETargetPlatformFeatures::CanCookPackages) ? static_cast<uint8>(FStripDataFlags::EStrippedData::None) : static_cast<uint8>(FStripDataFlags::EStrippedData::NeededForCooking);
 			ClassStripFlags = InClassFlags;
 		}
 		Record << SA_VALUE(TEXT("GlobalStripFlags"), GlobalStripFlags);
diff --git a/Engine/Source/Runtime/Engine/Private/StaticMesh.cpp b/Engine/Source/Runtime/Engine/Private/StaticMesh.cpp
--- a/Engine/Source/Runtime/Engine/Private/StaticMesh.cpp
+++ b/Engine/Source/Runtime/Engine/Private/StaticMesh.cpp
@@ -1858,6 +1858,36 @@ void FStaticMeshRenderData::Serialize(FArchive& Ar, UStaticMesh* Owner, bool bCo
 			}
 		}
 	}
+
+#if WITH_EDITORONLY_DATA
+	// when cooking for a cooked cooker, we need to save extra data it may later need from its own cooked assets
+	if (Ar.IsCooking())
+	{
+		FStripDataFlags StripFlags(Ar);
+		if (!StripFlags.IsDataNeededForCookingStripped())
+		{
+			// if we need to keep data needed for cooking, just save the collision data
+			UStaticMesh* OwnerStaticMesh = Cast<UStaticMesh>(Owner);
+			check(OwnerStaticMesh);
+
+			FTriMeshCollisionData CollisionData;
+			OwnerStaticMesh->GetPhysicsTriMeshData(&CollisionData, true);
+
+			Ar << CollisionData;
+		}
+	}
+#endif
+	if (bCooked && Ar.IsLoading())
+	{
+		FStripDataFlags StripFlags(Ar);
+#if WITH_EDITORONLY_DATA	// the below lines can only happen for a cooked cooker, which has editor data
+		if (!StripFlags.IsDataNeededForCookingStripped())
+		{
+			CollisionDataForCookedCooker = MakeUnique<FTriMeshCollisionData>();
+			Ar << *CollisionDataForCookedCooker;
+		}
+#endif
+	}
 }
 
 void FStaticMeshRenderData::InitResources(ERHIFeatureLevel::Type InFeatureLevel, UStaticMesh* Owner)
@@ -7380,6 +7410,16 @@ bool UStaticMesh::GetPhysicsTriMeshDataCheckComplex(struct FTriMeshCollisionData
 	bInUseAllTriData = true;
 #endif // #if !WITH_EDITORONLY_DATA
 
+#if WITH_EDITORONLY_DATA
+	// if we're a cooked cooker, just use the canned data
+	if (UNLIKELY(GetRenderData()->CollisionDataForCookedCooker))
+	{
+		*CollisionData = *GetRenderData()->CollisionDataForCookedCooker;
+	}
+	else
+	{
+#endif // #if !WITH_EDITORONLY_DATA
+
 	check(HasValidRenderData());
 
 	// Get the LOD level to use for collision
@@ -7439,7 +7479,10 @@ bool UStaticMesh::GetPhysicsTriMeshDataCheckComplex(struct FTriMeshCollisionData
 		}
 	}
 	CollisionData->bFlipNormals = true;
-	
+
+#if WITH_EDITORONLY_DATA
+	}
+#endif
 	// We only have a valid TriMesh if the CollisionData has vertices AND indices. For meshes with disabled section collision, it
 	// can happen that the indices will be empty, in which case we do not want to consider that as valid trimesh data
 	return CollisionData->Vertices.Num() > 0 && CollisionData->Indices.Num() > 0;
@@ -7499,6 +7542,14 @@ bool UStaticMesh::ContainsPhysicsTriMeshDataCheckComplex(bool bInUseAllTriData,
 	bInUseAllTriData = true;
 #endif // #if !WITH_EDITORONLY_DATA
 	
+#if WITH_EDITORONLY_DATA
+	// if we're a cooked cooker, just use the canned data
+	if (UNLIKELY(GetRenderData()->CollisionDataForCookedCooker))
+	{
+		return !GetRenderData()->CollisionDataForCookedCooker->Vertices.IsEmpty();
+	}
+#endif // #if !WITH_EDITORONLY_DATA
+
 	if(GetRenderData() == nullptr || GetRenderData()->LODResources.Num() == 0)
 	{
 		return false;
diff --git a/Engine/Source/Runtime/Engine/Public/EngineUtils.h b/Engine/Source/Runtime/Engine/Public/EngineUtils.h
--- a/Engine/Source/Runtime/Engine/Public/EngineUtils.h
+++ b/Engine/Source/Runtime/Engine/Public/EngineUtils.h
@@ -837,6 +837,11 @@ public:
 		AudioVisual = 2,
 		Server UE_DEPRECATED(5.4, "Use AudioVisual value instead of Server") = AudioVisual,
 
+		/* This flag means that all data needed to cook packages will be stripped. What it means is specific to 
+		  an asset type, and it might be a subset of both EditorOnly and AudioVisual. This flag is normally set, except
+		  when cooking for a "cooked cooker" target. */
+		NeededForCooking = 4,
+
 		// Add global flags here (up to 8 including the already defined ones).
 
 		/** All flags */
@@ -944,6 +949,16 @@ public:
 		return IsAudioVisualDataStripped();
 	}
 
+	/**
+	 * Checks if FStripDataFlags::NeededForCooking flag is set or not. It should be set for all non-content-worker targets
+	 *
+	 * @return true if FStripDataFlags::NeededForCooking is set, false otherwise.
+	 */
+	bool IsDataNeededForCookingStripped() const
+	{
+		return (GlobalStripFlags & static_cast<uint8>(FStripDataFlags::EStrippedData::NeededForCooking)) != 0;
+	}
+
 	/**
 	 * Checks if user defined flags are set or not.
 	 *
diff --git a/Engine/Source/Runtime/Engine/Public/StaticMeshResources.h b/Engine/Source/Runtime/Engine/Public/StaticMeshResources.h
--- a/Engine/Source/Runtime/Engine/Public/StaticMeshResources.h
+++ b/Engine/Source/Runtime/Engine/Public/StaticMeshResources.h
@@ -40,6 +40,9 @@
 #include "WeightedRandomSampler.h"
 #include "PerPlatformProperties.h"
 #include "RayTracingGeometry.h"
+#if WITH_EDITORONLY_DATA
+#include "Interface_CollisionDataProviderCore.h"
+#endif
 
 class FDistanceFieldVolumeData;
 class UBodySetup;
@@ -720,6 +723,12 @@ public:
 	/** The next cached derived data in the list. */
 	TUniquePtr<class FStaticMeshRenderData> NextCachedRenderData;
 
+	/**
+	 * Canned FTriMeshCollisionData for static meshes cooked for CookedCooker platform (see TCookedCookerTargetPlatform).
+	 * It is needed because the "cooked cooker" can be cooking a new spline mesh that is deforming a cooked static mesh, so it will request the SM's collision data.
+	 */
+	TUniquePtr<FTriMeshCollisionData> CollisionDataForCookedCooker;
+
 	/** Estimate of total compressed size of all rendering data, including Nanite data. */
 	uint64 EstimatedCompressedSize = 0;
 
diff --git a/Engine/Source/Runtime/PhysicsCore/Private/PhysicsCore.cpp b/Engine/Source/Runtime/PhysicsCore/Private/PhysicsCore.cpp
--- a/Engine/Source/Runtime/PhysicsCore/Private/PhysicsCore.cpp
+++ b/Engine/Source/Runtime/PhysicsCore/Private/PhysicsCore.cpp
@@ -1,5 +1,27 @@
 // Copyright Epic Games, Inc. All Rights Reserved.
 
 #include "PhysicsCore.h"
+#include "Interface_CollisionDataProviderCore.h"
 
-DEFINE_LOG_CATEGORY(LogPhysicsCore);
\ No newline at end of file
+DEFINE_LOG_CATEGORY(LogPhysicsCore);
+
+#if WITH_EDITORONLY_DATA
+void operator<<(class FArchive& Ar, FTriIndices& Indices)
+{
+	Ar << Indices.v0 << Indices.v1 << Indices.v2;
+}
+
+void operator<<(class FArchive& Ar, FTriMeshCollisionData& TriMeshCollisionData)
+{
+	Ar << TriMeshCollisionData.Vertices;
+	Ar << TriMeshCollisionData.Indices;
+	Ar << TriMeshCollisionData.MaterialIndices;
+	Ar << TriMeshCollisionData.UVs;
+	Ar << TriMeshCollisionData.Vertices;
+
+	FArchive_Serialize_BitfieldBool(Ar, TriMeshCollisionData.bFlipNormals);
+	FArchive_Serialize_BitfieldBool(Ar, TriMeshCollisionData.bDeformableMesh);
+	FArchive_Serialize_BitfieldBool(Ar, TriMeshCollisionData.bFastCook);
+	FArchive_Serialize_BitfieldBool(Ar, TriMeshCollisionData.bDisableActiveEdgePrecompute);
+}
+#endif
diff --git a/Engine/Source/Runtime/PhysicsCore/Public/Interface_CollisionDataProviderCore.h b/Engine/Source/Runtime/PhysicsCore/Public/Interface_CollisionDataProviderCore.h
--- a/Engine/Source/Runtime/PhysicsCore/Public/Interface_CollisionDataProviderCore.h
+++ b/Engine/Source/Runtime/PhysicsCore/Public/Interface_CollisionDataProviderCore.h
@@ -19,6 +19,9 @@ struct FTriIndices
 	}
 };
 
+#if WITH_EDITORONLY_DATA
+PHYSICSCORE_API void operator<<(class FArchive& Ar, FTriIndices& TriIndices);
+#endif
 
 // Description of triangle mesh collision data necessary for cooking physics data
 struct FTriMeshCollisionData
@@ -56,6 +59,10 @@ struct FTriMeshCollisionData
 	}
 };
 
+#if WITH_EDITORONLY_DATA
+PHYSICSCORE_API void operator<<(class FArchive& Ar, FTriMeshCollisionData& TriMeshCollisionData);
+#endif
+
 // Estimates of triangle mesh collision data necessary for cooking physics data
 struct FTriMeshCollisionDataEstimates
 {
