PSO precache fixes:
- make sure material ID is still valid when trying to boost (assets could have old ID cached after full recaching)
- don't precache decal component PSOs when material hasn't been post loaded yet

diff --git a/Engine/Source/Runtime/Engine/Private/Components/DecalComponent.cpp b/Engine/Source/Runtime/Engine/Private/Components/DecalComponent.cpp
--- a/Engine/Source/Runtime/Engine/Private/Components/DecalComponent.cpp
+++ b/Engine/Source/Runtime/Engine/Private/Components/DecalComponent.cpp
@@ -323,7 +323,7 @@ void UDecalComponent::PrecachePSOs()
 	// clear the current request data
 	PSOPrecacheCompileEvent = nullptr;
 
-	if (DecalMaterial)
+	if (DecalMaterial && !DecalMaterial->HasAnyFlags(RF_NeedPostLoad))
 	{
 		FPSOPrecacheParams PSOPrecacheParams;		
 		FPSOPrecacheVertexFactoryDataList VertexFactoryDataList;		
diff --git a/Engine/Source/Runtime/Engine/Private/PSOPrecacheMaterial.cpp b/Engine/Source/Runtime/Engine/Private/PSOPrecacheMaterial.cpp
--- a/Engine/Source/Runtime/Engine/Private/PSOPrecacheMaterial.cpp
+++ b/Engine/Source/Runtime/Engine/Private/PSOPrecacheMaterial.cpp
@@ -333,10 +333,9 @@ public:
 
 			const FMaterialPSOPrecacheParams& Params = MaterialPSORequests[MaterialPSORequestID];
 			FPrecacheData* FindResult = MaterialPSORequestData.Find(Params);
-			check(FindResult);
 
 			// Only process if not boosted yet and not completed yet
-			if (FindResult->Priority == EPSOPrecachePriority::High || FindResult->State == EState::Completed)
+			if (FindResult == nullptr || FindResult->Priority == EPSOPrecachePriority::High || FindResult->State == EState::Completed)
 			{
 				return;
 			}
