Fix for crash when generating SVT page topology during derived data building when the volume has a non-cubic page table resolution.

#rnx

diff --git a/Engine/Source/Runtime/Engine/Private/SparseVolumeTexture/SparseVolumeTexture.cpp b/Engine/Source/Runtime/Engine/Private/SparseVolumeTexture/SparseVolumeTexture.cpp
--- a/Engine/Source/Runtime/Engine/Private/SparseVolumeTexture/SparseVolumeTexture.cpp
+++ b/Engine/Source/Runtime/Engine/Private/SparseVolumeTexture/SparseVolumeTexture.cpp
@@ -129,7 +129,7 @@ static int32 ComputeNumMipLevels(const FIntVector3& InVirtualVolumeMin, const FI
 
 static const FString& GetDerivedDataVersion()
 {
-	static FString CachedVersionString = TEXT("A1A4EFBF-6596-418D-A1C1-11FC07BA0B7F");	// Bump this if you want to ignore all cached data so far.
+	static FString CachedVersionString = TEXT("4187DF14-A41A-45DD-957A-386EC727980D");	// Bump this if you want to ignore all cached data so far.
 	return CachedVersionString;
 }
 
@@ -836,6 +836,12 @@ FPageTopology FResources::BuildTopology(const FTextureData& InTextureData, const
 	// Recursively build the page mip hierarchy. The order of nodes is identical to how they're stored in bulk data.
 	auto BuildHierarchy = [&](const FIntVector3& InPageTableMipResolution, int32 InMipLevel, int32 InPageX, int32 InPageY, int32 InPageZ, uint32 InParentIndex, auto& InRecursionLambda) -> uint32
 	{
+		// With non-cubic virtual volumes it's possible that the parent node doesn't have a cubic (square in 2D) footprint of child nodes, so early out if that's the case.
+		if (!IsInBounds(FIntVector3(InPageX, InPageY, InPageZ), FIntVector3::ZeroValue, InPageTableMipResolution))
+		{
+			return INDEX_NONE;
+		}
+
 		// Load entry from dense page table. 
 		const int32 LinearPageIndex = InPageZ * (InPageTableMipResolution.X * InPageTableMipResolution.Y) + (InPageY * InPageTableMipResolution.X) + InPageX;
 		const uint32 PageTableEntry = InTextureData.MipMaps[InMipLevel].PageTable[LinearPageIndex];
@@ -876,7 +882,11 @@ FPageTopology FResources::BuildTopology(const FTextureData& InTextureData, const
 		return PageIndex;
 	};
 
-	BuildHierarchy(FIntVector3::ZeroValue, NumMipLevels - 1, 0, 0, 0, INDEX_NONE, BuildHierarchy);
+	// Recursively build the hierarchy, starting at the root mip.
+	const int32 NumPagesInLastMip = InTextureData.MipMaps.Last().PageTable.Num();
+	check(NumPagesInLastMip == 1);
+	const FIntVector3 RootMipPageTableResolution = FIntVector3(1, 1, 1);
+	BuildHierarchy(RootMipPageTableResolution, NumMipLevels - 1, 0, 0, 0, INDEX_NONE, BuildHierarchy);
 
 	// Validate
 	{
