Skeletal Mesh: Fix issues with smoothing groups and overlapping points when reconstructing morphs and import data from LODModels.

#rnx

diff --git a/Engine/Source/Runtime/Core/Private/UObject/FortniteMainBranchObjectVersion.cpp b/Engine/Source/Runtime/Core/Private/UObject/FortniteMainBranchObjectVersion.cpp
--- a/Engine/Source/Runtime/Core/Private/UObject/FortniteMainBranchObjectVersion.cpp
+++ b/Engine/Source/Runtime/Core/Private/UObject/FortniteMainBranchObjectVersion.cpp
@@ -12,7 +12,7 @@ TMap<FGuid, FGuid> FFortniteMainBranchObjectVersion::GetSystemGuids()
 	SystemGuids.Add(DevGuids.NANITE_DERIVEDDATA_VER, FGuid("D28C003F825F4C178F698938FB0E47F7"));
 	SystemGuids.Add(DevGuids.NIAGARASHADERMAP_DERIVEDDATA_VER, FGuid("D0DF49D408F34E38A98FA097E0280EC7"));
 	SystemGuids.Add(DevGuids.Niagara_LatestScriptCompileVersion, FGuid("CE33E29C6268469E8E18F1D96A55DFDA"));
-	SystemGuids.Add(DevGuids.SkeletalMeshDerivedDataVersion, FGuid("C9CF008376FC4F7D9996CB449719502B"));
+	SystemGuids.Add(DevGuids.SkeletalMeshDerivedDataVersion, FGuid("3BD78FFFB88E4DDF943FB2608A20C0D2"));
 	SystemGuids.Add(DevGuids.STATICMESH_DERIVEDDATA_VER, FGuid("543843B103794E0E9BA4BE80FB602F79"));
 	SystemGuids.Add(DevGuids.MaterialTranslationDDCVersion, FGuid("BBD650804D314CE7A44ADF9BB4F3F5F9"));
 	return SystemGuids;
diff --git a/Engine/Source/Runtime/Engine/Private/Rendering/SkeletalMeshLODImporterData.cpp b/Engine/Source/Runtime/Engine/Private/Rendering/SkeletalMeshLODImporterData.cpp
--- a/Engine/Source/Runtime/Engine/Private/Rendering/SkeletalMeshLODImporterData.cpp
+++ b/Engine/Source/Runtime/Engine/Private/Rendering/SkeletalMeshLODImporterData.cpp
@@ -1433,7 +1433,7 @@ void FSkeletalMeshImportData::ComputeSmoothGroupFromNormals()
 void FSkeletalMeshImportData::SetMorphTargets(
 	const TArray<TObjectPtr<UMorphTarget>>& InMorphTargets,
 	int32 InLODIndex,
-	const TArray<int32>& InVertexMap
+	const TArray<uint32>& InVertexMap
 	)
 {
 	MorphTargets.Reset();
@@ -2074,6 +2074,7 @@ bool FSkeletalMeshImportData::GetMeshDescription(const USkeletalMesh* InSkeletal
 	
 	// Avoid repeated allocations and reserve the target buffers right off the bat.
 	OutMeshDescription.ReserveNewPolygonGroups(Materials.Num());
+	OutMeshDescription.ReserveNewPolygons(Faces.Num());
 	OutMeshDescription.ReserveNewTriangles(Faces.Num());
 	OutMeshDescription.ReserveNewVertexInstances(Wedges.Num());
 	OutMeshDescription.ReserveNewVertices(Points.Num());
@@ -2330,6 +2331,8 @@ bool FSkeletalMeshImportData::GetMeshDescription(const USkeletalMesh* InSkeletal
 			VertexAndCountAttribute.Set(PartID, {Info.StartImportedVertex, Info.NumVertices});
 		}
 	}
+
+	OutMeshDescription.BuildIndexers();
 	
 	return true;
 }
@@ -2420,6 +2423,62 @@ void FSkeletalMeshImportData::CopySkinWeightsToMeshDescription(
 }
 
 
+// A simpler variant of FStaticMeshOperations::ConvertHardEdgesToSmoothGroup that assumes that hard edges always form closed regions.
+static void ConvertHardEdgesToSmoothMasks(
+	const FMeshDescription& InMeshDescription,
+	TArray<uint32>& OutSmoothMasks
+	)
+{
+	OutSmoothMasks.SetNumZeroed(InMeshDescription.Triangles().Num());
+
+	TSet<FTriangleID> ProcessedTriangles;
+	TArray<FTriangleID> TriangleQueue;
+	uint32 CurrentSmoothMask = 1;
+
+	const TEdgeAttributesConstRef<bool> IsEdgeHard = InMeshDescription.EdgeAttributes().GetAttributesRef<bool>(MeshAttribute::Edge::IsHard);	
+
+	for (FTriangleID SeedTriangleID: InMeshDescription.Triangles().GetElementIDs())
+	{
+		if (ProcessedTriangles.Contains(SeedTriangleID))
+		{
+			continue;
+		}
+		
+		TriangleQueue.Push(SeedTriangleID);
+		while (!TriangleQueue.IsEmpty())
+		{
+			const FTriangleID TriangleID = TriangleQueue.Pop(EAllowShrinking::No);
+			TArrayView<const FEdgeID> TriangleEdges = InMeshDescription.GetTriangleEdges(TriangleID);
+
+			OutSmoothMasks[TriangleID.GetValue()] = CurrentSmoothMask;
+			ProcessedTriangles.Add(TriangleID);
+
+			for (const FEdgeID EdgeID: TriangleEdges)
+			{
+				if (!IsEdgeHard.Get(EdgeID))
+				{
+					TArrayView<const FTriangleID> ConnectedTriangles = InMeshDescription.GetEdgeConnectedTriangleIDs(EdgeID);
+					for (const FTriangleID NeighborTriangleID: ConnectedTriangles)
+					{
+						if (!ProcessedTriangles.Contains(NeighborTriangleID))
+						{
+							TriangleQueue.Push(NeighborTriangleID);
+						}
+					}
+				}
+			}
+		}
+
+		CurrentSmoothMask <<= 1;
+		if (CurrentSmoothMask == 0)
+		{
+			CurrentSmoothMask = 1;
+		}
+	}
+}
+
+
+
 
 template<typename T>
 struct FCopyAttributeElement {};
@@ -2538,8 +2597,7 @@ FSkeletalMeshImportData FSkeletalMeshImportData::CreateFromMeshDescription(const
 
 	//Get the per face smoothing
 	TArray<uint32> FaceSmoothingMasks;
-	FaceSmoothingMasks.AddZeroed(InMeshDescription.Triangles().Num());
-	FSkeletalMeshOperations::ConvertHardEdgesToSmoothGroup(InMeshDescription, FaceSmoothingMasks);
+	ConvertHardEdgesToSmoothMasks(InMeshDescription, FaceSmoothingMasks);
 
 	//////////////////////////////////////////////////////////////////////////
 	// Copy the materials
@@ -2600,18 +2658,14 @@ FSkeletalMeshImportData FSkeletalMeshImportData::CreateFromMeshDescription(const
 	{
 		FPolygonGroupID PolygonGroupID = InMeshDescription.GetTrianglePolygonGroup(TriangleID);
 		TArrayView<const FVertexInstanceID> VertexInstances = InMeshDescription.GetTriangleVertexInstances(TriangleID);
-		int32 FaceIndex = TriangleID.GetValue();
-		if (!ensure(SkelMeshImportData.Faces.IsValidIndex(FaceIndex)))
-		{
-			//TODO log an error for the user
-			break;
-		}
-		SkeletalMeshImportData::FTriangle& Face = SkelMeshImportData.Faces[FaceIndex];
+		int32 TriangleIndex = TriangleID.GetValue();
+		
+		SkeletalMeshImportData::FTriangle& Face = SkelMeshImportData.Faces[TriangleIndex];
 		Face.MatIndex = PolygonGroupID.GetValue();
 		Face.SmoothingGroups = 0;
-		if (FaceSmoothingMasks.IsValidIndex(FaceIndex))
+		if (FaceSmoothingMasks.IsValidIndex(TriangleIndex))
 		{
-			Face.SmoothingGroups = FaceSmoothingMasks[FaceIndex];
+			Face.SmoothingGroups = FaceSmoothingMasks[TriangleIndex];
 		}
 		//Create the wedges
 		for (int32 Corner = 0; Corner < 3; ++Corner)
diff --git a/Engine/Source/Runtime/Engine/Private/SkeletalMesh.cpp b/Engine/Source/Runtime/Engine/Private/SkeletalMesh.cpp
--- a/Engine/Source/Runtime/Engine/Private/SkeletalMesh.cpp
+++ b/Engine/Source/Runtime/Engine/Private/SkeletalMesh.cpp
@@ -3208,7 +3208,7 @@ void USkeletalMesh::BeginPostLoadInternal(FSkinnedAssetPostLoadContext& Context)
 						// mesh itself. Try to back-fill from the skeletal mesh.
 						if (SerializeMeshData.MorphTargets.IsEmpty() && !GetMorphTargets().IsEmpty())
 						{
-							SerializeMeshData.SetMorphTargets(GetMorphTargets(), LODIndex, ThisLODModel.MeshToImportVertexMap);
+							SerializeMeshData.SetMorphTargets(GetMorphTargets(), LODIndex, ThisLODModel.GetRawPointIndices());
 						}
 
 						FMeshDescription MeshDescription;
@@ -3255,11 +3255,44 @@ void USkeletalMesh::BeginPostLoadInternal(FSkinnedAssetPostLoadContext& Context)
 		}
 
 		// If we didn't get any meshes from the bulk data, then try to recover them from the LODModel listings.
-		constexpr bool bInResetReductionAsNeeded = true;
+		for (int32 LODIndex = 0; LODIndex < GetImportedModel()->LODModels.Num(); ++LODIndex)
+		{
+			if (HasMeshDescription(LODIndex))
+			{
+				continue;
+			}
+
+			// If the mesh was not pulled out of the reduction data, we need to reset the LOD settings
+			// so that the mesh doesn't get reduced again if it gets regenerated.
+			FSkeletalMeshLODInfo* MeshLODInfo = GetLODInfo(LODIndex); 
+			const bool bReductionActive = IsReductionActive(LODIndex);
+			const bool bInlineReduction = (GetLODInfo(LODIndex)->ReductionSettings.BaseLOD == LODIndex);
+			if (bReductionActive && !bInlineReduction)
+			{
+				//Generated LOD (not inline) do not need imported data
+				continue;
+			}
 			
-		IMeshUtilities& MeshUtilities = FModuleManager::Get().LoadModuleChecked<IMeshUtilities>("MeshUtilities");
-		MeshUtilities.CreateImportDataFromLODModel(this, bInResetReductionAsNeeded);
-		
+			const FSkeletalMeshLODModel& LODModel = GetImportedModel()->LODModels[LODIndex];
+
+			FMeshDescription MeshDescription;
+			LODModel.GetMeshDescription(this, LODIndex, MeshDescription);
+			CreateMeshDescription(LODIndex, MoveTemp(MeshDescription));
+			CommitMeshDescription(LODIndex);
+
+			// Reset the reduction settings so that we don't re-reduce the mesh and possibly lose morph targets
+			// in the process.
+			FSkeletalMeshOptimizationSettings& ReductionSettings = MeshLODInfo->ReductionSettings;
+			
+			//Remove the reduction settings
+			ReductionSettings.NumOfTrianglesPercentage = 1.0f;
+			ReductionSettings.NumOfVertPercentage = 1.0f;
+			ReductionSettings.MaxNumOfTrianglesPercentage = MAX_uint32;
+			ReductionSettings.MaxNumOfVertsPercentage = MAX_uint32;
+			ReductionSettings.TerminationCriterion = SMTC_NumOfTriangles;
+			MeshLODInfo->bHasBeenSimplified = false;
+		}
+
 		if (GetLinkerCustomVersion(FEditorObjectVersion::GUID) < FEditorObjectVersion::SkeletalMeshBuildRefactor)
 		{
 			CreateUserSectionsDataForLegacyAssets();
diff --git a/Engine/Source/Runtime/Engine/Private/SkeletalMeshLODModel.cpp b/Engine/Source/Runtime/Engine/Private/SkeletalMeshLODModel.cpp
--- a/Engine/Source/Runtime/Engine/Private/SkeletalMeshLODModel.cpp
+++ b/Engine/Source/Runtime/Engine/Private/SkeletalMeshLODModel.cpp
@@ -1338,13 +1338,14 @@ void FSkeletalMeshLODModel::CopyStructure(FSkeletalMeshLODModel* Destination, co
 	Destination->BulkDataReadMutex = DestinationBulkDataReadMutex;
 }
 
-void FSkeletalMeshLODModel::GetMeshDescription(FMeshDescription& MeshDescription, const USkeletalMesh *Owner) const
+UE_DISABLE_OPTIMIZATION
+void FSkeletalMeshLODModel::GetMeshDescription(const USkeletalMesh *InSkeletalMesh, const int32 InLODIndex, FMeshDescription& OutMeshDescription) const
 {
 	using UE::AnimationCore::FBoneWeights;
 
-	MeshDescription.Empty();
+	OutMeshDescription.Empty();
 	
-	FSkeletalMeshAttributes MeshAttributes(MeshDescription);	
+	FSkeletalMeshAttributes MeshAttributes(OutMeshDescription);	
 	
 	// Register extra attributes for us.
 	MeshAttributes.Register();
@@ -1363,29 +1364,78 @@ void FSkeletalMeshLODModel::GetMeshDescription(FMeshDescription& MeshDescription
 	FSkeletalMeshAttributes::FBoneParentIndexAttributesRef BoneParentIndices = MeshAttributes.GetBoneParentIndices();
 	FSkeletalMeshAttributes::FBonePoseAttributesRef BonePoses = MeshAttributes.GetBonePoses();
 
+	TArray<TPair<FName, const FMorphTargetLODModel*>> MorphTargets;
+	for (UMorphTarget* MorphTargetSource: InSkeletalMesh->GetMorphTargets())
+	{
+		if (!MorphTargetSource->HasDataForLOD(InLODIndex))
+		{
+			continue;
+		}
+
+		FName Name = MorphTargetSource->GetFName();
+		if (Name.IsNone())
+		{
+			Name = TEXT("Unnamed");
+		}
+		
+		MorphTargets.Emplace(Name, &MorphTargetSource->GetMorphLODModels()[InLODIndex]);
+		MeshAttributes.RegisterMorphTargetAttribute(Name);
+	}
+
 	const int32 NumTriangles = IndexBuffer.Num() / 3;
 
-	const FReferenceSkeleton& RefSkeleton = Owner->GetRefSkeleton();
+	const FReferenceSkeleton& RefSkeleton = InSkeletalMesh->GetRefSkeleton();
 	const int NumBones = RefSkeleton.GetRawBoneNum();
 
-	MeshDescription.ReserveNewPolygonGroups(Sections.Num());
-	MeshDescription.ReserveNewTriangles(NumTriangles);
-	MeshDescription.ReserveNewVertexInstances(NumTriangles * 3);
-	MeshDescription.ReserveNewVertices(static_cast<int32>(NumVertices));
+	OutMeshDescription.ReserveNewPolygonGroups(Sections.Num());
+	OutMeshDescription.ReserveNewPolygons(NumTriangles);
+	OutMeshDescription.ReserveNewTriangles(NumTriangles);
+	OutMeshDescription.ReserveNewVertexInstances(NumTriangles * 3);
+	OutMeshDescription.ReserveNewVertices(static_cast<int32>(NumVertices));
 	MeshAttributes.ReserveNewBones(NumBones);
 
+	// Map the section vertices back to the import vertices to remove seams, but only if there's
+	// mapping available.
+	TArray<int32> SourceToTargetVertexMap; 
+	TSet<int32> ProcessedTargetVertex;
+
+	int32 TargetVertexCount = 0;
+	if (RawPointIndices2.Num() == NumVertices)
+	{
+		SourceToTargetVertexMap.Reserve(RawPointIndices2.Num());
+		for (const uint32 Index: RawPointIndices2)
+		{
+			SourceToTargetVertexMap.Add(Index);
+			ProcessedTargetVertex.Add(Index);
+
+			TargetVertexCount = FMath::Max(TargetVertexCount, static_cast<int32>(Index));
+		}
+
+		TargetVertexCount += 1;
+		ProcessedTargetVertex.Reset();
+	}
+	else
+	{
+		SourceToTargetVertexMap.Reserve(NumVertices);
+		for (uint32 Index = 0; Index < NumVertices; Index++)
+		{
+			SourceToTargetVertexMap.Add(Index);
+		}
+		TargetVertexCount = NumVertices;
+	}
+	
 	TArray<FVertexID> VertexIDs;
-	VertexIDs.Reserve(NumVertices);
-	for (int32 VertexIndex = 0; VertexIndex < int32(NumVertices); VertexIndex++)
+	VertexIDs.Reserve(TargetVertexCount);
+	for (int32 VertexIndex = 0; VertexIndex < TargetVertexCount; VertexIndex++)
 	{
-		VertexIDs.Add(MeshDescription.CreateVertex());
+		VertexIDs.Add(OutMeshDescription.CreateVertex());
 	}
 
 	// Ensure we have enough channels to store all the defined UV coordinates.
 	VertexInstanceUVs.SetNumChannels(static_cast<int32>(NumTexCoords));
 	
-	const TArray<FSkeletalMaterial>& Materials = Owner->GetMaterials();
-	const bool bHasVertexColors = EnumHasAllFlags(Owner->GetVertexBufferFlags(), ESkeletalMeshVertexFlags::HasVertexColors);
+	const TArray<FSkeletalMaterial>& Materials = InSkeletalMesh->GetMaterials();
+	const bool bHasVertexColors = EnumHasAllFlags(InSkeletalMesh->GetVertexBufferFlags(), ESkeletalMeshVertexFlags::HasVertexColors);
 
 	// Convert sections to polygon groups, each with their own material.
 	for (int32 SectionIndex = 0; SectionIndex < Sections.Num(); SectionIndex++)
@@ -1396,7 +1446,15 @@ void FSkeletalMeshLODModel::GetMeshDescription(FMeshDescription& MeshDescription
 		const TArray<FSoftSkinVertex>& SourceVertices = Section.SoftVertices;
 		for (int32 VertexIndex = 0; VertexIndex < SourceVertices.Num(); VertexIndex++)
 		{
-			const FVertexID VertexID = VertexIDs[VertexIndex + Section.BaseVertexIndex];
+			const int32 SourceVertexIndex = VertexIndex + Section.BaseVertexIndex;
+			const int32 TargetVertexIndex = SourceToTargetVertexMap[SourceVertexIndex];
+
+			if (ProcessedTargetVertex.Contains(TargetVertexIndex))
+			{
+				continue;
+			}
+			
+			const FVertexID VertexID = VertexIDs[TargetVertexIndex];
 
 			VertexPositions.Set(VertexID, SourceVertices[VertexIndex].Position);
 
@@ -1416,9 +1474,9 @@ void FSkeletalMeshLODModel::GetMeshDescription(FMeshDescription& MeshDescription
 
 		const FPolygonGroupID PolygonGroupID(Section.MaterialIndex);
 
-		if (!MeshDescription.IsPolygonGroupValid(PolygonGroupID))
+		if (!OutMeshDescription.IsPolygonGroupValid(PolygonGroupID))
 		{
-			MeshDescription.CreatePolygonGroupWithID(PolygonGroupID);
+			OutMeshDescription.CreatePolygonGroupWithID(PolygonGroupID);
 		}
 
 		if (ensure(Materials.IsValidIndex(Section.MaterialIndex)))
@@ -1436,8 +1494,10 @@ void FSkeletalMeshLODModel::GetMeshDescription(FMeshDescription& MeshDescription
 			for (int32 Corner = 0; Corner < 3; Corner++)
 			{
 				const int32 SourceVertexIndex = IndexBuffer[VertexIndexBase + Corner];
-				const FVertexID VertexID = VertexIDs[SourceVertexIndex];
-				const FVertexInstanceID VertexInstanceID = MeshDescription.CreateVertexInstance(VertexID);
+				const int32 TargetVertexIndex = SourceToTargetVertexMap[SourceVertexIndex];
+				
+				const FVertexID VertexID = VertexIDs[TargetVertexIndex];
+				const FVertexInstanceID VertexInstanceID = OutMeshDescription.CreateVertexInstance(VertexID);
 
 				const FSoftSkinVertex& SourceVertex = SourceVertices[SourceVertexIndex - Section.BaseVertexIndex];
 
@@ -1462,7 +1522,21 @@ void FSkeletalMeshLODModel::GetMeshDescription(FMeshDescription& MeshDescription
 				TriangleVertexInstanceIDs[Corner] = VertexInstanceID;
 			}
 
-			MeshDescription.CreateTriangle(PolygonGroupID, TriangleVertexInstanceIDs);
+			OutMeshDescription.CreateTriangle(PolygonGroupID, TriangleVertexInstanceIDs);
+		}
+	}
+
+	// Copy morph targets.
+	for (TPair<FName, const FMorphTargetLODModel*>& MorphSource: MorphTargets)
+	{
+		FMorphTargetVertexAttributesRef MorphTarget = MeshAttributes.GetVertexMorphTarget(MorphSource.Key);
+
+		for (const FMorphTargetDelta& Delta: MorphSource.Value->Vertices)
+		{
+			const int32 TargetVertexIndex = SourceToTargetVertexMap[Delta.SourceIdx];
+			const FVertexID VertexID = VertexIDs[TargetVertexIndex];
+			
+			MorphTarget.SetPositionAndTangentZDelta(VertexID, Delta.PositionDelta, Delta.TangentZDelta);
 		}
 	}
 
@@ -1479,5 +1553,6 @@ void FSkeletalMeshLODModel::GetMeshDescription(FMeshDescription& MeshDescription
 		BonePoses.Set(BoneID, BoneTransform);
 	}
 }
+UE_ENABLE_OPTIMIZATION
 
 #endif // WITH_EDITOR
diff --git a/Engine/Source/Runtime/Engine/Public/Rendering/SkeletalMeshLODImporterData.h b/Engine/Source/Runtime/Engine/Public/Rendering/SkeletalMeshLODImporterData.h
--- a/Engine/Source/Runtime/Engine/Public/Rendering/SkeletalMeshLODImporterData.h
+++ b/Engine/Source/Runtime/Engine/Public/Rendering/SkeletalMeshLODImporterData.h
@@ -421,7 +421,7 @@ public:
 	/*
 	 * Add morph target data from UMorphTarget in case there was none on the mesh itself.
 	 */
-	ENGINE_API void SetMorphTargets(const TArray<TObjectPtr<UMorphTarget>>& InMorphTargets, int32 InLODIndex, const TArray<int32>& InVertexMap);
+	ENGINE_API void SetMorphTargets(const TArray<TObjectPtr<UMorphTarget>>& InMorphTargets, int32 InLODIndex, const TArray<uint32>& InVertexMap);
 
 	/**
 	 * Returns a mesh description from the import data. If logging on failures is required, pass in a pointer 
diff --git a/Engine/Source/Runtime/Engine/Public/Rendering/SkeletalMeshLODModel.h b/Engine/Source/Runtime/Engine/Public/Rendering/SkeletalMeshLODModel.h
--- a/Engine/Source/Runtime/Engine/Public/Rendering/SkeletalMeshLODModel.h
+++ b/Engine/Source/Runtime/Engine/Public/Rendering/SkeletalMeshLODModel.h
@@ -567,7 +567,7 @@ public:
 	 * Fills in a representation of this model into the given mesh description object. Existing
 	 * mesh description data is emptied.
 	 */
-	void ENGINE_API GetMeshDescription(FMeshDescription& MeshDescription, const USkeletalMesh *Owner) const;
+	void ENGINE_API GetMeshDescription(const USkeletalMesh *InSkeletalMesh, const int32 InLODIndex, FMeshDescription& OutMeshDescription) const;
 };
 
 #endif // WITH_EDITOR
