Skeletal Mesh: Remove optimization macros.

#rnx

diff --git a/Engine/Source/Runtime/Engine/Private/SkeletalMeshLODModel.cpp b/Engine/Source/Runtime/Engine/Private/SkeletalMeshLODModel.cpp
--- a/Engine/Source/Runtime/Engine/Private/SkeletalMeshLODModel.cpp
+++ b/Engine/Source/Runtime/Engine/Private/SkeletalMeshLODModel.cpp
@@ -1338,7 +1338,6 @@ void FSkeletalMeshLODModel::CopyStructure(FSkeletalMeshLODModel* Destination, co
 	Destination->BulkDataReadMutex = DestinationBulkDataReadMutex;
 }
 
-UE_DISABLE_OPTIMIZATION
 void FSkeletalMeshLODModel::GetMeshDescription(const USkeletalMesh *InSkeletalMesh, const int32 InLODIndex, FMeshDescription& OutMeshDescription) const
 {
 	using UE::AnimationCore::FBoneWeights;
@@ -1553,6 +1552,5 @@ void FSkeletalMeshLODModel::GetMeshDescription(const USkeletalMesh *InSkeletalMe
 		BonePoses.Set(BoneID, BoneTransform);
 	}
 }
-UE_ENABLE_OPTIMIZATION
 
 #endif // WITH_EDITOR
