Remove world.h include from NavLinkCustomInterface.h

diff --git a/Engine/Source/Runtime/Engine/Classes/Engine/World.h b/Engine/Source/Runtime/Engine/Classes/Engine/World.h
--- a/Engine/Source/Runtime/Engine/Classes/Engine/World.h
+++ b/Engine/Source/Runtime/Engine/Classes/Engine/World.h
@@ -38,6 +38,7 @@
 #include "Subsystems/SubsystemCollection.h"
 #include "CollisionProfile.h"
 #include "RHIFeatureLevel.h"
+#include "Engine/WorldInitializationValues.h"
 
 #include "World.generated.h"
 
@@ -848,84 +849,6 @@ private:
 	static void BroadcastWorldPartitionUninitialized(UWorld* InWorld, UWorldPartition* InWorldPartition);
 };
 
-/** Struct containing a collection of optional parameters for initialization of a World. */
-struct FWorldInitializationValues
-{
-	FWorldInitializationValues()
-		: bInitializeScenes(true)
-		, bAllowAudioPlayback(true)
-		, bRequiresHitProxies(true)
-		, bCreatePhysicsScene(true)
-		, bCreateNavigation(true)
-		, bCreateAISystem(true)
-		, bShouldSimulatePhysics(true)
-		, bEnableTraceCollision(false)
-		, bForceUseMovementComponentInNonGameWorld(false)
-		, bTransactional(true)
-		, bCreateFXSystem(true)
-		, bCreateWorldPartition(false)
-		, bEnableWorldPartitionStreaming(true)
-	{
-	}
-
-	/** Should the scenes (physics, rendering) be created. */
-	uint32 bInitializeScenes:1;
-
-	/** Are sounds allowed to be generated from this world. */
-	uint32 bAllowAudioPlayback:1;
-
-	/** Should the render scene create hit proxies. */
-	uint32 bRequiresHitProxies:1;
-
-	/** Should the physics scene be created. bInitializeScenes must be true for this to be considered. */
-	uint32 bCreatePhysicsScene:1;
-
-	/** Should the navigation system be created for this world. */
-	uint32 bCreateNavigation:1;
-
-	/** Should the AI system be created for this world. */
-	uint32 bCreateAISystem:1;
-
-	/** Should physics be simulated in this world. */
-	uint32 bShouldSimulatePhysics:1;
-
-	/** Are collision trace calls valid within this world. */
-	uint32 bEnableTraceCollision:1;
-
-	/** Special flag to enable movement component in non game worlds (see UMovementComponent::OnRegister) */
-	uint32 bForceUseMovementComponentInNonGameWorld:1;
-
-	/** Should actions performed to objects in this world be saved to the transaction buffer. */
-	uint32 bTransactional:1;
-
-	/** Should the FX system be created for this world. */
-	uint32 bCreateFXSystem:1;
-
-	/** Should the world be partitioned */
-	uint32 bCreateWorldPartition:1;
-
-	/** If bCreateWorldPartition is set to true, this flag will init the streaming mode for the WorldPartition object (default to true to preserve previous behavior) */
-	uint32 bEnableWorldPartitionStreaming:1;
-
-	/** The default game mode for this world (if any) */
-	TSubclassOf<class AGameModeBase> DefaultGameMode;
-
-	FWorldInitializationValues& InitializeScenes(const bool bInitialize) { bInitializeScenes = bInitialize; return *this; }
-	FWorldInitializationValues& AllowAudioPlayback(const bool bAllow) { bAllowAudioPlayback = bAllow; return *this; }
-	FWorldInitializationValues& RequiresHitProxies(const bool bRequires) { bRequiresHitProxies = bRequires; return *this; }
-	FWorldInitializationValues& CreatePhysicsScene(const bool bCreate) { bCreatePhysicsScene = bCreate; return *this; }
-	FWorldInitializationValues& CreateNavigation(const bool bCreate) { bCreateNavigation = bCreate; return *this; }
-	FWorldInitializationValues& CreateAISystem(const bool bCreate) { bCreateAISystem = bCreate; return *this; }
-	FWorldInitializationValues& ShouldSimulatePhysics(const bool bInShouldSimulatePhysics) { bShouldSimulatePhysics = bInShouldSimulatePhysics; return *this; }
-	FWorldInitializationValues& EnableTraceCollision(const bool bInEnableTraceCollision) { bEnableTraceCollision = bInEnableTraceCollision; return *this; }
-	FWorldInitializationValues& ForceUseMovementComponentInNonGameWorld(const bool bInForceUseMovementComponentInNonGameWorld) { bForceUseMovementComponentInNonGameWorld = bInForceUseMovementComponentInNonGameWorld; return *this; }
-	FWorldInitializationValues& SetTransactional(const bool bInTransactional) { bTransactional = bInTransactional; return *this; }
-	FWorldInitializationValues& CreateFXSystem(const bool bCreate) { bCreateFXSystem = bCreate; return *this; }
-	FWorldInitializationValues& CreateWorldPartition(const bool bCreate) { bCreateWorldPartition = bCreate; return *this; }
-	FWorldInitializationValues& EnableWorldPartitionStreaming(const bool bEnableStreaming) { bEnableWorldPartitionStreaming = bEnableStreaming; return *this; }
-	FWorldInitializationValues& SetDefaultGameMode(TSubclassOf<class AGameModeBase> GameMode) { DefaultGameMode = GameMode; return *this; }
-};
-
 /** 
  * The World is the top level object representing a map or a sandbox in which Actors and Components will exist and be rendered.  
  *
diff --git a/Engine/Source/Runtime/Engine/Classes/Engine/WorldInitializationValues.h b/Engine/Source/Runtime/Engine/Classes/Engine/WorldInitializationValues.h
new file mode 100644
--- /dev/null
+++ b/Engine/Source/Runtime/Engine/Classes/Engine/WorldInitializationValues.h
@@ -0,0 +1,85 @@
+﻿// Copyright Epic Games, Inc. All Rights Reserved.
+
+#pragma once
+
+#include "Templates/SubclassOf.h"
+
+/** Struct containing a collection of optional parameters for initialization of a World. */
+struct FWorldInitializationValues
+{
+	FWorldInitializationValues()
+		: bInitializeScenes(true)
+		, bAllowAudioPlayback(true)
+		, bRequiresHitProxies(true)
+		, bCreatePhysicsScene(true)
+		, bCreateNavigation(true)
+		, bCreateAISystem(true)
+		, bShouldSimulatePhysics(true)
+		, bEnableTraceCollision(false)
+		, bForceUseMovementComponentInNonGameWorld(false)
+		, bTransactional(true)
+		, bCreateFXSystem(true)
+		, bCreateWorldPartition(false)
+		, bEnableWorldPartitionStreaming(true)
+	{
+	}
+
+	/** Should the scenes (physics, rendering) be created. */
+	uint32 bInitializeScenes:1;
+
+	/** Are sounds allowed to be generated from this world. */
+	uint32 bAllowAudioPlayback:1;
+
+	/** Should the render scene create hit proxies. */
+	uint32 bRequiresHitProxies:1;
+
+	/** Should the physics scene be created. bInitializeScenes must be true for this to be considered. */
+	uint32 bCreatePhysicsScene:1;
+
+	/** Should the navigation system be created for this world. */
+	uint32 bCreateNavigation:1;
+
+	/** Should the AI system be created for this world. */
+	uint32 bCreateAISystem:1;
+
+	/** Should physics be simulated in this world. */
+	uint32 bShouldSimulatePhysics:1;
+
+	/** Are collision trace calls valid within this world. */
+	uint32 bEnableTraceCollision:1;
+
+	/** Special flag to enable movement component in non game worlds (see UMovementComponent::OnRegister) */
+	uint32 bForceUseMovementComponentInNonGameWorld:1;
+
+	/** Should actions performed to objects in this world be saved to the transaction buffer. */
+	uint32 bTransactional:1;
+
+	/** Should the FX system be created for this world. */
+	uint32 bCreateFXSystem:1;
+
+	/** Should the world be partitioned */
+	uint32 bCreateWorldPartition:1;
+
+	/** If bCreateWorldPartition is set to true, this flag will init the streaming mode for the WorldPartition object (default to true to preserve previous behavior) */
+	uint32 bEnableWorldPartitionStreaming:1;
+
+	/** The default game mode for this world (if any) */
+	TSubclassOf<class AGameModeBase> DefaultGameMode;
+
+	FWorldInitializationValues& InitializeScenes(const bool bInitialize) { bInitializeScenes = bInitialize; return *this; }
+	FWorldInitializationValues& AllowAudioPlayback(const bool bAllow) { bAllowAudioPlayback = bAllow; return *this; }
+	FWorldInitializationValues& RequiresHitProxies(const bool bRequires) { bRequiresHitProxies = bRequires; return *this; }
+	FWorldInitializationValues& CreatePhysicsScene(const bool bCreate) { bCreatePhysicsScene = bCreate; return *this; }
+	FWorldInitializationValues& CreateNavigation(const bool bCreate) { bCreateNavigation = bCreate; return *this; }
+	FWorldInitializationValues& CreateAISystem(const bool bCreate) { bCreateAISystem = bCreate; return *this; }
+	FWorldInitializationValues& ShouldSimulatePhysics(const bool bInShouldSimulatePhysics) { bShouldSimulatePhysics = bInShouldSimulatePhysics; return *this; }
+	FWorldInitializationValues& EnableTraceCollision(const bool bInEnableTraceCollision) { bEnableTraceCollision = bInEnableTraceCollision; return *this; }
+	FWorldInitializationValues& ForceUseMovementComponentInNonGameWorld(const bool bInForceUseMovementComponentInNonGameWorld) { bForceUseMovementComponentInNonGameWorld = bInForceUseMovementComponentInNonGameWorld; return *this; }
+	FWorldInitializationValues& SetTransactional(const bool bInTransactional) { bTransactional = bInTransactional; return *this; }
+	FWorldInitializationValues& CreateFXSystem(const bool bCreate) { bCreateFXSystem = bCreate; return *this; }
+	FWorldInitializationValues& CreateWorldPartition(const bool bCreate) { bCreateWorldPartition = bCreate; return *this; }
+	FWorldInitializationValues& EnableWorldPartitionStreaming(const bool bEnableStreaming) { bEnableWorldPartitionStreaming = bEnableStreaming; return *this; }
+	FWorldInitializationValues& SetDefaultGameMode(TSubclassOf<class AGameModeBase> GameMode) { DefaultGameMode = GameMode; return *this; }
+};
+
+
diff --git a/Engine/Source/Runtime/NavigationSystem/Private/NavigationSystemTypes.cpp b/Engine/Source/Runtime/NavigationSystem/Private/NavigationSystemTypes.cpp
--- a/Engine/Source/Runtime/NavigationSystem/Private/NavigationSystemTypes.cpp
+++ b/Engine/Source/Runtime/NavigationSystem/Private/NavigationSystemTypes.cpp
@@ -312,7 +312,7 @@ FNavigationLink INavLinkCustomInterface::GetModifier(const INavLinkCustomInterfa
 	return LinkMod;
 }
 
-void INavLinkCustomInterface::OnPreWorldInitialization(UWorld* World, const UWorld::InitializationValues IVS)
+void INavLinkCustomInterface::OnPreWorldInitialization(UWorld* World, const FWorldInitializationValues IVS)
 {
 	PRAGMA_DISABLE_DEPRECATION_WARNINGS
 	ResetUniqueId();
diff --git a/Engine/Source/Runtime/NavigationSystem/Public/NavLinkCustomInterface.h b/Engine/Source/Runtime/NavigationSystem/Public/NavLinkCustomInterface.h
--- a/Engine/Source/Runtime/NavigationSystem/Public/NavLinkCustomInterface.h
+++ b/Engine/Source/Runtime/NavigationSystem/Public/NavLinkCustomInterface.h
@@ -10,7 +10,10 @@
 #include "UObject/Interface.h"
 #include "NavAreas/NavArea.h"
 #include "AI/Navigation/NavLinkDefinition.h"
+#if UE_ENABLE_INCLUDE_ORDER_DEPRECATED_IN_5_4
 #include "Engine/World.h"
+#endif
+#include "Engine/WorldInitializationValues.h"
 #include "NavLinkCustomInterface.generated.h"
 
 /** 
@@ -102,7 +105,7 @@ class INavLinkCustomInterface
 	UE_DEPRECATED(5.3, "LinkIds are now based on a FNavLinkId Hash. If your project is still using any of the old incremental Ids then this function must be called still (typically by existing engine code), otherwise it is not necessary.")
 	static NAVIGATIONSYSTEM_API void ResetUniqueId();
 
-	static NAVIGATIONSYSTEM_API void OnPreWorldInitialization(UWorld* World, const UWorld::InitializationValues IVS);
+	static NAVIGATIONSYSTEM_API void OnPreWorldInitialization(UWorld* World, const FWorldInitializationValues IVS);
 
 	UE_DEPRECATED(5.3, "LinkIds are now based on FNavLinkId using FNavLinkId::GenerateUniqueId().")
 	static uint32 NextUniqueId;
