Make SceneComponent caches allocated on-demand.

- Less than 20% of scene comp allocate WorldRotation cache, and less than 70% allocate RelativeRotationCache on a internal game content  (there's a win as long as the allocation rate is less than 86%).
- PerfB testing showed no adverse performance effect, although memory effect is also inconclusive.

diff --git a/Engine/Source/Runtime/Engine/Classes/Engine/EngineTypes.h b/Engine/Source/Runtime/Engine/Classes/Engine/EngineTypes.h
--- a/Engine/Source/Runtime/Engine/Classes/Engine/EngineTypes.h
+++ b/Engine/Source/Runtime/Engine/Classes/Engine/EngineTypes.h
@@ -1838,28 +1838,31 @@ struct FBasedPosition
 struct FRotationConversionCache
 {
 	FRotationConversionCache()
-		: CachedQuat(FQuat::Identity)
-		, CachedRotator(FRotator::ZeroRotator)
 	{
 	}
 
 	/** Convert a FRotator to FQuat. Uses the cached conversion if possible, and updates it if there was no match. */
 	FORCEINLINE_DEBUGGABLE FQuat RotatorToQuat(const FRotator& InRotator) const
 	{
-		if (CachedRotator != InRotator)
+		FPayload& Payload = GetOrCreatePayload();
+		if (LIKELY(Payload.CachedRotator != InRotator))
 		{
-			CachedRotator = InRotator.GetNormalized();
-			CachedQuat = CachedRotator.Quaternion();
+			Payload.CachedRotator = InRotator.GetNormalized();
+			Payload.CachedQuat = Payload.CachedRotator.Quaternion();
 		}
-		return CachedQuat;
+		return Payload.CachedQuat;
 	}
 
 	/** Convert a FRotator to FQuat. Uses the cached conversion if possible, but does *NOT* update the cache if there was no match. */
 	FORCEINLINE_DEBUGGABLE FQuat RotatorToQuat_ReadOnly(const FRotator& InRotator) const
 	{
-		if (CachedRotator == InRotator)
+		if (LIKELY(PayloadPtr.IsValid()))
 		{
-			return CachedQuat;
+			FPayload& Payload = *PayloadPtr;
+			if (LIKELY(Payload.CachedRotator == InRotator))
+			{
+				return Payload.CachedQuat;
+			}
 		}
 		return InRotator.Quaternion();
 	}
@@ -1867,20 +1870,25 @@ struct FRotationConversionCache
 	/** Convert a FQuat to FRotator. Uses the cached conversion if possible, and updates it if there was no match. */
 	FORCEINLINE_DEBUGGABLE FRotator QuatToRotator(const FQuat& InQuat) const
 	{
-		if (CachedQuat != InQuat)
+		FPayload& Payload = GetOrCreatePayload();
+		if (LIKELY(Payload.CachedQuat != InQuat))
 		{
-			CachedQuat = InQuat.GetNormalized();
-			CachedRotator = CachedQuat.Rotator();
+			Payload.CachedQuat = InQuat.GetNormalized();
+			Payload.CachedRotator = Payload.CachedQuat.Rotator();
 		}
-		return CachedRotator;
+		return Payload.CachedRotator;
 	}
 
 	/** Convert a FQuat to FRotator. Uses the cached conversion if possible, but does *NOT* update the cache if there was no match. */
 	FORCEINLINE_DEBUGGABLE FRotator QuatToRotator_ReadOnly(const FQuat& InQuat) const
 	{
-		if (CachedQuat == InQuat)
+		if (LIKELY(PayloadPtr.IsValid()))
 		{
-			return CachedRotator;
+			FPayload& Payload = *PayloadPtr;
+			if (LIKELY(Payload.CachedQuat == InQuat))
+			{
+				return Payload.CachedRotator;
+			}
 		}
 		return InQuat.GetNormalized().Rotator();
 	}
@@ -1888,20 +1896,25 @@ struct FRotationConversionCache
 	/** Version of QuatToRotator when the Quat is known to already be normalized. */
 	FORCEINLINE_DEBUGGABLE FRotator NormalizedQuatToRotator(const FQuat& InNormalizedQuat) const
 	{
-		if (CachedQuat != InNormalizedQuat)
+		FPayload& Payload = GetOrCreatePayload();
+		if (LIKELY(Payload.CachedQuat != InNormalizedQuat))
 		{
-			CachedQuat = InNormalizedQuat;
-			CachedRotator = InNormalizedQuat.Rotator();
+			Payload.CachedQuat = InNormalizedQuat;
+			Payload.CachedRotator = InNormalizedQuat.Rotator();
 		}
-		return CachedRotator;
+		return Payload.CachedRotator;
 	}
 
 	/** Version of QuatToRotator when the Quat is known to already be normalized. Does *NOT* update the cache if there was no match. */
 	FORCEINLINE_DEBUGGABLE FRotator NormalizedQuatToRotator_ReadOnly(const FQuat& InNormalizedQuat) const
 	{
-		if (CachedQuat == InNormalizedQuat)
+		if (LIKELY(PayloadPtr.IsValid()))
 		{
-			return CachedRotator;
+			FPayload& Payload = *PayloadPtr;
+			if (LIKELY(Payload.CachedQuat == InNormalizedQuat))
+			{
+				return Payload.CachedRotator;
+			}
 		}
 		return InNormalizedQuat.Rotator();
 	}
@@ -1909,18 +1922,57 @@ struct FRotationConversionCache
 	/** Return the cached Quat. */
 	FORCEINLINE_DEBUGGABLE FQuat GetCachedQuat() const
 	{
-		return CachedQuat;
+		if (UNLIKELY(!PayloadPtr.IsValid()))
+		{
+			return FQuat::Identity;
+		}
+		return PayloadPtr->CachedQuat;
 	}
 
 	/** Return the cached Rotator. */
 	FORCEINLINE_DEBUGGABLE FRotator GetCachedRotator() const
 	{
-		return CachedRotator;
+		if (UNLIKELY(!PayloadPtr.IsValid()))
+		{
+			return FRotator::ZeroRotator;
+		}
+		return PayloadPtr->CachedRotator;
+	}
+
+	FRotationConversionCache& operator=(const FRotationConversionCache& Other)
+	{
+		if (LIKELY(PayloadPtr.IsValid() || Other.PayloadPtr.IsValid()))
+		{
+			FPayload& Payload = GetOrCreatePayload();
+			Payload.CachedQuat = Other.GetCachedQuat();
+			Payload.CachedRotator = Other.GetCachedRotator();
+		}
+		return *this;
 	}
 
 private:
-	mutable FQuat		CachedQuat;		// FQuat matching CachedRotator such that CachedQuat.Rotator() == CachedRotator.
-	mutable FRotator	CachedRotator;	// FRotator matching CachedQuat such that CachedRotator.Quaternion() == CachedQuat.
+
+	struct FPayload
+	{
+		mutable FQuat				CachedQuat;		// FQuat matching CachedRotator such that CachedQuat.Rotator() == CachedRotator.
+		mutable FRotator			CachedRotator;	// FRotator matching CachedQuat such that CachedRotator.Quaternion() == CachedQuat.
+
+		FPayload()
+		: CachedQuat(FQuat::Identity)
+		, CachedRotator(FRotator::ZeroRotator)
+		{}
+	};
+
+	inline FPayload&				GetOrCreatePayload() const
+	{
+		if (UNLIKELY(!PayloadPtr.IsValid()))
+		{
+			PayloadPtr = MakeUnique<FPayload>();
+		}
+		return *PayloadPtr;
+	}
+
+	mutable TUniquePtr<FPayload>	PayloadPtr;
 };
 
 /** A line of subtitle text and the time at which it should be displayed. */
