Support resizing of the ParticleGPUSimulation resources. Cascade emitter are currently allocating a lot of render resources to support gpu simulation.

Ex:
position        PF_A32B32G32R32F   1024x1024 (16MB)
velocity        PF_FloatRGBA             1024x1024 (8MB)
Attribute1    PF_B8G8R8A8             1024x1024 (4MB)
Attribute2    PF_B8G8R8A8             1024x1024 (4MB)

Thoses textures are all double buffered and at full sizes, they represent 64 MB of render targets.

When GFXCascadeGpuSpriteAllowDynAllocs is enabled, resizing happens when the TileAllocator of the FParticleSimulationResources run out of tiles. We then allocate new tiles and resize the gpu simulation resources render targets, each time by a factor 2 (X,Y). This introduce a TileZ index (TilePageIndex) which allow us to compute a tile page index at runtime and retrieve the resized tile coordinates. Resizing is limited by the maximum morton index possible (65535) and also the original resource sizes provided by the device profile (fx.GPUSimulationTextureSizeX,fx.GPUSimulationTextureSizeY).

The resizing feature can be enable with fx.Cascade.GpuSpriteDynamicAllocations cvar, and the initial size is controlled by fx.GPUSimulationDynTextureSizeXY.

Depending on the platform and map content, this can saves alot of memorey, ranging from 7MB to 56MB depending on the platforms.

*Resizing feature currently dont support GPU sorting mode. Radix sort shader algo will need to support keys/values of different format than 32 bits to work with the dynamic resizing of the gpu sim resources. To work around this limitation, we currently, ignore gpu sorting mode on cascade emitters when resizing is enable. A Jira was created to track this issue, and potentially add support in the future.

diff --git a/Engine/Shaders/Private/ParticleBoundsShader.usf b/Engine/Shaders/Private/ParticleBoundsShader.usf
--- a/Engine/Shaders/Private/ParticleBoundsShader.usf
+++ b/Engine/Shaders/Private/ParticleBoundsShader.usf
@@ -11,6 +11,7 @@
 ------------------------------------------------------------------------------*/
 
 #include "Common.ush"
+#include "ParticleSimulationCommon.ush"
 
 /** The number of threads with which to perform raking. */
 #define RAKING_THREADS 32
@@ -22,7 +23,9 @@
 #define MIN_FLT (-3.402823466e+38f)
 
 /** Input buffer containing particle indices. */
-Buffer<float2> InParticleIndices;
+Buffer<float4> InParticleIndices;
+float2 TilePageScale;
+
 /** Texture containing particle positions. */
 Texture2D PositionTexture;
 
@@ -82,7 +85,8 @@ void ComputeParticleBounds(
 		if ( InputIndex < ParticleBounds.ParticleCount )
 		{
 			// Read in the particle index and its position.
-			const float2 ParticleIndex = float2(InParticleIndices[InputIndex]);
+			float3 ParticleIndices = InParticleIndices[InputIndex].xyz;
+			const float2 ParticleIndex = ParticleIndices.xy * TilePageScale.xy + GetTilePageOffset(ParticleIndices.z, TilePageScale);
 			int3 ParticleTexel = int3(ParticleIndex.xy * int2(TextureSizeX, TextureSizeY), 0);
 			const float4 ParticlePosition = PositionTexture.Load(ParticleTexel);
 			if ( ParticlePosition.w <= 1.0f )
diff --git a/Engine/Shaders/Private/ParticleGPUSpriteVertexFactory.ush b/Engine/Shaders/Private/ParticleGPUSpriteVertexFactory.ush
--- a/Engine/Shaders/Private/ParticleGPUSpriteVertexFactory.ush
+++ b/Engine/Shaders/Private/ParticleGPUSpriteVertexFactory.ush
@@ -8,6 +8,7 @@
 
 #include "VertexFactoryCommon.ush"
 #include "ParticleVertexFactoryCommon.ush"
+#include "ParticleSimulationCommon.ush"
 
 #define USE_PARTICLE_LIGHTING_OFFSET (FEATURE_LEVEL >= FEATURE_LEVEL_SM5 && !MATERIAL_SHADINGMODEL_UNLIT && MATERIAL_DOMAIN_SURFACE)
 #define USE_PARTICLE_POSITION (NEEDS_PARTICLE_POSITION || PASS_NEEDS_PRIMITIVE_VOLUME_BOUNDS)
@@ -18,10 +19,12 @@
 #define USE_PARTICLE_SPRITE_ROTATION (NEEDS_PARTICLE_SPRITE_ROTATION)
 
 /** Buffer containing particle indices. */
-Buffer<float2> ParticleIndices;
+Buffer<float4> ParticleIndices;
 
 /** Offset in to the particle indices buffer. */
 uint ParticleIndicesOffset;
+float3 TilePageScale;
+
 
 /** Texture containing positions for all particles. */
 Texture2D PositionTexture;
@@ -245,10 +248,11 @@ FVertexFactoryIntermediates GetVertexFactoryIntermediates( FVertexFactoryInput I
 {
 	// Sample the position and velocity textures to get the current state of the particle.
 	uint InstanceId = GetInstanceId(Input.InstanceId) * PARTICLES_PER_INSTANCE + Input.VertexId / 4;
-	float2 ParticleIndex = ParticleIndices[ParticleIndicesOffset + InstanceId];
-	float4 PositionSample = Texture2DSampleLevel(PositionTexture, PositionTextureSampler, ParticleIndex, 0);
-	float4 VelocitySample = Texture2DSampleLevel(VelocityTexture, VelocityTextureSampler, ParticleIndex, 0);
-	float4 AttributeSample = Texture2DSampleLevel(AttributesTexture, AttributesTextureSampler, ParticleIndex, 0);
+	float3 ParticleIndex = ParticleIndices[ParticleIndicesOffset + InstanceId].xyz;
+	ParticleIndex.xy = ParticleIndex.xy * TilePageScale.xy + GetTilePageOffset(ParticleIndex.z, TilePageScale.xy, TilePageScale.z);
+	float4 PositionSample = Texture2DSampleLevel(PositionTexture, PositionTextureSampler, ParticleIndex.xy, 0);
+	float4 VelocitySample = Texture2DSampleLevel(VelocityTexture, VelocityTextureSampler, ParticleIndex.xy, 0);
+	float4 AttributeSample = Texture2DSampleLevel(AttributesTexture, AttributesTextureSampler, ParticleIndex.xy, 0);
 
 	// For clarity, store some information in local variables.
 	const float RelativeTime = PositionSample.w;
diff --git a/Engine/Shaders/Private/ParticleInjectionShader.usf b/Engine/Shaders/Private/ParticleInjectionShader.usf
--- a/Engine/Shaders/Private/ParticleInjectionShader.usf
+++ b/Engine/Shaders/Private/ParticleInjectionShader.usf
@@ -5,6 +5,7 @@
 =============================================================================*/
 
 #include "Common.ush"
+#include "ParticleSimulationCommon.ush"
 
 /*-----------------------------------------------------------------------------
 	Shared declarations and functions.
@@ -38,7 +39,7 @@ struct FVertexInput
 	/** The simulation attributes for the particle. */
 	float4 SimulationAttributes : ATTRIBUTE3;
 	/** The offset in to the texture to store the particle. */
-	float2 ParticleIndex		: ATTRIBUTE4;
+	float3 ParticleIndex	: ATTRIBUTE4;
 	/** The texture coordinate of the corner of the sprite being rendered. */
 	float2 TexCoord				: ATTRIBUTE5;
 };
@@ -49,7 +50,7 @@ void VertexMain(
 	out float4 OutPosition : SV_POSITION
 	)
 {
-	float2 ParticleCoord = Input.TexCoord.xy * ParticleInjection.PixelScale.xy + Input.ParticleIndex.xy;
+	float2 ParticleCoord = Input.TexCoord.xy * ParticleInjection.TexCoordScale.xy + Input.ParticleIndex.xy * ParticleInjection.TilePageScale.xy + GetTilePageOffset(Input.ParticleIndex.z, ParticleInjection.TilePageScale, true);
 	OutPosition = float4(
 		ParticleCoord.xy * float2(2.0f,-2.0f) + float2(-1.0f,1.0f),
 		0,
diff --git a/Engine/Shaders/Private/ParticleSimulationCommon.ush b/Engine/Shaders/Private/ParticleSimulationCommon.ush
new file mode 100644
--- /dev/null
+++ b/Engine/Shaders/Private/ParticleSimulationCommon.ush
@@ -0,0 +1,22 @@
+// Copyright Epic Games, Inc. All Rights Reserved.
+
+/*=============================================================================
+	ParticleSimulationCommon.usf: Particle Simulation shared functions and defines.
+=============================================================================*/
+
+#pragma once
+
+#include "MortonCode.ush"
+
+float2 GetTilePageOffset(float ZTileIndex, float2 PageScale, bool bResize = true)
+{
+	float2 UV  = 0.0f;
+
+	if(bResize)
+	{
+		uint2 PageOffset = MortonDecode(ZTileIndex);
+		UV = PageOffset * PageScale;
+	}
+
+	return UV;
+}
\ No newline at end of file
diff --git a/Engine/Shaders/Private/ParticleSimulationShader.usf b/Engine/Shaders/Private/ParticleSimulationShader.usf
--- a/Engine/Shaders/Private/ParticleSimulationShader.usf
+++ b/Engine/Shaders/Private/ParticleSimulationShader.usf
@@ -6,6 +6,7 @@
 
 #include "Common.ush"
 #include "SceneTexturesCommon.ush"
+#include "ParticleSimulationCommon.ush"
 
 #if DISTANCE_FIELD_COLLISION
 	#include "DistanceField/GlobalDistanceFieldShared.ush"
@@ -29,9 +30,8 @@ struct FShaderInterpolants
 ------------------------------------------------------------------------------*/
 #if VERTEXSHADER
 
-/** The width and height of a tile in the particle simulation texture */
-float TileSizeX;
-float TileSizeY;
+float2 TexCoordScale;
+float2 TilePageScale;
 
 struct FVertexInput
 {
@@ -44,12 +44,12 @@ struct FVertexInput
 	uint InstanceId	: SV_InstanceID;
 #else // Mobile
 	/** Tile offsets in per-instance data */
-	float2 TileOffsets : ATTRIBUTE1;
+	float4 TileOffsets : ATTRIBUTE1;
 #endif
 };
 
 /** Buffer from which to read tile offsets. */
-Buffer<float2> TileOffsets;
+Buffer<float4> TileOffsets;
 
 void VertexMain(
 	in FVertexInput Input,
@@ -59,12 +59,12 @@ void VertexMain(
 {
 #if FEATURE_LEVEL >= FEATURE_LEVEL_SM4
 	uint InstanceId = Input.InstanceId * TILES_PER_INSTANCE + Input.VertexId / 4;
-	float2 TileOffset = TileOffsets[InstanceId];
+	float3 TileOffset = TileOffsets[InstanceId].xyz;
 #else
-	float2 TileOffset = Input.TileOffsets.xy;
+	float3 TileOffset = Input.TileOffsets.xyz;
 #endif
 	
-	float2 TileCoord = Input.TexCoord.xy * float2(TileSizeX,TileSizeY) + TileOffset;
+	float2 TileCoord = Input.TexCoord.xy * TexCoordScale.xy + TileOffset * TilePageScale.xy + GetTilePageOffset(TileOffset.z, TilePageScale, true);
 
 	OutPosition = float4(
 		TileCoord.xy * float2(2.0f,-2.0f) + float2(-1.0f,1.0f),
diff --git a/Engine/Source/Runtime/Engine/Private/Particles/FXSystem.cpp b/Engine/Source/Runtime/Engine/Private/Particles/FXSystem.cpp
--- a/Engine/Source/Runtime/Engine/Private/Particles/FXSystem.cpp
+++ b/Engine/Source/Runtime/Engine/Private/Particles/FXSystem.cpp
@@ -423,11 +423,24 @@ void FFXSystem::DrawDebug( FCanvas* Canvas )
 	}
 }
 
+DECLARE_GPU_DRAWCALL_STAT(FXSystemPreInitViews);
+
 void FFXSystem::PreInitViews(FRDGBuilder& GraphBuilder, bool bAllowGPUParticleUpdate, const TArrayView<const FSceneViewFamily*>& ViewFamilies, const FSceneViewFamily* CurrentFamily)
 {
 	if (RHISupportsGPUParticles())
 	{
-		AdvanceGPUParticleFrame(bAllowGPUParticleUpdate);
+		RDG_GPU_STAT_SCOPE(GraphBuilder, FXSystemPreInitViews);
+		RDG_CSV_STAT_EXCLUSIVE_SCOPE(GraphBuilder, FXSystem);
+
+		AddPass(
+			GraphBuilder,
+			RDG_EVENT_NAME("FFXSystem::PreInitViews"),
+			[this, bAllowGPUParticleUpdate](FRHICommandListImmediate& RHICmdList)
+			{
+
+				AdvanceGPUParticleFrame(RHICmdList, bAllowGPUParticleUpdate);
+			}
+		);
 	}
 }
 
diff --git a/Engine/Source/Runtime/Engine/Private/Particles/FXSystemPrivate.h b/Engine/Source/Runtime/Engine/Private/Particles/FXSystemPrivate.h
--- a/Engine/Source/Runtime/Engine/Private/Particles/FXSystemPrivate.h
+++ b/Engine/Source/Runtime/Engine/Private/Particles/FXSystemPrivate.h
@@ -100,7 +100,7 @@ struct FNewParticle
 	/** Random selection of orbit attributes. */
 	float RandomOrbit;
 	/** The offset at which to inject the new particle. */
-	FVector2f Offset;
+	FVector3f Offset;
 };
 
 
@@ -131,6 +131,8 @@ public:
 	FRHITexture2D* AttributesTextureRHI;
 	/** LWC tile offset, will be 0,0,0 for localspace emitters. */
 	FVector3f LWCTile;
+	/** Tile page offset factors associated with the GPU particle simulation resources. */
+	FVector3f TilePageScale;
 
 	FGPUSpriteVertexFactory()
 		: FParticleVertexFactoryBase(PVFT_MAX, ERHIFeatureLevel::Num)
@@ -139,6 +141,7 @@ public:
 		, VelocityTextureRHI(nullptr)
 		, AttributesTextureRHI(nullptr)
 		, LWCTile(FVector3f::ZeroVector)
+		, TilePageScale(FVector3f::OneVector)
 	{}
 
 	/**
@@ -320,7 +323,7 @@ private:
 	/**
 	 * Prepares GPU particles for simulation and rendering in the next frame.
 	 */
-	void AdvanceGPUParticleFrame(bool bAllowGPUParticleUpdate);
+	void AdvanceGPUParticleFrame(FRHICommandListImmediate& RHICmdList, bool bAllowGPUParticleUpdate);
 
 	bool UsesGlobalDistanceFieldInternal() const;
 	bool UsesDepthBufferInternal() const;
diff --git a/Engine/Source/Runtime/Engine/Private/Particles/ParticleGpuSimulation.cpp b/Engine/Source/Runtime/Engine/Private/Particles/ParticleGpuSimulation.cpp
--- a/Engine/Source/Runtime/Engine/Private/Particles/ParticleGpuSimulation.cpp
+++ b/Engine/Source/Runtime/Engine/Private/Particles/ParticleGpuSimulation.cpp
@@ -41,6 +41,7 @@
 #include "MeshMaterialShader.h"
 #include "Stats/StatsTrace.h"
 #include "ShaderParameterUtils.h"
+#include "RenderGraphUtils.h"
 
 DECLARE_CYCLE_STAT(TEXT("GPUSpriteEmitterInstance Init GT"), STAT_GPUSpriteEmitterInstance_Init, STATGROUP_Particles);
 DECLARE_GPU_STAT_NAMED(ParticleSimulation, TEXT("Particle Simulation"));
@@ -77,6 +78,22 @@ static FAutoConsoleVariableRef CVarFXCascadeGpuSpriteRenderingEnabled(
 	ECVF_Default
 );
 
+int32 GParticleSimulationDynTextureSizeXY = 16;
+FAutoConsoleVariableRef CVarParticleSimulationDynTextureSizeXY(
+	TEXT("fx.GPUSimulationDynTextureSizeXY"),
+	GParticleSimulationDynTextureSizeXY,
+	TEXT("GPU Particle simulation texture XY dimension when dynamic resizing is enabled (default=16)"),
+	ECVF_ReadOnly
+);
+
+static bool GFXCascadeGpuSpriteAllowDynAllocs = false;
+static FAutoConsoleVariableRef CVarFXCascadeGpuSpriteDynamicAllocations(
+	TEXT("fx.Cascade.GpuSpriteDynamicAllocations"),
+	GFXCascadeGpuSpriteAllowDynAllocs,
+	TEXT("Controls if gpu sprite are allocated dynamically"),
+	ECVF_ReadOnly
+);
+
 /** The tile size. Texture space is allocated in TileSize x TileSize units. */
 const int32 GParticleSimulationTileSize = 4;
 const int32 GParticlesPerTile = GParticleSimulationTileSize * GParticleSimulationTileSize;
@@ -84,11 +101,6 @@ const int32 GParticlesPerTile = GParticleSimulationTileSize * GParticleSimulatio
 /** Tile size must be power-of-two and <= each dimension of the simulation texture. */
 static_assert((GParticleSimulationTileSize & (GParticleSimulationTileSize - 1)) == 0, "Particle simulation tile size is not a power of two.");
 
-/** How many tiles are in the simulation textures. */
-int32 GParticleSimulationTileCountX = 0;
-int32 GParticleSimulationTileCountY = 0;
-int32 GParticleSimulationTileCount = 0;
-
 /** GPU particle rendering code assumes that the number of particles per instanced draw is <= 16. */
 static_assert(MAX_PARTICLES_PER_INSTANCE <= 16, "Max particles per instance is greater than 16.");
 /** Also, it must be a power of 2. */
@@ -128,21 +140,69 @@ public:
 
 	/** Default constructor. */
 	FParticleTileAllocator()
-		: FreeTileCount(GParticleSimulationTileCount)
+		: bDirtyAlloc(false)
 	{
-		/** Texture size must be power-of-two. */
-		check((GParticleSimulationTextureSizeX & (GParticleSimulationTextureSizeX - 1)) == 0); // fx.GPUSimulationTextureSizeX is not a power of two.
-		check((GParticleSimulationTextureSizeY & (GParticleSimulationTextureSizeY - 1)) == 0); // fx.GPUSimulationTextureSizeY is not a power of two.
+		InitialTileCount = 0;
+		bAllowResizing = false;
+	}
 
-		check(GParticleSimulationTileSize <= GParticleSimulationTextureSizeX); // Particle simulation tile size is larger than fx.GPUSimulationTextureSizeX.
-		check(GParticleSimulationTileSize <= GParticleSimulationTextureSizeY); // Particle simulation tile size is larger than fx.GPUSimulationTextureSizeY.
+	/**
+	 * Init the tile allocator
+	 *
+	 */
+	void Init(bool bAllowTileResizing, uint32 TileCount, uint32 InMaxTileCount)
+	{
+		/** Check if the tile can be resized at runtime. */
+		bAllowResizing = bAllowTileResizing;
+		InitialTileCount = TileCount;
+		MaxTileCount = InMaxTileCount;
 
-		FreeTiles.AddUninitialized(GParticleSimulationTileCount);
+		FreeTiles.AddUninitialized(InitialTileCount);
 
-		for ( int32 TileIndex = 0; TileIndex < GParticleSimulationTileCount; ++TileIndex )
+		for (uint32 TileIndex = 0; TileIndex < InitialTileCount; ++TileIndex)
 		{
-			FreeTiles[TileIndex] = GParticleSimulationTileCount - TileIndex - 1;
+			FreeTiles[TileIndex] = InitialTileCount - TileIndex - 1;
+		}
+
+		ResizeTileAllocCount = 0;
+	}
+
+	/**
+	 * Resize the tile allocator to support more tiles.
+	 * 
+	 * Each time we resize, the associated textures width and height are multiplied by 2.
+	 * This means that each time we increase the tile allocator, the tile count will grow
+	 * by a 4 ^ (resizecount) factor. Ex: Initial textures is driven by GParticleSimulationDynTextureSizeXY,
+	 * and contains one tile. The first resize will contain 4 tiles, the second resize will have 16 tiles, and so on.
+	 * The tile allocated from a resize are allocated following a Morton Z curve pattern.
+	 * 
+	 */
+
+	bool Resize()
+	{
+		check(bAllowResizing);
+
+		static uint32 MaxMortonIndex = 65535;// (2 ^ 16) -1
+		uint32 OldTileCount = InitialTileCount * FMath::Pow(4, (float)ResizeTileAllocCount);
+		uint32 ResizedTileCount = OldTileCount * 4;
+
+		// 1-check to make sure we dont bust the maximum tile allocation allowed (see GParticleSimulationTextureSizeX, GParticleSimulationTextureSizeY).
+		// 2-check to make sure we dont bust the morton index capacity (16 bits for each X,Y coord -> 2^16 - 1) == 65535)
+		if(ResizedTileCount <= MaxTileCount && ResizedTileCount <= MaxMortonIndex)
+		{ 
+			FreeTiles.AddUninitialized(ResizedTileCount - OldTileCount);
+
+			for (uint32 TileIndex = 0; TileIndex < (ResizedTileCount - OldTileCount); ++TileIndex)
+			{
+				FreeTiles[TileIndex] = ResizedTileCount - TileIndex - 1;
+			}
+
+			bDirtyAlloc = true;
+			ResizeTileAllocCount++;
+			return true;
 		}
+
+		return false;
 	}
 
 	/**
@@ -152,11 +212,16 @@ public:
 	uint32 Allocate()
 	{
 		FScopeLock Lock(&CriticalSection);
-		if ( FreeTileCount > 0 )
+		if (FreeTiles.Num() > 0)
 		{
-			FreeTileCount--;
-			return FreeTiles[FreeTileCount];
+			return FreeTiles.Pop();
 		}
+
+		if (bAllowResizing && Resize())
+		{
+			return FreeTiles.Pop();
+		}
+
 		return INDEX_NONE;
 	}
 
@@ -167,10 +232,9 @@ public:
 	void Free( int32 TileIndex )
 	{
 		FScopeLock Lock(&CriticalSection);
-		check( TileIndex < GParticleSimulationTileCount );
-		check( FreeTileCount < GParticleSimulationTileCount );
-		FreeTiles[FreeTileCount] = TileIndex;
-		FreeTileCount++;
+		check( TileIndex < GetAllocatedTileCount());
+		check( FreeTiles.Num() < GetAllocatedTileCount());
+		FreeTiles.Add(TileIndex);
 	}
 
 	/**
@@ -179,15 +243,52 @@ public:
 	int32 GetFreeTileCount() const
 	{
 		FScopeLock Lock(&CriticalSection);
-		return FreeTileCount;
+		return FreeTiles.Num();
+	}
+
+	/**
+	 * Returns the resize factor of the gpu sim textures.
+	 */
+	int32 GetResizeFactor() const
+	{
+		FScopeLock Lock(&CriticalSection);
+		return FMath::Pow(2, (float)ResizeTileAllocCount);
+	}
+
+	/**
+	 * Returns the max allocated tiles.
+	 */
+	int32 GetAllocatedTileCount() const
+	{
+		FScopeLock Lock(&CriticalSection);
+		return InitialTileCount * FMath::Pow(4, (float)ResizeTileAllocCount);
 	}
 
+	/**
+	 * Access function to know if the tile allocator supports resizing tiles.
+	 */
+	bool SupportResizingTiles() const
+	{
+		return bAllowResizing;
+	}
+
+	/* Tile allocator have resized. Need to update the GPU resources*/
+	bool bDirtyAlloc;
+
 private:
 
 	/** List of free tiles. */
 	TArray<uint32> FreeTiles;
-	/** How many tiles are in the free list. */
-	int32 FreeTileCount;
+
+	/** Support dynamic resizing of tiles*/
+	bool bAllowResizing;
+
+	/** Number of time the TileAllocator resized*/
+	uint32 ResizeTileAllocCount;
+
+	/** Tile information*/
+	uint32 InitialTileCount;
+	uint32 MaxTileCount;
 
 	mutable FCriticalSection CriticalSection;
 };
@@ -205,6 +306,10 @@ struct FParticleIndex
 	FFloat16 X;
 	/** The Y coordinate of the particle within the texture. */
 	FFloat16 Y;
+	/** Tile page index to support dynamic reallocation of tiles. */
+	FFloat16 Z;
+	/*Padding*/
+	FFloat16 W;
 };
 
 /**
@@ -214,6 +319,9 @@ class FParticleStateTextures : public FRenderResource
 {
 public:
 
+	const TCHAR* PaticlePositonTextureName = TEXT("FParticleStatePosition");
+	const TCHAR* PaticleVelocityTextureName = TEXT("FParticleStateVelocity");
+
 	/** Contains the positions of all simulating particles. */
 	FTexture2DRHIRef PositionTextureRHI;
 	/** Contains the velocity of all simulating particles. */
@@ -221,60 +329,41 @@ public:
 
 	bool bTexturesCleared;
 	int32 ParticleStateIndex = 0;
+	FIntPoint RenderTargetSize;
+
+	FTextureRHIRef CreateStateTextureRHI(const TCHAR* Name, EPixelFormat PixelFormat)
+	{
+		const FRHITextureCreateDesc Desc =
+			FRHITextureCreateDesc::Create2D(Name)
+			.SetExtent(RenderTargetSize.X, RenderTargetSize.Y)
+			.SetFormat(PixelFormat)
+			.SetClearValue(FClearValueBinding::Transparent)
+			.SetFlags(ETextureCreateFlags::RenderTargetable | ETextureCreateFlags::ShaderResource)
+			.SetInitialState(ERHIAccess::SRVMask);
+
+		return RHICreateTexture(Desc);
+	}
 
 	/**
 	 * Initialize RHI resources used for particle simulation.
 	 */
 	virtual void InitRHI(FRHICommandListBase& RHICmdList) override
 	{
-		const int32 SizeX = GParticleSimulationTextureSizeX;
-		const int32 SizeY = GParticleSimulationTextureSizeY;
-
 		// 32-bit per channel RGBA texture for position.
 		check( !IsValidRef( PositionTextureRHI ) );
-
-#define PARTICLE_STATE_POSITION_TEXTURE_NAME	TEXT("ParticleStatePosition")
-#define PARTICLE_STATE_VELOCITY_TEXTURE_NAME	TEXT("ParticleStateVelocity")
-
-		const static FLazyName ClassName(TEXT("FParticleStateTextures"));
-		{
-			const FRHITextureCreateDesc Desc =
-				FRHITextureCreateDesc::Create2D(PARTICLE_STATE_POSITION_TEXTURE_NAME)
-				.SetExtent(SizeX, SizeY)
-				.SetFormat(PF_A32B32G32R32F)
-				.SetClearValue(FClearValueBinding::Transparent)
-				.SetFlags(ETextureCreateFlags::RenderTargetable | ETextureCreateFlags::ShaderResource)
-				.SetInitialState(ERHIAccess::SRVMask)
-				.SetClassName(ClassName);
-
-			PositionTextureRHI = RHICreateTexture(Desc);
-		}
+		PositionTextureRHI = CreateStateTextureRHI(PaticlePositonTextureName, PF_A32B32G32R32F);
 
 		// 16-bit per channel RGBA texture for velocity.
-		check( !IsValidRef( VelocityTextureRHI ) );
-
-		{
-			const FRHITextureCreateDesc Desc =
-				FRHITextureCreateDesc::Create2D(PARTICLE_STATE_VELOCITY_TEXTURE_NAME)
-				.SetExtent(SizeX, SizeY)
-				.SetFormat(PF_FloatRGBA)
-				.SetClearValue(FClearValueBinding::Transparent)
-				.SetFlags(ETextureCreateFlags::RenderTargetable | ETextureCreateFlags::ShaderResource)
-				.SetInitialState(ERHIAccess::SRVMask)
-				.SetClassName(ClassName);
-
-			VelocityTextureRHI = RHICreateTexture(Desc);
-		}
+		check(!IsValidRef(VelocityTextureRHI));
+		VelocityTextureRHI = CreateStateTextureRHI(PaticleVelocityTextureName, PF_FloatRGBA);
 
 		// using FName's ability to append a number to a string (..._0) without an extra string allocation, except suffixing is done when number > 0 hence the +1 here : 
-		FName PositionTextureName(PARTICLE_STATE_POSITION_TEXTURE_NAME, ParticleStateIndex + 1); 
-		FName VelocityTextureName(PARTICLE_STATE_VELOCITY_TEXTURE_NAME, ParticleStateIndex + 1);
+		FName PositionTextureName(PaticlePositonTextureName, ParticleStateIndex + 1);
+		FName VelocityTextureName(PaticleVelocityTextureName, ParticleStateIndex + 1);
 		PositionTextureRHI->SetName(PositionTextureName);
 		VelocityTextureRHI->SetName(VelocityTextureName);
-		RHICmdList.BindDebugLabelName(PositionTextureRHI, *PositionTextureName.ToString());
-		RHICmdList.BindDebugLabelName(VelocityTextureRHI, *VelocityTextureName.ToString());
-#undef PARTICLE_STATE_VELOCITY_TEXTURE_NAME
-#undef PARTICLE_STATE_POSITION_TEXTURE_NAME
+		RHICmdList.BindDebugLabelName(PositionTextureRHI, PaticlePositonTextureName);
+		RHICmdList.BindDebugLabelName(VelocityTextureRHI, PaticleVelocityTextureName);
 
 		bTexturesCleared = false;
 	}
@@ -287,6 +376,44 @@ public:
 		PositionTextureRHI.SafeRelease();
 		VelocityTextureRHI.SafeRelease();
 	}
+
+	/**
+	 * ResizeRHI esources used for particle simulation.
+	 */
+	void ResizeRHI(FRHICommandListImmediate& RHICmdList, int32 SizeX, int32 SizeY)
+	{
+		RenderTargetSize.X = SizeX;
+		RenderTargetSize.Y = SizeY;
+
+		// allocate the new textures
+		FTexture2DRHIRef NewPositionTextureRHI = CreateStateTextureRHI(PaticlePositonTextureName, PF_A32B32G32R32F);
+		FTexture2DRHIRef NewVelocityTextureRHI = CreateStateTextureRHI(PaticleVelocityTextureName, PF_FloatRGBA);
+		
+		RHICmdList.Transition(FRHITransitionInfo(NewPositionTextureRHI, ERHIAccess::SRVMask, ERHIAccess::CopyDest));
+		RHICmdList.Transition(FRHITransitionInfo(PositionTextureRHI, ERHIAccess::SRVMask, ERHIAccess::CopySrc));
+
+		RHICmdList.Transition(FRHITransitionInfo(NewVelocityTextureRHI, ERHIAccess::SRVMask, ERHIAccess::CopyDest));
+		RHICmdList.Transition(FRHITransitionInfo(VelocityTextureRHI, ERHIAccess::SRVMask, ERHIAccess::CopySrc));
+
+		FRHICopyTextureInfo CopyTextureInfo;
+		CopyTextureInfo.Size = FIntVector(PositionTextureRHI->GetSizeX(), PositionTextureRHI->GetSizeY(), 0);
+
+		RHICmdList.CopyTexture(PositionTextureRHI, NewPositionTextureRHI, CopyTextureInfo);
+		RHICmdList.CopyTexture(VelocityTextureRHI, NewVelocityTextureRHI, CopyTextureInfo);
+
+		RHICmdList.Transition(FRHITransitionInfo(NewPositionTextureRHI, ERHIAccess::CopyDest, ERHIAccess::SRVMask));
+		RHICmdList.Transition(FRHITransitionInfo(PositionTextureRHI, ERHIAccess::CopySrc, ERHIAccess::SRVMask));
+
+		RHICmdList.Transition(FRHITransitionInfo(NewVelocityTextureRHI, ERHIAccess::CopyDest, ERHIAccess::SRVMask));
+		RHICmdList.Transition(FRHITransitionInfo(VelocityTextureRHI, ERHIAccess::CopySrc, ERHIAccess::SRVMask));
+
+		// release the old textures
+		ReleaseRHI();
+
+		// assign the new textures
+		PositionTextureRHI = NewPositionTextureRHI;
+		VelocityTextureRHI = NewVelocityTextureRHI;
+	}
 };
 
 /**
@@ -296,35 +423,36 @@ class FParticleAttributesTexture : public FRenderResource
 {
 public:
 
+	const TCHAR* PaticleAttributeTextureName = TEXT("FParticleAttributesTexture");
+
 	/** Contains the attributes of all simulating particles. */
 	FTextureRHIRef TextureRHI;
 
-	/**
-	 * Initialize RHI resources used for particle simulation.
-	 */
-	virtual void InitRHI(FRHICommandListBase&) override
-	{
-		const int32 SizeX = GParticleSimulationTextureSizeX;
-		const int32 SizeY = GParticleSimulationTextureSizeY;
-
-#define ATTRIBUTES_TEXTURE_NAME	TEXT("ParticleAttributes")
-		const static FLazyName ClassName(TEXT("FParticleAttributesTexture"));
+	FIntPoint RenderTargetSize = FIntPoint::ZeroValue;
 
+	FTextureRHIRef CreateAttributesTextureRHI()
+	{
 		const FRHITextureCreateDesc Desc =
-			FRHITextureCreateDesc::Create2D(ATTRIBUTES_TEXTURE_NAME)
-			.SetExtent(SizeX, SizeY)
+			FRHITextureCreateDesc::Create2D(PaticleAttributeTextureName)
+			.SetExtent(RenderTargetSize.X, RenderTargetSize.Y)
 			.SetFormat(PF_B8G8R8A8)
 			.SetClearValue(FClearValueBinding::Transparent)
 			.SetFlags(ETextureCreateFlags::RenderTargetable | ETextureCreateFlags::ShaderResource | ETextureCreateFlags::NoFastClear)
-			.SetInitialState(ERHIAccess::RTV)
-			.SetClassName(ClassName);
+			.SetInitialState(ERHIAccess::RTV);
+
+		return RHICreateTexture(Desc);
+	}
 
-		TextureRHI = RHICreateTexture(Desc);
+	/**
+	 * Initialize RHI resources used for particle simulation.
+	 */
+	virtual void InitRHI(FRHICommandListBase&) override
+	{
+		TextureRHI = CreateAttributesTextureRHI();
 
-		static FName AttributesTextureName(ATTRIBUTES_TEXTURE_NAME);
-		TextureRHI->SetName(AttributesTextureName);
-		RHIBindDebugLabelName(TextureRHI, ATTRIBUTES_TEXTURE_NAME);
-#undef ATTRIBUTES_TEXTURE_NAME
+ 		FName AttributesTextureName(PaticleAttributeTextureName);
+ 		TextureRHI->SetName(AttributesTextureName);
+ 		RHIBindDebugLabelName(TextureRHI, PaticleAttributeTextureName);
 
 		FRHICommandListImmediate& RHICmdList = FRHICommandListExecutor::GetImmediateCommandList();
 		ClearRenderTarget(RHICmdList, TextureRHI);
@@ -338,6 +466,35 @@ public:
 	{
 		TextureRHI.SafeRelease();
 	}
+
+	/**
+	 * 
+	 */
+	void ResizeRHI(FRHICommandListImmediate& RHICmdList, int32 SizeX, int32 SizeY)
+	{
+		RenderTargetSize.X = SizeX;
+		RenderTargetSize.Y = SizeY;
+
+		// allocate the new textures
+		FTexture2DRHIRef NewAttributesTextureRHI = CreateAttributesTextureRHI();
+
+		RHICmdList.Transition(FRHITransitionInfo(NewAttributesTextureRHI, ERHIAccess::SRVMask, ERHIAccess::CopyDest));
+		RHICmdList.Transition(FRHITransitionInfo(TextureRHI, ERHIAccess::SRVMask, ERHIAccess::CopySrc));
+
+		FRHICopyTextureInfo CopyTextureInfo;
+		CopyTextureInfo.Size = FIntVector(TextureRHI->GetSizeX(), TextureRHI->GetSizeY(), 0);
+
+		RHICmdList.CopyTexture(TextureRHI, NewAttributesTextureRHI, CopyTextureInfo);
+
+		RHICmdList.Transition(FRHITransitionInfo(NewAttributesTextureRHI, ERHIAccess::CopyDest, ERHIAccess::SRVMask));
+		RHICmdList.Transition(FRHITransitionInfo(TextureRHI, ERHIAccess::CopySrc, ERHIAccess::SRVMask));
+
+		// release the old textures
+		ReleaseRHI();
+
+		// assign the new textures
+		TextureRHI = NewAttributesTextureRHI;
+	}
 };
 
 /**
@@ -359,14 +516,45 @@ public:
 	/** LWC tile offset, will be 0,0,0 for localspace emitters. */
 	FVector3f LWCTile = FVector3f::ZeroVector;
 
+	FIntPoint SimulationResourcesSize = FIntPoint::ZeroValue;
+
+	uint32 ParticleSimulationTileCountX = 0;
+	uint32 ParticleSimulationTileCountY = 0;
+	uint32 ParticleSimulationTileCount = 0;
+
 	/**
 	 * Initialize resources.
 	 */
 	void Init()
 	{
+		bool bAllowTileResizing = GFXCascadeGpuSpriteAllowDynAllocs;
+		SimulationResourcesSize.X = bAllowTileResizing ? GParticleSimulationDynTextureSizeXY : GParticleSimulationTextureSizeX;
+		SimulationResourcesSize.Y = bAllowTileResizing ? GParticleSimulationDynTextureSizeXY : GParticleSimulationTextureSizeY;
+		
+		check((SimulationResourcesSize.X & (SimulationResourcesSize.X - 1)) == 0); // fx.GPUSimulationTextureSizeX,or fx.GPUSimulationDynTextureSizeXY is not a power of two.
+		check((SimulationResourcesSize.Y & (SimulationResourcesSize.Y - 1)) == 0); // fx.GPUSimulationTextureSizeY,or fx.GPUSimulationDynTextureSizeXY is not a power of two.
+		check(GParticleSimulationTileSize <= SimulationResourcesSize.X); // Particle simulation tile size is larger than fx.GPUSimulationTextureSizeX, or fx.GPUSimulationDynTextureSizeXY.
+		check(GParticleSimulationTileSize <= SimulationResourcesSize.Y); // Particle simulation tile size is larger than fx.GPUSimulationTextureSizeY, or fx.GPUSimulationDynTextureSizeXY.
+
+		/** How many tiles are in the simulation textures. */
+		ParticleSimulationTileCountX = SimulationResourcesSize.X / GParticleSimulationTileSize;
+		ParticleSimulationTileCountY = SimulationResourcesSize.Y / GParticleSimulationTileSize;
+		ParticleSimulationTileCount = ParticleSimulationTileCountX * ParticleSimulationTileCountY;
+
+		// Compute max tile count resizing (only valid when bAllowTileResizing == true)
+		uint32 MaxParticleSimulationTileCount = GParticleSimulationTextureSizeX / GParticleSimulationTileSize * GParticleSimulationTextureSizeY / GParticleSimulationTileSize;
+
+		// check if we support tile resizing for this particle simulation resources
+		TileAllocator.Init(bAllowTileResizing, ParticleSimulationTileCount, MaxParticleSimulationTileCount);
+
 		// Help debugging by identifying each state :
 		StateTextures[0].ParticleStateIndex = 0;
+		StateTextures[0].RenderTargetSize = SimulationResourcesSize;
 		StateTextures[1].ParticleStateIndex = 1;
+		StateTextures[1].RenderTargetSize = SimulationResourcesSize;
+
+		RenderAttributesTexture.RenderTargetSize = SimulationResourcesSize;
+		SimulationAttributesTexture.RenderTargetSize = SimulationResourcesSize;
 
 		FParticleSimulationResources* ParticleResources = this;
 		ENQUEUE_RENDER_COMMAND(FInitParticleSimulationResourcesCommand)([ParticleResources](FRHICommandList& RHICmdList)
@@ -387,6 +575,26 @@ public:
 		});
 	}
 
+	/**
+	 * Resize gpu simaluation resources.
+	 */
+	void Resize(FRHICommandListImmediate& RHICmdList)
+	{
+		check(IsInRenderingThread());
+		check(TileAllocator.SupportResizingTiles());
+
+		uint32 ScaleFactor = TileAllocator.GetResizeFactor();
+		SimulationResourcesSize.X = GParticleSimulationDynTextureSizeXY * ScaleFactor;
+		SimulationResourcesSize.Y = GParticleSimulationDynTextureSizeXY * ScaleFactor;
+		
+		// resize position and velocity
+		StateTextures[0].ResizeRHI(RHICmdList, SimulationResourcesSize.X, SimulationResourcesSize.Y);
+		StateTextures[1].ResizeRHI(RHICmdList, SimulationResourcesSize.X, SimulationResourcesSize.Y);
+		// resize attributes
+		RenderAttributesTexture.ResizeRHI(RHICmdList, SimulationResourcesSize.X, SimulationResourcesSize.Y);
+		SimulationAttributesTexture.ResizeRHI(RHICmdList, SimulationResourcesSize.X, SimulationResourcesSize.Y);
+	}
+
 	/**
 	 * Release resources.
 	 */
@@ -443,6 +651,23 @@ public:
 			return GetCurrentStateTextures();
 		}
 	}
+
+	/**
+	 * Update when AdvanceGPUParticleFrame is called .
+	 */
+	void Update(FRHICommandListImmediate& RHICmdList)
+	{
+		if(TileAllocator.SupportResizingTiles())
+		{ 
+			check(IsInRenderingThread());
+			if (TileAllocator.bDirtyAlloc)
+			{
+				Resize(RHICmdList);
+				TileAllocator.bDirtyAlloc = false;
+			}
+		}
+	}
+
 	/**
 	 * Allocate a particle tile.
 	 */
@@ -467,6 +692,22 @@ public:
 		return TileAllocator.GetFreeTileCount();
 	}
 
+	/**
+	 * Returns the number of allocated tiles.
+	 */
+	int32 GetAllocatedTileCount() const
+	{
+		return TileAllocator.GetAllocatedTileCount();
+	}
+
+	/**
+	 * Support resizing tile allocator
+	 */
+	bool SupportTileResizing() const
+	{
+		return TileAllocator.SupportResizingTiles();
+	}
+
 private:
 
 	/** Allocator for managing particle tiles. */
@@ -532,6 +773,7 @@ public:
 	{
 		ParticleIndices.Bind(ParameterMap, TEXT("ParticleIndices"));
 		ParticleIndicesOffset.Bind(ParameterMap, TEXT("ParticleIndicesOffset"));
+		TilePageScale.Bind(ParameterMap, TEXT("TilePageScale"));
 		PositionTexture.Bind(ParameterMap, TEXT("PositionTexture"));
 		PositionTextureSampler.Bind(ParameterMap, TEXT("PositionTextureSampler"));
 		VelocityTexture.Bind(ParameterMap, TEXT("VelocityTexture"));
@@ -557,6 +799,8 @@ private:
 	LAYOUT_FIELD(FShaderResourceParameter, ParticleIndices);
 	/** Offset in to the particle indices buffer. */
 	LAYOUT_FIELD(FShaderParameter, ParticleIndicesOffset);
+	/** Tile Page Offset when tileAllocator can dynamically resize*/
+	LAYOUT_FIELD(FShaderParameter, TilePageScale);
 	/** Texture containing positions for all particles. */
 	LAYOUT_FIELD(FShaderResourceParameter, PositionTexture);
 	LAYOUT_FIELD(FShaderResourceParameter, PositionTextureSampler);
@@ -719,6 +963,7 @@ void FGPUSpriteVertexFactoryShaderParametersVS::GetElementShaderBindings(
 	ShaderBindings.AddTexture(VelocityTexture, VelocityTextureSampler, SamplerStatePoint, GPUVF->VelocityTextureRHI);
 	ShaderBindings.AddTexture(AttributesTexture, AttributesTextureSampler, SamplerStatePoint, GPUVF->AttributesTextureRHI);
 	ShaderBindings.AddTexture(CurveTexture, CurveTextureSampler, SamplerStateLinear, GParticleCurveTexture.GetCurveTexture());
+	ShaderBindings.Add(TilePageScale, GPUVF->TilePageScale);
 }
 
 void FGPUSpriteVertexFactoryShaderParametersPS::GetElementShaderBindings(
@@ -899,16 +1144,25 @@ public:
 	}
 
 	BEGIN_SHADER_PARAMETER_STRUCT(FParameters, )
-		SHADER_PARAMETER(float, TileSizeX)
-		SHADER_PARAMETER(float, TileSizeY)
-		SHADER_PARAMETER_SRV(Buffer<float2>, TileOffsets)
+		SHADER_PARAMETER(FVector2f, TexCoordScale)
+		SHADER_PARAMETER(FVector2f, TilePageScale)
+		SHADER_PARAMETER_SRV(Buffer<float4>, TileOffsets)
 	END_SHADER_PARAMETER_STRUCT()
 
-	FParameters GetParameters(FParticleShaderParamRef TileOffsetsRef)
+	FParameters GetParameters(FParticleShaderParamRef TileOffsetsRef, const FParticleSimulationResources* ParticleSimulationResources)
 	{
 		FParameters Result{};
-		Result.TileSizeX = (float)GParticleSimulationTileSize / (float)GParticleSimulationTextureSizeX;
-		Result.TileSizeY = (float)GParticleSimulationTileSize / (float)GParticleSimulationTextureSizeY;
+		Result.TexCoordScale.X = (float)(GParticleSimulationTileSize) / (float)ParticleSimulationResources->SimulationResourcesSize.X;
+		Result.TexCoordScale.Y = (float)(GParticleSimulationTileSize) / (float)ParticleSimulationResources->SimulationResourcesSize.Y;
+		
+		Result.TilePageScale = FVector2f::UnitVector;
+
+		if(ParticleSimulationResources->SupportTileResizing())
+		{ 
+			Result.TilePageScale.X = (float)GParticleSimulationDynTextureSizeXY / (float)ParticleSimulationResources->SimulationResourcesSize.X;
+			Result.TilePageScale.Y = (float)GParticleSimulationDynTextureSizeXY / (float)ParticleSimulationResources->SimulationResourcesSize.Y;
+		}
+		
 		Result.TileOffsets = TileOffsetsRef;
 		return Result;
 	}
@@ -1050,7 +1304,7 @@ public:
 		// TexCoord.
 		Elements.Add(FVertexElement(0, 0, VET_Float2, 0, sizeof(FVector2f), /*bUseInstanceIndex=*/ false));
 		// TileOffsets
-		Elements.Add(FVertexElement(1, 0, VET_Float2, 1, sizeof(FVector2f), /*bUseInstanceIndex=*/ true));
+		Elements.Add(FVertexElement(1, 0, VET_Float4, 1, sizeof(FVector4f), /*bUseInstanceIndex=*/ true));
 		VertexDeclarationRHI = PipelineStateCache::GetOrCreateVertexDeclaration( Elements );
 	}
 
@@ -1089,19 +1343,27 @@ FORCEINLINE int32 ComputeAlignedTileCount(int32 TileCount)
  * @param TileCount - The number of tiles in the array.
  * @param AlignedTileCount - The number of tiles to create in buffer for aligned rendering.
  */
-static void BuildTileVertexBuffer(FRHICommandListBase& RHICmdList, FParticleBufferParamRef TileOffsetsRef, const uint32* Tiles, int32 TileCount, int32 AlignedTileCount)
+static void BuildTileVertexBuffer(FRHICommandListBase& RHICmdList, FParticleBufferParamRef TileOffsetsRef, const uint32* Tiles, int32 TileCount, int32 AlignedTileCount, const FParticleSimulationResources* ParticleSimulationResources)
 {
-	FVector2f* TileOffset = (FVector2f*)RHICmdList.LockBuffer( TileOffsetsRef, 0, AlignedTileCount * sizeof(FVector2f), RLM_WriteOnly );
+	FVector4f* TileOffset = (FVector4f*)RHICmdList.LockBuffer( TileOffsetsRef, 0, AlignedTileCount * sizeof(FVector4f), RLM_WriteOnly );
+	
+	float TileCountX = ParticleSimulationResources->ParticleSimulationTileCountX;
+	float TileCountY = ParticleSimulationResources->ParticleSimulationTileCountY;
+
 	for ( int32 Index = 0; Index < TileCount; ++Index )
 	{
-		const uint32 TileIndex = Tiles[Index];
-		TileOffset[Index].X = FMath::Fractional( (float)TileIndex / (float)GParticleSimulationTileCountX );
-		TileOffset[Index].Y = FMath::Fractional( FMath::TruncToFloat( (float)TileIndex / (float)GParticleSimulationTileCountX ) / (float)GParticleSimulationTileCountY );
+		// Take the modulo of the tile index with the tile count of the first allocation. In case the gpu resources are resized, the 
+		// tile coordinate will be ajusted by the TilePageIndex.
+		const uint32 TileIndex = Tiles[Index] % ParticleSimulationResources->ParticleSimulationTileCount;
+		TileOffset[Index].X = FMath::Fractional( (float)TileIndex / TileCountX);
+		TileOffset[Index].Y = FMath::Fractional( FMath::TruncToFloat( (float)TileIndex / TileCountX) / TileCountY);
+		TileOffset[Index].Z = Tiles[Index] / ParticleSimulationResources->ParticleSimulationTileCount;
 	}
 	for ( int32 Index = TileCount; Index < AlignedTileCount; ++Index )
 	{
 		TileOffset[Index].X = 100.0f;
 		TileOffset[Index].Y = 100.0f;
+		TileOffset[Index].Z = 0.0;
 	}
 	RHICmdList.UnlockBuffer( TileOffsetsRef );
 }
@@ -1304,7 +1566,7 @@ void ExecuteSimulationCommands(
 	{
 		const FSimulationCommandGPU& Command = SimulationCommands[CommandIndex];
 
-		const FParticleTileVS::FParameters VsParameters = VertexShader->GetParameters(Command.TileOffsetsShaderRef);
+		const FParticleTileVS::FParameters VsParameters = VertexShader->GetParameters(Command.TileOffsetsShaderRef, ParticleSimulationResources);
 
 		PsParameters.PerFrameParameters = GetParticlePerFrameSimulationShaderParameters(Command.PerFrameParameters, bUseFixDT);
 		for (int32 i = 0; i < MAX_VECTOR_FIELDS; ++i)
@@ -1391,7 +1653,7 @@ void ExecuteSimulationCommands(
  * Invokes the clear simulation shader for each particle in each tile.
  * @param Tiles - The list of tiles to clear.
  */
-void ClearTiles(FRHICommandList& RHICmdList, FGraphicsPipelineStateInitializer& GraphicsPSOInit, ERHIFeatureLevel::Type FeatureLevel, const TArray<uint32>& Tiles)
+void ClearTiles(FRHICommandList& RHICmdList, FGraphicsPipelineStateInitializer& GraphicsPSOInit, ERHIFeatureLevel::Type FeatureLevel, const TArray<uint32>& Tiles, const FParticleSimulationResources* ParticleSimulationResources)
 {
 	if (!CVarSimulateGPUParticles.GetValueOnAnyThread())
 	{
@@ -1430,9 +1692,9 @@ void ClearTiles(FRHICommandList& RHICmdList, FGraphicsPipelineStateInitializer&
 				
 		if (FeatureLevel <= ERHIFeatureLevel::ES3_1)
 		{
-			BuildTileVertexBuffer(RHICmdList, BufferParam, TilesPtr, TilesThisDrawCall, TilesThisDrawCall);
+			BuildTileVertexBuffer(RHICmdList, BufferParam, TilesPtr, TilesThisDrawCall, TilesThisDrawCall, ParticleSimulationResources);
 
-			const FParticleTileVS::FParameters VsParameters = VertexShader->GetParameters(ShaderParam);
+			const FParticleTileVS::FParameters VsParameters = VertexShader->GetParameters(ShaderParam, ParticleSimulationResources);
 			SetShaderParameters(RHICmdList, VertexShader, VertexShader.GetVertexShader(), VsParameters);
 
 			DrawParticleTiles(RHICmdList, BufferParam, TilesThisDrawCall);
@@ -1440,9 +1702,9 @@ void ClearTiles(FRHICommandList& RHICmdList, FGraphicsPipelineStateInitializer&
 		else
 		{
 			const int32 AlignedTilesThisDrawCall = ComputeAlignedTileCount(TilesThisDrawCall);
-			BuildTileVertexBuffer(RHICmdList, BufferParam, TilesPtr, TilesThisDrawCall, AlignedTilesThisDrawCall);
+			BuildTileVertexBuffer(RHICmdList, BufferParam, TilesPtr, TilesThisDrawCall, AlignedTilesThisDrawCall, ParticleSimulationResources);
 
-			const FParticleTileVS::FParameters VsParameters = VertexShader->GetParameters(ShaderParam);
+			const FParticleTileVS::FParameters VsParameters = VertexShader->GetParameters(ShaderParam, ParticleSimulationResources);
 			SetShaderParameters(RHICmdList, VertexShader, VertexShader.GetVertexShader(), VsParameters);
 
 			DrawAlignedParticleTiles(RHICmdList, AlignedTilesThisDrawCall);
@@ -1457,7 +1719,8 @@ void ClearTiles(FRHICommandList& RHICmdList, FGraphicsPipelineStateInitializer&
  * Uniform buffer to hold parameters for particle simulation.
  */
 BEGIN_GLOBAL_SHADER_PARAMETER_STRUCT( FParticleInjectionParameters, )
-	SHADER_PARAMETER( FVector2f, PixelScale )
+	SHADER_PARAMETER(FVector2f, TexCoordScale)
+	SHADER_PARAMETER(FVector2f, TilePageScale)
 END_GLOBAL_SHADER_PARAMETER_STRUCT()
 
 IMPLEMENT_GLOBAL_SHADER_PARAMETER_STRUCT(FParticleInjectionParameters, "ParticleInjection");
@@ -1492,11 +1755,20 @@ public:
 	/**
 	 * Sets parameters for particle injection.
 	 */
-	void SetParameters(FRHIBatchedShaderParameters& BatchedParameters)
+	void SetParameters(FRHIBatchedShaderParameters& BatchedParameters, const FParticleSimulationResources* ParticleSimulationResources)
 	{
 		FParticleInjectionParameters Parameters;
-		Parameters.PixelScale.X = 1.0f / GParticleSimulationTextureSizeX;
-		Parameters.PixelScale.Y = 1.0f / GParticleSimulationTextureSizeY;
+		Parameters.TexCoordScale.X = 1.0f / ParticleSimulationResources->SimulationResourcesSize.X;
+		Parameters.TexCoordScale.Y = 1.0f / ParticleSimulationResources->SimulationResourcesSize.Y;
+		
+		Parameters.TilePageScale = FVector2f::UnitVector;
+		
+		if (ParticleSimulationResources->SupportTileResizing())
+		{
+			Parameters.TilePageScale.X = (float)GParticleSimulationDynTextureSizeXY / (float)ParticleSimulationResources->SimulationResourcesSize.X;
+			Parameters.TilePageScale.Y = (float)GParticleSimulationDynTextureSizeXY / (float)ParticleSimulationResources->SimulationResourcesSize.Y;
+		}
+
 		FParticleInjectionBufferRef UniformBuffer = FParticleInjectionBufferRef::CreateUniformBufferImmediate( Parameters, UniformBuffer_SingleDraw );
 		SetUniformBufferParameter(BatchedParameters, GetUniformBufferParameter<FParticleInjectionParameters>(), UniformBuffer );
 	}
@@ -1575,8 +1847,8 @@ public:
 			Elements.Add(FVertexElement(0, Offset, VET_Float4, 3, Stride, /*bUseInstanceIndex=*/ true));
 			Offset += sizeof(FVector4f);
 			// ParticleIndex.
-			Elements.Add(FVertexElement(0, Offset, VET_Float2, 4, Stride, /*bUseInstanceIndex=*/ true));
-			Offset += sizeof(FVector2f);
+			Elements.Add(FVertexElement(0, Offset, VET_Float3, 4, Stride, /*bUseInstanceIndex=*/ true));
+			Offset += sizeof(FVector3f);
 		}
 
 		// Stream 1.
@@ -1604,7 +1876,7 @@ TGlobalResource<FParticleInjectionVertexDeclaration> GParticleInjectionVertexDec
  * @param NewParticles - A list of particles to inject in to the simulation.
  */
 template<bool StaticPropertiesOnly>
-void InjectNewParticles(FRHICommandList& RHICmdList, FGraphicsPipelineStateInitializer& GraphicsPSOInit,  ERHIFeatureLevel::Type FeatureLevel, const TArray<FNewParticle>& NewParticles)
+void InjectNewParticles(FRHICommandList& RHICmdList, FGraphicsPipelineStateInitializer& GraphicsPSOInit,  ERHIFeatureLevel::Type FeatureLevel, const TArray<FNewParticle>& NewParticles, const FParticleSimulationResources* ParticleSimulationResources)
 {
 	if (GIsRenderingThreadSuspended.Load(EMemoryOrder::Relaxed) || !CVarSimulateGPUParticles.GetValueOnAnyThread())
 	{
@@ -1639,7 +1911,7 @@ void InjectNewParticles(FRHICommandList& RHICmdList, FGraphicsPipelineStateIniti
 
 		SetGraphicsPipelineState(RHICmdList, GraphicsPSOInit, 0);
 
-		SetShaderParametersLegacyVS(RHICmdList, VertexShader);
+		SetShaderParametersLegacyVS(RHICmdList, VertexShader, ParticleSimulationResources);
 
 		// Stream 0: New particles.
 		RHICmdList.SetStreamSource(
@@ -1858,7 +2130,7 @@ static void VisualizeGPUSimulation(
  * @param VertexBuffer - The buffer with which to fill with particle indices.
  * @param InTiles - The list of tiles for which to generate indices.
  */
-static void BuildParticleVertexBuffer(FRHICommandListBase& RHICmdList, FRHIBuffer* VertexBufferRHI, const TArray<uint32>& InTiles )
+static void BuildParticleVertexBuffer(FRHICommandListBase& RHICmdList, FRHIBuffer* VertexBufferRHI, const TArray<uint32>& InTiles, const FParticleSimulationResources* ParticleSimulationResources)
 {
 	const int32 TileCount = InTiles.Num();
 	const int32 IndexCount = TileCount * GParticlesPerTile;
@@ -1866,25 +2138,37 @@ static void BuildParticleVertexBuffer(FRHICommandListBase& RHICmdList, FRHIBuffe
 	const int32 Stride = 1;
 	FParticleIndex* RESTRICT ParticleIndices = (FParticleIndex*)RHICmdList.LockBuffer( VertexBufferRHI, 0, BufferSize, RLM_WriteOnly );
 
+	bool bAllowResizing = ParticleSimulationResources->SupportTileResizing();
+	
+	float TileCountX = (float)ParticleSimulationResources->ParticleSimulationTileCountX;
+	float TileCountY = (float)ParticleSimulationResources->ParticleSimulationTileCountY;
+	float TextureSizeX = bAllowResizing ? GParticleSimulationDynTextureSizeXY: GParticleSimulationTextureSizeX;
+	float TextureSizeY = bAllowResizing ? GParticleSimulationDynTextureSizeXY : GParticleSimulationTextureSizeY;
+
 	for ( int32 Index = 0; Index < TileCount; ++Index )
 	{
-		const uint32 TileIndex = InTiles[Index];
+		// Take the mod of the tile index with the tile count of the first allocation. In case the gpu resources are resized, the 
+		// tile coordinate will be ajusted by the TilePageIndex.
+		const uint32 TileIndex = InTiles[Index] % ParticleSimulationResources->ParticleSimulationTileCount;
+		const uint32 PageIndex = InTiles[Index] / ParticleSimulationResources->ParticleSimulationTileCount;
+
 		const FVector2D TileOffset(
-			FMath::Fractional( (float)TileIndex / (float)GParticleSimulationTileCountX ),
-			FMath::Fractional( FMath::TruncToFloat( (float)TileIndex / (float)GParticleSimulationTileCountX ) / (float)GParticleSimulationTileCountY )
+			FMath::Fractional( (float)TileIndex / TileCountX),
+			FMath::Fractional( FMath::TruncToFloat( (float)TileIndex / TileCountX) / TileCountY)
 			);
 		for ( int32 ParticleY = 0; ParticleY < GParticleSimulationTileSize; ++ParticleY )
 		{
 			for ( int32 ParticleX = 0; ParticleX < GParticleSimulationTileSize; ++ParticleX )
 			{
-				const float IndexX = TileOffset.X + ((float)ParticleX / (float)GParticleSimulationTextureSizeX) + (0.5f / (float)GParticleSimulationTextureSizeX);
-				const float IndexY = TileOffset.Y + ((float)ParticleY / (float)GParticleSimulationTextureSizeY) + (0.5f / (float)GParticleSimulationTextureSizeY);
+				const float IndexX = TileOffset.X + ((float)ParticleX / TextureSizeX + (0.5f / TextureSizeX));
+				const float IndexY = TileOffset.Y + ((float)ParticleY / TextureSizeY + (0.5f / TextureSizeY));
 
 				// @todo faster float32 -> float16 conversion
 				//	use AVX2/F16C for _mm_cvtps_ph
 				ParticleIndices->X.Set(IndexX);
-				ParticleIndices->Y.Set(IndexY);					
-
+				ParticleIndices->Y.Set(IndexY);
+				ParticleIndices->Z.Set(PageIndex);
+				ParticleIndices->W.Set(0.0f);
 				// move to next particle
 				ParticleIndices += Stride;
 			}
@@ -1925,9 +2209,9 @@ public:
 	BEGIN_SHADER_PARAMETER_STRUCT(FParameters, )
 		SHADER_PARAMETER(int32, TextureSizeX)
 		SHADER_PARAMETER(int32, TextureSizeY)
-
+		SHADER_PARAMETER(FVector2f, TilePageScale)
 		SHADER_PARAMETER_UAV(RWBuffer<float4>, OutBounds)
-		SHADER_PARAMETER_SRV(Buffer<float2>, InParticleIndices)
+		SHADER_PARAMETER_SRV(Buffer<float4>, InParticleIndices)
 		SHADER_PARAMETER_TEXTURE(Texture2D, PositionTexture)
 
 		SHADER_PARAMETER_STRUCT_REF(FParticleBoundsParameters, ParticleBoundsParameters)
@@ -1967,13 +2251,16 @@ static FBox ComputeParticleBounds(
 	ERHIFeatureLevel::Type FeatureLevel,
 	FRHIShaderResourceView* VertexBufferSRV,
 	FRHITexture2D* PositionTextureRHI,
-	int32 ParticleCount )
+	int32 ParticleCount,
+	FParticleSimulationResources* ParticleSimulationResources)
 {
 	FBox BoundingBox;
 	FParticleBoundsParameters ParticleBoundsParameters;
 
 	if (ParticleCount > 0 && FeatureLevel >= ERHIFeatureLevel::SM5)
 	{
+		bool bSupportTileResizing = ParticleSimulationResources->SupportTileResizing();
+
 		// Determine how to break the work up over individual work groups.
 		const uint32 MaxGroupCount = 128;
 		const uint32 AlignedParticleCount = ((ParticleCount + PARTICLE_BOUNDS_THREADS - 1) & (~(PARTICLE_BOUNDS_THREADS - 1)));
@@ -2002,13 +2289,23 @@ static FBox ComputeParticleBounds(
 
 		{
 			FParticleBoundsCS::FParameters Parameters{};
-			Parameters.TextureSizeX = GParticleSimulationTextureSizeX;
-			Parameters.TextureSizeY = GParticleSimulationTextureSizeY;
+			Parameters.TextureSizeX = ParticleSimulationResources->SimulationResourcesSize.X;
+			Parameters.TextureSizeY = ParticleSimulationResources->SimulationResourcesSize.Y;
 			Parameters.OutBounds = BoundsVertexBufferUAV;
 			Parameters.InParticleIndices = VertexBufferSRV;
 			Parameters.PositionTexture = PositionTextureRHI;
 			Parameters.ParticleBoundsParameters = FParticleBoundsUniformBufferRef::CreateUniformBufferImmediate(ParticleBoundsParameters, UniformBuffer_SingleFrame);
 
+			FVector2f TilePageScale = FVector2f::UnitVector;
+			
+			if (bSupportTileResizing)
+			{
+				TilePageScale.X = (float)GParticleSimulationDynTextureSizeXY / (float)ParticleSimulationResources->SimulationResourcesSize.X;
+				TilePageScale.Y = (float)GParticleSimulationDynTextureSizeXY / (float)ParticleSimulationResources->SimulationResourcesSize.Y;
+			}
+			
+			Parameters.TilePageScale = TilePageScale;
+
 			SetShaderParameters(RHICmdList, ParticleBoundsCS, ParticleBoundsCS.GetComputeShader(), Parameters);
 		}
 
@@ -2114,7 +2411,7 @@ public:
 	/**
 	 * Initializes the vertex buffer from a list of tiles.
 	 */
-	void Init( FRHICommandListBase& RHICmdList, const TArray<uint32>& Tiles )
+	void Init( FRHICommandListBase& RHICmdList, const TArray<uint32>& Tiles, const FParticleSimulationResources* ParticleSimulationResources)
 	{
 		TileCount = Tiles.Num();
 		AlignedTileCount = ComputeAlignedTileCount(TileCount);
@@ -2122,7 +2419,7 @@ public:
 		if (Tiles.Num())
 		{
 			int32 BufferAlignedTileCount = (GMaxRHIFeatureLevel <= ERHIFeatureLevel::ES3_1 ? TileCount : AlignedTileCount);
-			BuildTileVertexBuffer(RHICmdList, VertexBufferRHI, Tiles.GetData(), Tiles.Num(), BufferAlignedTileCount);
+			BuildTileVertexBuffer(RHICmdList, VertexBufferRHI, Tiles.GetData(), Tiles.Num(), BufferAlignedTileCount, ParticleSimulationResources);
 		}
 	}
 
@@ -2134,11 +2431,12 @@ public:
 		if ( AlignedTileCount > 0 )
 		{
 			int32 BufferAlignedTileCount = (GMaxRHIFeatureLevel <= ERHIFeatureLevel::ES3_1 ? TileCount : AlignedTileCount);
-			const int32 TileBufferSize = BufferAlignedTileCount * sizeof(FVector2f);
+			const uint32 Stride = sizeof(FVector4f);
+			const int32 TileBufferSize = BufferAlignedTileCount * Stride;
 			check(TileBufferSize > 0);
 			FRHIResourceCreateInfo CreateInfo(TEXT("FParticleTileVertexBuffer"));
 			VertexBufferRHI = RHICmdList.CreateVertexBuffer( TileBufferSize, BUF_Static | BUF_KeepCPUAccessible | BUF_ShaderResource, CreateInfo );
-			VertexBufferSRV = RHICmdList.CreateShaderResourceView( VertexBufferRHI, /*Stride=*/ sizeof(FVector2f), PF_G32R32F );
+			VertexBufferSRV = RHICmdList.CreateShaderResourceView( VertexBufferRHI, Stride, PF_A32B32G32R32F);
 		}
 	}
 
@@ -2176,13 +2474,13 @@ public:
 	/**
 	 * Initializes the vertex buffer from a list of tiles.
 	 */
-	void Init(FRHICommandListBase& RHICmdList, const TArray<uint32>& Tiles )
+	void Init(FRHICommandListBase& RHICmdList, const TArray<uint32>& Tiles, const FParticleSimulationResources* ParticleSimulationResources)
 	{
 		ParticleCount = Tiles.Num() * GParticlesPerTile;
 		InitResource(FRHICommandListImmediate::Get());
 		if ( Tiles.Num() )
 		{
-			BuildParticleVertexBuffer( RHICmdList, VertexBufferRHI, Tiles );
+			BuildParticleVertexBuffer( RHICmdList, VertexBufferRHI, Tiles, ParticleSimulationResources);
 		}
 	}
 
@@ -2198,7 +2496,7 @@ public:
 			const EBufferUsageFlags Flags = BUF_Static | /*BUF_KeepCPUAccessible | */BUF_ShaderResource;
 			FRHIResourceCreateInfo CreateInfo(TEXT("FGPUParticleVertexBuffer"));
 			VertexBufferRHI = RHICmdList.CreateVertexBuffer(BufferSize, Flags, CreateInfo);
-			VertexBufferSRV = RHICmdList.CreateShaderResourceView(VertexBufferRHI, BufferStride, PF_G16R16F);
+			VertexBufferSRV = RHICmdList.CreateShaderResourceView(VertexBufferRHI, BufferStride, PF_FloatRGBA);
 		}
 	}
 
@@ -2216,7 +2514,8 @@ public:
 class FParticleSimulationGPU
 {
 public:
-
+	/** Pointer the the FX system with which the simulation is associated. */
+	FFXSystem* FXSystem;
 	/** The vertex buffer used to access tiles in the simulation. */
 	FParticleTileVertexBuffer TileVertexBuffer;
 	/** Reference to the GPU sprite resources. */
@@ -2270,7 +2569,8 @@ public:
 
 	/** Default constructor. */
 	FParticleSimulationGPU()
-		: EmitterSimulationResources(NULL)
+		: FXSystem(nullptr)
+		, EmitterSimulationResources(NULL)
 		, VectorFieldVisualizationVertexFactory(NULL)
 #if GPUPARTICLE_LOCAL_VF_ONLY
 		, LocalIntensity(0.0f)
@@ -2458,8 +2758,8 @@ void FParticleSimulationGPU::InitResources(const TArray<uint32>& Tiles, FGPUSpri
 				Simulation->TileVertexBuffer.ReleaseResource();
 
 				// Initialize new buffers with list of tiles.
-				Simulation->VertexBuffer.Init(RHICmdList, Tiles);
-				Simulation->TileVertexBuffer.Init(RHICmdList, Tiles);
+				Simulation->VertexBuffer.Init(RHICmdList, Tiles, Simulation->FXSystem->GetParticleSimulationResources());
+				Simulation->TileVertexBuffer.Init(RHICmdList, Tiles, Simulation->FXSystem->GetParticleSimulationResources());
 
 				// Store simulation resources for this emitter.
 				Simulation->GPUSpriteResources = InGPUSpriteResourcesRef;
@@ -2704,15 +3004,22 @@ public:
 				FGPUSpriteMeshDataUserData* MeshBatchUserData = nullptr;
 				if (bAllowSorting && SortMode == PSORTMODE_DistanceToView)
 				{
-					// Extensibility TODO: This call to AddSortedGPUSimulation is very awkward. When rendering a frame we need to
-					// accumulate all GPU particle emitters that need to be sorted. That is so they can be sorted in one big radix
-					// sort for efficiency. Ideally that state is per-scene renderer but the renderer doesn't know anything about particles.
-					FGPUSortManager::FAllocationInfo SortedIndicesInfo;
-					if (FXSystem->AddSortedGPUSimulation(RHICmdList, Simulation, View->ViewMatrices.GetViewOrigin(), bTranslucent, SortedIndicesInfo))
+					if (GFXCascadeGpuSpriteAllowDynAllocs)
 					{
-						MeshBatchUserData = &Collector.AllocateOneFrameResource<FGPUSpriteMeshDataUserData>();
-						MeshBatchUserData->SortedOffset = SortedIndicesInfo.BufferOffset;
-						MeshBatchUserData->SortedParticleIndicesSRV = SortedIndicesInfo.BufferSRV;
+						UE_LOG(LogParticles, Warning, TEXT("Cascade doesn't support gpu sorting with particle gpu simulation resources resizing"));
+					}
+					else
+					{ 
+						// Extensibility TODO: This call to AddSortedGPUSimulation is very awkward. When rendering a frame we need to
+						// accumulate all GPU particle emitters that need to be sorted. That is so they can be sorted in one big radix
+						// sort for efficiency. Ideally that state is per-scene renderer but the renderer doesn't know anything about particles.
+						FGPUSortManager::FAllocationInfo SortedIndicesInfo;
+						if (FXSystem->AddSortedGPUSimulation(RHICmdList, Simulation, View->ViewMatrices.GetViewOrigin(), bTranslucent, SortedIndicesInfo))
+						{
+							MeshBatchUserData = &Collector.AllocateOneFrameResource<FGPUSpriteMeshDataUserData>();
+							MeshBatchUserData->SortedOffset = SortedIndicesInfo.BufferOffset;
+							MeshBatchUserData->SortedParticleIndicesSRV = SortedIndicesInfo.BufferSRV;
+						}
 					}
 				}
 				check(Simulation->VertexBuffer.IsInitialized());
@@ -2734,6 +3041,12 @@ public:
 					VertexFactory.AttributesTextureRHI = ParticleSimulationResources->RenderAttributesTexture.TextureRHI;
 					VertexFactory.LWCTile = ParticleSimulationResources->LWCTile;
 
+					FVector3f TilePageScale;
+					TilePageScale.X = SimulationResources->SupportTileResizing() ? ((float)(GParticleSimulationDynTextureSizeXY) / (float)SimulationResources->SimulationResourcesSize.X) : 1.0f;
+					TilePageScale.Y = SimulationResources->SupportTileResizing() ? ((float)(GParticleSimulationDynTextureSizeXY) / (float)SimulationResources->SimulationResourcesSize.Y) : 1.0f;
+					TilePageScale.Z = SimulationResources->SupportTileResizing() ? 1.0f : 0.0f;
+					VertexFactory.TilePageScale = TilePageScale;
+
 					FMeshBatch& Mesh = Collector.AllocateMesh();
 					FMeshBatchElement& BatchElement = Mesh.Elements[0];
 					BatchElement.IndexBuffer = &GParticleIndexBuffer;
@@ -2904,6 +3217,7 @@ FGPUSpriteParticleEmitterInstance(FFXSystem* InFXSystem, FGPUSpriteEmitterInfo&
 		}
 		Simulation->bWantsCollision = InEmitterInfo.bEnableCollision;
 		Simulation->CollisionMode = InEmitterInfo.CollisionMode;
+		Simulation->FXSystem = InFXSystem;
 
 #if TRACK_TILE_ALLOCATIONS
 		TSet<class FGPUSpriteParticleEmitterInstance*>* EmitterSet = GPUSpriteParticleEmitterInstances.Find(FXSystem);
@@ -3354,10 +3668,10 @@ FGPUSpriteParticleEmitterInstance(FFXSystem* InFXSystem, FGPUSpriteEmitterInfo&
 
 				ReserveNewParticles(FirstBurstParticleIndex + BurstInfo.Count + SpawnInfo.Count);
 
-				BurstInfo.Count = AllocateTilesForParticles(NewParticles, BurstInfo.Count, ActiveTileCount);
+				BurstInfo.Count = AllocateTilesForParticles(NewParticles, BurstInfo.Count, ActiveTileCount, FXSystem->GetParticleSimulationResources());
 
 				int32 FirstSpawnParticleIndex = NewParticles.Num();
-				SpawnInfo.Count = AllocateTilesForParticles(NewParticles, SpawnInfo.Count, ActiveTileCount);
+				SpawnInfo.Count = AllocateTilesForParticles(NewParticles, SpawnInfo.Count, ActiveTileCount, FXSystem->GetParticleSimulationResources());
 				SpawnFraction += LeftoverBurst;
 
 				if (BurstInfo.Count > 0)
@@ -3539,7 +3853,8 @@ FGPUSpriteParticleEmitterInstance(FFXSystem* InFXSystem, FGPUSpriteEmitterInfo&
 					FeatureLevel,
 					EmitterInstance->Simulation->VertexBuffer.VertexBufferSRV,
 					EmitterInstance->FXSystem->GetParticleSimulationResources()->GetVisualizeStateTextures().PositionTextureRHI,
-					EmitterInstance->Simulation->VertexBuffer.ParticleCount
+					EmitterInstance->Simulation->VertexBuffer.ParticleCount,
+					EmitterInstance->FXSystem->GetParticleSimulationResources()
 					);
 			});
 		FlushRenderingCommands();
@@ -3729,18 +4044,25 @@ private:
 	 * @param ActiveTileCount - Number of active tiles, incremented each time a new tile is allocated.
 	 * @returns the number of particles which were successfully allocated.
 	 */
-	int32 AllocateTilesForParticles(TArray<FNewParticle>& InNewParticles, int32 NumNewParticles, int32& ActiveTileCount)
+	int32 AllocateTilesForParticles(TArray<FNewParticle>& InNewParticles, int32 NumNewParticles, int32& ActiveTileCount, const FParticleSimulationResources* ParticleSimulationResources)
 	{
 		if (!NumNewParticles)
 		{
 			return 0;
 		}
+
+		bool bAllowResizing = ParticleSimulationResources->SupportTileResizing();
+		float TileCountX = ParticleSimulationResources->ParticleSimulationTileCountX;
+		float TileCountY = ParticleSimulationResources->ParticleSimulationTileCountY;
+		float TextureSizeX = bAllowResizing ? GParticleSimulationDynTextureSizeXY : GParticleSimulationTextureSizeX;
+		float TextureSizeY = bAllowResizing ? GParticleSimulationDynTextureSizeXY : GParticleSimulationTextureSizeY;
+
 		// Need to allocate space in tiles for all new particles.
 		FParticleSimulationResources* SimulationResources = FXSystem->GetParticleSimulationResources();
-		uint32 TileIndex = (AllocatedTiles.IsValidIndex(TileToAllocateFrom)) ? AllocatedTiles[TileToAllocateFrom] : INDEX_NONE;
+		uint32 TileIndex = (AllocatedTiles.IsValidIndex(TileToAllocateFrom)) ? AllocatedTiles[TileToAllocateFrom] % ParticleSimulationResources->ParticleSimulationTileCount : INDEX_NONE;
 		FVector2D TileOffset(
-			FMath::Fractional((float)TileIndex / (float)GParticleSimulationTileCountX),
-			FMath::Fractional(FMath::TruncToFloat((float)TileIndex / (float)GParticleSimulationTileCountX) / (float)GParticleSimulationTileCountY)
+			FMath::Fractional((float)TileIndex / TileCountX),
+			FMath::Fractional(FMath::TruncToFloat((float)TileIndex / TileCountX) / TileCountY)
 			);
 
 		for (int32 ParticleIndex = 0; ParticleIndex < NumNewParticles; ++ParticleIndex)
@@ -3772,17 +4094,20 @@ private:
 				}
 
 				ActiveTileCount++;
-				TileIndex = AllocatedTiles[TileToAllocateFrom];
-				TileOffset.X = FMath::Fractional((float)TileIndex / (float)GParticleSimulationTileCountX);
-				TileOffset.Y = FMath::Fractional(FMath::TruncToFloat((float)TileIndex / (float)GParticleSimulationTileCountX) / (float)GParticleSimulationTileCountY);
+				// Take the mod of the tile index with the tile count of the first allocation. In case the gpu resources are resized, the 
+				// tile coordinate will be ajusted by the TilePageIndex.
+				TileIndex = AllocatedTiles[TileToAllocateFrom] % ParticleSimulationResources->ParticleSimulationTileCount;
+				TileOffset.X = FMath::Fractional((float)TileIndex / TileCountX);
+				TileOffset.Y = FMath::Fractional(FMath::TruncToFloat((float)TileIndex / TileCountX) / TileCountY);
 				FreeParticlesInTile = GParticlesPerTile;
 			}
 			FNewParticle& Particle = InNewParticles.AddDefaulted_GetRef();
 			const int32 SubTileIndex = GParticlesPerTile - FreeParticlesInTile;
 			const int32 SubTileX = SubTileIndex % GParticleSimulationTileSize;
 			const int32 SubTileY = SubTileIndex / GParticleSimulationTileSize;
-			Particle.Offset.X = TileOffset.X + ((float)SubTileX / (float)GParticleSimulationTextureSizeX);
-			Particle.Offset.Y = TileOffset.Y + ((float)SubTileY / (float)GParticleSimulationTextureSizeY);
+			Particle.Offset.X = TileOffset.X + ((float)SubTileX / TextureSizeX);
+			Particle.Offset.Y = TileOffset.Y + ((float)SubTileY / TextureSizeY);
+			Particle.Offset.Z = AllocatedTiles[TileToAllocateFrom] / ParticleSimulationResources->ParticleSimulationTileCount;
 			Particle.ResilienceAndTileIndex.AllocatedTileIndex = TileToAllocateFrom;
 			FreeParticlesInTile--;
 		}
@@ -4151,11 +4476,11 @@ void DumpTileAllocations()
 
 		UE_LOG(LogParticles,Display,TEXT("---"));
 		UE_LOG(LogParticles,Display,TEXT("Total Allocated: %d"), TotalAllocatedTiles);
-		UE_LOG(LogParticles,Display,TEXT("Free (est.): %d"), GParticleSimulationTileCount - TotalAllocatedTiles);
+		UE_LOG(LogParticles,Display,TEXT("Free (est.): %d"), FXSystem->GetParticleSimulationResources()->GetAllocatedTileCount() - TotalAllocatedTiles);
 		if (FXSystem)
 		{
 			UE_LOG(LogParticles,Display,TEXT("Free (actual): %d"), FXSystem->GetParticleSimulationResources()->GetFreeTileCount());
-			UE_LOG(LogParticles,Display,TEXT("Leaked: %d"), GParticleSimulationTileCount - TotalAllocatedTiles - FXSystem->GetParticleSimulationResources()->GetFreeTileCount());
+			UE_LOG(LogParticles,Display,TEXT("Leaked: %d"), FXSystem->GetParticleSimulationResources()->GetAllocatedTileCount() - TotalAllocatedTiles - FXSystem->GetParticleSimulationResources()->GetFreeTileCount());
 		}
 	}
 }
@@ -4176,12 +4501,8 @@ void FFXSystem::InitGPUSimulation()
 	LLM_SCOPE(ELLMTag::Particles);
 
 	check(ParticleSimulationResources == NULL);
-	ensure(GParticleSimulationTextureSizeX > 0 && GParticleSimulationTextureSizeY > 0);
-	/** How many tiles are in the simulation textures. */
-	GParticleSimulationTileCountX = GParticleSimulationTextureSizeX / GParticleSimulationTileSize;
-	GParticleSimulationTileCountY = GParticleSimulationTextureSizeY / GParticleSimulationTileSize;
-	GParticleSimulationTileCount = GParticleSimulationTileCountX * GParticleSimulationTileCountY;
-
+	ensure(GParticleSimulationTextureSizeX > 0 && GParticleSimulationTextureSizeY > 0 && GParticleSimulationDynTextureSizeXY > 0);
+	
 	ParticleSimulationResources = new FParticleSimulationResources();
 
 	InitGPUResources();
@@ -4311,8 +4632,11 @@ void FFXSystem::GenerateSortKeys(FRHICommandListImmediate& RHICmdList, int32 Bat
 		BatchId);
 }
 
-void FFXSystem::AdvanceGPUParticleFrame(bool bAllowGPUParticleUpdate)
+void FFXSystem::AdvanceGPUParticleFrame(FRHICommandListImmediate& RHICmdList, bool bAllowGPUParticleUpdate)
 {
+	// check if we need to grow the gpu simulation resources
+	ParticleSimulationResources->Update(RHICmdList);
+
 	if (bAllowGPUParticleUpdate)
 	{
 		// We double buffer, so swap the current and previous textures.
@@ -4662,7 +4986,7 @@ void FFXSystem::SimulateGPUParticles(
 			FGraphicsPipelineStateInitializer GraphicsPSOInit;
 			RHICmdList.ApplyCachedRenderTargets(GraphicsPSOInit);
 
-			RHICmdList.SetViewport(0.0f, 0.0f, 0.0f, (float)GParticleSimulationTextureSizeX, (float)GParticleSimulationTextureSizeY, 1.0f);
+			RHICmdList.SetViewport(0.0f, 0.0f, 0.0f, (float)ParticleSimulationResources->SimulationResourcesSize.X, (float)ParticleSimulationResources->SimulationResourcesSize.X, 1.0f);
 			GraphicsPSOInit.DepthStencilState = TStaticDepthStencilState<false, CF_Always>::GetRHI();
 			GraphicsPSOInit.RasterizerState = TStaticRasterizerState<FM_Solid, CM_None>::GetRHI();
 			GraphicsPSOInit.BlendState = TStaticBlendState<>::GetRHI();
@@ -4691,7 +5015,7 @@ void FFXSystem::SimulateGPUParticles(
 			{
 				SCOPED_DRAW_EVENT(RHICmdList, ParticleTilesClear);
 
-				ClearTiles(RHICmdList, GraphicsPSOInit, FeatureLevel, TilesToClear);
+				ClearTiles(RHICmdList, GraphicsPSOInit, FeatureLevel, TilesToClear, ParticleSimulationResources);
 			}
 		}
 		RHICmdList.EndRenderPass();
@@ -4726,13 +5050,13 @@ void FFXSystem::SimulateGPUParticles(
 			FGraphicsPipelineStateInitializer GraphicsPSOInit;
 			RHICmdList.ApplyCachedRenderTargets(GraphicsPSOInit);
 
-			RHICmdList.SetViewport(0.0f, 0.0f, 0.0f, (float)GParticleSimulationTextureSizeX, (float)GParticleSimulationTextureSizeY, 1.0f);
+			RHICmdList.SetViewport(0.0f, 0.0f, 0.0f, (float)ParticleSimulationResources->SimulationResourcesSize.X, ParticleSimulationResources->SimulationResourcesSize.Y, 1.0f);
 			GraphicsPSOInit.DepthStencilState = TStaticDepthStencilState<false, CF_Always>::GetRHI();
 			GraphicsPSOInit.RasterizerState = TStaticRasterizerState<FM_Solid, CM_None>::GetRHI();
 			GraphicsPSOInit.BlendState = TStaticBlendState<>::GetRHI();
 
 			// Inject particles.
-			InjectNewParticles<false>(RHICmdList, GraphicsPSOInit, this->FeatureLevel, NewParticles);
+			InjectNewParticles<false>(RHICmdList, GraphicsPSOInit, this->FeatureLevel, NewParticles, ParticleSimulationResources);
 
 			RHICmdList.EndRenderPass();
 
@@ -4855,7 +5179,7 @@ void FFXSystem::UpdateMultiGPUResources(FRHICommandListImmediate& RHICmdList)
 			});
 			RHICmdList.BeginRenderPass(RPInfo, TEXT("UpdateMultiGPUResources"));
 
-			RHICmdList.SetViewport(0.0f, 0.0f, 0.0f, (float)GParticleSimulationTextureSizeX, (float)GParticleSimulationTextureSizeY, 1.0f);
+			RHICmdList.SetViewport(0.0f, 0.0f, 0.0f, (float)ParticleSimulationResources->SimulationResourcesSize.X, (float)ParticleSimulationResources->SimulationResourcesSize.Y, 1.0f);
 			FGraphicsPipelineStateInitializer GraphicsPSOInit;
 			RHICmdList.ApplyCachedRenderTargets(GraphicsPSOInit);
 			GraphicsPSOInit.DepthStencilState = TStaticDepthStencilState<false, CF_Always>::GetRHI();
@@ -4863,7 +5187,7 @@ void FFXSystem::UpdateMultiGPUResources(FRHICommandListImmediate& RHICmdList)
 			GraphicsPSOInit.BlendState = TStaticBlendState<>::GetRHI();
 
 			// Inject particles.
-			InjectNewParticles<true>(RHICmdList, GraphicsPSOInit, this->FeatureLevel, this->LastFrameNewParticles);
+			InjectNewParticles<true>(RHICmdList, GraphicsPSOInit, this->FeatureLevel, this->LastFrameNewParticles, ParticleSimulationResources);
 
 			RHICmdList.EndRenderPass();
 		}
diff --git a/Engine/Source/Runtime/Engine/Private/Particles/ParticleResources.cpp b/Engine/Source/Runtime/Engine/Private/Particles/ParticleResources.cpp
--- a/Engine/Source/Runtime/Engine/Private/Particles/ParticleResources.cpp
+++ b/Engine/Source/Runtime/Engine/Private/Particles/ParticleResources.cpp
@@ -144,12 +144,12 @@ void FParticleScratchVertexBuffer::InitRHI(FRHICommandListBase& RHICmdList)
 
 	FRHIResourceCreateInfo CreateInfo(TEXT("FParticleScratchVertexBuffer"));
 	VertexBufferRHI = RHICmdList.CreateVertexBuffer(GParticleScratchVertexBufferSize, Flags, CreateInfo);
-	VertexBufferSRV_G32R32F = RHICmdList.CreateShaderResourceView(VertexBufferRHI, /*Stride=*/ sizeof(FVector2f), PF_G32R32F);
+	VertexBufferSRV_A32B32G32R32F = RHICmdList.CreateShaderResourceView(VertexBufferRHI, /*Stride=*/ sizeof(FVector4f), PF_A32B32G32R32F);
 }
 
 FParticleShaderParamRef FParticleScratchVertexBuffer::GetShaderParam()
 {
-	return VertexBufferSRV_G32R32F;
+	return VertexBufferSRV_A32B32G32R32F;
 }
 
 FParticleBufferParamRef FParticleScratchVertexBuffer::GetBufferParam()
@@ -160,7 +160,7 @@ FParticleBufferParamRef FParticleScratchVertexBuffer::GetBufferParam()
 /** Release RHI resources. */
 void FParticleScratchVertexBuffer::ReleaseRHI()
 {
-	VertexBufferSRV_G32R32F.SafeRelease();
+	VertexBufferSRV_A32B32G32R32F.SafeRelease();
 	FVertexBuffer::ReleaseRHI();
 }
 
diff --git a/Engine/Source/Runtime/Engine/Private/Particles/ParticleSimulationGPU.h b/Engine/Source/Runtime/Engine/Private/Particles/ParticleSimulationGPU.h
--- a/Engine/Source/Runtime/Engine/Private/Particles/ParticleSimulationGPU.h
+++ b/Engine/Source/Runtime/Engine/Private/Particles/ParticleSimulationGPU.h
@@ -20,8 +20,3 @@ extern int32 GParticleSimulationTextureSizeY;
 extern const int32 GParticleSimulationTileSize;
 extern const int32 GParticlesPerTile;
 
-/** How many tiles are in the simulation textures. */
-extern int32 GParticleSimulationTileCountX;
-extern int32 GParticleSimulationTileCountY;
-extern int32 GParticleSimulationTileCount;
-
diff --git a/Engine/Source/Runtime/Engine/Public/ParticleResources.h b/Engine/Source/Runtime/Engine/Public/ParticleResources.h
--- a/Engine/Source/Runtime/Engine/Public/ParticleResources.h
+++ b/Engine/Source/Runtime/Engine/Public/ParticleResources.h
@@ -74,8 +74,8 @@ public:
 
 private:
 
-	/** SRV in to the buffer as an array of FVector2D values. */
-	FShaderResourceViewRHIRef VertexBufferSRV_G32R32F;
+	/** SRV in to the buffer as an array of FVector4f values. */
+	FShaderResourceViewRHIRef VertexBufferSRV_A32B32G32R32F;
 };
 
 /** The global scratch vertex buffer. */
