Change cvar to be default instead of read_only if we need to hotfix.

Read the cache cvar value from the FXSystem when deciding if we add a sorted cascade emitter.

diff --git a/Engine/Source/Runtime/Engine/Private/Particles/ParticleGpuSimulation.cpp b/Engine/Source/Runtime/Engine/Private/Particles/ParticleGpuSimulation.cpp
--- a/Engine/Source/Runtime/Engine/Private/Particles/ParticleGpuSimulation.cpp
+++ b/Engine/Source/Runtime/Engine/Private/Particles/ParticleGpuSimulation.cpp
@@ -91,7 +91,7 @@ static FAutoConsoleVariableRef CVarFXCascadeGpuSpriteDynamicAllocations(
 	TEXT("fx.Cascade.GpuSpriteDynamicAllocations"),
 	GFXCascadeGpuSpriteAllowDynAllocs,
 	TEXT("Controls if gpu sprite are allocated dynamically"),
-	ECVF_ReadOnly
+	ECVF_Default
 );
 
 /** The tile size. Texture space is allocated in TileSize x TileSize units. */
@@ -3004,7 +3004,7 @@ public:
 				FGPUSpriteMeshDataUserData* MeshBatchUserData = nullptr;
 				if (bAllowSorting && SortMode == PSORTMODE_DistanceToView)
 				{
-					if (GFXCascadeGpuSpriteAllowDynAllocs)
+					if (FXSystem->GetParticleSimulationResources()->SupportTileResizing())
 					{
 						UE_LOG(LogParticles, Warning, TEXT("Cascade doesn't support gpu sorting with particle gpu simulation resources resizing"));
 					}
