Rename struct BuildRequest to FBuildRequest in FRayTracingGeometryManager

diff --git a/Engine/Source/Runtime/Engine/Private/Rendering/RayTracingGeometryManager.cpp b/Engine/Source/Runtime/Engine/Private/Rendering/RayTracingGeometryManager.cpp
--- a/Engine/Source/Runtime/Engine/Private/Rendering/RayTracingGeometryManager.cpp
+++ b/Engine/Source/Runtime/Engine/Private/Rendering/RayTracingGeometryManager.cpp
@@ -78,7 +78,7 @@ static float GetInitialBuildPriority(ERTAccelerationStructureBuildPriority InBui
 
 FRayTracingGeometryManager::BuildRequestIndex FRayTracingGeometryManager::RequestBuildAccelerationStructure(FRayTracingGeometry* InGeometry, ERTAccelerationStructureBuildPriority InPriority, EAccelerationStructureBuildMode InBuildMode)
 {
-	BuildRequest Request;
+	FBuildRequest Request;
 	Request.BuildPriority = GetInitialBuildPriority(InPriority);
 	Request.Owner = InGeometry;
 	Request.BuildMode = EAccelerationStructureBuildMode::Build;
@@ -246,7 +246,7 @@ void FRayTracingGeometryManager::ProcessBuildRequests(FRHIComputeCommandList& In
 
 		SortedRequests.Empty(); // free potentially allocated memory
 
-		for (BuildRequest& Request : GeometryBuildRequests)
+		for (FBuildRequest& Request : GeometryBuildRequests)
 		{
 			const bool bRemoveFromRequestArray = false; // can't modify array while iterating over it
 			SetupBuildParams(Request, BuildParams, bRemoveFromRequestArray);
@@ -263,12 +263,12 @@ void FRayTracingGeometryManager::ProcessBuildRequests(FRHIComputeCommandList& In
 			TRACE_CPUPROFILER_EVENT_SCOPE(SortRequests);
 
 			// Is there a fast way to extract all entries from sparse array?
-			for (const BuildRequest& Request : GeometryBuildRequests)
+			for (const FBuildRequest& Request : GeometryBuildRequests)
 			{
 				SortedRequests.Add(Request);
 			}
 
-			SortedRequests.Sort([](const BuildRequest& InLHS, const BuildRequest& InRHS)
+			SortedRequests.Sort([](const FBuildRequest& InLHS, const FBuildRequest& InRHS)
 				{
 					return InLHS.BuildPriority > InRHS.BuildPriority;
 				});
@@ -277,7 +277,7 @@ void FRayTracingGeometryManager::ProcessBuildRequests(FRHIComputeCommandList& In
 		// process n requests each 'frame'
 		uint64 PrimitivesBuild = 0;
 		bool bAddBuildRequest = true;
-		for (BuildRequest& Request : SortedRequests)
+		for (FBuildRequest& Request : SortedRequests)
 		{
 			if (bAddBuildRequest || Request.BuildPriority >= 1.0f) // always build immediate requests
 			{
@@ -306,7 +306,7 @@ void FRayTracingGeometryManager::ProcessBuildRequests(FRHIComputeCommandList& In
 	BuildParams.Reset();
 }
 
-void FRayTracingGeometryManager::SetupBuildParams(const BuildRequest& InBuildRequest, TArray<FRayTracingGeometryBuildParams>& InBuildParams, bool bRemoveFromRequestArray)
+void FRayTracingGeometryManager::SetupBuildParams(const FBuildRequest& InBuildRequest, TArray<FRayTracingGeometryBuildParams>& InBuildParams, bool bRemoveFromRequestArray)
 {
 	check(InBuildRequest.RequestIndex != INDEX_NONE && InBuildRequest.Owner->RayTracingBuildRequestIndex != INDEX_NONE);
 
diff --git a/Engine/Source/Runtime/Engine/Public/Rendering/RayTracingGeometryManager.h b/Engine/Source/Runtime/Engine/Public/Rendering/RayTracingGeometryManager.h
--- a/Engine/Source/Runtime/Engine/Public/Rendering/RayTracingGeometryManager.h
+++ b/Engine/Source/Runtime/Engine/Public/Rendering/RayTracingGeometryManager.h
@@ -36,7 +36,7 @@ public:
 
 private:
 
-	struct BuildRequest
+	struct FBuildRequest
 	{
 		BuildRequestIndex RequestIndex = INDEX_NONE;
 
@@ -45,17 +45,17 @@ private:
 		EAccelerationStructureBuildMode BuildMode;
 	};
 
-	void SetupBuildParams(const BuildRequest& InBuildRequest, TArray<FRayTracingGeometryBuildParams>& InBuildParams, bool bRemoveFromRequestArray = true);
+	void SetupBuildParams(const FBuildRequest& InBuildRequest, TArray<FRayTracingGeometryBuildParams>& InBuildParams, bool bRemoveFromRequestArray = true);
 
 	FCriticalSection RequestCS;
 
-	TSparseArray<BuildRequest> GeometryBuildRequests;
+	TSparseArray<FBuildRequest> GeometryBuildRequests;
 
 	// Used for keeping track of geometries when ray tracing is dynamic
 	TSparseArray<FRayTracingGeometry*> RegisteredGeometries;
 
 	// Working array with all active build build params in the RHI
-	TArray<BuildRequest> SortedRequests;
+	TArray<FBuildRequest> SortedRequests;
 	TArray<FRayTracingGeometryBuildParams> BuildParams;
 
 	// TODO: Investigate removing this critical section
