Skeletal Mesh: Preserve exact triangle ordering when converting from import data to mesh description.

#rnx

diff --git a/Engine/Source/Runtime/Engine/Private/Rendering/SkeletalMeshLODImporterData.cpp b/Engine/Source/Runtime/Engine/Private/Rendering/SkeletalMeshLODImporterData.cpp
--- a/Engine/Source/Runtime/Engine/Private/Rendering/SkeletalMeshLODImporterData.cpp
+++ b/Engine/Source/Runtime/Engine/Private/Rendering/SkeletalMeshLODImporterData.cpp
@@ -2129,90 +2129,60 @@ bool FSkeletalMeshImportData::GetMeshDescription(const USkeletalMesh* InSkeletal
 
 	TArray<uint32> FaceSmoothingMasks;
 	FaceSmoothingMasks.AddZeroed(Faces.Num());
+
+	TArray<FPolygonGroupID> MaterialGroups;
+	MaterialGroups.Reserve(Materials.Num());
+	for (int32 MaterialIndex = 0; MaterialIndex < Materials.Num(); MaterialIndex++)
+	{
+		const FPolygonGroupID PolygonGroupID = OutMeshDescription.CreatePolygonGroup();
+		PolygonGroupMaterialSlotNames.Set(PolygonGroupID, FName(*Materials[MaterialIndex].MaterialImportName));
+		MaterialGroups.Add(PolygonGroupID);
+	}
 	
-	for(int32 StartStride = 0, EndStride = 0; EndStride != FaceIndices.Num(); StartStride = EndStride)
+	for (int32 TriangleIndex = 0; TriangleIndex < Faces.Num(); TriangleIndex++)
 	{
-		FPolygonGroupID PolygonGroupID;
-		FName MaterialName;
-		if (!Materials.IsEmpty())
-		{
-			const int32 MaterialIndex = Faces[FaceIndices[StartStride]].MatIndex;
-		
-			EndStride = StartStride + 1;
-			while (EndStride < FaceIndices.Num() && MaterialIndex == Faces[FaceIndices[EndStride]].MatIndex)
-			{
-				EndStride++;
-			}
-
-			// Create a section for each material index. We re-use vertex instances if they are
-			// referred to multiple times by the FTriangle object. However, because the tangents
-			// are stored on the triangle, it's possible to end up with a recycled vertex instance
-			// that has a different tangent. This is a limitation for now. 
-			// Along the way we track smoothing groups and use that to define hard edges once the
-			// entire mesh is defined.
-			PolygonGroupID = MaterialIndex;
-			MaterialName = FName(*Materials[MaterialIndex].MaterialImportName);
-		}
-		else
-		{
-			// No materials defined, in which case we set a dummy material on the whole mesh.
-			// This can happen when getting a mesh description from alternate influences mesh.
-			PolygonGroupID = 0;
-			EndStride = FaceIndices.Num();
-		}
-		
-		if (!OutMeshDescription.IsPolygonGroupValid(PolygonGroupID))
-		{
-			OutMeshDescription.CreatePolygonGroupWithID(PolygonGroupID);
-		}
-		PolygonGroupMaterialSlotNames.Set(PolygonGroupID, MaterialName);
+		const SkeletalMeshImportData::FTriangle &Triangle = Faces[TriangleIndex];
 
-		
-		
-		for (int32 Idx = StartStride; Idx < EndStride; Idx++)
+		for (int32 Corner = 0; Corner < 3; Corner++)
 		{
-			const SkeletalMeshImportData::FTriangle &Triangle = Faces[FaceIndices[Idx]];
+			const uint32 WedgeId = Triangle.WedgeIndex[Corner];
+			const SkeletalMeshImportData::FVertex &Wedge = Wedges[WedgeId];
+			const FVertexID VertexID = VertexIDMap[Wedge.VertexIndex];
 
-			for (int32 Corner = 0; Corner < 3; Corner++)
+			FVertexInstanceID VertexInstanceID = VertexInstanceIDMap[WedgeId];
+			if (VertexInstanceID == INDEX_NONE)
 			{
-				const uint32 WedgeId = Triangle.WedgeIndex[Corner];
-				const SkeletalMeshImportData::FVertex &Wedge = Wedges[WedgeId];
-				const FVertexID VertexID = VertexIDMap[Wedge.VertexIndex];
+				VertexInstanceID = OutMeshDescription.CreateVertexInstance(VertexID);
 
-				FVertexInstanceID VertexInstanceID = VertexInstanceIDMap[WedgeId];
-				if (VertexInstanceID == INDEX_NONE)
+				if (bHasVertexColors)
 				{
-					VertexInstanceID = OutMeshDescription.CreateVertexInstance(VertexID);
-
-					if (bHasVertexColors)
-					{
-						// Don't perform sRGB conversion (which mirrors what CreateFromMeshDescription does).
-						VertexInstanceColors.Set(VertexInstanceID, Wedge.Color.ReinterpretAsLinear());
-					}
-					for (int32 UVIndex = 0; UVIndex < static_cast<int32>(NumTexCoords); UVIndex++)
-					{
-						VertexInstanceUVs.Set(VertexInstanceID, UVIndex, Wedge.UVs[UVIndex]);
-					}
-					VertexInstanceTangents.Set(VertexInstanceID, Triangle.TangentX[Corner]);
-					VertexInstanceNormals.Set(VertexInstanceID, Triangle.TangentZ[Corner]);
-					VertexInstanceBinormalSigns.Set(VertexInstanceID,
-						((Triangle.TangentZ[Corner] ^ Triangle.TangentX[Corner]) | Triangle.TangentY[Corner]) < 0 ? -1.0f : 1.0f);
-					
-					VertexInstanceIDMap[WedgeId] = VertexInstanceID;
+					// Don't perform sRGB conversion (which mirrors what CreateFromMeshDescription does).
+					VertexInstanceColors.Set(VertexInstanceID, Wedge.Color.ReinterpretAsLinear());
 				}
-
-				TriangleVertexInstanceIDs[Corner] = VertexInstanceID; 
+				for (int32 UVIndex = 0; UVIndex < static_cast<int32>(NumTexCoords); UVIndex++)
+				{
+					VertexInstanceUVs.Set(VertexInstanceID, UVIndex, Wedge.UVs[UVIndex]);
+				}
+				VertexInstanceTangents.Set(VertexInstanceID, Triangle.TangentX[Corner]);
+				VertexInstanceNormals.Set(VertexInstanceID, Triangle.TangentZ[Corner]);
+				VertexInstanceBinormalSigns.Set(VertexInstanceID,
+					((Triangle.TangentZ[Corner] ^ Triangle.TangentX[Corner]) | Triangle.TangentY[Corner]) < 0 ? -1.0f : 1.0f);
+				
+				VertexInstanceIDMap[WedgeId] = VertexInstanceID;
 			}
 
-			const FTriangleID TriangleID = OutMeshDescription.CreateTriangle(PolygonGroupID, TriangleVertexInstanceIDs);
-			const FPolygonID PolygonID = OutMeshDescription.GetTrianglePolygon(TriangleID);
+			TriangleVertexInstanceIDs[Corner] = VertexInstanceID; 
+		}
+
+		const FPolygonGroupID PolygonGroupID = MaterialGroups[Triangle.MatIndex];
+		const FTriangleID TriangleID = OutMeshDescription.CreateTriangle(PolygonGroupID, TriangleVertexInstanceIDs);
+		const FPolygonID PolygonID = OutMeshDescription.GetTrianglePolygon(TriangleID);
 
-			if (PolygonID.GetValue() >= FaceSmoothingMasks.Num())
-			{
-				FaceSmoothingMasks.SetNum(PolygonID.GetValue() + 1);
-			}
-			FaceSmoothingMasks[PolygonID.GetValue()] = Triangle.SmoothingGroups;
+		if (PolygonID.GetValue() >= FaceSmoothingMasks.Num())
+		{
+			FaceSmoothingMasks.SetNum(PolygonID.GetValue() + 1);
 		}
+		FaceSmoothingMasks[PolygonID.GetValue()] = Triangle.SmoothingGroups;
 	}
 
 	// Convert morph targets
