Skeletal Mesh: Allow importing alt influence meshes with mismatched bones.

#rnx

diff --git a/Engine/Source/Developer/MeshBuilder/Private/SkeletalMeshBuilder.cpp b/Engine/Source/Developer/MeshBuilder/Private/SkeletalMeshBuilder.cpp
--- a/Engine/Source/Developer/MeshBuilder/Private/SkeletalMeshBuilder.cpp
+++ b/Engine/Source/Developer/MeshBuilder/Private/SkeletalMeshBuilder.cpp
@@ -158,11 +158,20 @@ bool FSkeletalMeshBuilder::Build(const FSkeletalMeshBuildParameters& SkeletalMes
 
 		//Re-apply the alternate skinning it must be after the inline reduction
 		SlowTask.EnterProgressFrame(1.0f, NSLOCTEXT("SkeltalMeshBuilder", "RebuildAlternateSkinning", "Rebuilding alternate skinning..."));
-		const TArray<FSkinWeightProfileInfo>& SkinProfiles = SkeletalMesh->GetSkinWeightProfiles();
+		TArray<FSkinWeightProfileInfo> SkinProfiles = SkeletalMesh->GetSkinWeightProfiles();
 		for (int32 SkinProfileIndex = 0; SkinProfileIndex < SkinProfiles.Num(); ++SkinProfileIndex)
 		{
 			const FSkinWeightProfileInfo& ProfileInfo = SkinProfiles[SkinProfileIndex];
-			FLODUtilities::UpdateAlternateSkinWeights(SkeletalMesh, ProfileInfo.Name, LODIndex, Options);
+			if (!FLODUtilities::UpdateAlternateSkinWeights(SkeletalMesh, ProfileInfo.Name, LODIndex, Options))
+			{
+				// If we failed to update the profile, remove it from the skeletal mesh so that we don't try to rely on it being there
+				// later on in the mesh build.
+				SkeletalMesh->GetSkinWeightProfiles().RemoveAll([Name=ProfileInfo.Name](const FSkinWeightProfileInfo& InProfileInfo)
+				{
+					return InProfileInfo.Name == Name;
+				});
+				BuildLODModel.SkinWeightProfiles.Remove(ProfileInfo.Name);
+			}
 		}
 
 		// Copy vertex attribute definitions and their values from the import model.
diff --git a/Engine/Source/Developer/SkeletalMeshUtilitiesCommon/Private/LODUtilities.cpp b/Engine/Source/Developer/SkeletalMeshUtilitiesCommon/Private/LODUtilities.cpp
--- a/Engine/Source/Developer/SkeletalMeshUtilitiesCommon/Private/LODUtilities.cpp
+++ b/Engine/Source/Developer/SkeletalMeshUtilitiesCommon/Private/LODUtilities.cpp
@@ -2329,7 +2329,7 @@ bool FLODUtilities::UpdateAlternateSkinWeights(
 		FText Message = FText::Format(NSLOCTEXT("FLODUtilities_UpdateAlternateSkinWeights", "AlternateDataNotAvailable", 
 			"Asset {SkeletalMeshName} LOD {LODIndex} failed to import skin weight profile {ProfileName}. The alternate skinning imported source data is not available."), 
 			Args);
-		UE_LOG(LogLODUtilities, Warning, TEXT("%s"), *(Message.ToString()));
+		UE_LOG(LogLODUtilities, Display, TEXT("%s"), *(Message.ToString()));
 		return false;
 	}
 
diff --git a/Engine/Source/Runtime/Engine/Private/Rendering/SkeletalMeshLODImporterData.cpp b/Engine/Source/Runtime/Engine/Private/Rendering/SkeletalMeshLODImporterData.cpp
--- a/Engine/Source/Runtime/Engine/Private/Rendering/SkeletalMeshLODImporterData.cpp
+++ b/Engine/Source/Runtime/Engine/Private/Rendering/SkeletalMeshLODImporterData.cpp
@@ -1839,14 +1839,27 @@ static void CopySkinWeightsToAttribute(
 			int32 BoneIndex = RawInfluence.BoneIndex;
 			if (InBoneIndexMap)
 			{
-				// The map should be complete.
-				BoneIndex = (*InBoneIndexMap)[BoneIndex];
+				// If we can't map from a source bone to a target bone, we skip the weight. 
+				if (const int32* BoneIndexPtr = InBoneIndexMap->Find(BoneIndex))
+				{
+					BoneIndex = *BoneIndexPtr;
+				}
+				else
+				{
+					continue;
+				}
 			}
 			
 			FBoneWeight BoneWeight(static_cast<FBoneIndexType>(BoneIndex), RawInfluence.Weight);
 			BoneWeights.Add(BoneWeight);
 		}
 
+		if (BoneWeights.IsEmpty())
+		{
+			const FBoneWeight RootBoneWeight(0, 1.0f);
+			BoneWeights.Add(RootBoneWeight);
+		}
+
 		OutSkinWeightsAttribute.Set(InVertexIDMap[VertexIndex], BoneWeights);		
 	}
 }
@@ -2009,10 +2022,9 @@ bool FSkeletalMeshImportData::GetMeshDescription(const USkeletalMesh* InSkeletal
 				
 				if (InSkeletalMesh)
 				{
-					UE_ASSET_LOG(LogSkeletalMeshLODImporterData, Display, InSkeletalMesh, TEXT("Alternate influence profile '%s' binds to one or more bones (%s) that do not exist on base mesh's skeleton. Profile will be dropped since it may otherwise introduce visual errors."),
+					UE_ASSET_LOG(LogSkeletalMeshLODImporterData, Display, InSkeletalMesh, TEXT("Alternate influence profile '%s' binds to one or more bones (%s) that do not exist on base mesh's skeleton. Those bone bindings will be dropped, which may reduce visual quality."),
 						*AlternateInfluenceProfileName, *FString::Join(MissingBoneNames, TEXT(", ")));  
 				}
-				continue;
 			}
 
 			if (ensure(MeshAttributes.RegisterSkinWeightAttribute(*AlternateInfluenceProfileName)))
@@ -2335,8 +2347,7 @@ void FSkeletalMeshImportData::CopySkinWeightsToMeshDescription(
 	TMap<int32, int32> AltMeshBoneToBaseBoneMap;
 	for (int32 BoneIndex = 0; BoneIndex < InSkinWeightMesh.RefBonesBinary.Num(); BoneIndex++)
 	{
-		// It's ok if a matching bone isn't in the ref map, as long as it's not bound to the skin. This would have been checked
-		// in GetMeshDescription already.
+		// Any bone not in the base mesh will not get bound. If that results in no bindings for a vertex, the vertex will be bound to the root.
 		if (const int32* BaseBoneIndex = BaseBoneToIndexMap.Find(InSkinWeightMesh.RefBonesBinary[BoneIndex].Name))
 		{
 			AltMeshBoneToBaseBoneMap.Add(BoneIndex, *BaseBoneIndex);
diff --git a/Engine/Source/Runtime/SkeletalMeshDescription/Private/SkeletalMeshOperations.cpp b/Engine/Source/Runtime/SkeletalMeshDescription/Private/SkeletalMeshOperations.cpp
--- a/Engine/Source/Runtime/SkeletalMeshDescription/Private/SkeletalMeshOperations.cpp
+++ b/Engine/Source/Runtime/SkeletalMeshDescription/Private/SkeletalMeshOperations.cpp
@@ -150,8 +150,17 @@ bool FSkeletalMeshOperations::CopySkinWeightAttributeFromMesh(
 		{
 			for (FBoneWeight OriginalWeight: InWeights)
 			{
-				FBoneWeight NewWeight(static_cast<FBoneIndexType>((*SourceBoneIndexToTargetBoneIndexMap)[OriginalWeight.GetBoneIndex()]), OriginalWeight.GetRawWeight());
-				Weights.Add(NewWeight);
+				if (const int32* BoneIndexPtr = SourceBoneIndexToTargetBoneIndexMap->Find(OriginalWeight.GetBoneIndex()))
+				{
+					FBoneWeight NewWeight(static_cast<FBoneIndexType>(*BoneIndexPtr), OriginalWeight.GetRawWeight());
+					Weights.Add(NewWeight);
+				}
+			}
+
+			if (Weights.IsEmpty())
+			{
+				const FBoneWeight RootBoneWeight(0, 1.0f);
+				Weights.Add(RootBoneWeight);
 			}
 		}
 		else
