Skeletal Mesh: Handle models with an empty material list.

#rnx

diff --git a/Engine/Source/Runtime/Engine/Private/Rendering/SkeletalMeshLODImporterData.cpp b/Engine/Source/Runtime/Engine/Private/Rendering/SkeletalMeshLODImporterData.cpp
--- a/Engine/Source/Runtime/Engine/Private/Rendering/SkeletalMeshLODImporterData.cpp
+++ b/Engine/Source/Runtime/Engine/Private/Rendering/SkeletalMeshLODImporterData.cpp
@@ -2144,13 +2144,22 @@ bool FSkeletalMeshImportData::GetMeshDescription(const USkeletalMesh* InSkeletal
 
 	TArray<FPolygonGroupID> MaterialGroups;
 	MaterialGroups.Reserve(Materials.Num());
-	for (int32 MaterialIndex = 0; MaterialIndex < Materials.Num(); MaterialIndex++)
+	if (!Materials.IsEmpty())
+	{
+		for (int32 MaterialIndex = 0; MaterialIndex < Materials.Num(); MaterialIndex++)
+		{
+			const FPolygonGroupID PolygonGroupID = OutMeshDescription.CreatePolygonGroup();
+			PolygonGroupMaterialSlotNames.Set(PolygonGroupID, FName(*Materials[MaterialIndex].MaterialImportName));
+			MaterialGroups.Add(PolygonGroupID);
+		}
+	}
+	else
 	{
 		const FPolygonGroupID PolygonGroupID = OutMeshDescription.CreatePolygonGroup();
-		PolygonGroupMaterialSlotNames.Set(PolygonGroupID, FName(*Materials[MaterialIndex].MaterialImportName));
+		PolygonGroupMaterialSlotNames.Set(PolygonGroupID, NAME_None);
 		MaterialGroups.Add(PolygonGroupID);
 	}
-	
+
 	for (int32 TriangleIndex = 0; TriangleIndex < Faces.Num(); TriangleIndex++)
 	{
 		const SkeletalMeshImportData::FTriangle &Triangle = Faces[TriangleIndex];
@@ -2186,7 +2195,15 @@ bool FSkeletalMeshImportData::GetMeshDescription(const USkeletalMesh* InSkeletal
 			TriangleVertexInstanceIDs[Corner] = VertexInstanceID; 
 		}
 
-		const FPolygonGroupID PolygonGroupID = MaterialGroups[Triangle.MatIndex];
+		FPolygonGroupID PolygonGroupID;
+		if (MaterialGroups.IsValidIndex(Triangle.MatIndex))
+		{
+			PolygonGroupID = MaterialGroups[Triangle.MatIndex];
+		}
+		else
+		{
+			PolygonGroupID = MaterialGroups[0];
+		}
 		const FTriangleID TriangleID = OutMeshDescription.CreateTriangle(PolygonGroupID, TriangleVertexInstanceIDs);
 		const FPolygonID PolygonID = OutMeshDescription.GetTrianglePolygon(TriangleID);
 
