Fixed delegate registation from the ALT

Moved custom property node re-instancing handler so that it is called from the owning anim instance

diff --git a/Engine/Plugins/Animation/ControlRig/Source/ControlRig/Private/Sequencer/ControlRigLayerInstance.cpp b/Engine/Plugins/Animation/ControlRig/Source/ControlRig/Private/Sequencer/ControlRigLayerInstance.cpp
--- a/Engine/Plugins/Animation/ControlRig/Source/ControlRig/Private/Sequencer/ControlRigLayerInstance.cpp
+++ b/Engine/Plugins/Animation/ControlRig/Source/ControlRig/Private/Sequencer/ControlRigLayerInstance.cpp
@@ -170,4 +170,24 @@ UAnimInstance* UControlRigLayerInstance::GetSourceAnimInstance()
 	return GetProxyOnGameThread<FControlRigLayerInstanceProxy>().GetSourceAnimInstance();
 }
 
+#if WITH_EDITOR
+void UControlRigLayerInstance::HandleObjectsReinstanced(const TMap<UObject*, UObject*>& OldToNewInstanceMap)
+{
+	Super::HandleObjectsReinstanced(OldToNewInstanceMap);
+
+	static IConsoleVariable* UseLegacyAnimInstanceReinstancingBehavior = IConsoleManager::Get().FindConsoleVariable(TEXT("bp.UseLegacyAnimInstanceReinstancingBehavior"));
+	if(UseLegacyAnimInstanceReinstancingBehavior == nullptr || !UseLegacyAnimInstanceReinstancingBehavior->GetBool())
+	{
+		// Forward to control rig nodes
+		FControlRigLayerInstanceProxy& Proxy = GetProxyOnGameThread<FControlRigLayerInstanceProxy>();
+		for(TSharedPtr<FAnimNode_ControlRig_ExternalSource>& ControlRigNode : Proxy.ControlRigNodes)
+		{
+			if(ControlRigNode.IsValid())
+			{
+				ControlRigNode->HandleObjectsReinstanced(OldToNewInstanceMap);
+			}
+		}
+	}
+}
+#endif
 
diff --git a/Engine/Plugins/Animation/ControlRig/Source/ControlRig/Private/Sequencer/ControlRigLayerInstanceProxy.h b/Engine/Plugins/Animation/ControlRig/Source/ControlRig/Private/Sequencer/ControlRigLayerInstanceProxy.h
--- a/Engine/Plugins/Animation/ControlRig/Source/ControlRig/Private/Sequencer/ControlRigLayerInstanceProxy.h
+++ b/Engine/Plugins/Animation/ControlRig/Source/ControlRig/Private/Sequencer/ControlRigLayerInstanceProxy.h
@@ -108,6 +108,7 @@ public:
 // 	virtual FAnimNode_Base* GetCustomRootNode() 
 
 	friend struct FAnimNode_ControlRigInputPose;
+	friend class UControlRigLayerInstance;
 protected:
 	/** Sort Control Rig node*/
 	void SortControlRigNodes();
diff --git a/Engine/Plugins/Animation/ControlRig/Source/ControlRig/Public/Sequencer/ControlRigLayerInstance.h b/Engine/Plugins/Animation/ControlRig/Source/ControlRig/Public/Sequencer/ControlRigLayerInstance.h
--- a/Engine/Plugins/Animation/ControlRig/Source/ControlRig/Public/Sequencer/ControlRigLayerInstance.h
+++ b/Engine/Plugins/Animation/ControlRig/Source/ControlRig/Public/Sequencer/ControlRigLayerInstance.h
@@ -49,6 +49,10 @@ public:
 	virtual void SetSourceAnimInstance(UAnimInstance* SourceAnimInstance) override;
 	virtual bool DoesSupportDifferentSourceAnimInstance() const override { return true; }
 
+#if WITH_EDITOR
+	virtual void HandleObjectsReinstanced(const TMap<UObject*, UObject*>& OldToNewInstanceMap) override;
+#endif
+
 protected:
 	// UAnimInstance interface
 	virtual FAnimInstanceProxy* CreateAnimInstanceProxy() override;
diff --git a/Engine/Source/Runtime/Engine/Classes/Animation/AnimNode_CustomProperty.h b/Engine/Source/Runtime/Engine/Classes/Animation/AnimNode_CustomProperty.h
--- a/Engine/Source/Runtime/Engine/Classes/Animation/AnimNode_CustomProperty.h
+++ b/Engine/Source/Runtime/Engine/Classes/Animation/AnimNode_CustomProperty.h
@@ -53,8 +53,11 @@ public:
 	// FAnimNode_Base interface
 	ENGINE_API virtual void OnInitializeAnimInstance(const FAnimInstanceProxy* InProxy, const UAnimInstance* InAnimInstance) override;
 	virtual bool NeedsOnInitializeAnimInstance() const override { return true; }
+
+	// Handle object reinstancing in editor
+	ENGINE_API void HandleObjectsReinstanced(const TMap<UObject*, UObject*>& OldToNewInstanceMap);
 #endif
-	
+
 protected:
 	/** List of source properties to use, 1-1 with Dest names below, built by the compiler */
 	UPROPERTY(meta=(BlueprintCompilerGeneratedDefaults))
@@ -92,10 +95,6 @@ protected:
 	 * replacement.
 	 */
 	ENGINE_API virtual void HandleObjectsReinstanced_Impl(UObject* InSourceObject, UObject* InTargetObject, const TMap<UObject*, UObject*>& OldToNewInstanceMap);
-	
-private:
-	// Handle object reinstancing in editor
-	ENGINE_API void HandleObjectsReinstanced(const TMap<UObject*, UObject*>& OldToNewInstanceMap);
 #endif // WITH_EDITOR
 
 #if WITH_EDITORONLY_DATA
@@ -113,6 +112,8 @@ protected:
 #endif // ANIMNODE_STATS_VERBOSE
 
 	friend class UAnimGraphNode_CustomProperty;
+	friend class UAnimInstance;
+	friend class UControlRigLayerInstance;
 };
 
 template<>
diff --git a/Engine/Source/Runtime/Engine/Private/Animation/AnimInstance.cpp b/Engine/Source/Runtime/Engine/Private/Animation/AnimInstance.cpp
--- a/Engine/Source/Runtime/Engine/Private/Animation/AnimInstance.cpp
+++ b/Engine/Source/Runtime/Engine/Private/Animation/AnimInstance.cpp
@@ -125,7 +125,7 @@ UAnimInstance::UAnimInstance(const FObjectInitializer& ObjectInitializer)
 #endif
 
 #if WITH_EDITOR
-	if(!HasAnyFlags(RF_ClassDefaultObject))
+	if(!HasAnyFlags(RF_ClassDefaultObject) && !GetClass()->HasAnyClassFlags(CLASS_Native))
 	{
 		FCoreUObjectDelegates::OnObjectsReinstanced.AddUObject(this, &UAnimInstance::HandleObjectsReinstanced);
 	}
@@ -4230,6 +4230,20 @@ void UAnimInstance::HandleObjectsReinstanced(const TMap<UObject*, UObject*>& Old
 				MeshComponent->ClearMotionVector();
 			}
 		}
+
+		// Forward to custom property-based nodes even if it wasnt this object that was reinstanced, as they may reference different objects that may
+		// also have been reinstanced
+		if(IAnimClassInterface* AnimClassInterface = IAnimClassInterface::GetFromClass(GetClass()))
+		{
+			for(const FStructProperty* NodeProperty : AnimClassInterface->GetAnimNodeProperties())
+			{
+				if(NodeProperty && NodeProperty->Struct && NodeProperty->Struct->IsChildOf(FAnimNode_CustomProperty::StaticStruct()))
+				{
+					FAnimNode_CustomProperty* CustomPropertyNode = NodeProperty->ContainerPtrToValuePtr<FAnimNode_CustomProperty>(this);
+					CustomPropertyNode->HandleObjectsReinstanced(OldToNewInstanceMap);
+				}
+			}
+		}
 	}
 }
 #endif
diff --git a/Engine/Source/Runtime/Engine/Private/Animation/AnimNode_CustomProperty.cpp b/Engine/Source/Runtime/Engine/Private/Animation/AnimNode_CustomProperty.cpp
--- a/Engine/Source/Runtime/Engine/Private/Animation/AnimNode_CustomProperty.cpp
+++ b/Engine/Source/Runtime/Engine/Private/Animation/AnimNode_CustomProperty.cpp
@@ -10,16 +10,10 @@ FAnimNode_CustomProperty::FAnimNode_CustomProperty()
 	: FAnimNode_Base()
 	, TargetInstance(nullptr)
 {
-#if WITH_EDITOR
-	FCoreUObjectDelegates::OnObjectsReinstanced.AddRaw(this, &FAnimNode_CustomProperty::HandleObjectsReinstanced);
-#endif // WITH_EDITOR
 }
 
 FAnimNode_CustomProperty::~FAnimNode_CustomProperty()
 {
-#if WITH_EDITOR
-	FCoreUObjectDelegates::OnObjectsReinstanced.RemoveAll(this);
-#endif // WITH_EDITOR
 }
 
 void FAnimNode_CustomProperty::SetTargetInstance(UObject* InInstance)
