Animation debug draw items now flush in both post update and post evaluate

Queued items are now emptied when flushed

Adapted from:
PR #6072: Move QueuedDrawDebugItems drawing from FAnimInstanceProxy::PostUpdate to FAnimInstanceProxy::PostEvaluate

diff --git a/Engine/Source/Runtime/Engine/Private/Animation/AnimInstanceProxy.cpp b/Engine/Source/Runtime/Engine/Private/Animation/AnimInstanceProxy.cpp
--- a/Engine/Source/Runtime/Engine/Private/Animation/AnimInstanceProxy.cpp
+++ b/Engine/Source/Runtime/Engine/Private/Animation/AnimInstanceProxy.cpp
@@ -658,45 +658,53 @@ void FAnimInstanceProxy::PostUpdate(UAnimInstance* InAnimInstance) const
 	InAnimInstance->NotifyQueue.Append(NotifyQueue);
 	InAnimInstance->NotifyQueue.ApplyMontageNotifies(*this);
 
-	// Send Queued DrawDebug Commands.
 #if ENABLE_ANIM_DRAW_DEBUG
+	FlushQueuedDebugDrawItems(InAnimInstance->GetSkelMeshComponent()->GetOwner(), InAnimInstance->GetSkelMeshComponent()->GetWorld());
+#endif
+
+#if ENABLE_ANIM_LOGGING
+	FMessageLog MessageLog(GetTargetLogNameForCurrentWorldType());
+	const TArray<FLogMessageEntry>* Messages = LoggedMessagesMap.Find(NAME_Update);
+	if (ensureMsgf(Messages, TEXT("PreUpdate isn't called. This could potentially cause other issues.")))
+	{
+		for (const FLogMessageEntry& Message : *Messages)
+		{
+			MessageLog.AddMessage(Message);
+		}
+	}
+#endif
+}
+
+#if ENABLE_ANIM_DRAW_DEBUG
+void FAnimInstanceProxy::FlushQueuedDebugDrawItems(AActor* InActor, UWorld* InWorld) const
+{
 	for (const FQueuedDrawDebugItem& DebugItem : QueuedDrawDebugItems)
 	{
 		switch (DebugItem.ItemType)
 		{
 			case EDrawDebugItemType::OnScreenMessage: GEngine->AddOnScreenDebugMessage(INDEX_NONE, 0.f, DebugItem.Color, DebugItem.Message, false, DebugItem.TextScale); break;
-			case EDrawDebugItemType::InWorldMessage: DrawDebugString(InAnimInstance->GetSkelMeshComponent()->GetWorld(), DebugItem.StartLoc, DebugItem.Message, InAnimInstance->GetSkelMeshComponent()->GetOwner(), DebugItem.Color, DebugItem.LifeTime, false /*bDrawShadow*/, DebugItem.TextScale.X); break;
-			case EDrawDebugItemType::DirectionalArrow: DrawDebugDirectionalArrow(InAnimInstance->GetSkelMeshComponent()->GetWorld(), DebugItem.StartLoc, DebugItem.EndLoc.Value, DebugItem.Size, DebugItem.Color, DebugItem.bPersistentLines, DebugItem.LifeTime, DebugItem.DepthPriority, DebugItem.Thickness); break;
-			case EDrawDebugItemType::Sphere: DrawDebugSphere(InAnimInstance->GetSkelMeshComponent()->GetWorld(), DebugItem.Center, DebugItem.Radius, DebugItem.Segments, DebugItem.Color, DebugItem.bPersistentLines, DebugItem.LifeTime, DebugItem.DepthPriority, DebugItem.Thickness); break;
-			case EDrawDebugItemType::Line: DrawDebugLine(InAnimInstance->GetSkelMeshComponent()->GetWorld(), DebugItem.StartLoc, DebugItem.EndLoc.Value, DebugItem.Color, DebugItem.bPersistentLines, DebugItem.LifeTime, DebugItem.DepthPriority, DebugItem.Thickness); break;
-			case EDrawDebugItemType::CoordinateSystem: DrawDebugCoordinateSystem(InAnimInstance->GetSkelMeshComponent()->GetWorld(), DebugItem.StartLoc, DebugItem.Rotation, DebugItem.Size, DebugItem.bPersistentLines, DebugItem.LifeTime, DebugItem.DepthPriority, DebugItem.Thickness); break;
-			case EDrawDebugItemType::Point: DrawDebugPoint(InAnimInstance->GetSkelMeshComponent()->GetWorld(), DebugItem.StartLoc, DebugItem.Size, DebugItem.Color, DebugItem.bPersistentLines, DebugItem.LifeTime, DebugItem.DepthPriority); break;
+			case EDrawDebugItemType::InWorldMessage: DrawDebugString(InWorld, DebugItem.StartLoc, DebugItem.Message, InActor, DebugItem.Color, DebugItem.LifeTime, false /*bDrawShadow*/, DebugItem.TextScale.X); break;
+			case EDrawDebugItemType::DirectionalArrow: DrawDebugDirectionalArrow(InWorld, DebugItem.StartLoc, DebugItem.EndLoc.Value, DebugItem.Size, DebugItem.Color, DebugItem.bPersistentLines, DebugItem.LifeTime, DebugItem.DepthPriority, DebugItem.Thickness); break;
+			case EDrawDebugItemType::Sphere: DrawDebugSphere(InWorld, DebugItem.Center, DebugItem.Radius, DebugItem.Segments, DebugItem.Color, DebugItem.bPersistentLines, DebugItem.LifeTime, DebugItem.DepthPriority, DebugItem.Thickness); break;
+			case EDrawDebugItemType::Line: DrawDebugLine(InWorld, DebugItem.StartLoc, DebugItem.EndLoc.Value, DebugItem.Color, DebugItem.bPersistentLines, DebugItem.LifeTime, DebugItem.DepthPriority, DebugItem.Thickness); break;
+			case EDrawDebugItemType::CoordinateSystem: DrawDebugCoordinateSystem(InWorld, DebugItem.StartLoc, DebugItem.Rotation, DebugItem.Size, DebugItem.bPersistentLines, DebugItem.LifeTime, DebugItem.DepthPriority, DebugItem.Thickness); break;
+			case EDrawDebugItemType::Point: DrawDebugPoint(InWorld, DebugItem.StartLoc, DebugItem.Size, DebugItem.Color, DebugItem.bPersistentLines, DebugItem.LifeTime, DebugItem.DepthPriority); break;
 			case EDrawDebugItemType::Circle:
 			{
 				const FMatrix RotationMatrix = FRotationMatrix::MakeFromZ(DebugItem.Direction.Value);
 				const FVector ForwardVector = RotationMatrix.GetScaledAxis(EAxis::X);
 				const FVector RightVector = RotationMatrix.GetScaledAxis(EAxis::Y);
-				DrawDebugCircle(InAnimInstance->GetSkelMeshComponent()->GetWorld(), DebugItem.Center, DebugItem.Radius, DebugItem.Segments, DebugItem.Color, DebugItem.bPersistentLines, DebugItem.LifeTime, DebugItem.DepthPriority, DebugItem.Thickness, ForwardVector, RightVector, false);
+				DrawDebugCircle(InWorld, DebugItem.Center, DebugItem.Radius, DebugItem.Segments, DebugItem.Color, DebugItem.bPersistentLines, DebugItem.LifeTime, DebugItem.DepthPriority, DebugItem.Thickness, ForwardVector, RightVector, false);
 				break;
 			}
-			case EDrawDebugItemType::Cone: DrawDebugCone(InAnimInstance->GetSkelMeshComponent()->GetWorld(), DebugItem.Center, DebugItem.Direction.Value, DebugItem.Length, DebugItem.AngleWidth, DebugItem.AngleHeight, DebugItem.Segments, DebugItem.Color, DebugItem.bPersistentLines, DebugItem.LifeTime, DebugItem.DepthPriority, DebugItem.Thickness); break;
-			case EDrawDebugItemType::Capsule: DrawDebugCapsule(InAnimInstance->GetSkelMeshComponent()->GetWorld(), DebugItem.Center, DebugItem.Size, DebugItem.Radius, DebugItem.Rotation.Quaternion(), DebugItem.Color, DebugItem.bPersistentLines, DebugItem.LifeTime, 0, DebugItem.Thickness); break;
+			case EDrawDebugItemType::Cone: DrawDebugCone(InWorld, DebugItem.Center, DebugItem.Direction.Value, DebugItem.Length, DebugItem.AngleWidth, DebugItem.AngleHeight, DebugItem.Segments, DebugItem.Color, DebugItem.bPersistentLines, DebugItem.LifeTime, DebugItem.DepthPriority, DebugItem.Thickness); break;
+			case EDrawDebugItemType::Capsule: DrawDebugCapsule(InWorld, DebugItem.Center, DebugItem.Size, DebugItem.Radius, DebugItem.Rotation.Quaternion(), DebugItem.Color, DebugItem.bPersistentLines, DebugItem.LifeTime, 0, DebugItem.Thickness); break;
 		}
 	}
-#endif
 
-#if ENABLE_ANIM_LOGGING
-	FMessageLog MessageLog(GetTargetLogNameForCurrentWorldType());
-	const TArray<FLogMessageEntry>* Messages = LoggedMessagesMap.Find(NAME_Update);
-	if (ensureMsgf(Messages, TEXT("PreUpdate isn't called. This could potentially cause other issues.")))
-	{
-		for (const FLogMessageEntry& Message : *Messages)
-		{
-			MessageLog.AddMessage(Message);
-		}
-	}
-#endif
+	QueuedDrawDebugItems.Reset();
 }
+#endif
 
 void FAnimInstanceProxy::PostEvaluate(UAnimInstance* InAnimInstance)
 {
@@ -711,6 +719,10 @@ void FAnimInstanceProxy::PostEvaluate(UAnimInstance* InAnimInstance)
 	
 	ClearObjects();
 
+#if ENABLE_ANIM_DRAW_DEBUG
+	FlushQueuedDebugDrawItems(InAnimInstance->GetSkelMeshComponent()->GetOwner(), InAnimInstance->GetSkelMeshComponent()->GetWorld());
+#endif
+
 #if ENABLE_ANIM_LOGGING
 	FMessageLog MessageLog(GetTargetLogNameForCurrentWorldType());
 	if(const TArray<FLogMessageEntry>* Messages = LoggedMessagesMap.Find(NAME_Evaluate))
diff --git a/Engine/Source/Runtime/Engine/Public/Animation/AnimInstanceProxy.h b/Engine/Source/Runtime/Engine/Public/Animation/AnimInstanceProxy.h
--- a/Engine/Source/Runtime/Engine/Public/Animation/AnimInstanceProxy.h
+++ b/Engine/Source/Runtime/Engine/Public/Animation/AnimInstanceProxy.h
@@ -447,7 +447,7 @@ public:
 	ENGINE_API void GatherDebugData_WithRoot(FNodeDebugData& DebugData, FAnimNode_Base* InRootNode, FName InLayerName);
 
 #if ENABLE_ANIM_DRAW_DEBUG
-	TArray<FQueuedDrawDebugItem> QueuedDrawDebugItems;
+	mutable TArray<FQueuedDrawDebugItem> QueuedDrawDebugItems;
 
 	ENGINE_API void AnimDrawDebugOnScreenMessage(const FString& DebugMessage, const FColor& Color, const FVector2D& TextScale = FVector2D::UnitVector, ESceneDepthPriorityGroup DepthPriority = SDPG_World);
 	ENGINE_API void AnimDrawDebugInWorldMessage(const FString& DebugMessage, const FVector& TextLocation, const FColor& Color, float TextScale);
@@ -966,6 +966,11 @@ protected:
 		Sync.ResetAll();
 	}
 
+#if ENABLE_ANIM_DRAW_DEBUG
+	/** Send any queued DrawDebug commands and reset the queue */
+	void FlushQueuedDebugDrawItems(AActor* InActor, UWorld* InWorld) const;
+#endif
+
 private:
 
 	FName GetTargetLogNameForCurrentWorldType() const;
