Removed critical section used to control access to CachedRayTracingStateProxiesMap
- CL 30621379 moved register/unregister proxies with cached ray tracing state to FScene::Update(...) which runs on render thread
- CachedRayTracingStateProxiesMap is only accessed on render thread so the critical section is not required

diff --git a/Engine/Source/Runtime/Engine/Private/Rendering/RayTracingGeometryManager.cpp b/Engine/Source/Runtime/Engine/Private/Rendering/RayTracingGeometryManager.cpp
--- a/Engine/Source/Runtime/Engine/Private/Rendering/RayTracingGeometryManager.cpp
+++ b/Engine/Source/Runtime/Engine/Private/Rendering/RayTracingGeometryManager.cpp
@@ -176,7 +176,8 @@ void FRayTracingGeometryManager::Tick(FRHICommandList& RHICmdList)
 		}
 
 		{
-			FScopeLock ScopeLock(&CachedRayTracingStateProxiesCS);
+			checkf(IsInRenderingThread(), TEXT("Can only access CachedRayTracingStateProxiesMap on render thread otherwise need a critical section"));
+
 			for (TPair<const UStaticMesh*, TSet<FPrimitiveSceneProxy*>>& ProxiesSet : CachedRayTracingStateProxiesMap)
 			{
 				for (FPrimitiveSceneProxy* Proxy : ProxiesSet.Value)
@@ -328,7 +329,7 @@ void FRayTracingGeometryManager::SetupBuildParams(const FBuildRequest& InBuildRe
 
 void FRayTracingGeometryManager::RegisterProxyWithCachedRayTracingState(FPrimitiveSceneProxy* Proxy, const UStaticMesh* StaticMesh)
 {
-	FScopeLock ScopeLock(&CachedRayTracingStateProxiesCS);
+	checkf(IsInRenderingThread(), TEXT("Can only access CachedRayTracingStateProxiesMap on render thread otherwise need a critical section"));
 
 	TSet<FPrimitiveSceneProxy*>& ProxiesSet = CachedRayTracingStateProxiesMap.FindOrAdd(StaticMesh);
 	check(!ProxiesSet.Contains(Proxy));
@@ -338,7 +339,7 @@ void FRayTracingGeometryManager::RegisterProxyWithCachedRayTracingState(FPrimiti
 
 void FRayTracingGeometryManager::UnregisterProxyWithCachedRayTracingState(FPrimitiveSceneProxy* Proxy, const UStaticMesh* StaticMesh)
 {
-	FScopeLock ScopeLock(&CachedRayTracingStateProxiesCS);
+	checkf(IsInRenderingThread(), TEXT("Can only access CachedRayTracingStateProxiesMap on render thread otherwise need a critical section"));
 
 	check(CachedRayTracingStateProxiesMap.Contains(StaticMesh));
 
@@ -354,9 +355,7 @@ void FRayTracingGeometryManager::UnregisterProxyWithCachedRayTracingState(FPrimi
 
 void FRayTracingGeometryManager::RequestUpdateCachedRenderState(const UStaticMesh* StaticMesh)
 {
-	check(IsInRenderingThread());
-
-	FScopeLock ScopeLock(&CachedRayTracingStateProxiesCS);
+	checkf(IsInRenderingThread(), TEXT("Can only access CachedRayTracingStateProxiesMap on render thread otherwise need a critical section"));
 
 	const TSet<FPrimitiveSceneProxy*>* ProxiesSet = CachedRayTracingStateProxiesMap.Find(StaticMesh);
 
diff --git a/Engine/Source/Runtime/Engine/Public/Rendering/RayTracingGeometryManager.h b/Engine/Source/Runtime/Engine/Public/Rendering/RayTracingGeometryManager.h
--- a/Engine/Source/Runtime/Engine/Public/Rendering/RayTracingGeometryManager.h
+++ b/Engine/Source/Runtime/Engine/Public/Rendering/RayTracingGeometryManager.h
@@ -58,11 +58,6 @@ private:
 	TArray<FBuildRequest> SortedRequests;
 	TArray<FRayTracingGeometryBuildParams> BuildParams;
 
-	// TODO: Investigate removing this critical section
-	// RequestUpdateCachedRenderState(...) is always called on Render Thread
-	// Could potentially handled Register/Unregister proxies in FScene::UpdateAllPrimitiveSceneInfos which runs on Render Thread
-	// similar to how entries are added/removed to FScene.PrimitiveSceneProxies
-	FCriticalSection CachedRayTracingStateProxiesCS;
 	TMap<const UStaticMesh*, TSet<FPrimitiveSceneProxy*>> CachedRayTracingStateProxiesMap;
 };
 
