Prevent ANavSystemConfigOverride from using the world if it's in the process of being destroyed in PostUnregisterAllComponents (else it will try to recreated the original navsystem in the world currently being destroyed).

diff --git a/Engine/Source/Runtime/Engine/Classes/Engine/World.h b/Engine/Source/Runtime/Engine/Classes/Engine/World.h
--- a/Engine/Source/Runtime/Engine/Classes/Engine/World.h
+++ b/Engine/Source/Runtime/Engine/Classes/Engine/World.h
@@ -1111,6 +1111,9 @@ public:
 	/** Whether world object has been initialized via Init and has not yet had CleanupWorld called								*/
 	uint8 bIsWorldInitialized:1;
 
+	/** Indicates the the destruction process has begun for this world */
+	uint8 bIsBeingDestroyed:1;
+	
 	/** Is level streaming currently frozen?																					*/
 	uint8 bIsLevelStreamingFrozen:1;
 
diff --git a/Engine/Source/Runtime/Engine/Private/World.cpp b/Engine/Source/Runtime/Engine/Private/World.cpp
--- a/Engine/Source/Runtime/Engine/Private/World.cpp
+++ b/Engine/Source/Runtime/Engine/Private/World.cpp
@@ -597,6 +597,7 @@ UWorld::UWorld( const FObjectInitializer& ObjectInitializer )
 , FeatureLevel(GMaxRHIFeatureLevel)
 , bIsBuilt(false)
 , bIsWorldInitialized(false)
+, bIsBeingDestroyed(false)
 #if WITH_EDITOR
 , bDebugFrameStepExecutedThisFrame(false)
 , bToggledBetweenPIEandSIEThisFrame(false)
@@ -2384,6 +2385,8 @@ void UWorld::InitializeNewWorld(const InitializationValues IVS, bool bInSkipInit
 
 void UWorld::DestroyWorld( bool bInformEngineOfWorld, UWorld* NewWorld )
 {
+	bIsBeingDestroyed = true;
+	
 	// Clean up existing world and remove it from root set so it can be garbage collected.
 	bIsLevelStreamingFrozen = false;
 	SetShouldForceUnloadStreamingLevels(true);
diff --git a/Engine/Source/Runtime/NavigationSystem/Private/NavSystemConfigOverride.cpp b/Engine/Source/Runtime/NavigationSystem/Private/NavSystemConfigOverride.cpp
--- a/Engine/Source/Runtime/NavigationSystem/Private/NavSystemConfigOverride.cpp
+++ b/Engine/Source/Runtime/NavigationSystem/Private/NavSystemConfigOverride.cpp
@@ -186,7 +186,7 @@ void ANavSystemConfigOverride::PostUnregisterAllComponents()
 
 	// ApplyConfig was performed in PostRegisterAllComponents for Editor worlds so nothing to unregister for Game worlds.
 	UWorld* World = GetWorld();
-	if (World == nullptr || World->IsGameWorld() || World->WorldType == EWorldType::Inactive)
+	if (World == nullptr || World->IsGameWorld() || World->WorldType == EWorldType::Inactive || World->bIsBeingDestroyed)
 	{
 		return;
 	}
