Revert temporary fix CL 17521239 done for CitySample (9/15/2021)
- World Partition has been loading Actors and their references properly for a while now through IWorldPartitionActorLoaderInterface::ILoaderAdapter::RefreshLoadedState and Registering components first on them then running construction scripts (ULevel::AddLoadedActors)
- Removed property 'ReplacementSceneComponent' has been showing up in GC Leaks when editing level instances causing a failed check. (no 100% repro)

#rnx

diff --git a/Engine/Source/Runtime/Engine/Classes/Components/SceneComponent.h b/Engine/Source/Runtime/Engine/Classes/Components/SceneComponent.h
--- a/Engine/Source/Runtime/Engine/Classes/Components/SceneComponent.h
+++ b/Engine/Source/Runtime/Engine/Classes/Components/SceneComponent.h
@@ -312,23 +312,6 @@ public:
 	/** Returns the current scoped movement update, or NULL if there is none. @see FScopedMovementUpdate */
 	ENGINE_API FScopedMovementUpdate* GetCurrentScopedMovement() const;
 
-#if WITH_EDITORONLY_DATA
-	/**
-	 * @todo_ow: This is needed because of order of registration of Actors
-	 * 
-	 * In World Partition Levels loaded actors are Registered in an atomic fashion meaning we register all their
-	 * components and then call RerunConstructionScripts before loading the next actor. This means that a Parent can be Reconstructed and trash its components
-	 * without notifying its attached actors.
-	 * 
-	 * In Non World Partition Levels when adding actors to world we sort all actors based on hierarchy, register all actors and then RerunConstructionScripts on all actors
-	 * which allows handling of Attachments as attached actors are already registered when the parent gets reconstructed and so attachment is preserved.
-	 * 
-	 * ReplacementSceneComponent is there as a temp solution to allow finding of the replacement component of a trashed scene component. 
-	 */
-	UPROPERTY(Transient)
-	TObjectPtr<USceneComponent> ReplacementSceneComponent;
-#endif
-
 private:
 	/** Stack of current movement scopes. */
 	TArray<FScopedMovementUpdate*> ScopedMovementStack;
@@ -935,8 +918,6 @@ public:
 	ENGINE_API virtual void Serialize(FArchive& Ar) override;
 #if WITH_EDITORONLY_DATA
 	static ENGINE_API void AddReferencedObjects(UObject* InThis, FReferenceCollector& Collector);
-
-	ENGINE_API virtual void PostLoad() override;
 #endif
 
 #if WITH_EDITOR
diff --git a/Engine/Source/Runtime/Engine/Private/ActorConstruction.cpp b/Engine/Source/Runtime/Engine/Private/ActorConstruction.cpp
--- a/Engine/Source/Runtime/Engine/Private/ActorConstruction.cpp
+++ b/Engine/Source/Runtime/Engine/Private/ActorConstruction.cpp
@@ -751,18 +751,6 @@ void AActor::RerunConstructionScripts()
 					}
 				}
 			}
-#if WITH_EDITORONLY_DATA
-			if (OwningWorld && OwningWorld->GetWorldPartition() != nullptr)
-			{
-				if (USceneComponent* OldSceneComponent = Cast<USceneComponent>(ComponentData.OldComponent))
-				{
-					if (USceneComponent* NewSceneComponent = Cast<USceneComponent>(ResolvedNewComponent))
-					{
-						OldSceneComponent->ReplacementSceneComponent = NewSceneComponent;
-					}
-				}
-			}
-#endif
 
 			OldToNewComponentMapping.Add(ComponentData.OldComponent, ResolvedNewComponent);
 		}
diff --git a/Engine/Source/Runtime/Engine/Private/Components/SceneComponent.cpp b/Engine/Source/Runtime/Engine/Private/Components/SceneComponent.cpp
--- a/Engine/Source/Runtime/Engine/Private/Components/SceneComponent.cpp
+++ b/Engine/Source/Runtime/Engine/Private/Components/SceneComponent.cpp
@@ -91,24 +91,6 @@ void USceneComponent::AddReferencedObjects(UObject* InThis, FReferenceCollector&
 	Super::AddReferencedObjects(InThis, Collector);
 }
 
-void USceneComponent::PostLoad()
-{
-	Super::PostLoad();
-
-	if (AttachParent)
-	{
-		USceneComponent* TopReplacementComponent = AttachParent->ReplacementSceneComponent;
-		while (TopReplacementComponent && TopReplacementComponent->ReplacementSceneComponent)
-		{
-			TopReplacementComponent = TopReplacementComponent->ReplacementSceneComponent;
-		}
-
-		if (TopReplacementComponent)
-		{
-			SetAttachParent(TopReplacementComponent);
-		}
-	}
-}
 #endif
 
 #if WITH_EDITOR
