[Shaders] - If shader dev mode is on assume that any shader errors from materials are due to source code (and not the material graph).

- This way shader retry on error works for material shaders (right now it only works for global shaders and the default material).

todo: Fix a bug that causes the retry on compilation to still report the errors even after they are fixed.

diff --git a/Engine/Source/Runtime/Engine/Private/Materials/MeshMaterialShader.cpp b/Engine/Source/Runtime/Engine/Private/Materials/MeshMaterialShader.cpp
--- a/Engine/Source/Runtime/Engine/Private/Materials/MeshMaterialShader.cpp
+++ b/Engine/Source/Runtime/Engine/Private/Materials/MeshMaterialShader.cpp
@@ -30,6 +30,12 @@ static void PrepareMeshMaterialShaderCompileJob(EShaderPlatform Platform,
 	NewJob->bIsDefaultMaterial = Material->IsDefaultMaterial();
 	NewJob->bIsGlobalShader = false;
 
+	static IConsoleVariable* CVarShaderDevMode = IConsoleManager::Get().FindConsoleVariable(TEXT("r.ShaderDevelopmentMode"));
+	if (CVarShaderDevMode && CVarShaderDevMode->GetInt() != 0)
+	{
+		NewJob->bErrorsAreLikelyToBeCode = true;
+	}
+
 	const FMaterialShaderParameters MaterialParameters(Material);
 
 	// apply the vertex factory changes to the compile environment
