GameLayerManager: remove the restriction on Slate.GameLayer.AllCanvasesVisible that enforced only one OnChanged delegate to be bound.
This fixes the issue with the CVar not working when Console Variables Editor plugin is enabled.

diff --git a/Engine/Source/Runtime/Engine/Private/Slate/SGameLayerManager.cpp b/Engine/Source/Runtime/Engine/Private/Slate/SGameLayerManager.cpp
--- a/Engine/Source/Runtime/Engine/Private/Slate/SGameLayerManager.cpp
+++ b/Engine/Source/Runtime/Engine/Private/Slate/SGameLayerManager.cpp
@@ -47,6 +47,7 @@ namespace CanvasVisibility
 	static bool ViewportSlotVisibilityConsoleValue = true;
 	static bool DebugCanvasVisibilityConsoleValue = true;
 	static bool PlayerCanvasVisibilityConsoleValue = true;
+	static bool CanvasVisibilityCVarOnChangedBound = false;
 
 	enum class CanvasType
 	{
@@ -197,7 +198,7 @@ void SGameLayerManager::Construct(const SGameLayerManager::FArguments& InArgs)
 
 #if UE_SLATE_WITH_GAMELAYER_CANVAS_VISIBILITY_COMMANDS
 	CanvasVisibility::GameLayerManagerInstances.Add(this);
-	if (!CVarShowAllCanvases->AsVariable()->OnChangedDelegate().IsBound())
+	if (!CanvasVisibility::CanvasVisibilityCVarOnChangedBound)
 	{
 		auto ShowHUD = [](IConsoleVariable* Variable, CanvasVisibility::CanvasType CanvasType)
 		{
@@ -235,6 +236,7 @@ void SGameLayerManager::Construct(const SGameLayerManager::FArguments& InArgs)
 		CVarShowPlayerCanvas->AsVariable()->SetOnChangedCallback(FConsoleVariableDelegate::CreateLambda(ShowHUD, CanvasVisibility::CanvasType::PlayerCanvas));
 		CVarShowDebugCanvas->AsVariable()->SetOnChangedCallback(FConsoleVariableDelegate::CreateLambda(ShowHUD, CanvasVisibility::CanvasType::DebugCanvas));
 		CVarShowViewportSlot->AsVariable()->SetOnChangedCallback(FConsoleVariableDelegate::CreateLambda(ShowHUD, CanvasVisibility::CanvasType::Viewport));
+		CanvasVisibility::CanvasVisibilityCVarOnChangedBound = true;
 	}
 
 	ShowViewportSlot(CanvasVisibility::ViewportSlotVisibilityConsoleValue);
