SkinnedAsset - Prevent copy of weights array when empty.
- The initialization of a component with an empty asset was causing a crash within the editor.

diff --git a/Engine/Source/Runtime/Engine/Private/SkinnedAsset.cpp b/Engine/Source/Runtime/Engine/Private/SkinnedAsset.cpp
--- a/Engine/Source/Runtime/Engine/Private/SkinnedAsset.cpp
+++ b/Engine/Source/Runtime/Engine/Private/SkinnedAsset.cpp
@@ -439,6 +439,11 @@ void USkinnedAsset::FillComponentSpaceTransforms(const TArray<FTransform>& InBon
 	check(GetRefSkeleton().GetNum() == OutComponentSpaceTransforms.Num());
 
 	const int32 NumBones = InBoneSpaceTransforms.Num();
+	
+	if (!NumBones)
+	{
+		return;
+	}
 
 #if DO_GUARD_SLOW
 	/** Keep track of which bones have been processed for fast look up */
