[FActorInstanceHandle] IsValid and equality operator now resolve the handle when called from the game thread.
Also make sure to resolve the handle before serializing it.

diff --git a/Engine/Source/Runtime/Engine/Private/Engine/ActorInstanceHandle.cpp b/Engine/Source/Runtime/Engine/Private/Engine/ActorInstanceHandle.cpp
--- a/Engine/Source/Runtime/Engine/Private/Engine/ActorInstanceHandle.cpp
+++ b/Engine/Source/Runtime/Engine/Private/Engine/ActorInstanceHandle.cpp
@@ -165,10 +165,17 @@ void FActorInstanceHandle::ResolveHandle() const
 
 bool FActorInstanceHandle::IsValid() const
 {
-	// A handle properly setup from another thread that needs resolving is considered valid.
 	if (ResolutionStatus == EResolutionStatus::NeedsResolving)
 	{
-		return true;
+		if (IsInGameThread() || IsInParallelGameThread())
+		{
+			ResolveHandle();
+		}
+		else
+		{
+			// A handle properly setup from another thread that needs resolving is considered valid.
+			return !ReferenceObject.IsExplicitlyNull();
+		}
 	}
 
 	return (ManagerInterface.IsValid() && InstanceIndex != INDEX_NONE) || IsActorValid();
@@ -349,16 +356,30 @@ FActorInstanceHandle& FActorInstanceHandle::operator=(AActor* OtherActor)
 
 bool FActorInstanceHandle::operator==(const FActorInstanceHandle& Other) const
 {
-	// Handles that needs resolving can be compared using the reference object and index.
-	if (ResolutionStatus == EResolutionStatus::NeedsResolving
+	if (IsInGameThread() || IsInParallelGameThread())
+	{
+		// Both handles need to be resolved to perform a valid comparison but only on game thread.
+		if (ResolutionStatus == EResolutionStatus::NeedsResolving)
+		{
+			ResolveHandle();
+		}
+
+		if (Other.ResolutionStatus == EResolutionStatus::NeedsResolving)
+		{
+			Other.ResolveHandle();
+		}
+	}
+	else if (ResolutionStatus == EResolutionStatus::NeedsResolving
 		&& Other.ResolutionStatus == EResolutionStatus::NeedsResolving)
 	{
+		// Handles that needs resolving can be compared using the reference object and index.
 		return (ReferenceObject.HasSameIndexAndSerialNumber(Other.ReferenceObject)) && (InstanceIndex == Other.InstanceIndex);
 	}
-
-	// Both handles need to be resolved to perform a valid comparison.
-	ResolveHandle();
-	Other.ResolveHandle();
+	else
+	{
+		ensureMsgf(false, TEXT("Comparing resolved and non-resolved handles is not supported."));
+		return false;
+	}
 
 	// try to compare managers and indices first if we have them
 	if (ManagerInterface.IsValid() && Other.ManagerInterface.IsValid() && InstanceIndex != INDEX_NONE && Other.InstanceIndex != INDEX_NONE)
@@ -432,7 +453,11 @@ FArchive& operator<<(FArchive& Ar, FActorInstanceHandle& Handle)
 		Ar << Handle.InstanceIndex;
 		return Ar;
 	}
-	
+
+	if (Ar.IsSaving())
+	{
+		Handle.ResolveHandle();
+	}
 	Ar << Handle.ReferenceObject;
 	Ar << Handle.InstanceIndex;
 
