Add a bIsBeingCleanedUp flag to world

diff --git a/Engine/Source/Runtime/Engine/Classes/Engine/World.h b/Engine/Source/Runtime/Engine/Classes/Engine/World.h
--- a/Engine/Source/Runtime/Engine/Classes/Engine/World.h
+++ b/Engine/Source/Runtime/Engine/Classes/Engine/World.h
@@ -1110,9 +1110,6 @@ public:
 
 	/** Whether world object has been initialized via Init and has not yet had CleanupWorld called								*/
 	uint8 bIsWorldInitialized:1;
-
-	/** Indicates the the destruction process has begun for this world */
-	uint8 bIsBeingDestroyed:1;
 	
 	/** Is level streaming currently frozen?																					*/
 	uint8 bIsLevelStreamingFrozen:1;
@@ -1244,6 +1241,9 @@ private:
 	/** Whether InitWorld was ever called on this world since its creation. Not cleared to false during CleanupWorld			*/
 	uint8 bHasEverBeenInitialized: 1;
 
+	/** Indicates that the world is in the process of being cleaned up */
+	bool bIsBeingCleanedUp;
+	
 	/** Whether the world is currently in a BlockTillLevelStreamingCompleted() call */
 	uint32 IsInBlockTillLevelStreamingCompleted;
 
@@ -3009,6 +3009,9 @@ public:
 	 */
 	void DestroyWorld( bool bInformEngineOfWorld, UWorld* NewWorld = nullptr );
 
+	/** Returns true if the world is in the process of being cleaned up. */
+	bool IsBeingCleanedUp() const { return bIsBeingCleanedUp; }
+
 	/** 
 	 * Marks this world and all objects within as pending kill
 	 */
diff --git a/Engine/Source/Runtime/Engine/Private/World.cpp b/Engine/Source/Runtime/Engine/Private/World.cpp
--- a/Engine/Source/Runtime/Engine/Private/World.cpp
+++ b/Engine/Source/Runtime/Engine/Private/World.cpp
@@ -597,7 +597,6 @@ UWorld::UWorld( const FObjectInitializer& ObjectInitializer )
 , FeatureLevel(GMaxRHIFeatureLevel)
 , bIsBuilt(false)
 , bIsWorldInitialized(false)
-, bIsBeingDestroyed(false)
 #if WITH_EDITOR
 , bDebugFrameStepExecutedThisFrame(false)
 , bToggledBetweenPIEandSIEThisFrame(false)
@@ -605,6 +604,7 @@ UWorld::UWorld( const FObjectInitializer& ObjectInitializer )
 #endif
 , bShouldTick(true)
 , bHasEverBeenInitialized(false)
+, bIsBeingCleanedUp(false)
 , ActiveLevelCollectionIndex(INDEX_NONE)
 , AudioDeviceHandle()
 #if WITH_EDITOR
@@ -2385,8 +2385,6 @@ void UWorld::InitializeNewWorld(const InitializationValues IVS, bool bInSkipInit
 
 void UWorld::DestroyWorld( bool bInformEngineOfWorld, UWorld* NewWorld )
 {
-	bIsBeingDestroyed = true;
-	
 	// Clean up existing world and remove it from root set so it can be garbage collected.
 	bIsLevelStreamingFrozen = false;
 	SetShouldForceUnloadStreamingLevels(true);
@@ -5366,6 +5364,8 @@ void UWorld::CleanupWorld(bool bSessionEnded, bool bCleanupResources, UWorld* Ne
 
 void UWorld::CleanupWorldInternal(bool bSessionEnded, bool bCleanupResources, bool bWorldChanged)
 {
+	TGuardValue<bool> IsBeingCleanedUp(bIsBeingCleanedUp, true);
+	
 	if(CleanupWorldTag == CleanupWorldGlobalTag)
 	{
 		return;
diff --git a/Engine/Source/Runtime/NavigationSystem/Private/NavSystemConfigOverride.cpp b/Engine/Source/Runtime/NavigationSystem/Private/NavSystemConfigOverride.cpp
--- a/Engine/Source/Runtime/NavigationSystem/Private/NavSystemConfigOverride.cpp
+++ b/Engine/Source/Runtime/NavigationSystem/Private/NavSystemConfigOverride.cpp
@@ -186,7 +186,7 @@ void ANavSystemConfigOverride::PostUnregisterAllComponents()
 
 	// ApplyConfig was performed in PostRegisterAllComponents for Editor worlds so nothing to unregister for Game worlds.
 	UWorld* World = GetWorld();
-	if (World == nullptr || World->IsGameWorld() || World->WorldType == EWorldType::Inactive || World->bIsBeingDestroyed)
+	if (World == nullptr || World->IsGameWorld() || World->WorldType == EWorldType::Inactive || !World->IsInitialized() || World->IsBeingCleanedUp())
 	{
 		return;
 	}
