Texture prevent multi-threaded TextureSource access causing crash in lockmip
change BulkData RWLock to CritSec because it needs reentrancy
use it to protect the LockState variables as well

diff --git a/Engine/Source/Runtime/Engine/Classes/Engine/Texture.h b/Engine/Source/Runtime/Engine/Classes/Engine/Texture.h
--- a/Engine/Source/Runtime/Engine/Classes/Engine/Texture.h
+++ b/Engine/Source/Runtime/Engine/Classes/Engine/Texture.h
@@ -572,8 +572,14 @@ private:
 	friend class UTextureCubeArray;
 
 #if WITH_EDITOR
-	/** Protects simultaneous access to BulkData */
-	TDontCopy<FRWLock> BulkDataLock;
+	/** BulkDataLock protects simultaneous access to BulkData ;
+	BulkDataLock protects the BulkData, also LockState, NumLockedMips, LockedMipData.
+	(eg. it protects the functions LockMip, GetMipData, Compress, Decompress, etc.)
+	It does NOT protect scalar fields (eg. "SizeX").
+	It is intended to allow multiple read-like threads to read from the same TextureSource.
+	If you need to modify a TextureSource you must first manually flush all other threads that could be using it.
+	eg. using BlockOnAnyAsyncBuild. */
+	TDontCopy<FCriticalSection> BulkDataLock;
 	/** Owner for associating BulkData with a package */
 	UTexture* Owner;
 	/** TextureClass == Owner->GetTextureClass(); a copy is kept here for torn-off **/
diff --git a/Engine/Source/Runtime/Engine/Private/Texture.cpp b/Engine/Source/Runtime/Engine/Private/Texture.cpp
--- a/Engine/Source/Runtime/Engine/Private/Texture.cpp
+++ b/Engine/Source/Runtime/Engine/Private/Texture.cpp
@@ -494,6 +494,9 @@ bool UTexture::Modify(bool bAlwaysMarkDirty)
 	// Before applying any modification to the texture
 	// make sure no compilation is still ongoing.
 	BlockOnAnyAsyncBuild();
+	
+	// @@ if other textures are using me as a Composite , also block on THEM
+	//	 their build action may be reading from me on other threads, must block them before I am modified
 
 	return Super::Modify(bAlwaysMarkDirty);
 }
@@ -991,7 +994,7 @@ void UTexture::Serialize(FArchive& Ar)
 	if (!StripFlags.IsEditorDataStripped())
 	{
 #if WITH_EDITOR
-		FWriteScopeLock BulkDataExclusiveScope(Source.BulkDataLock.Get());
+		FScopeLock BulkDataExclusiveScope(&Source.BulkDataLock.Get());
 #endif
 
 		if (Ar.IsLoading() && Ar.CustomVer(FUE5MainStreamObjectVersion::GUID) < FUE5MainStreamObjectVersion::VirtualizedBulkDataHaveUniqueGuids)
@@ -2261,6 +2264,10 @@ void FTextureSource::InitWithCompressedSourceData(
 
 FTextureSource FTextureSource::CopyTornOff() const
 {
+#if WITH_EDITOR
+	FScopeLock BulkDataExclusiveScope(&const_cast<FCriticalSection &>(BulkDataLock.Get()));
+#endif
+
 	FTextureSource Result;
 	// Set the Torn off flag on Result.BulkData so that the copy constructor below will not set it
 	Result.BulkData.TearOff();
@@ -2277,12 +2284,12 @@ FTextureSource FTextureSource::CopyTornOff() const
 
 void FTextureSource::Compress()
 {
-	CheckTextureIsUnlocked(TEXT("Compress"));
-
 #if WITH_EDITOR
-	FWriteScopeLock BulkDataExclusiveScope(BulkDataLock.Get());
+	FScopeLock BulkDataExclusiveScope(&BulkDataLock.Get());
 #endif
 
+	CheckTextureIsUnlocked(TEXT("Compress"));
+
 	// if bUseOodleOnPNGz0 , do PNG filters but then use Oodle instead of zlib back-end LZ
 	//	should be faster to load and also smaller files (than traditional PNG+zlib)
 	bool bUseOodleOnPNGz0 = true;
@@ -2362,7 +2369,11 @@ void FTextureSource::Compress()
 FSharedBuffer FTextureSource::Decompress(IImageWrapperModule* ) const
 {
 	TRACE_CPUPROFILER_EVENT_SCOPE(FTextureSource::Decompress);
-	
+
+#if WITH_EDITOR
+	FScopeLock BulkDataExclusiveScope(&const_cast<FCriticalSection &>(BulkDataLock.Get()));
+#endif
+
 	// ImageWrapperModule argument ignored, not drilled through DecompressImage
 
 	int64 ExpectedTotalSize = CalcTotalSize();
@@ -2390,6 +2401,8 @@ FSharedBuffer FTextureSource::Decompress(IImageWrapperModule* ) const
 
 void FTextureSource::CheckTextureIsUnlocked(const TCHAR* DebugMessage)
 {
+	//note: BulkDataLock should be held before calling this
+
 	// Asserts if a FTextureSource is locked for read or write access, along with additional debug data
 	checkf(LockState == ELockState::None, TEXT("%s cannot be called when FTextureSource is locked for %s access [%s]"), 
 		DebugMessage,
@@ -2484,6 +2497,15 @@ uint8* FTextureSource::LockMip(int32 BlockIndex, int32 LayerIndex, int32 MipInde
 FMutableMemoryView FTextureSource::LockMipInternal(int32 BlockIndex, int32 LayerIndex, int32 MipIndex, ELockState RequestedLockState)
 {
 	TRACE_CPUPROFILER_EVENT_SCOPE(FTextureSource::LockMip);
+	
+#if WITH_EDITOR
+	FScopeLock BulkDataExclusiveScope(&BulkDataLock.Get());
+	// BulkDataLock protects LockState, NumLockedMips
+	// note that it does NOT actually protect the bits of the texture data
+	// that is, it's released when we leave this function, so it is not held during LockMip to UnlockMip
+	// that means multiple threads can have locked mips and act on them at the same time
+	// it only protects the lock accounting variables
+#endif
 
 	checkf(RequestedLockState != ELockState::None, TEXT("Cannot call FTextureSource::LockMipInternal with a RequestedLockState of type ELockState::None"));
 
@@ -2543,6 +2565,10 @@ FMutableMemoryView FTextureSource::LockMipInternal(int32 BlockIndex, int32 Layer
 void FTextureSource::UnlockMip(int32 BlockIndex, int32 LayerIndex, int32 MipIndex)
 {
 	TRACE_CPUPROFILER_EVENT_SCOPE(FTextureSource::UnlockMip);
+	
+#if WITH_EDITOR
+	FScopeLock BulkDataExclusiveScope(&BulkDataLock.Get());
+#endif
 
 	check(BlockIndex < GetNumBlocks());
 	check(LayerIndex < NumLayers);
@@ -2611,6 +2637,10 @@ bool FTextureSource::GetMipImage(FImage & OutImage, int32 BlockIndex, int32 Laye
 bool FTextureSource::GetMipData(TArray64<uint8>& OutMipData, int32 BlockIndex, int32 LayerIndex, int32 MipIndex, IImageWrapperModule* )
 {
 	TRACE_CPUPROFILER_EVENT_SCOPE(FTextureSource::GetMipData (TArray64));
+	
+#if WITH_EDITOR
+	FScopeLock BulkDataExclusiveScope(&BulkDataLock.Get());
+#endif
 
 	CheckTextureIsUnlocked(TEXT("GetMipData (TArray64)"));
 
@@ -2618,10 +2648,6 @@ bool FTextureSource::GetMipData(TArray64<uint8>& OutMipData, int32 BlockIndex, i
 
 	if (IsValid() && BlockIndex < GetNumBlocks() && LayerIndex < NumLayers && MipIndex < NumMips && HasPayloadData())
 	{
-#if WITH_EDITOR
-		FWriteScopeLock BulkDataExclusiveScope(BulkDataLock.Get());
-#endif
-
 		checkf(NumLockedMips == 0, TEXT("Attempting to access a locked FTextureSource"));
 		// LockedMipData should only be allocated if NumLockedMips > 0 so the following assert should have been caught
 		// by the one above. If it fires then it indicates that there is a lock/unlock mismatch as well as invalid access!
@@ -2656,17 +2682,17 @@ FTextureSource::FMipData FTextureSource::GetMipData(IImageWrapperModule* )
 {
 	TRACE_CPUPROFILER_EVENT_SCOPE(FTextureSource::GetMipData (FMipData));
 
-	CheckTextureIsUnlocked(TEXT("GetMipData (FMipData)"));
-	
-	check(LockedMipData.IsNull());
-	check(NumLockedMips == 0);
-
 #if WITH_EDITOR
 	// We can end up waiting here a lot as the bulk data gets serialized for entry in to the transaction buffer.
 	TRACE_CPUPROFILER_EVENT_SCOPE(FTextureSource::GetMipData (StartLock) );
-	FReadScopeLock _(BulkDataLock.Get());
+	FScopeLock _(&BulkDataLock.Get());
 #endif //WITH_EDITOR
 
+	CheckTextureIsUnlocked(TEXT("GetMipData (FMipData)"));
+	
+	check(LockedMipData.IsNull());
+	check(NumLockedMips == 0);
+
 	FSharedBuffer DecompressedData = Decompress();
 	return FMipData(*this, DecompressedData);
 }
@@ -2863,6 +2889,8 @@ FString FTextureSource::GetSourceCompressionAsString() const
 
 FSharedBuffer FTextureSource::TryDecompressData() const
 {
+	// BulkDataLock should be held before calling this!
+
 	if (NumLayers == 1 && NumSlices == 1 && Blocks.Num() == 0)
 	{
 		FSharedBuffer Payload = BulkData.GetPayload().Get();
@@ -2928,6 +2956,10 @@ FSharedBuffer FTextureSource::TryDecompressData() const
 
 void FTextureSource::ExportCustomProperties(FOutputDevice& Out, uint32 Indent)
 {
+#if WITH_EDITOR
+	FScopeLock BulkDataExclusiveScope(&const_cast<FCriticalSection &>(BulkDataLock.Get()));
+#endif
+
 	CheckTextureIsUnlocked(TEXT("ExportCustomProperties"));
 
 	FSharedBuffer Payload = BulkData.GetPayload().Get();
@@ -2945,6 +2977,10 @@ void FTextureSource::ExportCustomProperties(FOutputDevice& Out, uint32 Indent)
 
 void FTextureSource::ImportCustomProperties(const TCHAR* SourceText, FFeedbackContext* Warn)
 {
+#if WITH_EDITOR
+	FScopeLock BulkDataExclusiveScope(&const_cast<FCriticalSection &>(BulkDataLock.Get()));
+#endif
+
 	CheckTextureIsUnlocked(TEXT("ImportCustomProperties"));
 
 	if (FParse::Command(&SourceText, TEXT("TextureSourceData")))
@@ -3050,6 +3086,10 @@ void FTextureSource::ForceGenerateGuid()
 
 void FTextureSource::ReleaseSourceMemory()
 {
+#if WITH_EDITOR
+	FScopeLock BulkDataExclusiveScope(&BulkDataLock.Get());
+#endif
+
 	check( LockState == ELockState::None && NumLockedMips == 0 );
 
 	bHasHadBulkDataCleared = true;
@@ -3058,6 +3098,10 @@ void FTextureSource::ReleaseSourceMemory()
 
 void FTextureSource::RemoveSourceData()
 {
+#if WITH_EDITOR
+	FScopeLock BulkDataExclusiveScope(&BulkDataLock.Get());
+#endif
+
 	check( LockState == ELockState::None && NumLockedMips == 0 );
 
 	SizeX = 0;
@@ -3169,6 +3213,10 @@ int64 FTextureSource::CalcMipOffset(int32 BlockIndex, int32 LayerIndex, int32 Of
 //	calling it multiple times does not re-hash the data; it's harmless
 void FTextureSource::UseHashAsGuid()
 {
+#if WITH_EDITOR
+	FScopeLock BulkDataExclusiveScope(&const_cast<FCriticalSection &>(BulkDataLock.Get()));
+#endif
+
 	if (HasPayloadData())
 	{
 		CheckTextureIsUnlocked(TEXT("UseHashAsGuid"));
@@ -3241,12 +3289,12 @@ FGuid FTextureSource::GetId() const
 
 void FTextureSource::OperateOnLoadedBulkData(TFunctionRef<void(const FSharedBuffer& BulkDataBuffer)> Operation)
 {
-	checkf(LockState == ELockState::None, TEXT("OperateOnLoadedBulkData shouldn't be called in-between LockMip/UnlockMip"));
-
 #if WITH_EDITOR
-	FReadScopeLock BulkDataExclusiveScope(BulkDataLock.Get());
+	FScopeLock BulkDataExclusiveScope(&BulkDataLock.Get());
 #endif
 
+	checkf(LockState == ELockState::None, TEXT("OperateOnLoadedBulkData shouldn't be called in-between LockMip/UnlockMip"));
+
 	FSharedBuffer Payload = BulkData.GetPayload().Get();
 	Operation(Payload);
 }
