Texture fix race when changing a texture that is used as a composite source by other textures, they must not be actively building on other threads
also fix property change failed to notify composite change which caused a check

diff --git a/Engine/Source/Runtime/Engine/Private/Texture.cpp b/Engine/Source/Runtime/Engine/Private/Texture.cpp
--- a/Engine/Source/Runtime/Engine/Private/Texture.cpp
+++ b/Engine/Source/Runtime/Engine/Private/Texture.cpp
@@ -494,9 +494,19 @@ bool UTexture::Modify(bool bAlwaysMarkDirty)
 	// Before applying any modification to the texture
 	// make sure no compilation is still ongoing.
 	BlockOnAnyAsyncBuild();
-	
-	// @@ if other textures are using me as a Composite , also block on THEM
-	//	 their build action may be reading from me on other threads, must block them before I am modified
+		
+#if WITH_EDITORONLY_DATA
+	{
+		// if other textures are using me as a Composite , also block on THEM
+		//	 their build action may be reading from me on other threads, must block them before I am modified
+
+		FObjectCacheContextScope ObjectCache;
+		for (UTexture* Texture : ObjectCache.GetContext().GetTexturesAffectedByTexture(this))
+		{
+			Texture->BlockOnAnyAsyncBuild();
+		}
+	}
+#endif
 
 	return Super::Modify(bAlwaysMarkDirty);
 }
@@ -787,6 +797,9 @@ void UTexture::PostEditChangeProperty(FPropertyChangedEvent& PropertyChangedEven
 		static const FName SourceColorSpaceName = GET_MEMBER_NAME_CHECKED(FTextureSourceColorSettings, ColorSpace);
 		static const FName CompressionQualityName = GET_MEMBER_NAME_CHECKED(UTexture, CompressionQuality);
 		static const FName OodleTextureSdkVersionName = GET_MEMBER_NAME_CHECKED(UTexture, OodleTextureSdkVersion);
+PRAGMA_DISABLE_DEPRECATION_WARNINGS
+		static const FName CompositeTextureName = GET_MEMBER_NAME_CHECKED(UTexture, CompositeTexture);
+PRAGMA_ENABLE_DEPRECATION_WARNINGS
 #endif //WITH_EDITORONLY_DATA
 
 		const FName PropertyName = PropertyThatChanged->GetFName();
@@ -850,6 +863,10 @@ void UTexture::PostEditChangeProperty(FPropertyChangedEvent& PropertyChangedEven
 				OodleTextureSdkVersion = CachedGetLatestOodleSdkVersion();
 			}
 		}
+		else if ( PropertyName == CompositeTextureName )
+		{
+			NotifyIfCompositeTextureChanged();
+		}
 #endif //WITH_EDITORONLY_DATA
 	}
 
