Perform proper null-checking of the scene and view pointers when retrieving the shader platform from them when binding spline mesh vertex factory shader parameters.

diff --git a/Engine/Source/Runtime/Engine/Private/Components/SplineMeshComponent.cpp b/Engine/Source/Runtime/Engine/Private/Components/SplineMeshComponent.cpp
--- a/Engine/Source/Runtime/Engine/Private/Components/SplineMeshComponent.cpp
+++ b/Engine/Source/Runtime/Engine/Private/Components/SplineMeshComponent.cpp
@@ -197,7 +197,10 @@ void FSplineMeshVertexFactoryShaderParameters::GetElementShaderBindings(
 	FVertexInputStreamArray& VertexStreams
 ) const
 {
-	const bool bUseGPUScene = UseGPUScene(Scene->GetShaderPlatform(), FeatureLevel);
+	const EShaderPlatform ShaderPlatform = Scene ?
+		Scene->GetShaderPlatform() :
+		(View ? View->GetShaderPlatform() : GetFeatureLevelShaderPlatform(FeatureLevel));
+	const bool bUseGPUScene = UseGPUScene(ShaderPlatform, FeatureLevel);
 	const auto* LocalVertexFactory = static_cast<const FLocalVertexFactory*>(VertexFactory);
 	
 	if (BatchElement.bUserDataIsColorVertexBuffer)
