Rad Audio -> RAD Audio

diff --git a/Engine/Source/Developer/AudioFormatRad/Private/AudioFormatRad.cpp b/Engine/Source/Developer/AudioFormatRad/Private/AudioFormatRad.cpp
--- a/Engine/Source/Developer/AudioFormatRad/Private/AudioFormatRad.cpp
+++ b/Engine/Source/Developer/AudioFormatRad/Private/AudioFormatRad.cpp
@@ -40,7 +40,7 @@ class FAudioFormatRad : public IAudioFormat
 {
 	enum
 	{
-		/** Version for Rad Audio format, this becomes part of the DDC key. */
+		/** Version for RAD Audio format, this becomes part of the DDC key. */
 		UE_AUDIO_RAD_VER = 2,
 	};
 
@@ -95,7 +95,7 @@ public:
 
 		if (RadCompressError)
 		{
-			UE_LOG(LogAudioFormatRad, Warning, TEXT("Failed to encode Rad Audio: %hs"),  RadAErrorString(RadCompressError));
+			UE_LOG(LogAudioFormatRad, Warning, TEXT("Failed to encode RAD Audio: %hs"),  RadAErrorString(RadCompressError));
 			if (RadCompressError == RADA_COMPRESS_ERROR_RATE)
 			{
 				UE_LOG(LogAudioFormatRad, Warning, TEXT("Sample rate provided: %u"), InQualityInfo.SampleRate);
@@ -120,7 +120,7 @@ public:
 		check(InFormat == NAME_RADA);
 
 		//
-		// CookSurround passes us a bunch of mono buffers, but Rad audio wants a standard 
+		// CookSurround passes us a bunch of mono buffers, but RAD audio wants a standard 
 		// interleaved buffer
 		//
 		TArray<uint8> InterleavedSrcBuffers;
diff --git a/Engine/Source/Runtime/Engine/Classes/Sound/AudioSettings.h b/Engine/Source/Runtime/Engine/Classes/Sound/AudioSettings.h
--- a/Engine/Source/Runtime/Engine/Classes/Sound/AudioSettings.h
+++ b/Engine/Source/Runtime/Engine/Classes/Sound/AudioSettings.h
@@ -69,7 +69,7 @@ enum class EDefaultAudioCompressionType : uint8
 	PlatformSpecific,
 
 	// As BinkAudio, except better quality. Comparable CPU usage on newer platforms, higher CPU on older platforms.
-	RadAudio
+	RADAudio UMETA(DisplayName = "RAD Audio")
 };
 
 
diff --git a/Engine/Source/Runtime/Engine/Classes/Sound/SoundWave.h b/Engine/Source/Runtime/Engine/Classes/Sound/SoundWave.h
--- a/Engine/Source/Runtime/Engine/Classes/Sound/SoundWave.h
+++ b/Engine/Source/Runtime/Engine/Classes/Sound/SoundWave.h
@@ -339,7 +339,7 @@ enum class ESoundAssetCompressionType : uint8
 	ProjectDefined,
 
 	// As BinkAudio, except better quality. Comparable CPU usage.
-	RadAudio,
+	RADAudio UMETA(DisplayName = "RAD Audio"),
 };
 
 
@@ -349,7 +349,7 @@ namespace Audio
 	{
 		switch (InDecoderType)
 		{
-		case ESoundAssetCompressionType::RadAudio:				return NAME_RADA;
+		case ESoundAssetCompressionType::RADAudio:				return NAME_RADA;
 		case ESoundAssetCompressionType::BinkAudio:				return NAME_BINKA;
 		case ESoundAssetCompressionType::ADPCM:					return NAME_ADPCM;
 		case ESoundAssetCompressionType::PCM:					return NAME_PCM;
@@ -366,7 +366,7 @@ namespace Audio
 	{
 		switch (InDefaultCompressionType)
 		{
-			case EDefaultAudioCompressionType::RadAudio:			return ESoundAssetCompressionType::RadAudio;
+			case EDefaultAudioCompressionType::RADAudio:			return ESoundAssetCompressionType::RADAudio;
 			case EDefaultAudioCompressionType::BinkAudio:			return ESoundAssetCompressionType::BinkAudio;
 			case EDefaultAudioCompressionType::ADPCM:				return ESoundAssetCompressionType::ADPCM;
 			case EDefaultAudioCompressionType::PCM:					return ESoundAssetCompressionType::PCM;
diff --git a/Engine/Source/Runtime/RadAudioCodec/Module/Private/RadAudioInfo.cpp b/Engine/Source/Runtime/RadAudioCodec/Module/Private/RadAudioInfo.cpp
--- a/Engine/Source/Runtime/RadAudioCodec/Module/Private/RadAudioInfo.cpp
+++ b/Engine/Source/Runtime/RadAudioCodec/Module/Private/RadAudioInfo.cpp
@@ -9,7 +9,7 @@
 #include "rada_decode.h"
 
 #if !defined(PLATFORM_LITTLE_ENDIAN) || !PLATFORM_LITTLE_ENDIAN
-#error "Rad Audio hasn't been updated for big endian."
+#error "RAD Audio hasn't been updated for big endian."
 #endif
 
 DEFINE_LOG_CATEGORY_STATIC(LogRadAudioDecoder, Log, All);
diff --git a/Engine/Source/Runtime/RadAudioCodec/Module/Public/RadAudioInfo.h b/Engine/Source/Runtime/RadAudioCodec/Module/Public/RadAudioInfo.h
--- a/Engine/Source/Runtime/RadAudioCodec/Module/Public/RadAudioInfo.h
+++ b/Engine/Source/Runtime/RadAudioCodec/Module/Public/RadAudioInfo.h
@@ -6,7 +6,7 @@
 #include "AudioDecompress.h"
 
 /**
-*	Decoder thunk to the Rad Audio libraries. Also manages file parsing
+*	Decoder thunk to the RAD Audio libraries. Also manages file parsing
 *	for the cooked data.
 */
 class  FRadAudioInfo : public IStreamedCompressedInfo
