Skeletal Mesh: Include skin profiles when getting a mesh description from a LOD model

#rnx

diff --git a/Engine/Source/Runtime/Engine/Private/Rendering/SkeletalMeshLODImporterData.cpp b/Engine/Source/Runtime/Engine/Private/Rendering/SkeletalMeshLODImporterData.cpp
--- a/Engine/Source/Runtime/Engine/Private/Rendering/SkeletalMeshLODImporterData.cpp
+++ b/Engine/Source/Runtime/Engine/Private/Rendering/SkeletalMeshLODImporterData.cpp
@@ -1430,45 +1430,85 @@ void FSkeletalMeshImportData::ComputeSmoothGroupFromNormals()
 	}
 }
 
-void FSkeletalMeshImportData::SetMorphTargets(
-	const TArray<TObjectPtr<UMorphTarget>>& InMorphTargets,
-	int32 InLODIndex,
+void FSkeletalMeshImportData::AddMorphTarget(
+	FName InMorphTargetName,
+	const FMorphTargetLODModel& InMorphTargetModel,
 	const TArray<uint32>& InVertexMap
 	)
 {
-	MorphTargets.Reset();
-	MorphTargetNames.Reset();
-	MorphTargetModifiedPoints.Reset();
+	MorphTargetNames.Add(InMorphTargetName.ToString());
 
-	for (UMorphTarget* MorphTargetSrc: InMorphTargets)
+	FSkeletalMeshImportData& MorphTargetDst = MorphTargets.AddDefaulted_GetRef();
+	TSet<uint32>& ModifiedPoints = MorphTargetModifiedPoints.AddDefaulted_GetRef();
+
+	for (const FMorphTargetDelta& Delta: InMorphTargetModel.Vertices)
 	{
-		if (!MorphTargetSrc->HasDataForLOD(InLODIndex))
+		if (InVertexMap.IsValidIndex(Delta.SourceIdx))
 		{
-			continue;
+			const int32 MappedIndex = InVertexMap[Delta.SourceIdx];
+			if (Points.IsValidIndex(MappedIndex))
+			{
+				ModifiedPoints.Add(MappedIndex);
+				MorphTargetDst.Points.Add(Delta.PositionDelta + Points[MappedIndex]);
+			}
 		}
+	}
+}
 
-		const FMorphTargetLODModel& MorphTargetModel = MorphTargetSrc->GetMorphLODModels()[InLODIndex];
-
-		MorphTargetNames.Add(MorphTargetSrc->GetName());
-
-		FSkeletalMeshImportData& MorphTargetDst = MorphTargets.AddDefaulted_GetRef();
-		TSet<uint32>& ModifiedPoints = MorphTargetModifiedPoints.AddDefaulted_GetRef();
+void FSkeletalMeshImportData::AddSkinWeightProfile(
+	FName InProfileName,
+	const FImportedSkinWeightProfileData& InProfileData,
+	const TArray<int32>& InVertexMap,
+	const TArray<FBoneIndexType>& InBoneIndexMap
+	)
+{
+	AlternateInfluenceProfileNames.Add(InProfileName.ToString());
+	FSkeletalMeshImportData& TargetProfileData = AlternateInfluences.AddDefaulted_GetRef();
+
+	TargetProfileData.Points = Points;
+	TargetProfileData.Wedges = Wedges;
+	TargetProfileData.Faces = Faces;
+	TargetProfileData.RefBonesBinary = RefBonesBinary;
+	TargetProfileData.Influences.Reserve(InProfileData.SourceModelInfluences.Num());
+
+	// We can't use the SourceModelInfluences since it may come from a alt skin mesh that didn't have the exact
+	// same point location and was therefore projected onto base base mesh. However, those points no longer exist
+	// and so we can't rely on them. Instead we reverse map back the LOD model influences onto the import model
+	// using the vertex map. We need to keep track of which target point we've set since points may have gotten split,
+	// and the vertex map can point to the same point multiple times, and we don't want to end up with double weights.
+	TSet<int32> ProcessedVertex;
+	for (int32 VertexIndex = 0; VertexIndex < InProfileData.SkinWeights.Num(); VertexIndex++)
+	{
+		const FRawSkinWeight& SourceWeights = InProfileData.SkinWeights[VertexIndex];
+		const int32 TargetVertexIndex = InVertexMap[VertexIndex];
+		if (ProcessedVertex.Contains(TargetVertexIndex))
+		{
+			continue;
+		}
 
-		for (const FMorphTargetDelta& Delta: MorphTargetModel.Vertices)
+		ProcessedVertex.Add(TargetVertexIndex);
+		
+		for (int32 InfluenceIndex = 0; InfluenceIndex < MAX_TOTAL_INFLUENCES; InfluenceIndex++)
 		{
-			if (InVertexMap.IsValidIndex(Delta.SourceIdx))
+			if (SourceWeights.InfluenceWeights[InfluenceIndex] == 0)
 			{
-				const int32 MappedIndex = InVertexMap[Delta.SourceIdx];
-				if (Points.IsValidIndex(MappedIndex))
-				{
-					ModifiedPoints.Add(MappedIndex);
-					MorphTargetDst.Points.Add(Delta.PositionDelta + Points[MappedIndex]);
-				}
+				continue;
+			}
+
+			if (InBoneIndexMap.IsValidIndex(SourceWeights.InfluenceBones[InfluenceIndex]))
+			{
+				SkeletalMeshImportData::FRawBoneInfluence TargetInfluence;
+				TargetInfluence.VertexIndex = TargetVertexIndex;
+				TargetInfluence.BoneIndex = InBoneIndexMap[SourceWeights.InfluenceBones[InfluenceIndex]];
+				TargetInfluence.Weight = SourceWeights.InfluenceWeights[InfluenceIndex] / 65535.0f;
+				TargetProfileData.Influences.Add(TargetInfluence);
 			}
 		}
 	}
 }
 
+
+
 void FSkeletalMeshImportData::CleanUpUnusedMaterials()
 {
 	if (Materials.Num() <= 0)
diff --git a/Engine/Source/Runtime/Engine/Private/SkeletalMesh.cpp b/Engine/Source/Runtime/Engine/Private/SkeletalMesh.cpp
--- a/Engine/Source/Runtime/Engine/Private/SkeletalMesh.cpp
+++ b/Engine/Source/Runtime/Engine/Private/SkeletalMesh.cpp
@@ -3201,18 +3201,39 @@ void USkeletalMesh::BeginPostLoadInternal(FSkinnedAssetPostLoadContext& Context)
 					check(RawSkeletalMeshBulkData_DEPRECATED);
 					if (!RawSkeletalMeshBulkData_DEPRECATED->IsEmpty())
 					{
-						FSkeletalMeshImportData SerializeMeshData;
-						RawSkeletalMeshBulkData_DEPRECATED->LoadRawMesh(SerializeMeshData);
+						FSkeletalMeshImportData SkeletalMeshImportData;
+						RawSkeletalMeshBulkData_DEPRECATED->LoadRawMesh(SkeletalMeshImportData);
 
-						// Some older versions of the bulk data did not store the morph targets, but they're available on the skeletal
-						// mesh itself. Try to back-fill from the skeletal mesh.
-						if (SerializeMeshData.MorphTargets.IsEmpty() && !GetMorphTargets().IsEmpty())
+						// Some older versions of the bulk data did not store the morph targets or alternate skin profiles, but they're 
+						// available on the skeletal mesh's LOD model. Try to back-fill from the LOD model.
+						if (SkeletalMeshImportData.MorphTargets.IsEmpty() && !GetMorphTargets().IsEmpty())
 						{
-							SerializeMeshData.SetMorphTargets(GetMorphTargets(), LODIndex, ThisLODModel.GetRawPointIndices());
+							for (UMorphTarget* MorphTarget: GetMorphTargets())
+							{
+								if (!MorphTarget->HasDataForLOD(LODIndex))
+								{
+									continue;
+								}
+
+								const FMorphTargetLODModel& MorphTargetModel = MorphTarget->GetMorphLODModels()[LODIndex];
+
+								// Confusingly, morph targets need FSkeletalMeshLODModel::RawPointIndices2 to map back to the import model, whereas
+								// skin weight profiles need MeshToImportVertexMap.
+								SkeletalMeshImportData.AddMorphTarget(MorphTarget->GetFName(), MorphTargetModel, ThisLODModel.GetRawPointIndices());
+							}
+						}
+
+						if (SkeletalMeshImportData.AlternateInfluences.IsEmpty() && !GetSkinWeightProfiles().IsEmpty() && !ThisLODModel.SkinWeightProfiles.IsEmpty())
+						{
+							for (const TPair<FName, FImportedSkinWeightProfileData>& SkinWeightProfileInfo: ThisLODModel.SkinWeightProfiles)
+							{
+								SkeletalMeshImportData.AddSkinWeightProfile(SkinWeightProfileInfo.Key, SkinWeightProfileInfo.Value,
+									ThisLODModel.MeshToImportVertexMap, ThisLODModel.ActiveBoneIndices);
+							}
 						}
 
 						FMeshDescription MeshDescription;
-						if (SerializeMeshData.GetMeshDescription(this, MeshDescription))
+						if (SkeletalMeshImportData.GetMeshDescription(this, MeshDescription))
 						{
 							CreateMeshDescription(LODIndex, MoveTemp(MeshDescription));
 							CommitMeshDescription(LODIndex);
diff --git a/Engine/Source/Runtime/Engine/Private/SkeletalMeshLODModel.cpp b/Engine/Source/Runtime/Engine/Private/SkeletalMeshLODModel.cpp
--- a/Engine/Source/Runtime/Engine/Private/SkeletalMeshLODModel.cpp
+++ b/Engine/Source/Runtime/Engine/Private/SkeletalMeshLODModel.cpp
@@ -1381,6 +1381,11 @@ void FSkeletalMeshLODModel::GetMeshDescription(const USkeletalMesh *InSkeletalMe
 		MeshAttributes.RegisterMorphTargetAttribute(Name);
 	}
 
+	for (const TPair<FName, FImportedSkinWeightProfileData>& SkinWeightProfileInfo: SkinWeightProfiles)
+	{
+		MeshAttributes.RegisterSkinWeightAttribute(SkinWeightProfileInfo.Key);
+	}
+
 	const int32 NumTriangles = IndexBuffer.Num() / 3;
 
 	const FReferenceSkeleton& RefSkeleton = InSkeletalMesh->GetRefSkeleton();
@@ -1539,11 +1544,29 @@ void FSkeletalMeshLODModel::GetMeshDescription(const USkeletalMesh *InSkeletalMe
 		}
 	}
 
+	for (const TPair<FName, FImportedSkinWeightProfileData>& SkinWeightProfileInfo: SkinWeightProfiles)
+	{
+		FSkinWeightsVertexAttributesRef SkinWeightAttribute = MeshAttributes.GetVertexSkinWeights(SkinWeightProfileInfo.Key);
+		const FImportedSkinWeightProfileData& SkinWeightProfileData = SkinWeightProfileInfo.Value;
+
+		check(SkinWeightProfileData.SkinWeights.Num() == NumVertices);
+		
+		for (int32 Index = 0; Index < SkinWeightProfileData.SkinWeights.Num(); Index++)
+		{
+			const FRawSkinWeight& RawSkinWeights = SkinWeightProfileData.SkinWeights[Index];
+			const int32 TargetVertexIndex = SourceToTargetVertexMap[Index];
+			const FVertexID VertexID = VertexIDs[TargetVertexIndex];
+			
+			FBoneWeights Weights = FBoneWeights::Create(RawSkinWeights.InfluenceBones, RawSkinWeights.InfluenceWeights);
+			SkinWeightAttribute.Set(VertexID, Weights);
+		}
+	}
+	
 	// Set Bone Attributes
-	for (int Idx = 0; Idx < NumBones; ++Idx)
+	for (int Index = 0; Index < NumBones; ++Index)
 	{
-		const FMeshBoneInfo& BoneInfo = RefSkeleton.GetRawRefBoneInfo()[Idx];
-		const FTransform& BoneTransform = RefSkeleton.GetRawRefBonePose()[Idx];
+		const FMeshBoneInfo& BoneInfo = RefSkeleton.GetRawRefBoneInfo()[Index];
+		const FTransform& BoneTransform = RefSkeleton.GetRawRefBonePose()[Index];
 
 		const FBoneID BoneID = MeshAttributes.CreateBone();
 		
diff --git a/Engine/Source/Runtime/Engine/Public/Rendering/SkeletalMeshLODImporterData.h b/Engine/Source/Runtime/Engine/Public/Rendering/SkeletalMeshLODImporterData.h
--- a/Engine/Source/Runtime/Engine/Public/Rendering/SkeletalMeshLODImporterData.h
+++ b/Engine/Source/Runtime/Engine/Public/Rendering/SkeletalMeshLODImporterData.h
@@ -13,6 +13,7 @@
 #include "Animation/MorphTarget.h"
 #include "Templates/DontCopy.h"
 
+struct FImportedSkinWeightProfileData;
 struct FMeshDescription;
 class FSkeletalMeshLODModel;
 
@@ -421,8 +422,14 @@ public:
 	/*
 	 * Add morph target data from UMorphTarget in case there was none on the mesh itself.
 	 */
-	ENGINE_API void SetMorphTargets(const TArray<TObjectPtr<UMorphTarget>>& InMorphTargets, int32 InLODIndex, const TArray<uint32>& InVertexMap);
+	ENGINE_API void AddMorphTarget(FName InMorphTargetName, const FMorphTargetLODModel& InMorphTargetModel, const TArray<uint32>& InVertexMap);
 
+	/*
+	 * Add alternate skin profile from FImportedSkinWeightProfileData
+	 */
+	ENGINE_API void AddSkinWeightProfile(FName InProfileName, const FImportedSkinWeightProfileData& InProfileData, const TArray<int32>& InVertexMap, const TArray<FBoneIndexType>& InBoneIndexMap);
+	
+	
 	/**
 	 * Returns a mesh description from the import data. If logging on failures is required, pass in a pointer 
 	 * to the owning skeletal mesh. Otherwise leave as a \c nullptr. 
