SVT: Add AssetUserData to UStreamableSparseVolumeTexture.

diff --git a/Engine/Source/Runtime/Engine/Classes/SparseVolumeTexture/SparseVolumeTexture.h b/Engine/Source/Runtime/Engine/Classes/SparseVolumeTexture/SparseVolumeTexture.h
--- a/Engine/Source/Runtime/Engine/Classes/SparseVolumeTexture/SparseVolumeTexture.h
+++ b/Engine/Source/Runtime/Engine/Classes/SparseVolumeTexture/SparseVolumeTexture.h
@@ -6,12 +6,13 @@
 
 #include "Engine/VolumeTexture.h"
 
-#include "Serialization/EditorBulkData.h"
 #include "Containers/Array.h"
 #include "Containers/StaticArray.h"
+#include "Engine/TextureDefines.h"
+#include "Interfaces/Interface_AssetUserData.h"
+#include "Serialization/EditorBulkData.h"
 #include "UnrealClient.h"
 #include "UObject/ObjectSaveContext.h"
-#include "Engine/TextureDefines.h"
 
 #include "SparseVolumeTexture.generated.h"
 
@@ -366,7 +367,7 @@ private:
 
 // Represents a streamable SparseVolumeTexture asset and serves as base class for UStaticSparseVolumeTexture and UAnimatedSparseVolumeTexture. It has an array of USparseVolumeTextureFrame.
 UCLASS(MinimalAPI, ClassGroup = Rendering, BlueprintType)
-class UStreamableSparseVolumeTexture : public USparseVolumeTexture
+class UStreamableSparseVolumeTexture : public USparseVolumeTexture, public IInterface_AssetUserData
 {
 	GENERATED_UCLASS_BODY()
 
@@ -414,6 +415,10 @@ public:
 	TObjectPtr<class UAssetImportData> AssetImportData;
 #endif // WITH_EDITORONLY_DATA
 
+	/** Array of user data stored with the asset */
+	UPROPERTY(EditAnywhere, AdvancedDisplay, Instanced, Category = "Texture")
+	TArray<TObjectPtr<class UAssetUserData>> AssetUserData;
+
 	UStreamableSparseVolumeTexture();
 	virtual ~UStreamableSparseVolumeTexture() = default;
 
@@ -444,6 +449,13 @@ public:
 	ENGINE_API virtual void GetAssetRegistryTags(TArray<FAssetRegistryTag>& OutTags) const override;
 	//~ End UObject Interface.
 
+	//~ Begin IInterface_AssetUserData Interface
+	ENGINE_API virtual void AddAssetUserData(UAssetUserData* InUserData) override;
+	ENGINE_API virtual void RemoveUserDataOfClass(TSubclassOf<UAssetUserData> InUserDataClass) override;
+	ENGINE_API virtual UAssetUserData* GetAssetUserDataOfClass(TSubclassOf<UAssetUserData> InUserDataClass) override;
+	ENGINE_API virtual const TArray<UAssetUserData*>* GetAssetUserDataArray() const override;
+	//~ End IInterface_AssetUserData Interface
+
 	//~ Begin USparseVolumeTexture Interface.
 	virtual int32 GetNumFrames() const override { return Frames.Num(); }
 	virtual int32 GetNumMipLevels() const override { return NumMipLevels; }
diff --git a/Engine/Source/Runtime/Engine/Private/SparseVolumeTexture/SparseVolumeTexture.cpp b/Engine/Source/Runtime/Engine/Private/SparseVolumeTexture/SparseVolumeTexture.cpp
--- a/Engine/Source/Runtime/Engine/Private/SparseVolumeTexture/SparseVolumeTexture.cpp
+++ b/Engine/Source/Runtime/Engine/Private/SparseVolumeTexture/SparseVolumeTexture.cpp
@@ -1508,6 +1508,49 @@ void UStreamableSparseVolumeTexture::GetAssetRegistryTags(FAssetRegistryTagsCont
 #endif
 }
 
+void UStreamableSparseVolumeTexture::AddAssetUserData(UAssetUserData* InUserData)
+{
+	if (InUserData != nullptr)
+	{
+		UAssetUserData* ExistingData = GetAssetUserDataOfClass(InUserData->GetClass());
+		if (ExistingData != nullptr)
+		{
+			AssetUserData.Remove(ExistingData);
+		}
+		AssetUserData.Add(InUserData);
+	}
+}
+
+UAssetUserData* UStreamableSparseVolumeTexture::GetAssetUserDataOfClass(TSubclassOf<UAssetUserData> InUserDataClass)
+{
+	for (UAssetUserData* Datum : AssetUserData)
+	{
+		if (Datum != nullptr && Datum->IsA(InUserDataClass))
+		{
+			return Datum;
+		}
+	}
+	return nullptr;
+}
+
+void UStreamableSparseVolumeTexture::RemoveUserDataOfClass(TSubclassOf<UAssetUserData> InUserDataClass)
+{
+	for (int32 DataIdx = 0; DataIdx < AssetUserData.Num(); DataIdx++)
+	{
+		UAssetUserData* Datum = AssetUserData[DataIdx];
+		if (Datum != nullptr && Datum->IsA(InUserDataClass))
+		{
+			AssetUserData.RemoveAt(DataIdx);
+			return;
+		}
+	}
+}
+
+const TArray<UAssetUserData*>* UStreamableSparseVolumeTexture::GetAssetUserDataArray() const
+{
+	return &ToRawPtrTArrayUnsafe(AssetUserData);
+}
+
 #if WITH_EDITORONLY_DATA
 void UStreamableSparseVolumeTexture::RecacheFrames()
 {
