Skein / ViewportStatus: Added ShowFlags to enable/disable drawing of the additional outline colors per individual color.

#rnx

diff --git a/Engine/Source/Runtime/Engine/Public/ShowFlags.h b/Engine/Source/Runtime/Engine/Public/ShowFlags.h
--- a/Engine/Source/Runtime/Engine/Public/ShowFlags.h
+++ b/Engine/Source/Runtime/Engine/Public/ShowFlags.h
@@ -397,6 +397,12 @@ private:
 		SetVisualizeTSR(false);
 		SetEditingLevelInstance(false);
 		SetSelectionOutline(false);
+		SetSelectionOutlineColor0(false);
+		SetSelectionOutlineColor1(false);
+		SetSelectionOutlineColor2(false);
+		SetSelectionOutlineColor3(false);
+		SetSelectionOutlineColor4(false);
+		SetSelectionOutlineColor5(false);
 		SetDebugAI(false);
 		SetNavigation(false);
 		SetLightComplexity(false);
diff --git a/Engine/Source/Runtime/Engine/Public/ShowFlagsValues.inl b/Engine/Source/Runtime/Engine/Public/ShowFlagsValues.inl
--- a/Engine/Source/Runtime/Engine/Public/ShowFlagsValues.inl
+++ b/Engine/Source/Runtime/Engine/Public/ShowFlagsValues.inl
@@ -138,6 +138,13 @@ SHOWFLAG_FIXED_IN_SHIPPING(1, Diffuse, SFG_LightingComponents, NSLOCTEXT("Unreal
 SHOWFLAG_ALWAYS_ACCESSIBLE(Specular, SFG_LightingComponents, NSLOCTEXT("UnrealEd", "SpecularSF", "Specular"))
 /** Outline around selected objects in the editor */
 SHOWFLAG_FIXED_IN_SHIPPING(0, SelectionOutline, SFG_Hidden, NSLOCTEXT("UnrealEd", "SelectionOutlineSF", "Selection Outline"))
+/** Outline custom colors in the editor */
+SHOWFLAG_FIXED_IN_SHIPPING(0, SelectionOutlineColor0, SFG_Hidden, NSLOCTEXT("UnrealEd", "SelectionOutlineColor0", "Selection Outline Color 0"))
+SHOWFLAG_FIXED_IN_SHIPPING(0, SelectionOutlineColor1, SFG_Hidden, NSLOCTEXT("UnrealEd", "SelectionOutlineColor1", "Selection Outline Color 1"))
+SHOWFLAG_FIXED_IN_SHIPPING(0, SelectionOutlineColor2, SFG_Hidden, NSLOCTEXT("UnrealEd", "SelectionOutlineColor2", "Selection Outline Color 2"))
+SHOWFLAG_FIXED_IN_SHIPPING(0, SelectionOutlineColor3, SFG_Hidden, NSLOCTEXT("UnrealEd", "SelectionOutlineColor3", "Selection Outline Color 3"))
+SHOWFLAG_FIXED_IN_SHIPPING(0, SelectionOutlineColor4, SFG_Hidden, NSLOCTEXT("UnrealEd", "SelectionOutlineColor4", "Selection Outline Color 4"))
+SHOWFLAG_FIXED_IN_SHIPPING(0, SelectionOutlineColor5, SFG_Hidden, NSLOCTEXT("UnrealEd", "SelectionOutlineColor5", "Selection Outline Color 5"))
 /** If screen space reflections are enabled, for now SHOWFLAG_ALWAYS_ACCESSIBLE because it's exposed in SceneCapture */
 SHOWFLAG_ALWAYS_ACCESSIBLE(ScreenSpaceReflections, SFG_LightingFeatures, NSLOCTEXT("UnrealEd", "ScreenSpaceReflectionsSF", "Screen Space Reflections"))
 SHOWFLAG_ALWAYS_ACCESSIBLE(LumenReflections, SFG_LightingFeatures, NSLOCTEXT("UnrealEd", "LumenReflectionsSF", "Lumen Reflections"))
