Create rendering settings category for Heterogeneous Volumes

diff --git a/Engine/Source/Runtime/Engine/Classes/Engine/RendererSettings.h b/Engine/Source/Runtime/Engine/Classes/Engine/RendererSettings.h
--- a/Engine/Source/Runtime/Engine/Classes/Engine/RendererSettings.h
+++ b/Engine/Source/Runtime/Engine/Classes/Engine/RendererSettings.h
@@ -988,24 +988,24 @@ class URendererSettings : public UDeveloperSettings
 		ConfigRestartRequired = true))
 		uint32 bMeshStreaming : 1;
 
-	UPROPERTY(config, EditAnywhere, Category = Experimental, meta = (
-		ConsoleVariable = "r.HeterogeneousVolumes", DisplayName = "Enable Heterogeneous Volumes",
+	UPROPERTY(config, EditAnywhere, Category = "Heterogeneous Volumes", meta = (
+		ConsoleVariable = "r.HeterogeneousVolumes", DisplayName = "Heterogeneous Volumes (Experimental)",
 		ToolTip = "Enable rendering with the heterogeneous volumes subsystem.",
 		ConfigRestartRequired = false))
 	uint32 bEnableHeterogeneousVolumes : 1;
 
-	UPROPERTY(config, EditAnywhere, Category = Experimental, meta = (
-		ConsoleVariable = "r.Translucency.HeterogeneousVolumes", DisplayName = "Composite Heterogeneous Volumes with Translucency",
-		ToolTip = "Enable compositing with heterogeneous volumes when rendering translucency.",
-		ConfigRestartRequired = true))
-	uint32 bCompositeHeterogeneousVolumesWithTranslucency : 1;
-
-	UPROPERTY(config, EditAnywhere, Category = Experimental, meta = (
-		ConsoleVariable = "r.HeterogeneousVolumes.Shadows", DisplayName = "Enable shadow-casting with Heterogeneous Volumes",
+	UPROPERTY(config, EditAnywhere, Category = "Heterogeneous Volumes", meta = (
+		ConsoleVariable = "r.HeterogeneousVolumes.Shadows", DisplayName = "Shadow Casting",
 		ToolTip = "Enable heterogeneous volumes to cast shadows onto the environment.",
 		ConfigRestartRequired = true))
 		uint32 bShouldHeterogeneousVolumesCastShadows : 1;
 
+	UPROPERTY(config, EditAnywhere, Category = "Heterogeneous Volumes", meta = (
+		ConsoleVariable = "r.Translucency.HeterogeneousVolumes", DisplayName = "Composite with Translucency",
+		ToolTip = "Enable compositing with heterogeneous volumes when rendering translucency.",
+		ConfigRestartRequired = true))
+	uint32 bCompositeHeterogeneousVolumesWithTranslucency : 1;
+
 	UPROPERTY(config, EditAnywhere, Category=Editor, meta=(
 		ConsoleVariable="r.WireframeCullThreshold",DisplayName="Wireframe Cull Threshold",
 		ToolTip="Screen radius at which wireframe objects are culled. Larger values can improve performance when viewing a scene in wireframe."))
