Report connection closed to Network Insights as late as possible to include all events posted between Close and CleanUp.

diff --git a/Engine/Source/Runtime/Engine/Private/NetConnection.cpp b/Engine/Source/Runtime/Engine/Private/NetConnection.cpp
--- a/Engine/Source/Runtime/Engine/Private/NetConnection.cpp
+++ b/Engine/Source/Runtime/Engine/Private/NetConnection.cpp
@@ -1092,11 +1092,6 @@ void UNetConnection::Close(FNetResult&& CloseReason)
 		}
 
 		NetPing.Reset();
-
-		if (const uint32 MyConnectionId = GetConnectionId())
-		{
-			UE_NET_TRACE_CONNECTION_CLOSED(NetTraceId, MyConnectionId);
-		}
 	}
 
 	LogCallLastTime		= 0;
@@ -1310,11 +1305,17 @@ void UNetConnection::CleanUp()
 		}
 	}
 
+	const uint32 MyConnectionId = GetConnectionId();
+	if (MyConnectionId)
+	{
+		UE_NET_TRACE_CONNECTION_CLOSED(NetTraceId, MyConnectionId);
+	}
+
 	if (Driver != nullptr)
 	{
 		// It would be nicer to have the Driver handle this internally, but unfortunately the ServerConnection member is public.
 		// Otherwise we'd be able to do the appropriate logic in Add/Remove Client/ServerConnection
-		if (const uint32 MyConnectionId = GetConnectionId())
+		if (MyConnectionId)
 		{
 #if UE_WITH_IRIS
 			if (UReplicationSystem* ReplicationSystem = Driver->GetReplicationSystem())
