[Shaders] - Implement recompile shaders on file save functionality for the editor.

- Adds a directory watcher to watch for changes in //Engine/Shaders/... for usf/ush/h files.
- When a file changes issue a `recompileshaders /Engine/Shaders/Private/PostProcessGBufferHints.usf`command.
- This code should be WITH_EDITOR only.
- This is opt-in behavior and can be turned on with r.ShaderCompiler.RecompileShadersOnSave=1

diff --git a/Engine/Source/Runtime/Engine/Private/ShaderCompiler/ShaderCompiler.cpp b/Engine/Source/Runtime/Engine/Private/ShaderCompiler/ShaderCompiler.cpp
--- a/Engine/Source/Runtime/Engine/Private/ShaderCompiler/ShaderCompiler.cpp
+++ b/Engine/Source/Runtime/Engine/Private/ShaderCompiler/ShaderCompiler.cpp
@@ -66,6 +66,8 @@
 #include "DerivedDataRequestOwner.h"
 #include "TextureCompiler.h"
 #include "Rendering/StaticLightingSystemInterface.h"
+#include "DirectoryWatcherModule.h"
+#include "IDirectoryWatcher.h"
 #endif
 
 #if WITH_ODSC
@@ -85,6 +87,13 @@ DEFINE_LOG_CATEGORY(LogShaderCompilers);
 
 LLM_DEFINE_TAG(ShaderCompiler);
 
+static TAutoConsoleVariable<bool> CVarRecompileShadersOnSave(
+	TEXT("r.ShaderCompiler.RecompileShadersOnSave"),
+	false,
+	TEXT("When enabled, the editor will attempt to recompile any shader files that have changed when saved.  Useful for iterating on shaders in the editor.\n")
+	TEXT("Default: false"),
+	ECVF_ReadOnly);
+
 int32 GShaderCompilerJobCache = 1;
 static FAutoConsoleVariableRef CVarShaderCompilerJobCache(
 	TEXT("r.ShaderCompiler.JobCache"),
@@ -5651,6 +5660,93 @@ FShaderCompilingManager::FShaderCompilingManager() :
 	OutOfMemoryDelegateHandle = FCoreDelegates::GetOutOfMemoryDelegate().AddRaw(this, &FShaderCompilingManager::ReportMemoryUsage);
 
 	FAssetCompilingManager::Get().RegisterManager(this);
+
+#if WITH_EDITOR
+	static const bool bAllowShaderRecompileOnSave = CVarRecompileShadersOnSave.GetValueOnAnyThread();
+	if (bAllowShaderRecompileOnSave)
+	{
+		if (IDirectoryWatcher* DirectoryWatcher = FModuleManager::LoadModuleChecked<FDirectoryWatcherModule>(TEXT("DirectoryWatcher")).Get())
+		{
+			// Handle if we are watching a directory for changes.
+			{
+				UE_LOG(LogShaderCompilers, Display, TEXT("Register directory watchers for shader files."));
+
+				const TMap<FString, FString>& ShaderSourceDirectoryMappings = AllShaderSourceDirectoryMappings();
+
+				DirectoryWatcherHandles.Reserve(ShaderSourceDirectoryMappings.Num());
+
+				for (const auto& It : ShaderSourceDirectoryMappings)
+				{
+					FString DirectoryToWatch = It.Value;
+					if (FPaths::IsRelative(DirectoryToWatch))
+					{
+						DirectoryToWatch = FPaths::ConvertRelativePathToFull(DirectoryToWatch);
+					}
+
+					FDelegateHandle& DirectoryWatcherHandle = DirectoryWatcherHandles.Add(DirectoryToWatch, FDelegateHandle());
+
+					DirectoryWatcher->RegisterDirectoryChangedCallback_Handle(
+						DirectoryToWatch,
+						IDirectoryWatcher::FDirectoryChanged::CreateLambda([](const TArray<FFileChangeData>& InFileChangeDatas) {
+
+							TRACE_CPUPROFILER_EVENT_SCOPE(HandleDirectoryChanged);
+
+							if (!bAllowShaderRecompileOnSave)
+							{
+								return;
+							}
+
+							TArray<FString> ChangedShaderFiles;
+							for (const FFileChangeData& It : InFileChangeDatas)
+							{
+								if (It.Filename.EndsWith(TEXT(".usf")) || It.Filename.EndsWith(TEXT(".ush")) || It.Filename.EndsWith(TEXT(".h")))
+								{
+									UE_LOG(LogShaderCompilers, Display, TEXT("Detected change on %s"), *It.Filename);
+
+									ChangedShaderFiles.AddUnique(It.Filename);
+								}
+							}
+
+							if (ChangedShaderFiles.Num())
+							{
+								// Mappings from:
+								// Key:   /Engine to
+								// Value: ../../../Engine/Shaders
+								const TMap<FString, FString>& ShaderSourceDirectoryMappings = AllShaderSourceDirectoryMappings();
+
+								FString RemappedShaderFileName;
+								for (const auto& It : ShaderSourceDirectoryMappings)
+								{
+									// ChangedShaderFiles will be of format: ../../../Engine/Shaders/Private/PostProcessGBufferHints.usf
+									if (ChangedShaderFiles[0].StartsWith(It.Value))
+									{
+										// Change from relative path to Engine absolute path.
+										// i.e. change `../../../Engine/Shaders/Private/PostProcessGBufferHints.usf` to `/Engine/Shaders/Private/PostProcessGBufferHints.usf`
+										RemappedShaderFileName = ChangedShaderFiles[0].Replace(*It.Value, *It.Key);
+									}
+								}
+
+								// Issue a `recompileshaders /Engine/Shaders/Private/PostProcessGBufferHints.usf` command, which will just compile that shader source file.
+								RecompileShaders(*RemappedShaderFileName, *GLog);
+
+								UE_LOG(LogShaderCompilers, Display, TEXT("Ready for new shader file changes"));
+							}
+						}),
+						DirectoryWatcherHandle);
+
+					if (DirectoryWatcherHandle.IsValid())
+					{
+						UE_LOG(LogShaderCompilers, Display, TEXT("Watching %s -> %s"), *It.Key, *DirectoryToWatch);
+					}
+					else
+					{
+						UE_LOG(LogShaderCompilers, Error, TEXT("Failed to set up directory watcher %s -> %s"), *It.Key, *DirectoryToWatch);
+					}
+				}
+			}
+		}
+	}
+#endif // WITH_EDITOR
 }
 
 FShaderCompilingManager::~FShaderCompilingManager()
@@ -5669,6 +5765,20 @@ FShaderCompilingManager::~FShaderCompilingManager()
 
 	FCoreDelegates::GetOutOfMemoryDelegate().Remove(OutOfMemoryDelegateHandle);
 
+#if WITH_EDITOR
+	const bool bAllowShaderRecompileOnSave = CVarRecompileShadersOnSave.GetValueOnAnyThread();
+	if (bAllowShaderRecompileOnSave)
+	{
+		if (IDirectoryWatcher* DirectoryWatcher = FModuleManager::LoadModuleChecked<FDirectoryWatcherModule>(TEXT("DirectoryWatcher")).Get())
+		{
+			for (const auto& It : DirectoryWatcherHandles)
+			{
+				DirectoryWatcher->UnregisterDirectoryChangedCallback_Handle(It.Key, It.Value);
+			}
+		}
+	}
+#endif // WITH_EDITOR
+
 	FAssetCompilingManager::Get().UnregisterManager(this);
 }
 
diff --git a/Engine/Source/Runtime/Engine/Public/ShaderCompiler.h b/Engine/Source/Runtime/Engine/Public/ShaderCompiler.h
--- a/Engine/Source/Runtime/Engine/Public/ShaderCompiler.h
+++ b/Engine/Source/Runtime/Engine/Public/ShaderCompiler.h
@@ -741,6 +741,10 @@ private:
 	/** Delegate handle for delegate used to report memory usage during out-of-memory conditions. */
 	FDelegateHandle OutOfMemoryDelegateHandle;
 
+#if WITH_EDITOR
+	TMap<FString, FDelegateHandle> DirectoryWatcherHandles;
+#endif // WITH_EDITOR
+
 	/** Calculate NumShaderCompilingThreads, during construction or OnMachineResourcesChanged */
 	void CalculateNumberOfCompilingThreads(int32 NumberOfCores, int32 NumberOfCoresIncludingHyperthreads);
 
