[NetPhysics] Resimulation bugfixes

- BlockResimFrame blocked next frame also
- If a valid rewind frame was not found the requested frame was not cleared and would be re-used again producing the same result (no valid rewind frame)
- Resim runtime correction could divide by 0

diff --git a/Engine/Source/Runtime/Engine/Private/PhysicsEngine/PhysicsReplication.cpp b/Engine/Source/Runtime/Engine/Private/PhysicsEngine/PhysicsReplication.cpp
--- a/Engine/Source/Runtime/Engine/Private/PhysicsEngine/PhysicsReplication.cpp
+++ b/Engine/Source/Runtime/Engine/Private/PhysicsEngine/PhysicsReplication.cpp
@@ -1712,7 +1712,7 @@ bool FPhysicsReplicationAsync::ResimulationReplication(Chaos::FPBDRigidParticleH
 
 	const int32 LocalFrame = Target.ServerFrame - Target.FrameOffset;
 
-	if (LocalFrame > RewindData->CurrentFrame() || LocalFrame < RewindData->GetEarliestFrame_Internal())
+	if (LocalFrame >= RewindData->CurrentFrame() || LocalFrame < RewindData->GetEarliestFrame_Internal())
 	{
 		if (LocalFrame > 0 && (RewindData->CurrentFrame() - RewindData->GetEarliestFrame_Internal()) == RewindData->Capacity())
 		{
@@ -1759,7 +1759,7 @@ bool FPhysicsReplicationAsync::ResimulationReplication(Chaos::FPBDRigidParticleH
 	}
 #endif
 
-	if (Target.TickCount == 0 && LocalFrame > RewindData->GetBlockedResimFrame() && ShouldTriggerResim)
+	if (ShouldTriggerResim && Target.TickCount == 0 && LocalFrame > RewindData->GetBlockedResimFrame())
 	{
 		// Trigger resimulation
 		RigidsSolver->GetEvolution()->GetIslandManager().SetParticleResimFrame(Handle, LocalFrame);
@@ -1772,7 +1772,7 @@ bool FPhysicsReplicationAsync::ResimulationReplication(Chaos::FPBDRigidParticleH
 	{
 		const int32 NumPredictedFrames = RigidsSolver->GetCurrentFrame() - LocalFrame - Target.TickCount;
 
-		if (Target.TickCount <= NumPredictedFrames)
+		if (Target.TickCount <= NumPredictedFrames && NumPredictedFrames > 0)
 		{
 			// Calculate correction to position
 			const float CorrectionAmountX = PhysicsReplicationCVars::ResimulationCVars::PosStabilityMultiplier / NumPredictedFrames; // Same result as (ErrorOffset / NumPredictedFrames) * PosStabilityMultiplier
@@ -1794,10 +1794,10 @@ bool FPhysicsReplicationAsync::ResimulationReplication(Chaos::FPBDRigidParticleH
 			// Apply correction
 			Handle->SetX(CorrectedX);
 			Handle->SetR(CorrectedR);
-
-			// Keep target for NumPredictedFrames time to perform runtime corrections with until a new target is received
-			bClearTarget = Target.TickCount >= NumPredictedFrames;
 		}
+
+		// Keep target for NumPredictedFrames time to perform runtime corrections with until a new target is received
+		bClearTarget = Target.TickCount >= NumPredictedFrames;
 	}
 
 	return bClearTarget;
diff --git a/Engine/Source/Runtime/Experimental/Chaos/Private/PBDRigidsSolver.cpp b/Engine/Source/Runtime/Experimental/Chaos/Private/PBDRigidsSolver.cpp
--- a/Engine/Source/Runtime/Experimental/Chaos/Private/PBDRigidsSolver.cpp
+++ b/Engine/Source/Runtime/Experimental/Chaos/Private/PBDRigidsSolver.cpp
@@ -1874,6 +1874,8 @@ namespace Chaos
 
 			if (ResimStep < 0)
 			{
+				// Clear ResimFrame if no valid resim frame was found
+				MRewindData->SetResimFrame(INDEX_NONE);
 				return;
 			}
 
diff --git a/Engine/Source/Runtime/Experimental/Chaos/Private/RewindData.cpp b/Engine/Source/Runtime/Experimental/Chaos/Private/RewindData.cpp
--- a/Engine/Source/Runtime/Experimental/Chaos/Private/RewindData.cpp
+++ b/Engine/Source/Runtime/Experimental/Chaos/Private/RewindData.cpp
@@ -280,7 +280,7 @@ bool FRewindData::RewindToFrame(int32 Frame)
 	}
 
 	//If we need to save and we are right on the edge of the buffer, we can't go back to earliest frame
-	if(Frame == EarliestFrame && bNeedsSave && FramesSaved == Managers.Capacity())
+	if (Frame == EarliestFrame && bNeedsSave && FramesSaved == Managers.Capacity())
 	{
 #if DEBUG_REWIND_DATA
 		UE_LOG(LogTemp, Log, TEXT("COMMON | PT | RewindToFrame | Failed due to rewinding to last available frame and bNeedsSave is set to true"));
@@ -940,7 +940,7 @@ int32 FRewindData::FindValidResimFrame(const int32 RequestedFrame)
 	}
 
 	// First frame of the history datas
-	const int32 EarliestFrame = FMath::Max(GetEarliestFrame_Internal(), BlockResimFrame + 1);
+	const int32 EarliestFrame = FMath::Max(GetEarliestFrame_Internal(), BlockResimFrame);
 	bool bHasTargetHistory = false;
 
 	for (ValidFrame = RequestedFrame; ValidFrame > EarliestFrame; ValidFrame--)
