Prevent a crash on the player controller when attempting to restart the player if the game mode was null. So far the only known case of this happening is if you delete everything from an active PIE session and then attempt to respawn the pawn.

Now instead of crashing, we wil log that it is happening and early exit.

diff --git a/Engine/Source/Runtime/Engine/Private/PlayerController.cpp b/Engine/Source/Runtime/Engine/Private/PlayerController.cpp
--- a/Engine/Source/Runtime/Engine/Private/PlayerController.cpp
+++ b/Engine/Source/Runtime/Engine/Private/PlayerController.cpp
@@ -3281,6 +3281,15 @@ void APlayerController::ServerRestartPlayer_Implementation()
 	if ( IsInState(NAME_Inactive) || (IsInState(NAME_Spectating) && bPlayerIsWaiting) )
 	{
 		AGameModeBase* const GameMode = GetWorld()->GetAuthGameMode();
+
+		// This can happen if you do something like delete a bunch of stuff at runtime in PIE or something like that.
+		// We need to check here to prevent a crash
+		if (!IsValid(GameMode))
+		{
+			UE_LOG(LogPlayerController, Warning, TEXT("[APlayerController::ServerRestartPlayer_Implementation] Player Controller '%s' requested restart but the game mode is null! Nothing will happen."), *GetNameSafe(this));
+			return;
+		}
+		
 		if ( !GameMode->PlayerCanRestart(this) )
 		{
 			return;
