PR #11313: Update `FilterArray`'s result be auto-cast to the filter type

#preflght 65b3f5988b73b25919600fd6

diff --git a/Engine/Source/Runtime/Engine/Classes/Kismet/KismetArrayLibrary.h b/Engine/Source/Runtime/Engine/Classes/Kismet/KismetArrayLibrary.h
--- a/Engine/Source/Runtime/Engine/Classes/Kismet/KismetArrayLibrary.h
+++ b/Engine/Source/Runtime/Engine/Classes/Kismet/KismetArrayLibrary.h
@@ -233,7 +233,7 @@ class UKismetArrayLibrary : public UBlueprintFunctionLibrary
 	 *@param	FilterClass		The Actor sub-class type that acts as the filter, only objects derived from it will be returned.
 	 *@return	An array containing only those objects which are derived from the class specified.
 	*/
-	UFUNCTION(BlueprintCallable, meta=(DisplayName = "Filter Array"), Category="Utilities|Array")
+	UFUNCTION(BlueprintCallable, meta=(DisplayName = "Filter Array", DeterminesOutputType = "FilterClass", DynamicOutputParam = "FilteredArray"), Category = "Utilities|Array")
 	static ENGINE_API void FilterArray(const TArray<AActor*>& TargetArray, TSubclassOf<class AActor> FilterClass, TArray<AActor*>& FilteredArray);
 
 	/** 
