LevelInstanceSubystem: Handle OnCleanupWorld to properly cleanup currently loaded LevelInstances as CleanupWorld doesn't do proper streaming out of levels

Reported through UDN : This is an issue with UWorld::ReInitWorld() (which is called on "Rebuild All"), it calls CleanupWorld / InitWorld which was leaving the LevelInstances in a bad state.

diff --git a/Engine/Source/Runtime/Engine/Private/LevelInstance/LevelInstanceSubsystem.cpp b/Engine/Source/Runtime/Engine/Private/LevelInstance/LevelInstanceSubsystem.cpp
--- a/Engine/Source/Runtime/Engine/Private/LevelInstance/LevelInstanceSubsystem.cpp
+++ b/Engine/Source/Runtime/Engine/Private/LevelInstance/LevelInstanceSubsystem.cpp
@@ -90,6 +90,7 @@ void ULevelInstanceSubsystem::Initialize(FSubsystemCollectionBase& Collection)
 		FEditorDelegates::OnAssetsPreDelete.AddUObject(this, &ULevelInstanceSubsystem::OnAssetsPreDelete);
 		FEditorDelegates::PreSaveWorldWithContext.AddUObject(this, &ULevelInstanceSubsystem::OnPreSaveWorldWithContext);
 		FWorldDelegates::OnPreWorldRename.AddUObject(this, &ULevelInstanceSubsystem::OnPreWorldRename);
+		FWorldDelegates::OnWorldCleanup.AddUObject(this, &ULevelInstanceSubsystem::OnWorldCleanup);
 	}
 #endif
 }
@@ -100,6 +101,7 @@ void ULevelInstanceSubsystem::Deinitialize()
 	FEditorDelegates::OnAssetsPreDelete.RemoveAll(this);
 	FEditorDelegates::PreSaveWorldWithContext.RemoveAll(this);
 	FWorldDelegates::OnPreWorldRename.RemoveAll(this);
+	FWorldDelegates::OnWorldCleanup.RemoveAll(this);
 #endif
 }
 
@@ -413,6 +415,68 @@ void ULevelInstanceSubsystem::OnPreWorldRename(UWorld* InWorld, const TCHAR* InN
 	}
 }
 
+void ULevelInstanceSubsystem::OnWorldCleanup(UWorld* InWorld, bool bSessionEnded, bool bCleanupResources)
+{
+	if (InWorld == GetWorld() && !InWorld->IsGameWorld() && bCleanupResources)
+	{
+		// LevelInstanceSubsystem doesn't support being Deinitialized and then Initialized without doing the following cleanup code (which happens with UWorld::ReInitWorld())
+		// because UWorld::CleanupWorldInternal doesn't do a clean Streaming out of StreamingLevels which is fine for regular StreamingLevels but since LevelInstance StreamingLevels are tied to Actors being registered/unregistered
+		// We need to do a cleanup here to make sure a call to UWorld::ReInitWorld() can properly reinitialize/stream in those levels instances again.
+		TArray<ULevelStreaming*> StreamingLevels;
+		ForEachLevelStreaming([&StreamingLevels](ULevelStreaming* LevelStreaming)
+		{
+			if (ULevel* LoadedLevel = LevelStreaming->GetLoadedLevel())
+			{
+				// Avoid GC Leak by restoring OwningWorld to its original value (stop pointing to World being cleaned up)
+				LoadedLevel->OwningWorld = LoadedLevel->GetTypedOuter<UWorld>();
+				// Level Streaming isn't going to be properly RemovedFromWorld so here we remove the annotation
+				ULevelStreaming::RemoveLevelAnnotation(LoadedLevel);
+				// Make sure Level can't be reused if world gets re-initialized through UWorld::ReInitWorld()
+				// This will make sure that the Level Package and its OFPA actor packages get trashed so they can't get reused (ULevel::CleanupLevel)
+				LoadedLevel->SetForceCantReuseUnloadedButStillAround(true);
+			}
+			StreamingLevels.Add(LevelStreaming);
+			return true;
+		});
+
+		if (StreamingLevels.Num())
+		{
+			InWorld->RemoveStreamingLevels(StreamingLevels);
+		}
+
+		LoadedLevelInstances.Empty();
+		if (LevelInstanceEdit)
+		{
+			// If we are inside an Edit, null out streaming pointer so that destructor doesn't cause a call to RemoveLevelsFromWorld
+			LevelInstanceEdit->LevelStreaming = nullptr;
+			LevelInstanceEdit.Reset();
+
+			if (ILevelInstanceEditorModule* EditorModule = (ILevelInstanceEditorModule*)FModuleManager::Get().GetModule("LevelInstanceEditor"))
+			{
+				EditorModule->DeactivateEditorMode();
+			}
+		}	
+	}
+}
+
+void ULevelInstanceSubsystem::ForEachLevelStreaming(TFunctionRef<bool(ULevelStreaming*)> Operation) const
+{
+	// Make sure Levels are properly trashed when cleaning up world (can't be reused)
+	for (auto& [LevelInstanceID, LoadedLevelInstance] : LoadedLevelInstances)
+	{
+		if (!Operation(LoadedLevelInstance.LevelStreaming))
+		{
+			return;
+		}
+	}
+
+	if (LevelInstanceEdit)
+	{
+		Operation(LevelInstanceEdit->LevelStreaming);
+	}
+}
+
+
 void ULevelInstanceSubsystem::RegisterLoadedLevelStreamingLevelInstanceEditor(ULevelStreamingLevelInstanceEditor* LevelStreaming)
 {
 	if (!bIsCreatingLevelInstance)
@@ -1771,7 +1835,10 @@ ULevelInstanceSubsystem::FLevelInstanceEdit::FLevelInstanceEdit(ULevelStreamingL
 ULevelInstanceSubsystem::FLevelInstanceEdit::~FLevelInstanceEdit()
 {
 	EditorObject->ExitEdit();
-	ULevelStreamingLevelInstanceEditor::Unload(LevelStreaming);
+	if (LevelStreaming)
+	{
+		ULevelStreamingLevelInstanceEditor::Unload(LevelStreaming);
+	}
 }
 
 UWorld* ULevelInstanceSubsystem::FLevelInstanceEdit::GetEditWorld() const
diff --git a/Engine/Source/Runtime/Engine/Public/LevelInstance/LevelInstanceSubsystem.h b/Engine/Source/Runtime/Engine/Public/LevelInstance/LevelInstanceSubsystem.h
--- a/Engine/Source/Runtime/Engine/Public/LevelInstance/LevelInstanceSubsystem.h
+++ b/Engine/Source/Runtime/Engine/Public/LevelInstance/LevelInstanceSubsystem.h
@@ -180,6 +180,7 @@ private:
 	ENGINE_API void OnAssetsPreDelete(const TArray<UObject*>& Objects);
 	ENGINE_API void OnPreSaveWorldWithContext(UWorld* InWorld, FObjectPreSaveContext ObjectSaveContext);
 	ENGINE_API void OnPreWorldRename(UWorld* InWorld, const TCHAR* InName, UObject* NewOuter, ERenameFlags Flags, bool& bShouldFailRename);
+	void OnWorldCleanup(UWorld* InWorld, bool bSessionEnded, bool bCleanupResources);
 
 	ENGINE_API void RegisterLoadedLevelStreamingLevelInstanceEditor(ULevelStreamingLevelInstanceEditor* LevelStreaming);
 
@@ -222,6 +223,8 @@ private:
 	ENGINE_API bool IsLevelInstanceEditDirty(const FLevelInstanceEdit* LevelInstanceEdit) const;
 	ENGINE_API bool PromptUserForCommit(const FLevelInstanceEdit* InLevelInstanceEdit, bool& bOutDiscard, bool bForceCommit = false) const;
 
+	void ForEachLevelStreaming(TFunctionRef<bool(ULevelStreaming*)> Operation) const;
+
 	struct FLevelsToRemoveScope
 	{
 		FLevelsToRemoveScope(ULevelInstanceSubsystem* InOwner);
