AssetCompilingManager : identify the case that we are running a task even though we expect to exceed memory budget
do not change behavior, it will still run like before, just log about it

diff --git a/Engine/Source/Runtime/Engine/Private/AssetCompilingManager.cpp b/Engine/Source/Runtime/Engine/Private/AssetCompilingManager.cpp
--- a/Engine/Source/Runtime/Engine/Private/AssetCompilingManager.cpp
+++ b/Engine/Source/Runtime/Engine/Private/AssetCompilingManager.cpp
@@ -264,22 +264,36 @@ public:
 
 		int32 DynamicMaxConcurrency = GetAdjustedMaxConcurrency(Priority);
 		int32 Concurrency = FQueuedThreadPoolWrapper::GetCurrentConcurrency();
+		
+		int64 MemoryLimit = GetMemoryLimit();
+
+		// we check MemoryLimit which uses current system available memory against TotalRequiredMemory
+		//	which includes already running tasks
+		// those tasks may have already allocated, so they have been counted twice
+		//  this means we can count tasks as using twice as much memory as they actually do
+		//  eg. one 4 GB task can block all concurrency when there was 8 GB of room
+		//    because after it allocated there is 4 GB free but we also then subtract off the 4 GB the task requires
+		// this is conservative, so okay if we want to avoid OOM and are okay with less concurrency than possible
 
-		// Never limit below a concurrency of 1 or else we'll starve the asset processing and never be scheduled again
-		// The idea for now is to let assets get built one by one when starving on memory, which should not increase OOM compared to the old synchronous behavior.
-		if ( TotalRequiredMemory > 0 && Concurrency > 0 )
+		if ( TotalRequiredMemory > 0 && TotalRequiredMemory >= MemoryLimit )
 		{
-			int64 MemoryLimit = GetMemoryLimit();
+			if ( Concurrency == 0 )
+			{
+				// Never limit below a concurrency of 1 or else we'll starve the asset processing and never be scheduled again
+				// The idea for now is to let assets get built one by one when starving on memory, which should not increase OOM compared to the old synchronous behavior.
+				
+				UE_LOG(LogAsyncCompilation, Display, TEXT("AssetCompile memory estimate is greater than available, but we're running it anyway, likely to fail! ")
+					TEXT("RequiredMemory = %.3f MB + %.3f MB MemoryLimit = %.3f MB "),
+					NewRequiredMemory/(1024*1024.f),
+					TotalEstimatedMemory/(1024*1024.f),
+					MemoryLimit/(1024*1024.f));
 
-			// we check MemoryLimit which uses current system available memory against TotalRequiredMemory
-			//	which includes already running tasks
-			// those tasks may have already allocated, so they have been counted twice
-			//  this means we can count tasks as using twice as much memory as they actually do
-			//  eg. one 4 GB task can block all concurrency when there was 8 GB of room
-			//    because after it allocated there is 4 GB free but we also then subtract off the 4 GB the task requires
-			// this is conservative, so okay if we want to avoid OOM and are okay with less concurrency than possible
+				// @todo : ? pause the main thread? pause shader compilers? trigger a GC ?
 
-			if ( TotalRequiredMemory >= MemoryLimit )
+				DynamicMaxConcurrency = 1;
+				// this will cause CanSchedule() to return true
+			}
+			else
 			{
 				UE_LOG(LogAsyncCompilation, Verbose, TEXT("CONCURRENCY LIMITED to %d/%d; RequiredMemory = %.3f MB + %.3f MB MemoryLimit = %.3f MB "),
 					Concurrency,DynamicMaxConcurrency,
