Clarifying Sample Rate Quality tooltip.

diff --git a/Engine/Source/Runtime/Engine/Classes/Sound/SoundWave.h b/Engine/Source/Runtime/Engine/Classes/Sound/SoundWave.h
--- a/Engine/Source/Runtime/Engine/Classes/Sound/SoundWave.h
+++ b/Engine/Source/Runtime/Engine/Classes/Sound/SoundWave.h
@@ -446,7 +446,10 @@ public:
 	UPROPERTY(meta = (DeprecatedProperty, DeprecationMessage = "5.0 - Property is deprecated. Streaming priority has no effect with stream caching enabled."))
 	int32 StreamingPriority;
 
-	/** Quality of sample rate conversion for platforms that opt into resampling during cook. The sample rate for each enumeration is definable per platform in platform target settings. */
+	/** Determines the max sample rate to use if the platform enables "Resampling For Device" in project settings. 
+	*	For example, if the platform enables Resampling For Device and specifies 32000 for High, then setting High here will
+	*	force the sound wave to be _at most_ 32000. Does nothing if Resampling For Device is disabled.
+	*/
 	UPROPERTY(EditAnywhere, Category = "Format|Quality")
 	ESoundwaveSampleRateSettings SampleRateQuality;
 
