Fix edge cases around re-importing meshes with bones inserted into the hierarchy

This adds a CVar (a.Skeleton.AllowIncompatibleSkeletalMeshMerge) to enable the addition of bones into the middle of skeleton hierarchies via import. In general this is not safe to do and will likely result in deformed meshes and corrupt animations. In the case this is being built for, the affected bones are not skinned to any existing meshes, and the number of animations that will also need re-importing is limited.

This splits USkeleton::MergeBonesToBoneTree into two modification passes:
- First, any mesh bones that already dont exist in the skeleton are added
- Second, any bones with differing parents are fixed up

Fixed precision issues with FReferenceSkeleton::SetParent where non-normalized rotations could be generated in the reference pose.
Fixed skeleton replacement not causing a recompression when skeleton bone indices could change (added new model notification).

Also fixed up incorrect error message on reimport of differing hierarchies

diff --git a/Engine/Plugins/Animation/AnimationData/Source/AnimationData/Private/AnimSequencerController.cpp b/Engine/Plugins/Animation/AnimationData/Source/AnimationData/Private/AnimSequencerController.cpp
--- a/Engine/Plugins/Animation/AnimationData/Source/AnimationData/Private/AnimSequencerController.cpp
+++ b/Engine/Plugins/Animation/AnimationData/Source/AnimationData/Private/AnimSequencerController.cpp
@@ -2252,6 +2252,9 @@ void UAnimSequencerController::UpdateWithSkeleton(USkeleton* TargetSkeleton, boo
 
 			// Forcefully re-generate legacy data structures
 			Model->RegenerateLegacyCurveData();
+
+			// Notify of skeleton change
+			Model->GetNotifier().Notify(EAnimDataModelNotifyType::SkeletonChanged);
 		}
 		CloseBracket();
 	}	
diff --git a/Engine/Source/Developer/AnimationDataController/Private/AnimDataController.cpp b/Engine/Source/Developer/AnimationDataController/Private/AnimDataController.cpp
--- a/Engine/Source/Developer/AnimationDataController/Private/AnimDataController.cpp
+++ b/Engine/Source/Developer/AnimationDataController/Private/AnimDataController.cpp
@@ -2057,7 +2057,14 @@ bool UAnimDataController::DuplicateAttribute(const FAnimationAttributeIdentifier
 
 void UAnimDataController::UpdateWithSkeleton(USkeleton* TargetSkeleton, bool bShouldTransact)
 {
-	RemoveBoneTracksMissingFromSkeleton(TargetSkeleton);
+	OpenBracket(LOCTEXT("SettingNewskeleton", "Updating Skeleton for Animation Data Model"), bShouldTransact);
+	{
+		RemoveBoneTracksMissingFromSkeleton(TargetSkeleton);
+
+		// Notify of skeleton change
+		Model->GetNotifier().Notify(EAnimDataModelNotifyType::SkeletonChanged);
+	}
+	CloseBracket();
 }
 
 void UAnimDataController::PopulateWithExistingModel(TScriptInterface<IAnimationDataModel> InModel)
diff --git a/Engine/Source/Editor/UnrealEd/Private/ImportUtils/SkeletalMeshImportUtils.cpp b/Engine/Source/Editor/UnrealEd/Private/ImportUtils/SkeletalMeshImportUtils.cpp
--- a/Engine/Source/Editor/UnrealEd/Private/ImportUtils/SkeletalMeshImportUtils.cpp
+++ b/Engine/Source/Editor/UnrealEd/Private/ImportUtils/SkeletalMeshImportUtils.cpp
@@ -86,8 +86,8 @@ bool SkeletalMesUtilsImpl::SkeletonsAreCompatible(const FReferenceSkeleton& NewS
 				if (!bFailNoError)
 				{
 					UnFbx::FFbxImporter* FFbxImporter = UnFbx::FFbxImporter::GetInstance();
-					FFbxImporter->AddTokenizedErrorMessage(FTokenizedMessage::Create(EMessageSeverity::Error, FText::Format(LOCTEXT("MeshHasDifferentRoot", "Root Bone is '{0}' instead of '{1}'.\nDiscarding existing LODs."),
-						FText::FromName(NewBoneName), FText::FromName(NewParentName))), FFbxErrors::SkeletalMesh_DifferentRoots);
+					FFbxImporter->AddTokenizedErrorMessage(FTokenizedMessage::Create(EMessageSeverity::Error, FText::Format(LOCTEXT("MeshHasDifferentParent", "Parent Bone of '{0}' is '{1}' instead of '{2}'. Discarding existing LODs."),
+						FText::FromName(NewBoneName), FText::FromName(NewParentName), FText::FromName(ExistParentName))), FFbxErrors::SkeletalMesh_DifferentRoots);
 				}
 				return false;
 			}
diff --git a/Engine/Source/Runtime/Engine/Classes/Animation/AnimData/AnimDataNotifications.h b/Engine/Source/Runtime/Engine/Classes/Animation/AnimData/AnimDataNotifications.h
--- a/Engine/Source/Runtime/Engine/Classes/Animation/AnimData/AnimDataNotifications.h
+++ b/Engine/Source/Runtime/Engine/Classes/Animation/AnimData/AnimDataNotifications.h
@@ -78,6 +78,9 @@ enum class EAnimDataModelNotifyType : uint8
 	/** Indicates all data stored on the model has been reset. Type of payload: FEmptyPayload */
 	Reset,
 
+	/** Indicates that the skeleton changed. Type of payload: FEmptyPayload */
+	SkeletonChanged,
+
 	Invalid // The max for this enum (used for guarding)
 };
 
diff --git a/Engine/Source/Runtime/Engine/Private/Animation/AnimSequence.cpp b/Engine/Source/Runtime/Engine/Private/Animation/AnimSequence.cpp
--- a/Engine/Source/Runtime/Engine/Private/Animation/AnimSequence.cpp
+++ b/Engine/Source/Runtime/Engine/Private/Animation/AnimSequence.cpp
@@ -4362,7 +4362,7 @@ void UAnimSequence::OnModelModified(const EAnimDataModelNotifyType& NotifyType,
 				const auto LengthChangingNotifies = { EAnimDataModelNotifyType::SequenceLengthChanged, EAnimDataModelNotifyType::FrameRateChanged, EAnimDataModelNotifyType::Reset };
 				const auto ResamplingNotifies = { EAnimDataModelNotifyType::TrackAdded, EAnimDataModelNotifyType::TrackChanged, EAnimDataModelNotifyType::TrackRemoved,  EAnimDataModelNotifyType::Populated };
 				const auto RecompressNotifies = { EAnimDataModelNotifyType::CurveAdded, EAnimDataModelNotifyType::CurveChanged, EAnimDataModelNotifyType::CurveRemoved, EAnimDataModelNotifyType::CurveFlagsChanged, EAnimDataModelNotifyType::CurveScaled,
-				EAnimDataModelNotifyType::AttributeAdded, EAnimDataModelNotifyType::AttributeChanged, EAnimDataModelNotifyType::AttributeRemoved };
+				EAnimDataModelNotifyType::AttributeAdded, EAnimDataModelNotifyType::AttributeChanged, EAnimDataModelNotifyType::AttributeRemoved, EAnimDataModelNotifyType::SkeletonChanged };
 
 				bShouldMarkPackageDirty = NotifyCollector.WasDataModified();
 				
@@ -4451,6 +4451,16 @@ void UAnimSequence::OnModelModified(const EAnimDataModelNotifyType& NotifyType,
 			UpdateCompressedCurveName(TypedPayload.Identifier.CurveName, TypedPayload.NewIdentifier.CurveName);
 			break;
 		}
+
+		case EAnimDataModelNotifyType::SkeletonChanged:
+		{
+			if (NotifyCollector.IsNotWithinBracket())
+			{
+				HandleTrackDataChanged();
+			}
+			break;
+		}
+
 		default:
 		{
 			checkf(false, TEXT("Missing case statement for animation model notify type"));
diff --git a/Engine/Source/Runtime/Engine/Private/Animation/ReferenceSkeleton.cpp b/Engine/Source/Runtime/Engine/Private/Animation/ReferenceSkeleton.cpp
--- a/Engine/Source/Runtime/Engine/Private/Animation/ReferenceSkeleton.cpp
+++ b/Engine/Source/Runtime/Engine/Private/Animation/ReferenceSkeleton.cpp
@@ -253,6 +253,7 @@ namespace FReferenceSkeletonLocals
 					if (Element.Parent)
 					{
 						Transforms[Element.RawIndex] *= GetGlobalTransformArg(*Element.Parent, GetGlobalTransformArg);
+						Transforms[Element.RawIndex].NormalizeRotation();
 					}
 
 					TransformCached[Element.RawIndex] = true;
@@ -424,6 +425,7 @@ int32 FReferenceSkeleton::SetParent(const FName InBoneName, const FName InParent
 		if (NewParentIdx != INDEX_NONE)
 		{
 			Pose = Transforms[NewIndex].GetRelativeTransform(Transforms[Element.Parent->RawIndex]);
+			Pose.NormalizeRotation();
 		}
 
 		// update name to index map
diff --git a/Engine/Source/Runtime/Engine/Private/Animation/Skeleton.cpp b/Engine/Source/Runtime/Engine/Private/Animation/Skeleton.cpp
--- a/Engine/Source/Runtime/Engine/Private/Animation/Skeleton.cpp
+++ b/Engine/Source/Runtime/Engine/Private/Animation/Skeleton.cpp
@@ -66,6 +66,8 @@ const FName USkeleton::AnimTrackCurveMappingName = FName(TEXT("AnimationTrackCur
 const FName FAnimSlotGroup::DefaultGroupName = FName(TEXT("DefaultGroup"));
 const FName FAnimSlotGroup::DefaultSlotName = FName(TEXT("DefaultSlot"));
 
+TAutoConsoleVariable<bool> CVarAllowIncompatibleSkeletalMeshMerge(TEXT("a.Skeleton.AllowIncompatibleSkeletalMeshMerge"), 0, TEXT("When importing or otherwise merging in skeletal mesh bones, allow 'incompatible' hierarchies with bone insertions."));
+
 void SerializeReferencePose(FArchive& Ar, FReferencePose& P, UObject* Outer)
 {
 	Ar.UsingCustomVersion(FAnimPhysObjectVersion::GUID);
@@ -1024,6 +1026,7 @@ bool USkeleton::MergeBonesToBoneTree(const USkinnedAsset* InSkinnedAsset, const
 	// see if it needs all animation data to remap - only happens when bone structure CHANGED - added
 	bool bSuccess = false;
 	bool bShouldHandleHierarchyChange = false;
+
 	// clear cache data since it won't work anymore once this is done
 	ClearCacheData();
 
@@ -1035,38 +1038,93 @@ bool USkeleton::MergeBonesToBoneTree(const USkinnedAsset* InSkinnedAsset, const
 	}
 	else
 	{
-		// can we play? - hierarchy matches
-		if( IsCompatibleMesh(InSkinnedAsset) )
+		// Check if we can merge in the bones
+		const bool bDoParentChainCheck = !CVarAllowIncompatibleSkeletalMeshMerge.GetValueOnGameThread();
+		if( IsCompatibleMesh(InSkinnedAsset, bDoParentChainCheck) )
 		{
 			TRACE_CPUPROFILER_EVENT_SCOPE(USkeleton::MergeBonesToBoneTree::CompatibleBranch);
 			// Exclude bones who do not have a parent.
 			TArray<int32> FilteredRequiredBones;
 			FAnimationRuntime::ExcludeBonesWithNoParents(RequiredRefBones, InSkinnedAsset->GetRefSkeleton(), FilteredRequiredBones);
 
-			FReferenceSkeletonModifier RefSkelModifier(ReferenceSkeleton, this);
-
-			for (int32 Index=0; Index<FilteredRequiredBones.Num(); Index++)
+			// Two modifier passes: add bones that dont already exist, then set bone parents that have changed
 			{
-				const int32& MeshBoneIndex = FilteredRequiredBones[Index];
-				const int32& SkeletonBoneIndex = ReferenceSkeleton.FindRawBoneIndex(InSkinnedAsset->GetRefSkeleton().GetBoneName(MeshBoneIndex));
-				
-				// Bone doesn't already exist. Add it.
-				if( SkeletonBoneIndex == INDEX_NONE )
+				FReferenceSkeletonModifier RefSkelModifier(ReferenceSkeleton, this);
+
+				// Check for bone's existence
+				for (const int32 MeshBoneIndex : FilteredRequiredBones)
 				{
-					FMeshBoneInfo NewMeshBoneInfo = InSkinnedAsset->GetRefSkeleton().GetRefBoneInfo()[MeshBoneIndex];
-					// Fix up ParentIndex for our new Skeleton.
-					if( ReferenceSkeleton.GetRawBoneNum() == 0 )
+					const int32 SkeletonBoneIndex = ReferenceSkeleton.FindRawBoneIndex(InSkinnedAsset->GetRefSkeleton().GetBoneName(MeshBoneIndex));
+					
+					// Bone doesn't already exist. Add it.
+					if( SkeletonBoneIndex == INDEX_NONE )
 					{
-						NewMeshBoneInfo.ParentIndex = INDEX_NONE; // root
+						FMeshBoneInfo NewMeshBoneInfo = InSkinnedAsset->GetRefSkeleton().GetRefBoneInfo()[MeshBoneIndex];
+						// Fix up ParentIndex for our new Skeleton.
+						if( ReferenceSkeleton.GetRawBoneNum() == 0 )
+						{
+							NewMeshBoneInfo.ParentIndex = INDEX_NONE; // root
+						}
+						else
+						{
+							NewMeshBoneInfo.ParentIndex = ReferenceSkeleton.FindRawBoneIndex(InSkinnedAsset->GetRefSkeleton().GetBoneName(InSkinnedAsset->GetRefSkeleton().GetParentIndex(MeshBoneIndex)));
+						}
+
+						RefSkelModifier.Add(NewMeshBoneInfo, InSkinnedAsset->GetRefSkeleton().GetRefBonePose()[MeshBoneIndex]);
+						BoneTree.AddDefaulted(1);
+						bShouldHandleHierarchyChange = true;
 					}
-					else
+				}
+			}
+
+			{
+				FReferenceSkeletonModifier RefSkelModifier(ReferenceSkeleton, this);
+
+				TMap<FName, FBoneNode> NameToBoneNode;
+				const int32 NumBones = BoneTree.Num();
+				check(NumBones == ReferenceSkeleton.GetRawRefBoneInfo().Num());
+
+				// Check for different parents
+				for (const int32 MeshBoneIndex : FilteredRequiredBones)
+				{
+					const FName BoneName = InSkinnedAsset->GetRefSkeleton().GetBoneName(MeshBoneIndex);
+					const int32 SkeletonBoneIndex = ReferenceSkeleton.FindRawBoneIndex(BoneName);
+
+					if(SkeletonBoneIndex != INDEX_NONE)
 					{
-						NewMeshBoneInfo.ParentIndex = ReferenceSkeleton.FindRawBoneIndex(InSkinnedAsset->GetRefSkeleton().GetBoneName(InSkinnedAsset->GetRefSkeleton().GetParentIndex(MeshBoneIndex)));
+						// Bone exists, check if it is in the same place in the hierarchy
+						const int32 MeshParentIndex = InSkinnedAsset->GetRefSkeleton().GetParentIndex(MeshBoneIndex);
+						const int32 SkeletonParentIndex = ReferenceSkeleton.GetRawParentIndex(SkeletonBoneIndex);
+						const FName MeshParentName = MeshParentIndex != INDEX_NONE ? InSkinnedAsset->GetRefSkeleton().GetBoneName(MeshParentIndex) : NAME_None;
+						const FName SkeletonParentName = SkeletonParentIndex != INDEX_NONE ? ReferenceSkeleton.GetRawRefBoneInfo()[SkeletonParentIndex].Name : NAME_None;
+
+						if(MeshParentName != SkeletonParentName)
+						{
+							// Cache the bone tree if we are going to change the structure
+							if(NameToBoneNode.Num() == 0)
+							{
+								NameToBoneNode.Reserve(NumBones);
+								for(int32 BoneTreeIndex = 0; BoneTreeIndex < NumBones; ++BoneTreeIndex)
+								{
+									NameToBoneNode.Add(ReferenceSkeleton.GetRawRefBoneInfo()[BoneTreeIndex].Name, BoneTree[BoneTreeIndex]);
+								}
+							}
+
+							RefSkelModifier.SetParent(BoneName, MeshParentName);
+							bShouldHandleHierarchyChange = true;
+						}
 					}
+				}
 
-					RefSkelModifier.Add(NewMeshBoneInfo, InSkinnedAsset->GetRefSkeleton().GetRefBonePose()[MeshBoneIndex]);
-					BoneTree.AddZeroed(1);
-					bShouldHandleHierarchyChange = true;
+				if(NameToBoneNode.Num() > 0)
+				{
+					// Setting parent can re-order bones, so the skeleton's BoneTree needs to update to reflect that
+					BoneTree.Reset();
+					for(int32 BoneTreeIndex = 0; BoneTreeIndex < NumBones; ++BoneTreeIndex)
+					{
+						FName BoneName = ReferenceSkeleton.GetRawRefBoneInfo()[BoneTreeIndex].Name;
+						BoneTree.Add(NameToBoneNode.FindChecked(BoneName));
+					}
 				}
 			}
 
