Expose some functionality needed for more reliable collision interactions with the ImmediatePhysics simulation (RBAN). Changes needed to support PhysicsControl/RigidBodyWithControl development.

diff --git a/Engine/Source/Runtime/Engine/Classes/PhysicsEngine/BodyInstance.h b/Engine/Source/Runtime/Engine/Classes/PhysicsEngine/BodyInstance.h
--- a/Engine/Source/Runtime/Engine/Classes/PhysicsEngine/BodyInstance.h
+++ b/Engine/Source/Runtime/Engine/Classes/PhysicsEngine/BodyInstance.h
@@ -973,6 +973,12 @@ public:
 	/** Get current transform in world space from physics body. */
 	ENGINE_API FTransform GetUnrealWorldTransform_AssumesLocked(bool bWithProjection = true, bool bForceGlobalPose = false) const;
 
+	/** Get the kinematic target transform in world space from physics body. Will only be relevant/useful if the body is kinematic */
+	ENGINE_API FTransform GetKinematicTarget() const;
+
+	/** Get the kinematic target transform in world space from physics body. Will only be relevant/useful if the body is kinematic */
+	ENGINE_API FTransform GetKinematicTarget_AssumesLocked() const;
+
 	/**
 	 *	Move the physics body to a new pose.
 	 *	@param	bTeleport	If true, no velocity is inferred on the kinematic body from this movement, but it moves right away.
diff --git a/Engine/Source/Runtime/Engine/Private/PhysicsEngine/BodyInstance.cpp b/Engine/Source/Runtime/Engine/Private/PhysicsEngine/BodyInstance.cpp
--- a/Engine/Source/Runtime/Engine/Private/PhysicsEngine/BodyInstance.cpp
+++ b/Engine/Source/Runtime/Engine/Private/PhysicsEngine/BodyInstance.cpp
@@ -2613,6 +2613,27 @@ FVector FBodyInstance::GetUnrealWorldVelocity() const
 	return OutVelocity;
 }
 
+FTransform FBodyInstance::GetKinematicTarget_AssumesLocked() const
+{
+	FTransform TM;
+	if (FPhysicsInterface::IsValid(ActorHandle))
+	{
+		TM = FPhysicsInterface::GetKinematicTarget_AssumesLocked(ActorHandle);
+	}
+	return TM;
+}
+
+FTransform FBodyInstance::GetKinematicTarget() const
+{
+	FTransform TM;
+	FPhysicsCommand::ExecuteRead(ActorHandle, [&](const FPhysicsActorHandle& Actor)
+		{
+			TM = FPhysicsInterface::GetKinematicTarget_AssumesLocked(Actor);
+		});
+
+	return TM;
+}
+
 FVector FBodyInstance::GetUnrealWorldVelocity_AssumesLocked() const
 {
 	FVector LinVel(EForceInit::ForceInitToZero);
diff --git a/Engine/Source/Runtime/Engine/Private/PhysicsEngine/ImmediatePhysics/ImmediatePhysicsChaos/ImmediatePhysicsActorHandle_Chaos.cpp b/Engine/Source/Runtime/Engine/Private/PhysicsEngine/ImmediatePhysics/ImmediatePhysicsChaos/ImmediatePhysicsActorHandle_Chaos.cpp
--- a/Engine/Source/Runtime/Engine/Private/PhysicsEngine/ImmediatePhysics/ImmediatePhysicsChaos/ImmediatePhysicsActorHandle_Chaos.cpp
+++ b/Engine/Source/Runtime/Engine/Private/PhysicsEngine/ImmediatePhysics/ImmediatePhysicsChaos/ImmediatePhysicsActorHandle_Chaos.cpp
@@ -16,7 +16,7 @@
 #include "PhysicsEngine/BodyUtils.h"
 #include "PhysicsProxy/SingleParticlePhysicsProxy.h"
 
-//PRAGMA_DISABLE_OPTIMIZATION
+//UE_DISABLE_OPTIMIZATION
 
 namespace ImmediatePhysics_Chaos
 {
@@ -336,6 +336,15 @@ namespace ImmediatePhysics_Chaos
 		}
 	}
 
+	bool FActorHandle::GetHasCollision() const
+	{
+		return ParticleHandle->HasCollision();
+	}
+
+	void FActorHandle::SetHasCollision(bool bCollision)
+	{
+		ParticleHandle->SetHasCollision(bCollision);
+	}
 
 	void FActorHandle::InitWorldTransform(const FTransform& WorldTM)
 	{
@@ -453,6 +462,11 @@ namespace ImmediatePhysics_Chaos
 		return ParticleHandle->CastToRigidParticle() != nullptr && ParticleHandle->ObjectState() == Chaos::EObjectStateType::Dynamic;
 	}
 
+	bool FActorHandle::CouldBeDynamic() const
+	{
+		return ParticleHandle->CastToRigidParticle() != nullptr;
+	}
+
 	bool FActorHandle::IsGravityEnabled() const
 	{
 		using namespace Chaos;
diff --git a/Engine/Source/Runtime/Engine/Private/PhysicsEngine/ImmediatePhysics/ImmediatePhysicsChaos/ImmediatePhysicsJointHandle_Chaos.cpp b/Engine/Source/Runtime/Engine/Private/PhysicsEngine/ImmediatePhysics/ImmediatePhysicsChaos/ImmediatePhysicsJointHandle_Chaos.cpp
--- a/Engine/Source/Runtime/Engine/Private/PhysicsEngine/ImmediatePhysics/ImmediatePhysicsChaos/ImmediatePhysicsJointHandle_Chaos.cpp
+++ b/Engine/Source/Runtime/Engine/Private/PhysicsEngine/ImmediatePhysics/ImmediatePhysicsChaos/ImmediatePhysicsJointHandle_Chaos.cpp
@@ -9,7 +9,7 @@
 
 #include "PhysicsEngine/ConstraintInstance.h"
 
-//PRAGMA_DISABLE_OPTIMIZATION
+//UE_DISABLE_OPTIMIZATION
 
 static_assert((int32)Chaos::EJointMotionType::Free == (int32)EAngularConstraintMotion::ACM_Free, "Chaos::EJointMotionType and EAngularConstraintMotion mismatch");
 static_assert((int32)Chaos::EJointMotionType::Limited == (int32)EAngularConstraintMotion::ACM_Limited, "Chaos::EJointMotionType and EAngularConstraintMotion mismatch");
diff --git a/Engine/Source/Runtime/Engine/Private/PhysicsEngine/ImmediatePhysics/ImmediatePhysicsChaos/ImmediatePhysicsSimulation_Chaos.cpp b/Engine/Source/Runtime/Engine/Private/PhysicsEngine/ImmediatePhysics/ImmediatePhysicsChaos/ImmediatePhysicsSimulation_Chaos.cpp
--- a/Engine/Source/Runtime/Engine/Private/PhysicsEngine/ImmediatePhysics/ImmediatePhysicsChaos/ImmediatePhysicsSimulation_Chaos.cpp
+++ b/Engine/Source/Runtime/Engine/Private/PhysicsEngine/ImmediatePhysics/ImmediatePhysicsChaos/ImmediatePhysicsSimulation_Chaos.cpp
@@ -19,8 +19,7 @@
 #include "ChaosVisualDebugger/ChaosVisualDebuggerTrace.h"
 #include "Stats/StatsTrace.h"
 
-
-//PRAGMA_DISABLE_OPTIMIZATION
+//UE_DISABLE_OPTIMIZATION
 
 DECLARE_CYCLE_STAT(TEXT("FSimulation::Simulate_Chaos"), STAT_ImmediateSimulate_Chaos, STATGROUP_ImmediatePhysics);
 DECLARE_CYCLE_STAT(TEXT("FSimulation::Simulate_Chaos::InertiaConditioning"), STAT_ImmediateSimulate_Chaos_InertiaConditioning, STATGROUP_ImmediatePhysics);
@@ -503,6 +502,36 @@ namespace ImmediatePhysics_Chaos
         Implementation->Collisions.GetConstraintAllocator().RemoveParticle(ActorHandle->GetParticle());
     }
 
+	void FSimulation::SetIsKinematic(FActorHandle* ActorHandle, bool bKinematic)
+	{
+		bool bWasKinematic = ActorHandle->GetIsKinematic();
+		if (bKinematic != bWasKinematic)
+		{
+			ActorHandle->SetIsKinematic(bKinematic);
+			Implementation->bActorsDirty = true;
+		}
+	}
+
+	void FSimulation::SetEnabled(FActorHandle* ActorHandle, bool bEnable)
+	{
+		bool bWasEnabled = ActorHandle->GetEnabled();
+		if (bEnable != bWasEnabled)
+		{
+			ActorHandle->SetEnabled(bEnable);
+			Implementation->bActorsDirty = true;
+		}
+	}
+
+	void FSimulation::SetHasCollision(FActorHandle* ActorHandle, bool bCollision)
+	{
+		bool bWasCollision = ActorHandle->GetHasCollision();
+		if (bCollision != bWasCollision)
+		{
+			ActorHandle->SetHasCollision(bCollision);
+			Implementation->bActorsDirty = true;
+		}
+	}
+
 	FJointHandle* FSimulation::CreateJoint(FConstraintInstance* ConstraintInstance, FActorHandle* Body1, FActorHandle* Body2)
 	{
 		FJointHandle* JointHandle = new FJointHandle(&Implementation->Joints, ConstraintInstance, Body1, Body2);
@@ -577,7 +606,7 @@ namespace ImmediatePhysics_Chaos
 		for (FActorHandle* OtherActorHandle : Implementation->ActorHandles)
 		{
 			FGeometryParticleHandle* Particle1 = OtherActorHandle->GetParticle();
-			if ((OtherActorHandle != ActorHandle) && OtherActorHandle->IsSimulated())
+			if ((OtherActorHandle != ActorHandle) && OtherActorHandle->CouldBeDynamic())
 			{
 				Implementation->PotentiallyCollidingPairs.Emplace(FParticlePair(Particle0, Particle1));
 			}
@@ -697,7 +726,8 @@ namespace ImmediatePhysics_Chaos
 			{
 				bool bAnyDisabled = FGenericParticleHandle(ParticlePair[0])->Disabled() || FGenericParticleHandle(ParticlePair[1])->Disabled();
 				bool bAnyDynamic = FGenericParticleHandle(ParticlePair[0])->IsDynamic() || FGenericParticleHandle(ParticlePair[1])->IsDynamic();
-				if (bAnyDynamic && !bAnyDisabled)
+				bool bBothCollide = ParticlePair[0]->HasCollision() && ParticlePair[1]->HasCollision();
+				if (bBothCollide && bAnyDynamic && !bAnyDisabled)
 				{
 					Implementation->ActivePotentiallyCollidingPairs.Add(ParticlePair);
 				}
diff --git a/Engine/Source/Runtime/Engine/Public/Physics/ImmediatePhysics/ImmediatePhysicsChaos/ImmediatePhysicsActorHandle_Chaos.h b/Engine/Source/Runtime/Engine/Public/Physics/ImmediatePhysics/ImmediatePhysicsChaos/ImmediatePhysicsActorHandle_Chaos.h
--- a/Engine/Source/Runtime/Engine/Public/Physics/ImmediatePhysics/ImmediatePhysicsChaos/ImmediatePhysicsActorHandle_Chaos.h
+++ b/Engine/Source/Runtime/Engine/Public/Physics/ImmediatePhysics/ImmediatePhysicsChaos/ImmediatePhysicsActorHandle_Chaos.h
@@ -20,8 +20,20 @@ namespace ImmediatePhysics_Chaos
 
 		ENGINE_API bool GetEnabled() const;
 
+		/**
+		 * Note that this will not update the colliding pairs in the simulation, so better to change the kinematic
+		 * state there.
+		 */
 		ENGINE_API void SetEnabled(bool bEnabled);
 
+		ENGINE_API bool GetHasCollision() const;
+
+		/**
+		 * Note that this will not update the colliding pairs in the simulation, so better to change the kinematic
+		 * state there.
+		 */
+		ENGINE_API void SetHasCollision(bool bHasCollision);
+
 		/** Sets the world transform, zeroes velocity, etc.*/
 		ENGINE_API void InitWorldTransform(const FTransform& WorldTM);
 
@@ -33,6 +45,9 @@ namespace ImmediatePhysics_Chaos
 		 * was possible (even if nothing changed). Returns false if it fails - e.g. because there was no current 
 		 * ParticleHandle, or if the particle type was not EParticleType::Rigid (unless the current type was
 		 * EParticleType::Kinematic and bKinematic was true)
+		 * 
+		 * Note that this will not update the colliding pairs in the simulation, so better to change the kinematic
+		 * state there.
 		 */
 		ENGINE_API bool SetIsKinematic(bool bKinematic);
 
@@ -54,6 +69,9 @@ namespace ImmediatePhysics_Chaos
 		/** Whether the body is static */
 		ENGINE_API bool IsStatic() const;
 
+		/** Returns true the body is or could be dynamic (i.e. currently simulated, or a rigid body that is currently kinematic) */
+		ENGINE_API bool CouldBeDynamic() const;
+
 		/** Sets whether world gravity should apply to this actor, assuming it is dynamic */
 		ENGINE_API void SetGravityEnabled(bool bEnable);
 
diff --git a/Engine/Source/Runtime/Engine/Public/Physics/ImmediatePhysics/ImmediatePhysicsChaos/ImmediatePhysicsSimulation_Chaos.h b/Engine/Source/Runtime/Engine/Public/Physics/ImmediatePhysics/ImmediatePhysicsChaos/ImmediatePhysicsSimulation_Chaos.h
--- a/Engine/Source/Runtime/Engine/Public/Physics/ImmediatePhysics/ImmediatePhysicsChaos/ImmediatePhysicsSimulation_Chaos.h
+++ b/Engine/Source/Runtime/Engine/Public/Physics/ImmediatePhysics/ImmediatePhysicsChaos/ImmediatePhysicsSimulation_Chaos.h
@@ -41,6 +41,12 @@ namespace ImmediatePhysics_Chaos
 
 		ENGINE_API void DestroyActorCollisions(FActorHandle* ActorHandle);
 
+		ENGINE_API void SetIsKinematic(FActorHandle* ActorHandle, bool bKinematic);
+
+		ENGINE_API void SetEnabled(FActorHandle* ActorHandle, bool bEnable);
+
+		ENGINE_API void SetHasCollision(FActorHandle* ActorHandle, bool bHasCollision);
+
 		/** Create a physical joint and add it to the simulation */
 		ENGINE_API FJointHandle* CreateJoint(FConstraintInstance* ConstraintInstance, FActorHandle* Body1, FActorHandle* Body2);
 
diff --git a/Engine/Source/Runtime/Experimental/Chaos/Private/Chaos/Collision/ParticlePairMidPhase.cpp b/Engine/Source/Runtime/Experimental/Chaos/Private/Chaos/Collision/ParticlePairMidPhase.cpp
--- a/Engine/Source/Runtime/Experimental/Chaos/Private/Chaos/Collision/ParticlePairMidPhase.cpp
+++ b/Engine/Source/Runtime/Experimental/Chaos/Private/Chaos/Collision/ParticlePairMidPhase.cpp
@@ -20,6 +20,8 @@
 #include "Chaos/ConvexOptimizer.h"
 #include "ProfilingDebugging/CountersTrace.h"
 
+//UE_DISABLE_OPTIMIZATION
+
 TRACE_DECLARE_INT_COUNTER_EXTERN(ChaosTraceCounter_MidPhase_NumShapePair);
 TRACE_DECLARE_INT_COUNTER_EXTERN(ChaosTraceCounter_MidPhase_NumGeneric);
 
diff --git a/Engine/Source/Runtime/Experimental/Chaos/Private/Chaos/Joint/PBDJointContainerSolver.cpp b/Engine/Source/Runtime/Experimental/Chaos/Private/Chaos/Joint/PBDJointContainerSolver.cpp
--- a/Engine/Source/Runtime/Experimental/Chaos/Private/Chaos/Joint/PBDJointContainerSolver.cpp
+++ b/Engine/Source/Runtime/Experimental/Chaos/Private/Chaos/Joint/PBDJointContainerSolver.cpp
@@ -150,7 +150,14 @@ namespace Chaos
 			// If this triggers, Reset was called with the wrong constraint count
 			check(ContainerConstraintIndices.Num() < ContainerConstraintIndices.Max());
 
-			ContainerConstraintIndices.Add(InContainerConstraintIndex);
+			// Only add a constraint if it is working on at least one dynamic body
+			const FGenericParticleHandle Particle0 = GetJointParticle(ConstraintContainer, InContainerConstraintIndex, 0);
+			const FGenericParticleHandle Particle1 = GetJointParticle(ConstraintContainer, InContainerConstraintIndex, 1);
+
+			if (Particle0->IsDynamic() || Particle1->IsDynamic())
+			{
+				ContainerConstraintIndices.Add(InContainerConstraintIndex);
+			}
 		}
 
 		void FPBDJointContainerSolver::AddBodies(FSolverBodyContainer& SolverBodyContainer)
diff --git a/Engine/Source/Runtime/Experimental/Chaos/Private/Chaos/PBDJointConstraints.cpp b/Engine/Source/Runtime/Experimental/Chaos/Private/Chaos/PBDJointConstraints.cpp
--- a/Engine/Source/Runtime/Experimental/Chaos/Private/Chaos/PBDJointConstraints.cpp
+++ b/Engine/Source/Runtime/Experimental/Chaos/Private/Chaos/PBDJointConstraints.cpp
@@ -285,7 +285,7 @@ namespace Chaos
 
 	void FPBDJointConstraintHandle::SetConstraintEnabled(bool bInEnabled)
 	{
-		return ConcreteContainer()->SetConstraintEnabled(ConstraintIndex, bInEnabled);
+		ConcreteContainer()->SetConstraintEnabled(ConstraintIndex, bInEnabled);
 	}
 
 
