Change the description of UBluerpintExtension::HandleGatherSearchData.

diff --git a/Engine/Source/Runtime/Engine/Public/Blueprint/BlueprintExtension.h b/Engine/Source/Runtime/Engine/Public/Blueprint/BlueprintExtension.h
--- a/Engine/Source/Runtime/Engine/Public/Blueprint/BlueprintExtension.h
+++ b/Engine/Source/Runtime/Engine/Public/Blueprint/BlueprintExtension.h
@@ -89,7 +89,7 @@ private:
 	virtual void HandleGenerateFunctionGraphs(FKismetCompilerContext* CompilerContext) {}
 
 	/**
-	 * Override this function to define custom function generation logic for gathering search data
+	 * Override this function to define additional search data for the specified blueprint
 	 */
 	virtual FSearchData HandleGatherSearchData(const UBlueprint* OwningBlueprint) const { return FSearchData(); }
 
