World Partition:
- Streaming sources velocity is now a vector instead of a float.
- Streaming sources velocity is now frame rate independant using a moving average instead of an history buffer.
- Streaming sources can now use the velocity vector instead of the orientation to compute streaming cells priority.
- Fix streaming sources velocity not updated in local world partitions.
- Display streaming sources velocity vector in debug display.

This can be used for streaming sources attached to moving objects with a different orientation, like a player in a moving car, looking backwards.

diff --git a/Engine/Source/Runtime/Engine/Private/WorldPartition/RuntimeHashSet/WorldPartitionRuntimeHashSetDebugDraw.cpp b/Engine/Source/Runtime/Engine/Private/WorldPartition/RuntimeHashSet/WorldPartitionRuntimeHashSetDebugDraw.cpp
--- a/Engine/Source/Runtime/Engine/Private/WorldPartition/RuntimeHashSet/WorldPartitionRuntimeHashSetDebugDraw.cpp
+++ b/Engine/Source/Runtime/Engine/Private/WorldPartition/RuntimeHashSet/WorldPartitionRuntimeHashSetDebugDraw.cpp
@@ -219,14 +219,15 @@ bool UWorldPartitionRuntimeHashSet::Draw2D(FWorldPartitionDraw2DContext& DrawCon
 				const int32 MaxSegments = FMath::Max(4, FMath::CeilToInt(64 * Angle / 360.f));
 				const float AngleIncrement = Angle / MaxSegments;
 				const FVector2D Axis = FVector2D(Shape.GetAxis());
-				const FVector Startup = FRotator(0, -0.5f * Angle, 0).RotateVector(Shape.GetScaledAxis());
-
+				const FVector Startup = FRotator(0, -0.5f * Angle, 0).RotateVector(Shape.GetScaledAxis());				
 				FVector2D LineStart = FVector2D(Startup);
+
 				if (!Shape.IsSphere())
 				{
 					// Draw sector start axis
 					DrawContext.PushDrawSegment(GridScreenBounds, WorldToScreen(Center2D), WorldToScreen(Center2D + LineStart), Color, 2);
 				}
+
 				// Draw sector Arc
 				for (int32 i = 1; i <= MaxSegments; i++)
 				{
@@ -234,12 +235,22 @@ bool UWorldPartitionRuntimeHashSet::Draw2D(FWorldPartitionDraw2DContext& DrawCon
 					DrawContext.PushDrawSegment(GridScreenBounds, WorldToScreen(Center2D + LineStart), WorldToScreen(Center2D + LineEnd), Color, 2);
 					LineStart = LineEnd;
 				}
+
 				// If sphere, close circle, else draw sector end axis
 				DrawContext.PushDrawSegment(GridScreenBounds, WorldToScreen(Center2D + LineStart), WorldToScreen(Center2D + (Shape.IsSphere() ? FVector2D(Startup) : FVector2D::ZeroVector)), Color, 2);
 
 				// Draw direction vector
-				DrawContext.PushDrawSegment(GridScreenBounds, WorldToScreen(Center2D), WorldToScreen(Center2D + Axis * Shape.GetRadius()), Color, 2);
+				DrawContext.PushDrawSegment(GridScreenBounds, WorldToScreen(Center2D), WorldToScreen(Center2D + Axis * Shape.GetRadius()), Color, 4);
 			});
+
+			// Draw velocity vector			
+			const FVector2D Velocity2D = FVector2D(Source.Velocity);
+
+			if (!Velocity2D.IsNearlyZero())
+			{
+				const FVector2D Center2D = FVector2D(Source.Location);
+				DrawContext.PushDrawSegment(GridScreenBounds, WorldToScreen(Center2D), WorldToScreen(Center2D + Velocity2D * StreamingData->LoadingRange * 0.5f), Color, 1);
+			}
 		}
 	};
 
diff --git a/Engine/Source/Runtime/Engine/Private/WorldPartition/WorldPartitionReplay.cpp b/Engine/Source/Runtime/Engine/Private/WorldPartition/WorldPartitionReplay.cpp
--- a/Engine/Source/Runtime/Engine/Private/WorldPartition/WorldPartitionReplay.cpp
+++ b/Engine/Source/Runtime/Engine/Private/WorldPartition/WorldPartitionReplay.cpp
@@ -238,7 +238,7 @@ bool AWorldPartitionReplay::GetReplayStreamingSources(TArray<FWorldPartitionStre
 
 						const FVector Location = FMath::Lerp(Source1.Location, Source2.Location, LerpPercent);
 						const FQuat Rotation = FQuat::FastLerp(FQuat(Source1.Rotation), FQuat(Source2.Rotation), LerpPercent).GetNormalized();
-						const float Velocity = FMath::Lerp(Source1.Velocity, Source2.Velocity, LerpPercent);
+						const FVector Velocity = FMath::Lerp(Source1.Velocity, Source2.Velocity, LerpPercent);
 						OutStreamingSources.Add(FWorldPartitionReplayStreamingSource(Source1.Name, Location, Rotation.Rotator(), Source1.TargetState, Source1.bBlockOnSlowLoading, Source1.Priority, Velocity));
 					}
 					else
diff --git a/Engine/Source/Runtime/Engine/Private/WorldPartition/WorldPartitionRuntimeCellData.cpp b/Engine/Source/Runtime/Engine/Private/WorldPartition/WorldPartitionRuntimeCellData.cpp
--- a/Engine/Source/Runtime/Engine/Private/WorldPartition/WorldPartitionRuntimeCellData.cpp
+++ b/Engine/Source/Runtime/Engine/Private/WorldPartition/WorldPartitionRuntimeCellData.cpp
@@ -60,7 +60,8 @@ void UWorldPartitionRuntimeCellData::AppendStreamingSourceInfo(const FWorldParti
 
 	// Compute cosine angle from cell to source direction ratio
 	const FVector CellToSource = SourceShape.GetCenter() - ContentBounds.GetClosestPointTo(SourceShape.GetCenter());
-	const float SourceCosAngle = ContentBounds.IsInsideOrOn(SourceShape.GetCenter()) ? -1.0f : (SourceShape.GetAxis() | CellToSource.GetSafeNormal());
+	const FVector SourceAxis = FVector(SourceShape.GetAxis() + Source.Velocity * (Source.bUseVelocityContributionToCellsSorting ? 1.0f : 0.0f)).GetSafeNormal();
+	const float SourceCosAngle = ContentBounds.IsInsideOrOn(SourceShape.GetCenter()) ? -1.0f : (SourceAxis | CellToSource.GetSafeNormal());
 	const float SourceCosAngleRatio = SourceCosAngle * 0.5f + 0.5f;
 
 	// Compute final cell priority for this source
diff --git a/Engine/Source/Runtime/Engine/Private/WorldPartition/WorldPartitionStreamingPolicy.cpp b/Engine/Source/Runtime/Engine/Private/WorldPartition/WorldPartitionStreamingPolicy.cpp
--- a/Engine/Source/Runtime/Engine/Private/WorldPartition/WorldPartitionStreamingPolicy.cpp
+++ b/Engine/Source/Runtime/Engine/Private/WorldPartition/WorldPartitionStreamingPolicy.cpp
@@ -102,6 +102,15 @@ void UWorldPartitionStreamingPolicy::UpdateStreamingSources(bool bCanOptimizeUpd
 	const uint32 NewUpdateStreamingSourcesHash = WorldPartitionSubsystem->GetStreamingSourcesHash();
 	if (bCanOptimizeUpdate && (UpdateStreamingSourcesHash == NewUpdateStreamingSourcesHash))
 	{
+		TArray<FWorldPartitionStreamingSource> LocalStreamingSources;
+		WorldPartitionSubsystem->GetStreamingSources(WorldPartition, LocalStreamingSources);
+		check(LocalStreamingSources.Num() == StreamingSources.Num());
+		const FTransform WorldToLocal = WorldPartition->GetInstanceTransform().Inverse();
+		for (int32 i=0; i<LocalStreamingSources.Num(); i++)
+		{
+			check(StreamingSources[i].Name == LocalStreamingSources[i].Name);
+			StreamingSources[i].Velocity = WorldToLocal.TransformVector(LocalStreamingSources[i].Velocity);
+		}
 		return;
 	}
 
diff --git a/Engine/Source/Runtime/Engine/Private/WorldPartition/WorldPartitionStreamingSource.cpp b/Engine/Source/Runtime/Engine/Private/WorldPartition/WorldPartitionStreamingSource.cpp
--- a/Engine/Source/Runtime/Engine/Private/WorldPartition/WorldPartitionStreamingSource.cpp
+++ b/Engine/Source/Runtime/Engine/Private/WorldPartition/WorldPartitionStreamingSource.cpp
@@ -53,8 +53,8 @@ FString FWorldPartitionStreamingSource::ToString() const
 		bBlockOnSlowLoading ? TEXT("Blocking") : TEXT("NonBlocking"),
 		(int64)Location.X, (int64)Location.Y, (int64)Location.Z,
 		*Rotation.ToCompactString(),
-		Velocity,
-		(int32)(Velocity * 2.23694f)
+		Velocity.Size() * 0.01,
+		(int32)(Velocity.Size() * 0.0223694f)
 	);
 
 	if (Shapes.Num())
@@ -95,15 +95,15 @@ void FWorldPartitionStreamingSource::UpdateHash()
 {
 	// Update old values when they are changing enough, to avoid the case where we are on the edge of a quantization unit.
 	if (!UWorldPartitionStreamingPolicy::IsUpdateStreamingOptimEnabled() || 
-		(FVector::Dist(Location, OldLocation) > FWorldPartitionStreamingSource::LocationQuantization))
+		(FVector::Dist(Location, OldLocation) > LocationQuantization))
 	{
 		OldLocation = Location;
 	}
 	
 	if (!UWorldPartitionStreamingPolicy::IsUpdateStreamingOptimEnabled() ||
-		(FMath::Abs(Rotation.Pitch - OldRotation.Pitch) > FWorldPartitionStreamingSource::RotationQuantization) ||
-		(FMath::Abs(Rotation.Yaw - OldRotation.Yaw) > FWorldPartitionStreamingSource::RotationQuantization) ||
-		(FMath::Abs(Rotation.Roll - OldRotation.Roll) > FWorldPartitionStreamingSource::RotationQuantization))
+		(FMath::Abs(Rotation.Pitch - OldRotation.Pitch) > RotationQuantization) ||
+		(FMath::Abs(Rotation.Yaw - OldRotation.Yaw) > RotationQuantization) ||
+		(FMath::Abs(Rotation.Roll - OldRotation.Roll) > RotationQuantization))
 	{
 		OldRotation = Rotation;
 	}
diff --git a/Engine/Source/Runtime/Engine/Private/WorldPartition/WorldPartitionSubsystem.cpp b/Engine/Source/Runtime/Engine/Private/WorldPartition/WorldPartitionSubsystem.cpp
--- a/Engine/Source/Runtime/Engine/Private/WorldPartition/WorldPartitionSubsystem.cpp
+++ b/Engine/Source/Runtime/Engine/Private/WorldPartition/WorldPartitionSubsystem.cpp
@@ -959,7 +959,7 @@ void UWorldPartitionSubsystem::UpdateStreamingSources()
 	}
 
 	StreamingSourcesHash = 0;
-	const float CurrentTime = World->GetTimeSeconds();
+	const double CurrentTime = World->GetTimeSeconds();
 	for (FWorldPartitionStreamingSource& StreamingSource : StreamingSources)
 	{
 		// Update streaming sources velocity
@@ -1525,49 +1525,46 @@ void UWorldPartitionSubsystem::DrawStreamingStatusLegend(class UCanvas* Canvas,
 FStreamingSourceVelocity::FStreamingSourceVelocity(const FName& InSourceName)
 	: bIsValid(false)
 	, SourceName(InSourceName)
-	, LastIndex(INDEX_NONE)
 	, LastUpdateTime(-1.0)
-	, VelocityHistorySum(0.f)
-{
-	VelocityHistory.SetNumZeroed(VELOCITY_HISTORY_SAMPLE_COUNT);
-}
+{}
 
-float FStreamingSourceVelocity::GetAverageVelocity(const FVector& NewPosition, const float CurrentTime)
+FVector FStreamingSourceVelocity::GetAverageVelocity(const FVector& NewPosition, double CurrentTime)
 {
 	bIsValid = true;
 
-	const double TeleportDistance = 100;
-	const float MaxDeltaSeconds = 5.f;
-	const bool bIsFirstCall = (LastIndex == INDEX_NONE);
-	const float DeltaSeconds = bIsFirstCall ? 0.f : (CurrentTime - LastUpdateTime);
-	const double Distance = bIsFirstCall ? 0.f : ((NewPosition - LastPosition) * 0.01).Size();
-	if (bIsFirstCall)
+	const bool bNewSource = (LastUpdateTime <= 0.0);
+	const double DeltaSeconds = bNewSource ? 1.0 : (CurrentTime - LastUpdateTime);
+	LastUpdateTime = CurrentTime;
+
+	const FVector AbsMovement = (NewPosition - LastPosition);
+	const FVector AbsVelocity = AbsMovement / DeltaSeconds;
+	LastPosition = NewPosition;
+
+	const double TeleportDistance = 10000;
+	const double MaxDeltaSeconds = 5.0;
+	const double Distance = AbsMovement.Size();
+
+	if (bNewSource)
 	{
 		UE_LOG(LogWorldPartition, Verbose, TEXT("New Streaming Source: %s -> Position: %s"), *SourceName.ToString(), *NewPosition.ToString());
-		LastIndex = 0;
+		AvgVelocity = FVector::Zero();
 	}
-
-	ON_SCOPE_EXIT
+	else if (Distance > TeleportDistance)
 	{
-		LastUpdateTime = CurrentTime;
-		LastPosition = NewPosition;
-	};
-
-	// Handle invalid cases
-	if (bIsFirstCall || (DeltaSeconds <= 0.f) || (DeltaSeconds > MaxDeltaSeconds) || (Distance > TeleportDistance))
+		UE_LOG(LogWorldPartition, Verbose, TEXT("Detected Streaming Source Teleport: %s -> Last Position: %s -> New Position: %s"), *SourceName.ToString(), *LastPosition.ToString(), *NewPosition.ToString());
+		AvgVelocity = FVector::Zero();
+	}
+	else if  (DeltaSeconds > MaxDeltaSeconds)
 	{
-		UE_CLOG(Distance > TeleportDistance, LogWorldPartition, Verbose, TEXT("Detected Streaming Source Teleport: %s -> Last Position: %s -> New Position: %s"), *SourceName.ToString(), *LastPosition.ToString(), *NewPosition.ToString());
-		return 0.f;
+		UE_LOG(LogWorldPartition, Verbose, TEXT("Detected Inactive Streaming Source: %s -> Last Position: %s -> New Position: %s"), *SourceName.ToString(), *LastPosition.ToString(), *NewPosition.ToString());
+		AvgVelocity = FVector::Zero();
+	}
+	else
+	{
+		// Compute the new value in a weighted moving average series
+		const double AvgWeight = FMath::Clamp(DeltaSeconds * 100, 0, 1);
+		AvgVelocity = AvgVelocity * (1.0 - AvgWeight) + AbsVelocity * AvgWeight;
 	}
 
-	// Compute velocity (m/s)
-	ensureMsgf(Distance < MAX_flt, TEXT("Invalid distance %lf computed using position %s and position %s"), Distance, *NewPosition.ToString(), *LastPosition.ToString());
-	const float Velocity = (float)Distance / DeltaSeconds;
-	// Update velocity history buffer and sum
-	LastIndex = (LastIndex + 1) % VELOCITY_HISTORY_SAMPLE_COUNT;
-	VelocityHistorySum = FMath::Max<float>(0.f, (VelocityHistorySum + Velocity - VelocityHistory[LastIndex]));
-	VelocityHistory[LastIndex] = Velocity;
-
-	// return average
-	return (VelocityHistorySum / (float)VELOCITY_HISTORY_SAMPLE_COUNT);
+	return AvgVelocity;
 }
\ No newline at end of file
diff --git a/Engine/Source/Runtime/Engine/Public/WorldPartition/WorldPartitionReplay.h b/Engine/Source/Runtime/Engine/Public/WorldPartition/WorldPartitionReplay.h
--- a/Engine/Source/Runtime/Engine/Public/WorldPartition/WorldPartitionReplay.h
+++ b/Engine/Source/Runtime/Engine/Public/WorldPartition/WorldPartitionReplay.h
@@ -20,7 +20,7 @@ struct FWorldPartitionReplayStreamingSource : public FWorldPartitionStreamingSou
 	{
 	}
 
-	FWorldPartitionReplayStreamingSource(FName InName, const FVector& InLocation, const FRotator& InRotation, EStreamingSourceTargetState InTargetState, bool bInBlockOnSlowLoading, EStreamingSourcePriority InPriority, float InVelocity)
+	FWorldPartitionReplayStreamingSource(FName InName, const FVector& InLocation, const FRotator& InRotation, EStreamingSourceTargetState InTargetState, bool bInBlockOnSlowLoading, EStreamingSourcePriority InPriority, FVector InVelocity)
 		: FWorldPartitionStreamingSource(InName, InLocation, InRotation, InTargetState, bInBlockOnSlowLoading, InPriority, false, InVelocity)
 	{
 		bReplay = true;
diff --git a/Engine/Source/Runtime/Engine/Public/WorldPartition/WorldPartitionStreamingSource.h b/Engine/Source/Runtime/Engine/Public/WorldPartition/WorldPartitionStreamingSource.h
--- a/Engine/Source/Runtime/Engine/Public/WorldPartition/WorldPartitionStreamingSource.h
+++ b/Engine/Source/Runtime/Engine/Public/WorldPartition/WorldPartitionStreamingSource.h
@@ -340,7 +340,8 @@ struct FWorldPartitionStreamingSource
 		, TargetState(EStreamingSourceTargetState::Activated)
 		, bBlockOnSlowLoading(false)
 		, Priority(EStreamingSourcePriority::Default)
-		, Velocity(0.f)
+		, Velocity(FVector::Zero())
+		, bUseVelocityContributionToCellsSorting(false)
 		, DebugColor(ForceInit)
 		, TargetBehavior(EStreamingSourceTargetBehavior::Include)
 		, bReplay(false)
@@ -353,7 +354,7 @@ struct FWorldPartitionStreamingSource
 		, ExtraAngle(0)
 	{}
 
-	FWorldPartitionStreamingSource(FName InName, const FVector& InLocation, const FRotator& InRotation, EStreamingSourceTargetState InTargetState, bool bInBlockOnSlowLoading, EStreamingSourcePriority InPriority, bool bRemote, float InVelocity = 0.f)
+	FWorldPartitionStreamingSource(FName InName, const FVector& InLocation, const FRotator& InRotation, EStreamingSourceTargetState InTargetState, bool bInBlockOnSlowLoading, EStreamingSourcePriority InPriority, bool bRemote, FVector InVelocity = FVector::Zero())
 		: Name(InName)
 		, Location(InLocation)
 		, Rotation(InRotation)
@@ -361,6 +362,7 @@ struct FWorldPartitionStreamingSource
 		, bBlockOnSlowLoading(bInBlockOnSlowLoading)
 		, Priority(InPriority)
 		, Velocity(InVelocity)
+		, bUseVelocityContributionToCellsSorting(false)
 		, DebugColor(ForceInit)
 		, TargetBehavior(EStreamingSourceTargetBehavior::Include)
 		, bReplay(false)
@@ -387,6 +389,7 @@ struct FWorldPartitionStreamingSource
 		TargetState = Other.TargetState;
 		bBlockOnSlowLoading = Other.bBlockOnSlowLoading;
 		Priority = Other.Priority;
+		bUseVelocityContributionToCellsSorting = Other.bUseVelocityContributionToCellsSorting;
 		Velocity = Other.Velocity;
 		DebugColor = Other.DebugColor;
 		TargetBehavior = Other.TargetBehavior;
@@ -435,7 +438,10 @@ struct FWorldPartitionStreamingSource
 	EStreamingSourcePriority Priority;
 
 	/** Source velocity (computed automatically). */
-	float Velocity;
+	FVector Velocity;
+	
+	/** Use velocity contribution to streaming cells priority streaming. */
+	bool bUseVelocityContributionToCellsSorting;
 
 	/** Color used for debugging. */
 	FColor DebugColor;
diff --git a/Engine/Source/Runtime/Engine/Public/WorldPartition/WorldPartitionSubsystem.h b/Engine/Source/Runtime/Engine/Public/WorldPartition/WorldPartitionSubsystem.h
--- a/Engine/Source/Runtime/Engine/Public/WorldPartition/WorldPartitionSubsystem.h
+++ b/Engine/Source/Runtime/Engine/Public/WorldPartition/WorldPartitionSubsystem.h
@@ -32,17 +32,14 @@ struct FStreamingSourceVelocity
 	FStreamingSourceVelocity(const FName& InSourceName);
 	void Invalidate() { bIsValid = false; }
 	bool IsValid() { return bIsValid; }
-	float GetAverageVelocity(const FVector& NewPosition, const float CurrentTime);
+	FVector GetAverageVelocity(const FVector& NewPosition, double CurrentTime);
 
 private:
-	enum { VELOCITY_HISTORY_SAMPLE_COUNT = 16 };
 	bool bIsValid;
 	FName SourceName;
-	int32 LastIndex;
-	float LastUpdateTime;
+	double LastUpdateTime;
 	FVector LastPosition;
-	float VelocityHistorySum;
-	TArray<float, TInlineAllocator<VELOCITY_HISTORY_SAMPLE_COUNT>> VelocityHistory;
+	FVector AvgVelocity;
 };
 
 /**
