World Partition: move IsSpatiallyLoaded and RuntimeGrid actor properties in the advanced section, since they have a huge impact on streaming performances and memory usage, while easy to change.

#preflkight none

diff --git a/Engine/Source/Runtime/Engine/Classes/GameFramework/Actor.h b/Engine/Source/Runtime/Engine/Classes/GameFramework/Actor.h
--- a/Engine/Source/Runtime/Engine/Classes/GameFramework/Actor.h
+++ b/Engine/Source/Runtime/Engine/Classes/GameFramework/Actor.h
@@ -673,7 +673,7 @@ protected:
 	 * Determine in which partition grid this actor will be placed in the partition (if the world is partitioned).
 	 * If None, the decision will be left to the partition.
 	 */
-	UPROPERTY(EditAnywhere, Category=WorldPartition)
+	UPROPERTY(EditAnywhere, AdvancedDisplay, Category=WorldPartition)
 	FName RuntimeGrid;
 #endif
 
@@ -1154,7 +1154,7 @@ protected:
 	 *	If true, this actor will be loaded when in the range of any streaming sources and if (1) in no data layers, or (2) one or more of its data layers are enabled.
 	 *	If false, this actor will be loaded if (1) in no data layers, or (2) one or more of its data layers are enabled.
 	 */
-	UPROPERTY(EditAnywhere, Category=WorldPartition)
+	UPROPERTY(EditAnywhere, AdvancedDisplay, Category=WorldPartition)
 	uint8 bIsSpatiallyLoaded : 1;
 
 private:
