WorldPartition: added experimental actor descriptor mutators event, used to mutate descriptors during streaming generation.

diff --git a/Engine/Source/Runtime/Engine/Private/WorldPartition/WorldPartitionStreamingGeneration.cpp b/Engine/Source/Runtime/Engine/Private/WorldPartition/WorldPartitionStreamingGeneration.cpp
--- a/Engine/Source/Runtime/Engine/Private/WorldPartition/WorldPartitionStreamingGeneration.cpp
+++ b/Engine/Source/Runtime/Engine/Private/WorldPartition/WorldPartitionStreamingGeneration.cpp
@@ -592,6 +592,7 @@ class FWorldPartitionStreamingGenerator
 		FPerInstanceData InstanceData;
 		TSet<FPerInstanceData> UniquePerInstanceData;
 		TMap<FGuid, FSetElementId> PerInstanceData;
+		TMap<FGuid, FActorDescViewMutator> ActorDescViewMutators;
 	};
 
 	void ResolveRuntimeSpatiallyLoaded(FStreamingGenerationActorDescView& ActorDescView)
@@ -1246,7 +1247,89 @@ class FWorldPartitionStreamingGenerator
 		});
 	}
 
-	/**
+	/** 
+	 * Experimental: apply actor descriptor view mutators.
+	 */
+	bool MutateContainerInstanceDescriptors(const FStreamingGenerationContainerInstanceCollection& ActorDescCollection)
+	{
+		TUniquePtr<const IStreamingGenerationContext> MutatorStreamingGenerationContext = MakeUnique<FStreamingGenerationContext>(this, ActorDescCollection);
+
+		// Gather actor descriptor mutators
+		if (WorldPartitionContext->OnGenerateStreamingActorDescsMutatePhase.IsBound())
+		{
+			TArray<FActorDescViewMutatorInstance> ActorDescsMutatorsInstances;
+			WorldPartitionContext->OnGenerateStreamingActorDescsMutatePhase.Broadcast(MutatorStreamingGenerationContext.Get(), ActorDescsMutatorsInstances);
+
+			// Apply actor descriptor mutators to their respective containers
+			for (const FActorDescViewMutatorInstance& ActorDescMutatorInstance : ActorDescsMutatorsInstances)
+			{
+				FContainerCollectionInstanceDescriptor& ContainerCollectionInstanceDescriptor = ContainerCollectionInstanceDescriptorsMap.FindChecked(ActorDescMutatorInstance.ContainerId);
+				FActorDescViewMutator& ActorDescViewMutator = ContainerCollectionInstanceDescriptor.ActorDescViewMutators.FindOrAdd(ActorDescMutatorInstance.ActorGuid);
+
+				ActorDescViewMutator.bIsSpatiallyLoaded = ActorDescMutatorInstance.bIsSpatiallyLoaded;
+				ActorDescViewMutator.RuntimeGrid = ActorDescMutatorInstance.RuntimeGrid;
+			}
+
+			// Build the containers tree representation
+			TMultiMap<FActorContainerID, FActorContainerID> InvertedContainersHierarchy;
+			for (auto& [ContainerID, ContainerCollectionInstanceDescriptor] : ContainerCollectionInstanceDescriptorsMap)
+			{
+				if (!ContainerID.IsMainContainer())
+				{
+					InvertedContainersHierarchy.Add(ContainerCollectionInstanceDescriptor.ParentID, ContainerID);
+				}
+			}
+
+			// Apply mutators to per instance data
+			auto ApplyActorDescViewMutators = [this, &InvertedContainersHierarchy](const FActorContainerID& ContainerID, TMap<FGuid, FActorDescViewMutator> ActorDescViewMutators)
+			{
+				auto DumpContainerInstancesRecursive = [this, &InvertedContainersHierarchy](const FActorContainerID& ContainerID, TMap<FGuid, FActorDescViewMutator> ActorDescViewMutators, auto& RecursiveFunc) -> void
+				{
+					FContainerCollectionInstanceDescriptor& ContainerCollectionInstanceDescriptor = ContainerCollectionInstanceDescriptorsMap.FindChecked(ContainerID);
+					ActorDescViewMutators.Append(ContainerCollectionInstanceDescriptor.ActorDescViewMutators);
+
+					for (const auto& [ActorGuid, ActorDescViewMutator] : ActorDescViewMutators)
+					{
+						FContainerCollectionInstanceDescriptor::FPerInstanceData PerInstanceData = ContainerCollectionInstanceDescriptor.GetPerInstanceData(ActorGuid);
+
+						if (ActorDescViewMutator.bIsSpatiallyLoaded.IsSet())
+						{
+							PerInstanceData.bIsSpatiallyLoaded = ActorDescViewMutator.bIsSpatiallyLoaded.GetValue();
+						}
+
+						if (ActorDescViewMutator.RuntimeGrid.IsSet())
+						{
+							PerInstanceData.RuntimeGrid = ActorDescViewMutator.RuntimeGrid.GetValue();
+						}
+
+						ContainerCollectionInstanceDescriptor.AddPerInstanceData(ActorGuid, PerInstanceData);
+					}
+
+					TArray<FActorContainerID> ChildContainersIDs;
+					InvertedContainersHierarchy.MultiFind(ContainerID, ChildContainersIDs);
+					ChildContainersIDs.Sort();
+
+					if (ChildContainersIDs.Num())
+					{
+						for (const FActorContainerID& ChildContainerID : ChildContainersIDs)
+						{
+							RecursiveFunc(ChildContainerID, ActorDescViewMutators, RecursiveFunc);
+						}
+					}
+				};
+
+				DumpContainerInstancesRecursive(ContainerID, ActorDescViewMutators, DumpContainerInstancesRecursive);
+			};
+
+			ApplyActorDescViewMutators(FActorContainerID(), TMap<FGuid, FActorDescViewMutator>());
+
+			return true;
+		}
+
+		return false;
+	}
+
+	/** 
 	 * Perform various validations on the container descriptor instance, and adjust it based on different requirements. This needs to happen before updating
 	 * containers bounds because some actor descriptor views might change grid placement, etc.
 	 */
@@ -1431,6 +1514,11 @@ public:
 		}
 
 		ValidateContainerInstanceDescriptors();
+
+		if (MutateContainerInstanceDescriptors(ContainerInstanceCollection))
+		{
+			ValidateContainerInstanceDescriptors();
+		}
 	}
 
 	static TUniquePtr<FArchive> CreateDumpStateLogArchive(const TCHAR* Suffix, bool bTimeStamped = true)
diff --git a/Engine/Source/Runtime/Engine/Public/WorldPartition/WorldPartition.h b/Engine/Source/Runtime/Engine/Public/WorldPartition/WorldPartition.h
--- a/Engine/Source/Runtime/Engine/Public/WorldPartition/WorldPartition.h
+++ b/Engine/Source/Runtime/Engine/Public/WorldPartition/WorldPartition.h
@@ -252,6 +252,12 @@ public:
 	DECLARE_MULTICAST_DELEGATE_OneParam(FWorldPartitionGenerateStreamingDelegate, TArray<FString>*);
 	FWorldPartitionGenerateStreamingDelegate OnPreGenerateStreaming;
 
+	/**
+	 * Experimental: event used to gather actor descriptor mutators.
+	 */
+	DECLARE_MULTICAST_DELEGATE_TwoParams(FWorldPartitionGenerateStreamingActorDescsMutatePhase, const IStreamingGenerationContext* StreamingGenerationContext, TArray<FActorDescViewMutatorInstance>& ActorDescsMutatorsInstances);
+	FWorldPartitionGenerateStreamingActorDescsMutatePhase OnGenerateStreamingActorDescsMutatePhase;
+
 	ENGINE_API void RemapSoftObjectPath(FSoftObjectPath& ObjectPath) const;
 	ENGINE_API bool IsValidPackageName(const FString& InPackageName);
 
diff --git a/Engine/Source/Runtime/Engine/Public/WorldPartition/WorldPartitionStreamingGeneration.h b/Engine/Source/Runtime/Engine/Public/WorldPartition/WorldPartitionStreamingGeneration.h
--- a/Engine/Source/Runtime/Engine/Public/WorldPartition/WorldPartitionStreamingGeneration.h
+++ b/Engine/Source/Runtime/Engine/Public/WorldPartition/WorldPartitionStreamingGeneration.h
@@ -245,4 +245,22 @@ public:
 
 using FActorDescViewMap UE_DEPRECATED(5.4, "Use FStreamingGenerationActorDescViewMap instead") = FStreamingGenerationActorDescViewMap;
 
+/**
+ * Experimental: an actor descriptor mutator used by external code to change properties on actor descriptors.
+ */
+struct FActorDescViewMutator
+{
+	TOptional<bool> bIsSpatiallyLoaded;
+	TOptional<FName> RuntimeGrid;
+};
+
+struct FActorDescViewMutatorInstance : FActorDescViewMutator
+{
+	/** The actor GUID to mutate */
+	FGuid ActorGuid;
+
+	/** Target container ID to mutate this actor */
+	FActorContainerID ContainerId;
+};
+
 #endif // WITH_EDITOR
