Fix crash cooking Content Worker.

- Not all collision data can be gotten.

diff --git a/Engine/Source/Runtime/Engine/Private/StaticMesh.cpp b/Engine/Source/Runtime/Engine/Private/StaticMesh.cpp
--- a/Engine/Source/Runtime/Engine/Private/StaticMesh.cpp
+++ b/Engine/Source/Runtime/Engine/Private/StaticMesh.cpp
@@ -1871,7 +1871,10 @@ void FStaticMeshRenderData::Serialize(FArchive& Ar, UStaticMesh* Owner, bool bCo
 			check(OwnerStaticMesh);
 
 			FTriMeshCollisionData CollisionData;
-			OwnerStaticMesh->GetPhysicsTriMeshData(&CollisionData, true);
+			if (OwnerStaticMesh->GetBodySetup() && OwnerStaticMesh->GetBodySetup()->GetCollisionTraceFlag() != CTF_UseComplexAsSimple && OwnerStaticMesh->ContainsPhysicsTriMeshData(true))
+			{
+				OwnerStaticMesh->GetPhysicsTriMeshData(&CollisionData, true);
+			}
 
 			Ar << CollisionData;
 		}
