Fix issues with proxy tracking when raytracing is disabled.
- don't register/unregister proxies with FRayTracingGeometryManager unless raytracing is allowed.
- initialize RayTracingGeometryGroupHandle to INDE_NONE.
- dont RequestUpdateCachedRenderState unless raytracing is enabled.

Submitted on behalf of Tiago.Costa

diff --git a/Engine/Source/Runtime/Engine/Private/Rendering/RayTracingGeometryManager.cpp b/Engine/Source/Runtime/Engine/Private/Rendering/RayTracingGeometryManager.cpp
--- a/Engine/Source/Runtime/Engine/Private/Rendering/RayTracingGeometryManager.cpp
+++ b/Engine/Source/Runtime/Engine/Private/Rendering/RayTracingGeometryManager.cpp
@@ -346,6 +346,12 @@ void FRayTracingGeometryManager::SetupBuildParams(const FBuildRequest& InBuildRe
 void FRayTracingGeometryManager::RegisterProxyWithCachedRayTracingState(FPrimitiveSceneProxy* Proxy, RayTracing::GeometryGroupHandle InRayTracingGeometryGroupHandle)
 {
 	checkf(IsInRenderingThread(), TEXT("Can only access RegisteredGroups on render thread otherwise need a critical section"));
+
+	if (!IsRayTracingAllowed())
+	{
+		return;
+	}
+
 	checkf(RegisteredGroups.IsValidIndex(InRayTracingGeometryGroupHandle), TEXT("InRayTracingGeometryGroupHandle must be valid"));
 
 	TSet<FPrimitiveSceneProxy*>& ProxiesSet = RegisteredGroups[InRayTracingGeometryGroupHandle].ProxiesWithCachedRayTracingState;
@@ -358,6 +364,13 @@ void FRayTracingGeometryManager::UnregisterProxyWithCachedRayTracingState(FPrimi
 {
 	checkf(IsInRenderingThread(), TEXT("Can only access RegisteredGroups on render thread otherwise need a critical section"));
 
+	if (!IsRayTracingAllowed())
+	{
+		return;
+	}
+
+	checkf(RegisteredGroups.IsValidIndex(InRayTracingGeometryGroupHandle), TEXT("InRayTracingGeometryGroupHandle must be valid"));
+
 	TSet<FPrimitiveSceneProxy*>& ProxiesSet = RegisteredGroups[InRayTracingGeometryGroupHandle].ProxiesWithCachedRayTracingState;
 
 	verify(ProxiesSet.Remove(Proxy) == 1);
@@ -366,6 +379,11 @@ void FRayTracingGeometryManager::UnregisterProxyWithCachedRayTracingState(FPrimi
 void FRayTracingGeometryManager::RequestUpdateCachedRenderState(RayTracing::GeometryGroupHandle InRayTracingGeometryGroupHandle)
 {
 	checkf(IsInRenderingThread(), TEXT("Can only access RegisteredGroups on render thread otherwise need a critical section"));
+	
+	if (InRayTracingGeometryGroupHandle == INDEX_NONE)
+	{
+		return;
+	}
 
 	const TSet<FPrimitiveSceneProxy*>& ProxiesSet = RegisteredGroups[InRayTracingGeometryGroupHandle].ProxiesWithCachedRayTracingState;
 
diff --git a/Engine/Source/Runtime/Engine/Private/StaticMesh.cpp b/Engine/Source/Runtime/Engine/Private/StaticMesh.cpp
--- a/Engine/Source/Runtime/Engine/Private/StaticMesh.cpp
+++ b/Engine/Source/Runtime/Engine/Private/StaticMesh.cpp
@@ -2027,7 +2027,10 @@ void UStaticMesh::RequestUpdateCachedRenderState() const
 	}
 
 #if RHI_RAYTRACING
-	((FRayTracingGeometryManager*)GRayTracingGeometryManager)->RequestUpdateCachedRenderState(GetRenderData()->RayTracingGeometryGroupHandle);
+	if (IsRayTracingEnabled())
+	{
+		((FRayTracingGeometryManager*)GRayTracingGeometryManager)->RequestUpdateCachedRenderState(GetRenderData()->RayTracingGeometryGroupHandle);
+	}
 #endif
 
 	// TODO: Need to mark all DynamicRayTracingGeometries used in FStaticMeshSceneProxy referencing this StaticMesh as either invalid or request a recreation (UE-139474)
diff --git a/Engine/Source/Runtime/Engine/Public/StaticMeshResources.h b/Engine/Source/Runtime/Engine/Public/StaticMeshResources.h
--- a/Engine/Source/Runtime/Engine/Public/StaticMeshResources.h
+++ b/Engine/Source/Runtime/Engine/Public/StaticMeshResources.h
@@ -699,7 +699,7 @@ public:
 	FBoxSphereBounds Bounds;
 
 #if RHI_RAYTRACING
-	RayTracing::GeometryGroupHandle RayTracingGeometryGroupHandle;
+	RayTracing::GeometryGroupHandle RayTracingGeometryGroupHandle = INDEX_NONE;
 #endif
 
 	bool IsInitialized() const
