SkeletalMesh: Convert LOD model vertex colors as linear rather than sRGB.

#rnx

diff --git a/Engine/Source/Runtime/Engine/Private/SkeletalMeshLODModel.cpp b/Engine/Source/Runtime/Engine/Private/SkeletalMeshLODModel.cpp
--- a/Engine/Source/Runtime/Engine/Private/SkeletalMeshLODModel.cpp
+++ b/Engine/Source/Runtime/Engine/Private/SkeletalMeshLODModel.cpp
@@ -1520,7 +1520,7 @@ void FSkeletalMeshLODModel::GetMeshDescription(const USkeletalMesh *InSkeletalMe
 
 				if (bHasVertexColors)
 				{
-					VertexInstanceColors.Set(VertexInstanceID, FVector4f(FLinearColor(SourceVertex.Color)));
+					VertexInstanceColors.Set(VertexInstanceID, FVector4f(SourceVertex.Color.ReinterpretAsLinear()));
 				}
 
 				TriangleVertexInstanceIDs[Corner] = VertexInstanceID;
