Cleanup compression settings for platform extensions.

diff --git a/Engine/Source/Runtime/Engine/Engine.Build.cs b/Engine/Source/Runtime/Engine/Engine.Build.cs
--- a/Engine/Source/Runtime/Engine/Engine.Build.cs
+++ b/Engine/Source/Runtime/Engine/Engine.Build.cs
@@ -448,12 +448,14 @@ public class Engine : ModuleRules
 		if (Target.IsInPlatformGroup(UnrealPlatformGroup.Android))
 		{
 			PrivateIncludePathModuleNames.Add("AndroidRuntimeSettings");
+			PrivateDefinitions.Add("ENABLE_PLATFORM_COMPRESSION_OVERRIDES=1");
 		}
 
 		if (Target.IsInPlatformGroup(UnrealPlatformGroup.IOS))
 		{
 			PublicIncludePathModuleNames.Add("IOSPlatformFeatures");
 			PrivateIncludePathModuleNames.Add("IOSRuntimeSettings");
+			PrivateDefinitions.Add("ENABLE_PLATFORM_COMPRESSION_OVERRIDES=1");
 		}
 
 		PublicDefinitions.Add("GPUPARTICLE_LOCAL_VF_ONLY=0");
diff --git a/Engine/Source/Runtime/Engine/Private/AudioCompressionSettingsUtils.cpp b/Engine/Source/Runtime/Engine/Private/AudioCompressionSettingsUtils.cpp
--- a/Engine/Source/Runtime/Engine/Private/AudioCompressionSettingsUtils.cpp
+++ b/Engine/Source/Runtime/Engine/Private/AudioCompressionSettingsUtils.cpp
@@ -4,7 +4,9 @@
 #include "AudioStreamingCache.h"
 #include "Misc/DataDrivenPlatformInfoRegistry.h"
 
-#define ENABLE_PLATFORM_COMPRESSION_OVERRIDES 1
+#ifndef ENABLE_PLATFORM_COMPRESSION_OVERRIDES
+	#define ENABLE_PLATFORM_COMPRESSION_OVERRIDES 0
+#endif
 
 #if PLATFORM_ANDROID && ENABLE_PLATFORM_COMPRESSION_OVERRIDES
 #include "AndroidRuntimeSettings.h"
@@ -51,34 +53,32 @@ FAutoConsoleVariableRef CVarChunkSlotNumScalar(
 	TEXT("1.0: is the lower limit"),
 	ECVF_Default);
 
+#if ENABLE_PLATFORM_COMPRESSION_OVERRIDES && PLATFORM_ANDROID
 const FPlatformRuntimeAudioCompressionOverrides* FPlatformCompressionUtilities::GetRuntimeCompressionOverridesForCurrentPlatform()
 {
-#if PLATFORM_ANDROID && ENABLE_PLATFORM_COMPRESSION_OVERRIDES
 	static const UAndroidRuntimeSettings* Settings = GetDefault<UAndroidRuntimeSettings>();
 	if (Settings)
 	{
 		return &(Settings->CompressionOverrides);
 	}
-
-#elif PLATFORM_IOS && ENABLE_PLATFORM_COMPRESSION_OVERRIDES
+	return nullptr;
+}
+#elif ENABLE_PLATFORM_COMPRESSION_OVERRIDES && PLATFORM_IOS
+const FPlatformRuntimeAudioCompressionOverrides* FPlatformCompressionUtilities::GetRuntimeCompressionOverridesForCurrentPlatform()
+{
 	static const UIOSRuntimeSettings* Settings = GetDefault<UIOSRuntimeSettings>();
-
 	if (Settings)
 	{
 		return &(Settings->CompressionOverrides);
 	}
-
-#elif PLATFORM_SWITCH && ENABLE_PLATFORM_COMPRESSION_OVERRIDES
-	static const USwitchRuntimeSettings* Settings = GetDefault<USwitchRuntimeSettings>();
-
-	if (Settings)
-	{
-		return &(Settings->CompressionOverrides);
-	}
-
-#endif // PLATFORM_ANDROID
 	return nullptr;
 }
+#elif !ENABLE_PLATFORM_COMPRESSION_OVERRIDES
+const FPlatformRuntimeAudioCompressionOverrides* FPlatformCompressionUtilities::GetRuntimeCompressionOverridesForCurrentPlatform()
+{
+	return nullptr;
+}
+#endif // ENABLE_PLATFORM_COMPRESSION_OVERRIDES
 
 void CacheAudioCookOverrides(FPlatformAudioCookOverrides& OutOverrides, const TCHAR* InPlatformName=nullptr)
 {
