#UE Quick fix for crash in validation

diff --git a/Engine/Source/Runtime/Engine/Private/WorldPartition/WorldPartitionStreamingGeneration.cpp b/Engine/Source/Runtime/Engine/Private/WorldPartition/WorldPartitionStreamingGeneration.cpp
--- a/Engine/Source/Runtime/Engine/Private/WorldPartition/WorldPartitionStreamingGeneration.cpp
+++ b/Engine/Source/Runtime/Engine/Private/WorldPartition/WorldPartitionStreamingGeneration.cpp
@@ -1255,7 +1255,7 @@ class FWorldPartitionStreamingGenerator
 		TUniquePtr<const IStreamingGenerationContext> MutatorStreamingGenerationContext = MakeUnique<FStreamingGenerationContext>(this, ActorDescCollection);
 
 		// Gather actor descriptor mutators
-		if (WorldPartitionContext->OnGenerateStreamingActorDescsMutatePhase.IsBound())
+		if (WorldPartitionContext && WorldPartitionContext->OnGenerateStreamingActorDescsMutatePhase.IsBound())
 		{
 			TArray<FActorDescViewMutatorInstance> ActorDescsMutatorsInstances;
 			WorldPartitionContext->OnGenerateStreamingActorDescsMutatePhase.Broadcast(MutatorStreamingGenerationContext.Get(), ActorDescsMutatorsInstances);
