remove needless platform include.

 - should have been removed in previous revision

#rnx

diff --git a/Engine/Source/Runtime/Engine/Private/AudioCompressionSettingsUtils.cpp b/Engine/Source/Runtime/Engine/Private/AudioCompressionSettingsUtils.cpp
--- a/Engine/Source/Runtime/Engine/Private/AudioCompressionSettingsUtils.cpp
+++ b/Engine/Source/Runtime/Engine/Private/AudioCompressionSettingsUtils.cpp
@@ -16,10 +16,6 @@
 #include "IOSRuntimeSettings.h"
 #endif
 
-#if PLATFORM_SWITCH && ENABLE_PLATFORM_COMPRESSION_OVERRIDES
-#include "SwitchRuntimeSettings.h"
-#endif
-
 #include "Misc/ConfigCacheIni.h"
 
 static float CookOverrideCachingIntervalCvar = 60.0f;
