World Partition Runtime HashSet: push partition debug color to level colorization mode.

diff --git a/Engine/Source/Runtime/Engine/Private/WorldPartition/RuntimeHashSet/WorldPartitionRuntimeHashSetStreamingGeneration.cpp b/Engine/Source/Runtime/Engine/Private/WorldPartition/RuntimeHashSet/WorldPartitionRuntimeHashSetStreamingGeneration.cpp
--- a/Engine/Source/Runtime/Engine/Private/WorldPartition/RuntimeHashSet/WorldPartitionRuntimeHashSetStreamingGeneration.cpp
+++ b/Engine/Source/Runtime/Engine/Private/WorldPartition/RuntimeHashSet/WorldPartitionRuntimeHashSetStreamingGeneration.cpp
@@ -188,6 +188,7 @@ bool UWorldPartitionRuntimeHashSet::GenerateStreaming(UWorldPartitionStreamingPo
 		RuntimeCell->SetBlockOnSlowLoading(CellDescInstance.bBlockOnSlowStreaming);
 		RuntimeCell->SetIsHLOD(CellDescInstance.SourcePartition->HLODIndex != INDEX_NONE);
 		RuntimeCell->SetGuid(CellUniqueId.Guid);
+		RuntimeCell->SetCellDebugColor(CellDescInstance.SourcePartition->DebugColor);
 
 		UWorldPartitionRuntimeCellData* RuntimeCellData = RuntimeCell->RuntimeCellData;
 		RuntimeCellData->DebugName = CellUniqueId.Name;
diff --git a/Engine/Source/Runtime/Engine/Private/WorldPartition/WorldPartitionLevelStreamingDynamic.cpp b/Engine/Source/Runtime/Engine/Private/WorldPartition/WorldPartitionLevelStreamingDynamic.cpp
--- a/Engine/Source/Runtime/Engine/Private/WorldPartition/WorldPartitionLevelStreamingDynamic.cpp
+++ b/Engine/Source/Runtime/Engine/Private/WorldPartition/WorldPartitionLevelStreamingDynamic.cpp
@@ -191,8 +191,6 @@ void UWorldPartitionLevelStreamingDynamic::CreateRuntimeLevel()
 
 	// Attach ourself to Level cleanup to do our own cleanup
 	OnCleanupLevelDelegateHandle = RuntimeLevel->OnCleanupLevel.AddUObject(this, &UWorldPartitionLevelStreamingDynamic::OnCleanupLevel);
-
-	LevelColor = FLinearColor::MakeRandomSeededColor(GetTypeHash(GetName()));
 }
 
 /**
@@ -307,6 +305,11 @@ bool UWorldPartitionLevelStreamingDynamic::RequestLevel(UWorld* InPersistentWorl
 			CreateRuntimeLevel();
 			check(RuntimeLevel);
 
+			if (const UWorldPartitionRuntimeLevelStreamingCell* RuntimeLevelStreamingCell = StreamingCell.Get())
+			{
+				LevelColor = RuntimeLevelStreamingCell->GetCellDebugColor();
+			}
+
 			UPackage* CellLevelPackage = RuntimeLevel->GetPackage();
 			check(CellLevelPackage);
 			check(UWorld::FindWorldInPackage(CellLevelPackage));
diff --git a/Engine/Source/Runtime/Engine/Public/WorldPartition/WorldPartitionRuntimeCell.h b/Engine/Source/Runtime/Engine/Public/WorldPartition/WorldPartitionRuntimeCell.h
--- a/Engine/Source/Runtime/Engine/Public/WorldPartition/WorldPartitionRuntimeCell.h
+++ b/Engine/Source/Runtime/Engine/Public/WorldPartition/WorldPartitionRuntimeCell.h
@@ -268,6 +268,9 @@ class UWorldPartitionRuntimeCell : public UObject, public IWorldPartitionCell
 	const FGuid& GetGuid() const { return CellGuid; }
 	void SetGuid(const FGuid& InCellGuid) { CellGuid = InCellGuid; }
 
+	const FLinearColor& GetCellDebugColor() const { return CellDebugColor; }
+	void SetCellDebugColor(const FLinearColor& InCellDebugColor) { CellDebugColor = InCellDebugColor; }
+
 	const FGuid& GetSourceCellGuid() const { return SourceCellGuid; }
 	void SetSourceCellGuid(const FGuid& InSourceCellGuid) { SourceCellGuid = InSourceCellGuid; }
 
@@ -308,6 +311,9 @@ private:
 	UPROPERTY()
 	FGuid ContentBundleID;
 
+	UPROPERTY()
+	FLinearColor CellDebugColor;
+
 protected:
 	UPROPERTY()
 	FGuid CellGuid;
