Level Instance: When doing a RebuildAll (UWorld::ReInit) previous LI Level might stay in memory (not GCed) so make sure that CachedEditorInstanceActorPtr is reset on the LevelInstanceComponent on Unload so it can resolve to newly loaded Level Actor

- Was hittinga fail check when running a "Build All" after having moved a Level Instance

#rnx

diff --git a/Engine/Source/Runtime/Engine/Private/LevelInstance/LevelInstanceActorImpl.cpp b/Engine/Source/Runtime/Engine/Private/LevelInstance/LevelInstanceActorImpl.cpp
--- a/Engine/Source/Runtime/Engine/Private/LevelInstance/LevelInstanceActorImpl.cpp
+++ b/Engine/Source/Runtime/Engine/Private/LevelInstance/LevelInstanceActorImpl.cpp
@@ -51,6 +51,11 @@ void FLevelInstanceActorImpl::UnregisterLevelInstance()
 
 		LevelInstance->UnloadLevelInstance();
 
+#if WITH_EDITOR
+		// Make sure our Component doesn't keep a pointer to the Editor Actor once unloaded
+		LevelInstance->GetLevelInstanceComponent()->CachedEditorInstanceActorPtr.Reset();
+#endif
+
 		// To avoid processing PostUnregisterAllComponents multiple times (BP Recompile is one use case)
 		LevelInstanceID = FLevelInstanceID();
 	}
diff --git a/Engine/Source/Runtime/Engine/Public/LevelInstance/LevelInstanceComponent.h b/Engine/Source/Runtime/Engine/Public/LevelInstance/LevelInstanceComponent.h
--- a/Engine/Source/Runtime/Engine/Public/LevelInstance/LevelInstanceComponent.h
+++ b/Engine/Source/Runtime/Engine/Public/LevelInstance/LevelInstanceComponent.h
@@ -52,6 +52,7 @@ private:
 	ENGINE_API void SetActiveFilter(const FWorldPartitionActorFilter& InFilter);
 	ENGINE_API bool IsEditFilter() const;
 
+	friend class FLevelInstanceActorImpl;
 	TWeakObjectPtr<AActor> CachedEditorInstanceActorPtr;
 #endif
 
