External Data Layer V1

diff --git a/Engine/Source/Editor/SceneOutliner/Private/SceneOutlinerModule.cpp b/Engine/Source/Editor/SceneOutliner/Private/SceneOutlinerModule.cpp
--- a/Engine/Source/Editor/SceneOutliner/Private/SceneOutlinerModule.cpp
+++ b/Engine/Source/Editor/SceneOutliner/Private/SceneOutlinerModule.cpp
@@ -27,6 +27,7 @@
 #include "WorldPartition/WorldPartitionActorDescInstance.h"
 #include "WorldPartition/DataLayer/DataLayerInstance.h"
 #include "WorldPartition/DataLayer/DataLayerManager.h"
+#include "WorldPartition/DataLayer/ExternalDataLayerAsset.h"
 #include "WorldPartition/ContentBundle/ContentBundleEngineSubsystem.h"
 #include "WorldPartition/ContentBundle/ContentBundleDescriptor.h"
 #include "LevelInstance/LevelInstanceSubsystem.h"
@@ -315,6 +316,24 @@ void FSceneOutlinerModule::CreateActorInfoColumns(FSceneOutlinerInitializationOp
 		return Result;
 	});
 
+	FGetTextForItem ExternalDatalayerInfoText = FGetTextForItem::CreateLambda([](const ISceneOutlinerTreeItem& Item) -> FString
+	{
+		const UExternalDataLayerAsset* ExternalDataLayerAsset = nullptr;
+		if (const FActorTreeItem* ActorItem = Item.CastTo<FActorTreeItem>())
+		{
+			if (AActor* Actor = ActorItem->Actor.Get())
+			{
+				ExternalDataLayerAsset = Actor->GetExternalDataLayerAsset();
+			}
+		}
+		else if (const FActorDescTreeItem* ActorDescItem = Item.CastTo<FActorDescTreeItem>())
+		{
+			ExternalDataLayerAsset = ActorDescItem->GetExternalDataLayerAsset();
+		}
+
+		return ExternalDataLayerAsset ? ExternalDataLayerAsset->GetName() : TEXT("");
+	});
+
 	FGetTextForItem DataLayerInfoText = FGetTextForItem::CreateLambda([](const ISceneOutlinerTreeItem& Item) -> FString
 	{
 		TStringBuilder<128> Builder;
@@ -324,20 +343,23 @@ void FSceneOutlinerModule::CreateActorInfoColumns(FSceneOutlinerInitializationOp
 		{
 			for (const UDataLayerInstance* DataLayerInstance : DataLayerInstances)
 			{
-				bool bIsAlreadyInSet = false;
-				DataLayerShortNames.Add(DataLayerInstance->GetDataLayerShortName(), &bIsAlreadyInSet);
-				if (!bIsAlreadyInSet)
+				if (!DataLayerInstance->IsA<UExternalDataLayerInstance>())
 				{
-					if (Builder.Len())
-					{
-						Builder += TEXT(", ");
-					}
-					// Put a '*' in front of DataLayers that are not part of of the main world
-					if (bPartOfOtherLevel)
+					bool bIsAlreadyInSet = false;
+					DataLayerShortNames.Add(DataLayerInstance->GetDataLayerShortName(), &bIsAlreadyInSet);
+					if (!bIsAlreadyInSet)
 					{
-						Builder += "*";
+						if (Builder.Len())
+						{
+							Builder += TEXT(", ");
+						}
+						// Put a '*' in front of DataLayers that are not part of of the main world
+						if (bPartOfOtherLevel)
+						{
+							Builder += "*";
+						}
+						Builder += DataLayerInstance->GetDataLayerShortName();
 					}
-					Builder += DataLayerInstance->GetDataLayerShortName();
 				}
 			}
 		};
@@ -361,7 +383,7 @@ void FSceneOutlinerModule::CreateActorInfoColumns(FSceneOutlinerInitializationOp
 						if (const UDataLayerManager* DataLayerManager = UDataLayerManager::GetDataLayerManager(OwningWorld))
 						{
 							TSet<const UDataLayerInstance*> DataLayerInstances;
-							DataLayerInstances.Append(DataLayerManager->GetDataLayerInstances(ActorDescInstance->GetDataLayerInstanceNames()));
+							DataLayerInstances.Append(DataLayerManager->GetDataLayerInstances(ActorDescInstance->GetDataLayerInstanceNames().ToArray()));
 							if (ULevelInstanceSubsystem* LevelInstanceSubsystem = UWorld::GetSubsystem<ULevelInstanceSubsystem>(OwningWorld))
 							{
 								UWorld* OuterWorld = ActorDescContainerInstance->GetTypedOuter<UWorld>();
@@ -529,7 +551,7 @@ void FSceneOutlinerModule::CreateActorInfoColumns(FSceneOutlinerInitializationOp
 		return FString();
 	});
 
-	auto AddTextInfoColumn = [&InInitOptions](FName ColumnID, TAttribute<FText> ColumnName, FGetTextForItem ColumnInfo)
+	auto AddTextInfoColumn = [&InInitOptions](FName ColumnID, TAttribute<FText> ColumnName, FGetTextForItem ColumnInfo, const FText& ColumnTooltip = FText::GetEmpty())
 	{
 		InInitOptions.ColumnMap.Add(
 			ColumnID,
@@ -540,7 +562,7 @@ void FSceneOutlinerModule::CreateActorInfoColumns(FSceneOutlinerInitializationOp
 					&FTextInfoColumn::CreateTextInfoColumn,
 					ColumnID,
 					ColumnInfo,
-					FText::GetEmpty()),
+					ColumnTooltip),
 				true,
 				TOptional<float>(),
 				ColumnName));
@@ -561,6 +583,7 @@ void FSceneOutlinerModule::CreateActorInfoColumns(FSceneOutlinerInitializationOp
 	AddTextInfoColumn(FSceneOutlinerBuiltInColumnTypes::Level(), LevelColumnName, LevelInfoText);
 	AddTextInfoColumn(FSceneOutlinerBuiltInColumnTypes::Layer(), FSceneOutlinerBuiltInColumnTypes::Layer_Localized(), LayerInfoText);
 	AddTextInfoColumn(FSceneOutlinerBuiltInColumnTypes::DataLayer(), FSceneOutlinerBuiltInColumnTypes::DataLayer_Localized(), DataLayerInfoText);
+	AddTextInfoColumn(FSceneOutlinerBuiltInColumnTypes::ExternalDataLayer(), FSceneOutlinerBuiltInColumnTypes::ExternalDataLayer_Localized(), ExternalDatalayerInfoText);
 	AddTextInfoColumn(FSceneOutlinerBuiltInColumnTypes::ContentBundle(), FSceneOutlinerBuiltInColumnTypes::ContentBundle_Localized(), ContentBundleInfoText);
 	AddTextInfoColumn(FSceneOutlinerBuiltInColumnTypes::SubPackage(), FSceneOutlinerBuiltInColumnTypes::SubPackage_Localized(), SubPackageInfoText);
 	AddTextInfoColumn(FSceneOutlinerBuiltInColumnTypes::Socket(), FSceneOutlinerBuiltInColumnTypes::Socket_Localized(), SocketInfoText);
diff --git a/Engine/Source/Editor/SceneOutliner/Public/ActorDescTreeItem.h b/Engine/Source/Editor/SceneOutliner/Public/ActorDescTreeItem.h
--- a/Engine/Source/Editor/SceneOutliner/Public/ActorDescTreeItem.h
+++ b/Engine/Source/Editor/SceneOutliner/Public/ActorDescTreeItem.h
@@ -9,6 +9,7 @@
 #include "WorldPartition/WorldPartitionHandle.h"
 
 class UActorDescContainerInstance;
+class UExternalDataLayerAsset;
 class FWorldPartitionActorDesc;
 class UToolMenu;
 
@@ -70,12 +71,16 @@ public:
 	virtual const FGuid& GetGuid() const override { return ActorGuid; }
 	/* End IActorBaseTreeItem Implementation */
 
+	UExternalDataLayerAsset* GetExternalDataLayerAsset() const;
+
 	void FocusActorBounds() const;
 
 protected:
 	FString DisplayString;
 
 private:
+	mutable TSoftObjectPtr<UExternalDataLayerAsset> CachedExternalDataLayerAsset;
+
 	void CopyActorFilePathtoClipboard() const;
 	FGuid ActorGuid;
 };
diff --git a/Engine/Source/Editor/SceneOutliner/Public/SceneOutlinerPublicTypes.h b/Engine/Source/Editor/SceneOutliner/Public/SceneOutlinerPublicTypes.h
--- a/Engine/Source/Editor/SceneOutliner/Public/SceneOutlinerPublicTypes.h
+++ b/Engine/Source/Editor/SceneOutliner/Public/SceneOutlinerPublicTypes.h
@@ -50,6 +50,7 @@ struct FSceneOutlinerBuiltInColumnTypes
 	DEFINE_SCENEOUTLINER_BUILTIN_COLUMN_TYPE(Level, "Level", "SceneOutlinerLevelColumn", "Level");
 	DEFINE_SCENEOUTLINER_BUILTIN_COLUMN_TYPE(Layer, "Layer", "SceneOutlinerLayerColumn", "Layer");
 	DEFINE_SCENEOUTLINER_BUILTIN_COLUMN_TYPE(DataLayer, "Data Layer", "SceneOutlinerDataLayerColumn", "Data Layer");
+	DEFINE_SCENEOUTLINER_BUILTIN_COLUMN_TYPE(ExternalDataLayer, "External Data Layer", "SceneOutlinerExternalDataLayerColumn", "External Data Layer");
 	DEFINE_SCENEOUTLINER_BUILTIN_COLUMN_TYPE(ContentBundle, "Content Bundle", "SceneOutlinerContentBundleColumn", "Content Bundle");
 	DEFINE_SCENEOUTLINER_BUILTIN_COLUMN_TYPE(SubPackage, "Sub Package", "SceneOutlinerSubPackageColumn", "Sub Package");
 	DEFINE_SCENEOUTLINER_BUILTIN_COLUMN_TYPE(Pinned, "Pinned", "SceneOutlinerPinnedColumn", "Pinned");
diff --git a/Engine/Source/Editor/UnrealEd/Classes/Settings/EditorExperimentalSettings.h b/Engine/Source/Editor/UnrealEd/Classes/Settings/EditorExperimentalSettings.h
--- a/Engine/Source/Editor/UnrealEd/Classes/Settings/EditorExperimentalSettings.h
+++ b/Engine/Source/Editor/UnrealEd/Classes/Settings/EditorExperimentalSettings.h
@@ -143,6 +143,9 @@ public:
 	UPROPERTY(EditAnywhere, config, Category = WorldPartition)
 	bool bEnableWorldPartitionActorFilters;
 
+	UPROPERTY(EditAnywhere, config, Category = WorldPartition)
+	bool bEnableWorldPartitionExternalDataLayers;
+
 	/**
 	 * Returns an event delegate that is executed when a setting has changed.
 	 *
diff --git a/Engine/Source/Editor/UnrealEd/Private/Commandlets/WorldPartitionCommandletHelpers.cpp b/Engine/Source/Editor/UnrealEd/Private/Commandlets/WorldPartitionCommandletHelpers.cpp
--- a/Engine/Source/Editor/UnrealEd/Private/Commandlets/WorldPartitionCommandletHelpers.cpp
+++ b/Engine/Source/Editor/UnrealEd/Private/Commandlets/WorldPartitionCommandletHelpers.cpp
@@ -13,7 +13,7 @@
 
 namespace WorldPartitionCommandletHelpers
 {
-	DEFINE_LOG_CATEGORY(LogWorldParitionCommandletUtils)
+	DEFINE_LOG_CATEGORY(LogWorldPartitionCommandletUtils)
 
 	UWorld* LoadAndInitWorld(const FString& LevelToLoad)
 	{
@@ -34,21 +34,21 @@ namespace WorldPartitionCommandletHelpers
 		TRACE_CPUPROFILER_EVENT_SCOPE(UWorldPartitionConvertCommandlet::LoadWorld);
 
 		SET_WARN_COLOR(COLOR_WHITE);
-		UE_LOG(LogWorldParitionCommandletUtils, Log, TEXT("Loading level %s."), *LevelToLoad);
+		UE_LOG(LogWorldPartitionCommandletUtils, Log, TEXT("Loading level %s."), *LevelToLoad);
 		CLEAR_WARN_COLOR();
 
 		// This will convert incomplete package name to a fully qualified path, avoiding calling it several times (takes ~50s)
 		FString FullLevelToLoadPath;
 		if (!FPackageName::SearchForPackageOnDisk(LevelToLoad, &FullLevelToLoadPath))
 		{
-			UE_LOG(LogWorldParitionCommandletUtils, Error, TEXT("Unknown level '%s'"), *LevelToLoad);
+			UE_LOG(LogWorldPartitionCommandletUtils, Error, TEXT("Unknown level '%s'"), *LevelToLoad);
 			return nullptr;
 		}
 
 		UPackage* MapPackage = LoadPackage(nullptr, *FullLevelToLoadPath, LOAD_None);
 		if (!MapPackage)
 		{
-			UE_LOG(LogWorldParitionCommandletUtils, Error, TEXT("Error loading %s."), *FullLevelToLoadPath);
+			UE_LOG(LogWorldPartitionCommandletUtils, Error, TEXT("Error loading %s."), *FullLevelToLoadPath);
 			return nullptr;
 		}
 
@@ -60,7 +60,7 @@ namespace WorldPartitionCommandletHelpers
 		TRACE_CPUPROFILER_EVENT_SCOPE(UWorldPartitionConvertCommandlet::InitWorld);
 
 		SET_WARN_COLOR(COLOR_WHITE);
-		UE_LOG(LogWorldParitionCommandletUtils, Log, TEXT("Initializing level %s."), *World->GetName());
+		UE_LOG(LogWorldPartitionCommandletUtils, Log, TEXT("Initializing level %s."), *World->GetName());
 		CLEAR_WARN_COLOR();
 
 		// Setup the world.
@@ -91,7 +91,7 @@ namespace WorldPartitionCommandletHelpers
 	{
 		TRACE_CPUPROFILER_EVENT_SCOPE(CheckoutPackages);
 
-		UE_LOG(LogWorldParitionCommandletUtils, Log, TEXT("Checking out %d Packages."), PackagesToCheckout.Num());
+		UE_LOG(LogWorldPartitionCommandletUtils, Log, TEXT("Checking out %d Packages."), PackagesToCheckout.Num());
 		for (UPackage* Package : PackagesToCheckout)
 		{
 			if (!Checkout(Package, SCHelper))
@@ -105,7 +105,7 @@ namespace WorldPartitionCommandletHelpers
 
 	bool Checkout(UPackage* PackagesToCheckout, FPackageSourceControlHelper& SCHelper)
 	{
-		UE_LOG(LogWorldParitionCommandletUtils, Verbose, TEXT("Checking out Package %s."), *SourceControlHelpers::PackageFilename(PackagesToCheckout));
+		UE_LOG(LogWorldPartitionCommandletUtils, Verbose, TEXT("Checking out Package %s."), *SourceControlHelpers::PackageFilename(PackagesToCheckout));
 
 		FString PackageFileName = SourceControlHelpers::PackageFilename(PackagesToCheckout);
 		if (FPlatformFileManager::Get().GetPlatformFile().FileExists(*PackageFileName))
@@ -123,7 +123,7 @@ namespace WorldPartitionCommandletHelpers
 	{
 		TRACE_CPUPROFILER_EVENT_SCOPE(SavePackages);
 
-		UE_LOG(LogWorldParitionCommandletUtils, Log, TEXT("Saving %d packages."), PackagesToSave.Num());
+		UE_LOG(LogWorldPartitionCommandletUtils, Log, TEXT("Saving %d packages."), PackagesToSave.Num());
 		for (UPackage* PackageToSave : PackagesToSave)
 		{
 			if (!Save(PackagesToSave))
@@ -141,7 +141,7 @@ namespace WorldPartitionCommandletHelpers
 		FSavePackageArgs SaveArgs;
 		SaveArgs.TopLevelFlags = RF_Standalone;
 
-		UE_LOG(LogWorldParitionCommandletUtils, Verbose, TEXT("Saving Package %s."), *PackageFileName);
+		UE_LOG(LogWorldPartitionCommandletUtils, Verbose, TEXT("Saving Package %s."), *PackageFileName);
 		return UPackage::SavePackage(PackageToSave, nullptr, *PackageFileName, SaveArgs);
 	}
 
@@ -162,7 +162,7 @@ namespace WorldPartitionCommandletHelpers
 
 	bool AddToSourceControl(UPackage* PackageToAdd, FPackageSourceControlHelper& SCHelper)
 	{
-		UE_LOG(LogWorldParitionCommandletUtils, Verbose, TEXT("Adding Package %s."), *SourceControlHelpers::PackageFilename(PackageToAdd));
+		UE_LOG(LogWorldPartitionCommandletUtils, Verbose, TEXT("Adding Package %s."), *SourceControlHelpers::PackageFilename(PackageToAdd));
 
 		return SCHelper.AddToSourceControl(PackageToAdd);
 	}
@@ -176,7 +176,7 @@ namespace WorldPartitionCommandletHelpers
 
 	bool Delete(UPackage* PackageToDelete, FPackageSourceControlHelper& SCHelper)
 	{
-		UE_LOG(LogWorldParitionCommandletUtils, Verbose, TEXT("Deleting Package %s."), *SourceControlHelpers::PackageFilename(PackageToDelete));
+		UE_LOG(LogWorldPartitionCommandletUtils, Verbose, TEXT("Deleting Package %s."), *SourceControlHelpers::PackageFilename(PackageToDelete));
 
 		return SCHelper.Delete(PackageToDelete);
 	}
diff --git a/Engine/Source/Editor/UnrealEd/Private/Commandlets/WorldPartitionDataLayerToAssetCommandlet.cpp b/Engine/Source/Editor/UnrealEd/Private/Commandlets/WorldPartitionDataLayerToAssetCommandlet.cpp
--- a/Engine/Source/Editor/UnrealEd/Private/Commandlets/WorldPartitionDataLayerToAssetCommandlet.cpp
+++ b/Engine/Source/Editor/UnrealEd/Private/Commandlets/WorldPartitionDataLayerToAssetCommandlet.cpp
@@ -323,7 +323,7 @@ bool UDataLayerToAssetCommandlet::InitializeFromCommandLine(TArray<FString>& Tok
 	if (Switches.Contains(TEXT("Verbose")))
 	{
 		LogDataLayerToAssetCommandlet.SetVerbosity(ELogVerbosity::Verbose);
-		WorldPartitionCommandletHelpers::LogWorldParitionCommandletUtils.SetVerbosity(ELogVerbosity::Verbose);
+		WorldPartitionCommandletHelpers::LogWorldPartitionCommandletUtils.SetVerbosity(ELogVerbosity::Verbose);
 	}
 
 	bPerformSavePackages = Switches.Contains(TEXT("NoSave")) == false;
@@ -507,10 +507,10 @@ bool UDataLayerToAssetCommandlet::RemapActorDataLayersToAssets(TStrongObjectPtr<
 		}
 		else
 		{
-			const TArray<FName>& ActDescDataLayers = ActorDescInstance->GetDataLayerInstanceNames();
+			const FDataLayerInstanceNames& ActDescDataLayers = ActorDescInstance->GetDataLayerInstanceNames();
 			if (!ActDescDataLayers.IsEmpty())
 			{
-				FString DataLayerString = FString::JoinBy(ActDescDataLayers, TEXT(", "), [](const FName& DataLayerName) { return DataLayerName.ToString(); });
+				FString DataLayerString = FString::JoinBy(ActDescDataLayers.ToArray(), TEXT(", "), [](const FName& DataLayerName) { return DataLayerName.ToString(); });
 
 				UE_LOG(LogDataLayerToAssetCommandlet, Error, TEXT("Actor %s failed to load. Its data layers %s will not be remapped to a data layer asset."),
 					*ActorDescInstance->GetActorName().ToString(), *DataLayerString);
diff --git a/Engine/Source/Editor/UnrealEd/Private/EditorActorFolders.cpp b/Engine/Source/Editor/UnrealEd/Private/EditorActorFolders.cpp
--- a/Engine/Source/Editor/UnrealEd/Private/EditorActorFolders.cpp
+++ b/Engine/Source/Editor/UnrealEd/Private/EditorActorFolders.cpp
@@ -591,7 +591,8 @@ void FActorFolders::OnExecuteActorEditorContextAction(UWorld* InWorld, const EAc
 			SetActorEditorContextFolder(*InWorld, FFolder::GetWorldRootFolder(InWorld));
 			break;
 		case EActorEditorContextAction::PushContext:
-			GetOrCreateWorldFolders(*InWorld).PushActorEditorContext();
+		case EActorEditorContextAction::PushDuplicateContext:
+			GetOrCreateWorldFolders(*InWorld).PushActorEditorContext(InType == EActorEditorContextAction::PushDuplicateContext);
 			break;
 		case EActorEditorContextAction::PopContext:
 			if (UWorldFolders** Folders = (UWorldFolders**)WorldFolders.Find(InWorld))
diff --git a/Engine/Source/Editor/UnrealEd/Private/EditorEngine.cpp b/Engine/Source/Editor/UnrealEd/Private/EditorEngine.cpp
--- a/Engine/Source/Editor/UnrealEd/Private/EditorEngine.cpp
+++ b/Engine/Source/Editor/UnrealEd/Private/EditorEngine.cpp
@@ -4406,6 +4406,15 @@ bool UEditorEngine::CanParentActors( const AActor* ParentActor, const AActor* Ch
 		return false;
 	}
 
+	if (ChildActor->GetExternalDataLayerAsset() != ParentActor->GetExternalDataLayerAsset())
+	{
+		if (ReasonText)
+		{
+			*ReasonText = NSLOCTEXT("ActorAttachmentError", "WrongExternalDataLayer_AttachmentError", "Actors need to be in the same external data layer");
+		}
+		return false;
+	}
+
 	if(ParentRoot->IsAttachedTo( ChildRoot ))
 	{
 		if (ReasonText)
diff --git a/Engine/Source/Editor/UnrealEd/Private/EditorExporters.cpp b/Engine/Source/Editor/UnrealEd/Private/EditorExporters.cpp
--- a/Engine/Source/Editor/UnrealEd/Private/EditorExporters.cpp
+++ b/Engine/Source/Editor/UnrealEd/Private/EditorExporters.cpp
@@ -79,7 +79,7 @@
 #include "IMaterialBakingModule.h"
 #include "MaterialBakingStructures.h"
 #include "MaterialOptions.h"
-
+#include "WorldPartition/DataLayer/ExternalDataLayerAsset.h"
 #include "Components/SkeletalMeshComponent.h"
 #include "Camera/CameraComponent.h"
 #include "Components/LightComponent.h"
@@ -563,10 +563,12 @@ bool ULevelExporterT3D::ExportText( const FExportObjectInnerContext* Context, UO
 				FString ActorFolderPath = (Level->IsUsingActorFolders() ? FString::Printf(TEXT(" ActorFolderPath=\"%s\""), *Actor->GetFolderPath().ToString()) : TEXT(""));
 				FString CopyPasteId = (Actor->CopyPasteId != INDEX_NONE) ? FString::Printf(TEXT(" CopyPasteId=%d"), Actor->CopyPasteId ) : TEXT("");
 				FString ContentBundleGuid = (Actor->GetContentBundleGuid().IsValid() ? FString::Printf(TEXT(" ActorContentBundleGuid=%s"), *Actor->GetContentBundleGuid().ToString()) : TEXT(""));
-				Ar.Logf( TEXT("%sBegin Actor Class=%s Name=%s Archetype=%s%s%s%s%s%s%s%s"), 
+				FString ExternalDataLayer = (Actor->GetExternalDataLayerAsset() ? FString::Printf(TEXT(" ExternalDataLayerAsset=%s"), *FSoftObjectPath(Actor->GetExternalDataLayerAsset()).ToString()) : TEXT(""));
+
+				Ar.Logf( TEXT("%sBegin Actor Class=%s Name=%s Archetype=%s%s%s%s%s%s%s%s%s"), 
 					FCString::Spc(TextIndent), *Actor->GetClass()->GetPathName(), *Actor->GetName(),
 					*FObjectPropertyBase::GetExportPath(Actor->GetArchetype(), nullptr, nullptr, (PortFlags | PPF_Delimited) & ~PPF_ExportsNotFullyQualified), 
-					*ParentActorString, *SocketNameString, *GroupActor, *GroupFolder, *ActorFolderPath, *CopyPasteId, *ContentBundleGuid );
+					*ParentActorString, *SocketNameString, *GroupActor, *GroupFolder, *ActorFolderPath, *CopyPasteId, *ContentBundleGuid, *ExternalDataLayer);
 
 				// When exporting for diffs, export paths can cause false positives. since diff files don't get imported, we can
 				// skip adding this info the file.
diff --git a/Engine/Source/Editor/UnrealEd/Private/Factories/EditorFactories.cpp b/Engine/Source/Editor/UnrealEd/Private/Factories/EditorFactories.cpp
--- a/Engine/Source/Editor/UnrealEd/Private/Factories/EditorFactories.cpp
+++ b/Engine/Source/Editor/UnrealEd/Private/Factories/EditorFactories.cpp
@@ -201,6 +201,8 @@
 #include "AssetToolsModule.h"
 #include "IAssetTools.h"
 #include "WorldPartition/ContentBundle/ContentBundleActivationScope.h"
+#include "WorldPartition/DataLayer/ExternalDataLayerAsset.h"
+#include "ActorEditorContext/ScopedActorEditorContextSetExternalDataLayerAsset.h"
 
 #include "DDSFile.h"
 #include "IESConverter.h"
@@ -918,6 +920,11 @@ UObject* ULevelFactory::FactoryCreateText
 				FParse::Value(Str, TEXT("ActorContentBundleGuid="), ActorContentBundleGuid);
 				FContentBundleActivationScope Scope(ActorContentBundleGuid);
 
+				FString ExternalDataLayerAssetPathStr;
+				FParse::Value(Str, TEXT("ExternalDataLayerAsset="), ExternalDataLayerAssetPathStr);
+				UExternalDataLayerAsset* ExternalDataLayerAsset = Cast<UExternalDataLayerAsset>(FSoftObjectPath(ExternalDataLayerAssetPathStr).TryLoad());
+				FScopedActorEditorContextSetExternalDataLayerAsset EDLScope(ExternalDataLayerAsset);
+
 				// If we're pasting from a class that belongs to a map we need to duplicate the class and use that instead
 				if (FBlueprintEditorUtils::IsAnonymousBlueprintClass(TempClass))
 				{
diff --git a/Engine/Source/Editor/UnrealEd/Private/LevelEditorDragDropHandler.cpp b/Engine/Source/Editor/UnrealEd/Private/LevelEditorDragDropHandler.cpp
--- a/Engine/Source/Editor/UnrealEd/Private/LevelEditorDragDropHandler.cpp
+++ b/Engine/Source/Editor/UnrealEd/Private/LevelEditorDragDropHandler.cpp
@@ -13,6 +13,7 @@
 #include "Editor/UnrealEdEngine.h"
 #include "Editor.h"
 #include "HitProxies.h"
+#include "Misc/MessageDialog.h"
 
 #define LOCTEXT_NAMESPACE "UnrealEd"
 
@@ -20,48 +21,110 @@ ULevelEditorDragDropHandler::ULevelEditorDragDropHandler()
 {
 }
 
+bool ULevelEditorDragDropHandler::PassesFilter(UWorld* World, const FAssetData& AssetData, TUniquePtr<FLevelEditorDragDropWorldSurrogateReferencingObject>* OutWorldSurrogateReferencingObject)
+{
+	if (World && bRunAssetFilter)
+	{
+		auto AssetPassesFilter = [](const FAssetData& InAssetData, const UObject* InReferencingAsset, FText* OutFailureReason = nullptr)
+		{
+			FAssetReferenceFilterContext AssetReferenceFilterContext;
+			AssetReferenceFilterContext.ReferencingAssets.Add(FAssetData(InReferencingAsset));
+
+			TSharedPtr<IAssetReferenceFilter> AssetReferenceFilter = GEditor->MakeAssetReferenceFilter(AssetReferenceFilterContext);
+			return AssetReferenceFilter.IsValid() ? AssetReferenceFilter->PassesFilter(InAssetData, OutFailureReason) : true;
+		};
+
+		ULevel* CurrentLevel = World->GetCurrentLevel();
+		UWorld* CurrentLevelOuterWorld = CurrentLevel ? Cast<UWorld>(CurrentLevel->GetOuter()) : nullptr;
+		UWorld* ReferencingWorld = CurrentLevelOuterWorld ? CurrentLevelOuterWorld : World;
+
+		FText FailureReason;
+		bool bPassesFilter = AssetPassesFilter(AssetData, ReferencingWorld, &FailureReason);
+		if (!bPassesFilter)
+		{
+			// Try to find with referencing asset matching the AssetData
+			if (OnLevelEditorDragDropWorldSurrogateReferencingObjectDelegate.IsBound())
+			{
+				TUniquePtr<FLevelEditorDragDropWorldSurrogateReferencingObject> Object = OnLevelEditorDragDropWorldSurrogateReferencingObjectDelegate.Execute(ReferencingWorld, AssetData.GetSoftObjectPath());
+				const UObject* ReferencingAsset = Object.IsValid() ? Object->GetValue() : nullptr;
+				bPassesFilter = ReferencingAsset ? AssetPassesFilter(AssetData, ReferencingAsset) : false;
+				if (bPassesFilter)
+				{
+					if (OutWorldSurrogateReferencingObject)
+					{
+						*OutWorldSurrogateReferencingObject = MoveTemp(Object);
+					}
+				}
+			}
+		}
+
+		if (!bPassesFilter)
+		{
+			bCanDrop = false;
+			HintText = FailureReason;
+			return false;
+		}
+	}
+
+	return true;
+}
+
 bool ULevelEditorDragDropHandler::PreviewDropObjectsAtCoordinates(int32 MouseX, int32 MouseY, UWorld* World, FViewport* Viewport, const FAssetData& AssetData)
 {
-	bCanDrop = false;
+	bCanDrop = true;
 	HintText = FText::GetEmpty();
 
-	if ( !AssetData.IsValid() )
+	if (!AssetData.IsValid())
 	{
 		bCanDrop = false;
 		return false;
 	}
 
+	return PassesFilter(World, AssetData);
+}
+
+bool ULevelEditorDragDropHandler::PreDropObjectsAtCoordinates(int32 MouseX, int32 MouseY, UWorld* World, FViewport* Viewport, const TArray<UObject*>& DroppedObjects, TArray<AActor*>& OutNewActors)
+{
+	WorldSurrogateReferencingObject.Reset();
 	if (World && bRunAssetFilter)
 	{
-		ULevel* CurrentLevel = World->GetCurrentLevel();
-		UWorld* CurrentLevelOuterWorld = CurrentLevel ? Cast<UWorld>(CurrentLevel->GetOuter()) : nullptr;  
-		UWorld* ReferencingWorld = CurrentLevelOuterWorld ? CurrentLevelOuterWorld : World;
-		FAssetReferenceFilterContext AssetReferenceFilterContext;
-		AssetReferenceFilterContext.ReferencingAssets.Add(FAssetData(ReferencingWorld));
-		
-		TSharedPtr<IAssetReferenceFilter> AssetReferenceFilter = GEditor->MakeAssetReferenceFilter(AssetReferenceFilterContext);
-		if (AssetReferenceFilter.IsValid())
+		if (DroppedObjects.IsEmpty() ||
+			(!PassesFilter(World, FAssetData(DroppedObjects[0]), &WorldSurrogateReferencingObject)))
+		{
+			return false;
+		}
+	
+		auto GetValue = [](const TUniquePtr<FLevelEditorDragDropWorldSurrogateReferencingObject>& InSurrogateObject)
+		{
+			const FLevelEditorDragDropWorldSurrogateReferencingObject* Object = InSurrogateObject.Get();
+			return Object ? Object->GetValue() : nullptr;
+		};
+
+		for (int i=1; i<DroppedObjects.Num(); ++i)
 		{
-			FText FailureReason;
-			if (!AssetReferenceFilter->PassesFilter(AssetData, &FailureReason))
+			TUniquePtr<FLevelEditorDragDropWorldSurrogateReferencingObject> DroppedSurrogateObject;
+			if (!PassesFilter(World, FAssetData(DroppedObjects[i]), &DroppedSurrogateObject) || 
+				(GetValue(DroppedSurrogateObject) != GetValue(WorldSurrogateReferencingObject)))
 			{
-				bCanDrop = false;
-				HintText = FailureReason;
+				FMessageDialog::Open(EAppMsgType::Ok, LOCTEXT("DroppedObjectsDifferentSourcePluginPath", "Dropped objects are not part of the same source plugin path."));
+				WorldSurrogateReferencingObject.Reset();
 				return false;
 			}
 		}
-	}
-
-	return true;
-}
 
-bool ULevelEditorDragDropHandler::PreDropObjectsAtCoordinates(int32 MouseX, int32 MouseY, UWorld* World, FViewport* Viewport, const TArray<UObject*>& DroppedObjects, TArray<AActor*>& OutNewActors)
-{
+		return !WorldSurrogateReferencingObject.IsValid() || WorldSurrogateReferencingObject->OnPreDropObjects(World, DroppedObjects);
+	}
 	return true;
 }
 
 bool ULevelEditorDragDropHandler::PostDropObjectsAtCoordinates(int32 MouseX, int32 MouseY, UWorld* World, FViewport* Viewport, const TArray<UObject*>& DroppedObjects)
 {
+	if (WorldSurrogateReferencingObject.IsValid())
+	{
+		bool bSuccess = WorldSurrogateReferencingObject->OnPostDropObjects(World, DroppedObjects);
+		WorldSurrogateReferencingObject.Reset();
+		return bSuccess;
+	}
 	return true;
 }
 
diff --git a/Engine/Source/Editor/UnrealEd/Private/Subsystems/ActorEditorContextSubsystem.cpp b/Engine/Source/Editor/UnrealEd/Private/Subsystems/ActorEditorContextSubsystem.cpp
--- a/Engine/Source/Editor/UnrealEd/Private/Subsystems/ActorEditorContextSubsystem.cpp
+++ b/Engine/Source/Editor/UnrealEd/Private/Subsystems/ActorEditorContextSubsystem.cpp
@@ -104,7 +104,7 @@ void UActorEditorContextSubsystem::ResetContext(IActorEditorContextClient* Clien
 	}
 }
 
-void UActorEditorContextSubsystem::PushContext()
+void UActorEditorContextSubsystem::PushContext(bool bDuplicateContext)
 {
 	UWorld* World = GetWorld();
 	if (Clients.IsEmpty() || !World)
@@ -114,7 +114,7 @@ void UActorEditorContextSubsystem::PushContext()
 
 	for (IActorEditorContextClient* Client : Clients)
 	{
-		Client->OnExecuteActorEditorContextAction(World, EActorEditorContextAction::PushContext);
+		Client->OnExecuteActorEditorContextAction(World, bDuplicateContext ? EActorEditorContextAction::PushDuplicateContext : EActorEditorContextAction::PushContext);
 	}
 	ActorEditorContextSubsystemChanged.Broadcast();
 }
diff --git a/Engine/Source/Editor/UnrealEd/Private/WorldFolders.cpp b/Engine/Source/Editor/UnrealEd/Private/WorldFolders.cpp
--- a/Engine/Source/Editor/UnrealEd/Private/WorldFolders.cpp
+++ b/Engine/Source/Editor/UnrealEd/Private/WorldFolders.cpp
@@ -228,11 +228,14 @@ bool UWorldFolders::SetActorEditorContextFolder(const FFolder& InFolder)
 	return false;
 }
 
-void UWorldFolders::PushActorEditorContext()
+void UWorldFolders::PushActorEditorContext(bool bDuplicateContext)
 {
 	Modify();
 	CurrentFolderStack.Push(CurrentFolder);
-	CurrentFolder.Reset();
+	if (!bDuplicateContext)
+	{
+		CurrentFolder.Reset();
+	}
 }
 
 void UWorldFolders::PopActorEditorContext()
diff --git a/Engine/Source/Editor/UnrealEd/Private/WorldPartition/WorldPartitionFoliageBuilder.cpp b/Engine/Source/Editor/UnrealEd/Private/WorldPartition/WorldPartitionFoliageBuilder.cpp
--- a/Engine/Source/Editor/UnrealEd/Private/WorldPartition/WorldPartitionFoliageBuilder.cpp
+++ b/Engine/Source/Editor/UnrealEd/Private/WorldPartition/WorldPartitionFoliageBuilder.cpp
@@ -122,7 +122,7 @@ bool UWorldPartitionFoliageBuilder::RunInternal(UWorld* World, const FCellInfo&
 		{
 			FPackageSourceControlHelper PackageHelper;
 
-			const FDataLayerEditorContext DataLayerEditorContext(World, ActorDescInstance->GetDataLayerInstanceNames());
+			const FDataLayerEditorContext DataLayerEditorContext(World, ActorDescInstance->GetDataLayerInstanceNames().ToArray());
 			const FActorPartitionIdentifier ActorPartitionIdentifier(PartitionActorDesc->GetActorNativeClass(), IFA->GetGridGuid(), DataLayerEditorContext.GetHash());
 			const bool bShouldIncludeGridSizeInName = ActorPartitionIdentifier.GetClass()->GetDefaultObject<APartitionActor>()->ShouldIncludeGridSizeInName(World, ActorPartitionIdentifier);
 			const FString ExpectedActorName = APartitionActor::GetActorName(World, ActorPartitionIdentifier, PartitionActorDesc->GridSize, 
diff --git a/Engine/Source/Editor/UnrealEd/Public/Commandlets/WorldPartitionCommandletHelpers.h b/Engine/Source/Editor/UnrealEd/Public/Commandlets/WorldPartitionCommandletHelpers.h
--- a/Engine/Source/Editor/UnrealEd/Public/Commandlets/WorldPartitionCommandletHelpers.h
+++ b/Engine/Source/Editor/UnrealEd/Public/Commandlets/WorldPartitionCommandletHelpers.h
@@ -12,7 +12,7 @@ class UWorld;
 
 namespace WorldPartitionCommandletHelpers
 {
-	UNREALED_API DECLARE_LOG_CATEGORY_EXTERN(LogWorldParitionCommandletUtils, Log, All);
+	UNREALED_API DECLARE_LOG_CATEGORY_EXTERN(LogWorldPartitionCommandletUtils, Log, All);
 
 	UWorld* LoadAndInitWorld(const FString& LevelToLoad);
 	UWorld* LoadWorld(const FString& LevelToLoad);
diff --git a/Engine/Source/Editor/UnrealEd/Public/IActorEditorContextClient.h b/Engine/Source/Editor/UnrealEd/Public/IActorEditorContextClient.h
--- a/Engine/Source/Editor/UnrealEd/Public/IActorEditorContextClient.h
+++ b/Engine/Source/Editor/UnrealEd/Public/IActorEditorContextClient.h
@@ -26,6 +26,7 @@ enum class EActorEditorContextAction
 	ApplyContext,
 	ResetContext,
 	PushContext,
+	PushDuplicateContext,
 	PopContext,
 	InitializeContextFromActor,
 };
diff --git a/Engine/Source/Editor/UnrealEd/Public/LevelEditorDragDropHandler.h b/Engine/Source/Editor/UnrealEd/Public/LevelEditorDragDropHandler.h
--- a/Engine/Source/Editor/UnrealEd/Public/LevelEditorDragDropHandler.h
+++ b/Engine/Source/Editor/UnrealEd/Public/LevelEditorDragDropHandler.h
@@ -4,11 +4,30 @@
 
 #include "CoreMinimal.h"
 #include "Internationalization/Text.h"
+#include "UObject/GCObject.h"
 #include "LevelEditorDragDropHandler.generated.h"
 
 class FViewport;
 struct FAssetData;
 
+struct FLevelEditorDragDropWorldSurrogateReferencingObject : public FGCObject
+{
+public:
+	FLevelEditorDragDropWorldSurrogateReferencingObject(const UObject* InSurrogateObject) { SurrogateObject = InSurrogateObject; }
+	bool IsValid() const { return GetValue() != nullptr; }
+	const UObject* GetValue() const { return SurrogateObject; }
+	virtual bool OnPreDropObjects(UWorld* World, const TArray<UObject*>& DroppedObjects) { return true; }
+	virtual bool OnPostDropObjects(UWorld* World, const TArray<UObject*>& DroppedObjects) { return true; }
+
+	//~ Begin FGCObject Interface
+	virtual void AddReferencedObjects(FReferenceCollector& Collector) override { Collector.AddReferencedObject(SurrogateObject); }
+	virtual FString GetReferencerName() const override { return TEXT("FLevelEditorDragDropWorldSurrogateReferencingObject"); }
+	//~ End FGCObject Interface
+
+protected:
+	TObjectPtr<const UObject> SurrogateObject;
+};
+
 UCLASS(transient, MinimalAPI)
 class ULevelEditorDragDropHandler : public UObject
 {
@@ -24,8 +43,13 @@ public:
 
 	bool GetCanDrop() const { return bCanDrop; }
 	FText GetPreviewDropHintText() const { return HintText; }
+
+	DECLARE_DELEGATE_RetVal_TwoParams(TUniquePtr<FLevelEditorDragDropWorldSurrogateReferencingObject>, FOnLevelEditorDragDropWorldSurrogateReferencingObject, UWorld*, const FSoftObjectPath&);
+	FOnLevelEditorDragDropWorldSurrogateReferencingObject& OnLevelEditorDragDropWorldSurrogateReferencingObject() { return OnLevelEditorDragDropWorldSurrogateReferencingObjectDelegate; }
 	
 protected:
+	bool PassesFilter(UWorld* World, const FAssetData& AssetData, TUniquePtr<FLevelEditorDragDropWorldSurrogateReferencingObject>* OutWorldSurrogateReferencingObject = nullptr);
+
 	bool bRunAssetFilter = true;
 
 	/** True if it's valid to drop the object at the location queried */
@@ -33,4 +57,7 @@ protected:
 
 	/** Optional hint text that may be returned to the user. */
 	FText HintText;
+
+	TUniquePtr<FLevelEditorDragDropWorldSurrogateReferencingObject> WorldSurrogateReferencingObject;
+	FOnLevelEditorDragDropWorldSurrogateReferencingObject OnLevelEditorDragDropWorldSurrogateReferencingObjectDelegate;
 };
diff --git a/Engine/Source/Editor/UnrealEd/Public/Subsystems/ActorEditorContextSubsystem.h b/Engine/Source/Editor/UnrealEd/Public/Subsystems/ActorEditorContextSubsystem.h
--- a/Engine/Source/Editor/UnrealEd/Public/Subsystems/ActorEditorContextSubsystem.h
+++ b/Engine/Source/Editor/UnrealEd/Public/Subsystems/ActorEditorContextSubsystem.h
@@ -31,7 +31,7 @@ public:
 	UNREALED_API void UnregisterClient(IActorEditorContextClient* Client);
 	UNREALED_API void ResetContext();
 	UNREALED_API void ResetContext(IActorEditorContextClient* Client);
-	UNREALED_API void PushContext();
+	UNREALED_API void PushContext(bool bDuplicateContext = false);
 	UNREALED_API void PopContext();
 	UNREALED_API void InitializeContextFromActor(AActor* Actor);
 	UNREALED_API TArray<IActorEditorContextClient*> GetDisplayableClients() const;
diff --git a/Engine/Source/Editor/UnrealEd/Public/WorldFolders.h b/Engine/Source/Editor/UnrealEd/Public/WorldFolders.h
--- a/Engine/Source/Editor/UnrealEd/Public/WorldFolders.h
+++ b/Engine/Source/Editor/UnrealEd/Public/WorldFolders.h
@@ -82,7 +82,7 @@ public:
 	UNREALED_API bool SetIsFolderExpanded(const FFolder& InFolder, bool bIsExpanded);
 	UNREALED_API FFolder GetActorEditorContextFolder(bool bMustMatchCurrentLevel) const;
 	UNREALED_API bool SetActorEditorContextFolder(const FFolder& InFolder);
-	UNREALED_API void PushActorEditorContext();
+	UNREALED_API void PushActorEditorContext(bool bDuplicateContext = false);
 	UNREALED_API void PopActorEditorContext();
 	UNREALED_API bool ContainsFolder(const FFolder& InFolder) const;
 	UNREALED_API void ForEachFolder(TFunctionRef<bool(const FFolder&)> Operation);
diff --git a/Engine/Source/Editor/WorldPartitionEditor/Private/WorldPartition/ContentBundle/ContentBundleEditorSubsystem.cpp b/Engine/Source/Editor/WorldPartitionEditor/Private/WorldPartition/ContentBundle/ContentBundleEditorSubsystem.cpp
--- a/Engine/Source/Editor/WorldPartitionEditor/Private/WorldPartition/ContentBundle/ContentBundleEditorSubsystem.cpp
+++ b/Engine/Source/Editor/WorldPartitionEditor/Private/WorldPartition/ContentBundle/ContentBundleEditorSubsystem.cpp
@@ -83,7 +83,8 @@ void UContentBundleEditingSubmodule::OnExecuteActorEditorContextAction(UWorld* I
 		}
 		break;
 	case EActorEditorContextAction::PushContext:
-		PushContentBundleEditing();
+	case EActorEditorContextAction::PushDuplicateContext:
+		PushContentBundleEditing(InType == EActorEditorContextAction::PushDuplicateContext);
 		break;
 	case EActorEditorContextAction::PopContext:
 		PopContentBundleEditing();
@@ -233,12 +234,15 @@ void UContentBundleEditingSubmodule::StopEditing(TSharedPtr<FContentBundleEditor
 	UE_LOG(LogContentBundle, Log, TEXT("[CB: %s] Content Bundle is no longer being edited"), *ContentBundleEditor->GetDescriptor()->GetDisplayName());
 }
 
-void UContentBundleEditingSubmodule::PushContentBundleEditing()
+void UContentBundleEditingSubmodule::PushContentBundleEditing(bool bDuplicateContext)
 {
 	Modify();
 
 	EditingContentBundlesStack.Add(EditingContentBundleGuid);
-	DeactivateCurrentContentBundleEditing();
+	if (!bDuplicateContext)
+	{
+		DeactivateCurrentContentBundleEditing();
+	}
 }
 
 void UContentBundleEditingSubmodule::PopContentBundleEditing()
@@ -399,14 +403,19 @@ bool UContentBundleEditorSubsystem::DeactivateContentBundleEditing(TSharedPtr<FC
 	return ContentBundleEditingSubModule->DeactivateContentBundleEditing(ContentBundleEditor);
 }
 
+bool UContentBundleEditorSubsystem::DeactivateCurrentContentBundleEditing() const
+{
+	return ContentBundleEditingSubModule->DeactivateCurrentContentBundleEditing();
+}
+
 bool UContentBundleEditorSubsystem::IsContentBundleEditingActivated(TSharedPtr<FContentBundleEditor>& ContentBundleEditor) const
 {
 	return ContentBundleEditor.IsValid() && IsEditingContentBundle(ContentBundleEditor->GetDescriptor()->GetGuid());
 }
 
-void UContentBundleEditorSubsystem::PushContentBundleEditing()
+void UContentBundleEditorSubsystem::PushContentBundleEditing(bool bDuplicateContext)
 {
-	ContentBundleEditingSubModule->PushContentBundleEditing();
+	ContentBundleEditingSubModule->PushContentBundleEditing(bDuplicateContext);
 }
 
 void UContentBundleEditorSubsystem::PopContentBundleEditing()
diff --git a/Engine/Source/Editor/WorldPartitionEditor/Private/WorldPartition/ContentBundle/Outliner/ContentBundleMode.cpp b/Engine/Source/Editor/WorldPartitionEditor/Private/WorldPartition/ContentBundle/Outliner/ContentBundleMode.cpp
--- a/Engine/Source/Editor/WorldPartitionEditor/Private/WorldPartition/ContentBundle/Outliner/ContentBundleMode.cpp
+++ b/Engine/Source/Editor/WorldPartitionEditor/Private/WorldPartition/ContentBundle/Outliner/ContentBundleMode.cpp
@@ -6,6 +6,7 @@
 #include "WorldPartition/ContentBundle/ContentBundleEditorSubsystem.h"
 #include "WorldPartition/ContentBundle/ContentBundleEditor.h"
 #include "WorldPartition/ContentBundle/ContentBundleDescriptor.h"
+#include "DataLayer/DataLayerEditorSubsystem.h"
 #include "Framework/Commands/UIAction.h"
 #include "Framework/Commands/UICommandList.h"
 #include "Toolkits/GlobalEditorCommonCommands.h"
@@ -181,7 +182,11 @@ void FContentBundleMode::RegisterContextMenu()
 							FCanExecuteAction::CreateLambda([this, SceneOutliner]
 							{
 								TSharedPtr<FContentBundleEditor> SelectedContentBundleEditor = GetSelectedContentBundlePin(SceneOutliner);
-								return SelectedContentBundleEditor && !SelectedContentBundleEditor->IsBeingEdited();
+								if (SelectedContentBundleEditor && !SelectedContentBundleEditor->IsBeingEdited())
+								{
+									return !UDataLayerEditorSubsystem::Get()->GetActorEditorContextCurrentExternalDataLayer();
+								}
+								return false;
 							})
 						));
 
@@ -333,8 +338,11 @@ void FContentBundleMode::ToggleContentBundleActivation(const TWeakPtr<FContentBu
 				}
 				else
 				{
-					const FScopedTransaction Transaction(LOCTEXT("MakeCurrentContentBundle", "Make Current Content Bundle"));
-					ContentBundleEditingSubModule->ActivateContentBundleEditing(ContentBundlePin);
+					if (!UDataLayerEditorSubsystem::Get()->GetActorEditorContextCurrentExternalDataLayer())
+					{
+						const FScopedTransaction Transaction(LOCTEXT("MakeCurrentContentBundle", "Make Current Content Bundle"));
+						ContentBundleEditingSubModule->ActivateContentBundleEditing(ContentBundlePin);
+					}
 				}
 			}
 		}
diff --git a/Engine/Source/Editor/WorldPartitionEditor/Private/WorldPartitionEditorModule.cpp b/Engine/Source/Editor/WorldPartitionEditor/Private/WorldPartitionEditorModule.cpp
--- a/Engine/Source/Editor/WorldPartitionEditor/Private/WorldPartitionEditorModule.cpp
+++ b/Engine/Source/Editor/WorldPartitionEditor/Private/WorldPartitionEditorModule.cpp
@@ -9,6 +9,7 @@
 #include "WorldPartition/LoaderAdapter/LoaderAdapterShape.h"
 #include "WorldPartition/SWorldPartitionEditor.h"
 #include "WorldPartition/SWorldPartitionEditorGridSpatialHash.h"
+#include "WorldPartition/Customizations/ExternalDataLayerUIDStructCustomization.h"
 #include "WorldPartition/Customizations/WorldPartitionDetailsCustomization.h"
 #include "WorldPartition/Customizations/WorldPartitionHLODDetailsCustomization.h"
 #include "WorldPartition/Customizations/WorldPartitionRuntimeSpatialHashDetailsCustomization.h"
@@ -208,6 +209,7 @@ void FWorldPartitionEditorModule::StartupModule()
 	PropertyEditor.RegisterCustomClassLayout("WorldPartitionHLOD", FOnGetDetailCustomizationInstance::CreateStatic(&FWorldPartitionHLODDetailsCustomization::MakeInstance));
 	PropertyEditor.RegisterCustomClassLayout("WorldDataLayers", FOnGetDetailCustomizationInstance::CreateStatic(&FWorldDataLayersActorDetails::MakeInstance));
 	PropertyEditor.RegisterCustomClassLayout("WorldPartitionEditorPerProjectUserSettings", FOnGetDetailCustomizationInstance::CreateStatic(&FWorldPartitionEditorPerProjectUserSettingsCustomization::MakeInstance));
+	PropertyEditor.RegisterCustomPropertyTypeLayout("ExternalDataLayerUID", FOnGetPropertyTypeCustomizationInstance::CreateStatic(&FExternalDataLayerUIDStructCustomization::MakeInstance));
 
 	FWorldPartitionClassDescRegistry().Get().Initialize();
 
diff --git a/Engine/Source/Editor/WorldPartitionEditor/Public/WorldPartition/ContentBundle/ContentBundleEditorSubsystem.h b/Engine/Source/Editor/WorldPartitionEditor/Public/WorldPartition/ContentBundle/ContentBundleEditorSubsystem.h
--- a/Engine/Source/Editor/WorldPartitionEditor/Public/WorldPartition/ContentBundle/ContentBundleEditorSubsystem.h
+++ b/Engine/Source/Editor/WorldPartitionEditor/Public/WorldPartition/ContentBundle/ContentBundleEditorSubsystem.h
@@ -70,7 +70,7 @@ protected:
 	//~ End UContentBundleEditorSubsystemModule interface
 
 private:
-	void PushContentBundleEditing();
+	void PushContentBundleEditing(bool bDuplicateContext);
 	void PopContentBundleEditing();
 	void StartEditing(TSharedPtr<FContentBundleEditor>& ContentBundleEditor);
 	void StopEditing(TSharedPtr<FContentBundleEditor>& ContentBundleEditor);
@@ -132,9 +132,10 @@ public:
 	bool IsEditingContentBundle(const FGuid& ContentBundleGuid) const;
 	bool ActivateContentBundleEditing(TSharedPtr<FContentBundleEditor>& ContentBundleEditor) const override;
 	bool DeactivateContentBundleEditing(TSharedPtr<FContentBundleEditor>& ContentBundleEditor) const override;
+	virtual bool DeactivateCurrentContentBundleEditing() const override;
 	bool IsContentBundleEditingActivated(TSharedPtr<FContentBundleEditor>& ContentBundleEditor) const;
 
-	virtual void PushContentBundleEditing() override;
+	virtual void PushContentBundleEditing(bool bDuplicateContext) override;
 	virtual void PopContentBundleEditing() override;
 
 private:
diff --git a/Engine/Source/Editor/WorldPartitionEditor/Public/WorldPartition/WorldPartitionEditorSettings.h b/Engine/Source/Editor/WorldPartitionEditor/Public/WorldPartition/WorldPartitionEditorSettings.h
--- a/Engine/Source/Editor/WorldPartitionEditor/Public/WorldPartition/WorldPartitionEditorSettings.h
+++ b/Engine/Source/Editor/WorldPartitionEditor/Public/WorldPartition/WorldPartitionEditorSettings.h
@@ -17,7 +17,7 @@ public:
 
 	virtual void PostEditChangeProperty(FPropertyChangedEvent& PropertyChangedEvent) override;
 
-	UPROPERTY(Config, EditAnywhere, Category = MapConversion, Meta = (ToolTip = "Commandlet class to use for World Parition conversion"))
+	UPROPERTY(Config, EditAnywhere, Category = MapConversion, Meta = (ToolTip = "Commandlet class to use for World Partition conversion"))
 	TSubclassOf<UWorldPartitionConvertCommandlet> CommandletClass;
 
 	UPROPERTY(Config, EditAnywhere, Category = Foliage, Meta = (ClampMin = 3200, ToolTip= "Editor grid size used for instance foliage actors in World Partition worlds"))
diff --git a/Engine/Source/Editor/WorldPartitionEditor/WorldPartitionEditor.Build.cs b/Engine/Source/Editor/WorldPartitionEditor/WorldPartitionEditor.Build.cs
--- a/Engine/Source/Editor/WorldPartitionEditor/WorldPartitionEditor.Build.cs
+++ b/Engine/Source/Editor/WorldPartitionEditor/WorldPartitionEditor.Build.cs
@@ -27,7 +27,8 @@ public class WorldPartitionEditor : ModuleRules
 				"RenderCore",
 				"Renderer",
 				"RHI",
-				"SceneOutliner"
+				"SceneOutliner",
+				"DataLayerEditor"
 			}
 		);
 
diff --git a/Engine/Source/Runtime/AssetRegistry/Private/AssetRegistry.cpp b/Engine/Source/Runtime/AssetRegistry/Private/AssetRegistry.cpp
--- a/Engine/Source/Runtime/AssetRegistry/Private/AssetRegistry.cpp
+++ b/Engine/Source/Runtime/AssetRegistry/Private/AssetRegistry.cpp
@@ -5419,6 +5419,8 @@ bool FAssetRegistryImpl::ShouldSkipGatheredAsset(FAssetData& AssetData)
 		// /PackageRoot/__ExternalActors__/RelPathFromPackageRootToMap/#/##/#######
 		// OR
 		// /PackageRoot/__ExternalActors__/ContentBundle/######/RelPathFromPackageRootToMap/#/##/#######
+		// OR
+		// /PackageRoot/__ExternalActors__/EDL/######/ObjectPathPackageRoot/RelPathFromPackageRootToMap/#/##/#######
 		// Package roots do not need to be the same; ContentBundles can be injected into /Game maps from plugins
 		PackageNamePackageRoot = FPackageName::SplitPackageNameRoot(PackageNameStr, &PackageNameRelPath);
 		ObjectPathPackageRoot = FPackageName::SplitPackageNameRoot(ObjectPathPackageName, &ObjectPathRelPath);
@@ -5432,13 +5434,27 @@ bool FAssetRegistryImpl::ShouldSkipGatheredAsset(FAssetData& AssetData)
 			return true;
 		}
 
-		bool bAllowValidation = true;
+		bool bIsEDLActor = false;
+		bool bIsPluginActor = false;
 		FStringView PackageNameRelPathAfterExternalActorRoot = PackageNameRelPath.RightChop(ExternalActorsFolderName.Len() + 1);
 		FStringView ContentBundleDirName(TEXTVIEW("ContentBundle"));
+		FStringView ExternalDataLayerDirName(TEXTVIEW("EDL"));
 		if (PackageNameRelPathAfterExternalActorRoot.StartsWith(ContentBundleDirName))
 		{
-			bAllowValidation = false; // Don't allow validation unless we succeed in finding the new relpath
 			PackageNameRelPathAfterExternalActorRoot.RightChopInline(ContentBundleDirName.Len());
+			bIsPluginActor = true;
+		}
+		else if (PackageNameRelPathAfterExternalActorRoot.StartsWith(ExternalDataLayerDirName))
+		{
+			PackageNameRelPathAfterExternalActorRoot.RightChopInline(ExternalDataLayerDirName.Len());
+			bIsEDLActor = true;
+			bIsPluginActor = true;
+		}
+
+		bool bAllowValidation = true;
+		if (bIsPluginActor)
+		{
+			bAllowValidation = false; // Don't allow validation unless we succeed in finding the new relpath
 			if (PackageNameRelPathAfterExternalActorRoot.StartsWith(TEXT("/")))
 			{
 				PackageNameRelPathAfterExternalActorRoot.RightChopInline(1);
@@ -5447,7 +5463,23 @@ bool FAssetRegistryImpl::ShouldSkipGatheredAsset(FAssetData& AssetData)
 				if (NextSlash != INDEX_NONE)
 				{
 					PackageNameRelPathAfterExternalActorRoot.RightChopInline(NextSlash + 1);
-					bAllowValidation = true;
+					// EDL path keeps ObjectPathPackageRoot
+					if (bIsEDLActor)
+					{
+						if (PackageNameRelPathAfterExternalActorRoot.StartsWith(ObjectPathPackageRoot))
+						{
+							PackageNameRelPathAfterExternalActorRoot.RightChopInline(ObjectPathPackageRoot.Len());
+							if (PackageNameRelPathAfterExternalActorRoot.StartsWith(TEXT("/")))
+							{
+								PackageNameRelPathAfterExternalActorRoot.RightChopInline(1);
+								bAllowValidation = true;
+							}
+						}
+					}
+					else
+					{
+						bAllowValidation = true;
+					}
 				}
 			}
 		}
diff --git a/Engine/Source/Runtime/Core/Public/UObject/UE5ReleaseStreamObjectVersion.h b/Engine/Source/Runtime/Core/Public/UObject/UE5ReleaseStreamObjectVersion.h
--- a/Engine/Source/Runtime/Core/Public/UObject/UE5ReleaseStreamObjectVersion.h
+++ b/Engine/Source/Runtime/Core/Public/UObject/UE5ReleaseStreamObjectVersion.h
@@ -153,6 +153,9 @@ struct FUE5ReleaseStreamObjectVersion
 
 		// Disabled opencolorio display configuration by default
 		OpenColorIODisabledDisplayConfigurationDefault,
+
+		// Serialize ExternalDataLayerAsset in WorldPartitionActorDesc
+		WorldPartitionExternalDataLayers,
 		
 		// -----<new versions can be added above this line>-------------------------------------------------
 		VersionPlusOne,
diff --git a/Engine/Source/Runtime/Engine/Classes/Atmosphere/AtmosphericFog.h b/Engine/Source/Runtime/Engine/Classes/Atmosphere/AtmosphericFog.h
--- a/Engine/Source/Runtime/Engine/Classes/Atmosphere/AtmosphericFog.h
+++ b/Engine/Source/Runtime/Engine/Classes/Atmosphere/AtmosphericFog.h
@@ -19,7 +19,7 @@ class UE_DEPRECATED(4.26, "Please use the SkyAtmosphere actor instead.") AAtmosp
 
 private:
 #if WITH_EDITOR
-	virtual bool IsDataLayerTypeSupported(TSubclassOf<UDataLayerInstance> DataLayerType) const override { return true; }
+	virtual bool ActorTypeSupportsDataLayer() const override { return true; }
 #endif
 
 	/** Main fog component */
diff --git a/Engine/Source/Runtime/Engine/Classes/Components/HeterogeneousVolumeComponent.h b/Engine/Source/Runtime/Engine/Classes/Components/HeterogeneousVolumeComponent.h
--- a/Engine/Source/Runtime/Engine/Classes/Components/HeterogeneousVolumeComponent.h
+++ b/Engine/Source/Runtime/Engine/Classes/Components/HeterogeneousVolumeComponent.h
@@ -128,7 +128,7 @@ private:
 	TObjectPtr<class UHeterogeneousVolumeComponent> HeterogeneousVolumeComponent;
 
 #if WITH_EDITOR
-	virtual bool IsDataLayerTypeSupported(TSubclassOf<UDataLayerInstance> DataLayerType)  const override { return true; }
+	virtual bool ActorTypeSupportsDataLayer()  const override { return true; }
 #endif
 
 };
\ No newline at end of file
diff --git a/Engine/Source/Runtime/Engine/Classes/Components/SkyAtmosphereComponent.h b/Engine/Source/Runtime/Engine/Classes/Components/SkyAtmosphereComponent.h
--- a/Engine/Source/Runtime/Engine/Classes/Components/SkyAtmosphereComponent.h
+++ b/Engine/Source/Runtime/Engine/Classes/Components/SkyAtmosphereComponent.h
@@ -293,7 +293,7 @@ class ASkyAtmosphere : public AInfo
 
 private:
 #if WITH_EDITOR
-	virtual bool IsDataLayerTypeSupported(TSubclassOf<UDataLayerInstance> DataLayerType) const override { return true; }
+	virtual bool ActorTypeSupportsDataLayer() const override { return true; }
 #endif
 
 	UPROPERTY(BlueprintReadOnly, VisibleAnywhere, Category = Atmosphere, meta = (AllowPrivateAccess = "true"))
diff --git a/Engine/Source/Runtime/Engine/Classes/Engine/Level.h b/Engine/Source/Runtime/Engine/Classes/Engine/Level.h
--- a/Engine/Source/Runtime/Engine/Classes/Engine/Level.h
+++ b/Engine/Source/Runtime/Engine/Classes/Engine/Level.h
@@ -1236,10 +1236,21 @@ public:
 
 	/**
 	 * Create an package for this actor
+	 * @param InLevelPackage the level package used when building the actor package name
+	 * @param InActorPackagingScheme the packaging scheme to use
 	 * @param InActorPath the fully qualified actor path, in the format: 'Outermost.Outer.Name'
 	 * @return the created package
 	 */
-	static ENGINE_API UPackage* CreateActorPackage(UPackage* InLevelPackage, EActorPackagingScheme ActorPackagingScheme, const FString& InActorPath);
+	static ENGINE_API UPackage* CreateActorPackage(UPackage* InLevelPackage, EActorPackagingScheme InActorPackagingScheme, const FString& InActorPath);
+
+	/**
+	 * Create an package for this actor
+	 * @param InBaseDir the base directory used when building the actor package name
+	 * @param InActorPackagingScheme the packaging scheme to use
+	 * @param InActorPath the fully qualified actor path, in the format: 'Outermost.Outer.Name'
+	 * @return the created package
+	 */
+	static ENGINE_API UPackage* CreateActorPackage(const FString& InBaseDir, EActorPackagingScheme InActorPackagingScheme, const FString& InActorPath);
 
 	/**
 	 * Detach or reattach all level actors to from/to their external package
@@ -1371,6 +1382,7 @@ private:
 	TSet<FGuid> GetDeletedAndUnreferencedActorFolders() const;
 
 	friend struct FLevelActorFoldersHelper;
+	friend struct FSetWorldPartitionRuntimeCell;
 	friend class FWorldPartitionLevelHelper;
 	friend class UActorFolder;
 	friend class AWorldDataLayers;
@@ -1389,6 +1401,18 @@ private:
 #endif // WITH_EDITOR
 };
 
+#if WITH_EDITOR
+struct FSetWorldPartitionRuntimeCell
+{
+private:
+	FSetWorldPartitionRuntimeCell(ULevel* InLevel, const FSoftObjectPath& InWorldPartitionRuntimeCell)
+	{
+		InLevel->WorldPartitionRuntimeCell = InWorldPartitionRuntimeCell;
+	}
+	friend class FWorldPartitionLevelHelper;
+	friend class UWorldPartition;
+};
+#endif 
 
 /**
  * Macro for wrapping Delegates in TScopedCallback
diff --git a/Engine/Source/Runtime/Engine/Classes/GameFramework/Actor.h b/Engine/Source/Runtime/Engine/Classes/GameFramework/Actor.h
--- a/Engine/Source/Runtime/Engine/Classes/GameFramework/Actor.h
+++ b/Engine/Source/Runtime/Engine/Classes/GameFramework/Actor.h
@@ -47,6 +47,7 @@ class FActorTransactionAnnotation;
 class FComponentInstanceDataCache;
 class UDEPRECATED_DataLayer;
 class UDataLayerAsset;
+class UExternalDataLayerAsset;
 class UDataLayerInstance;
 class AWorldDataLayers;
 class IWorldPartitionCell;
@@ -971,6 +972,9 @@ protected:
 
 	TArray<TSoftObjectPtr<UDataLayerAsset>> PreEditChangeDataLayers;
 
+	UPROPERTY(VisibleAnywhere, AdvancedDisplay, Category = DataLayers, TextExportTransient, NonTransactional)
+	TObjectPtr<const UExternalDataLayerAsset> ExternalDataLayerAsset;
+
 public:
 	/** The copy/paste id used to remap actors during copy operations */
 	uint32 CopyPasteId;
@@ -993,8 +997,9 @@ public:
 	 * Set the actor packaging mode.
 	 * @param bExternal will set the actor packaging mode to external if true, to internal otherwise
 	 * @param bShouldDirty should dirty or not the level package
+	 * @param ActorExternalPackage if non-null and bExternal is true, will use the provided package instead of creating one
 	 */
-	ENGINE_API void SetPackageExternal(bool bExternal, bool bShouldDirty = true);
+	ENGINE_API void SetPackageExternal(bool bExternal, bool bShouldDirty = true, UPackage* ActorExternalPackage = nullptr);
 
 	/**
 	 * Determine how this actor should be referenced by the level when external (saved in its own package).
@@ -1307,14 +1312,15 @@ public:
 	// DataLayers functions.
 #if WITH_EDITOR
 protected:
-	virtual bool IsDataLayerTypeSupported(TSubclassOf<UDataLayerInstance> DataLayerType) const { return true; }
+	UE_DEPRECATED(5.4, "Use ActorTypeSupportsDataLayer instead.")
+	virtual bool IsDataLayerTypeSupported(TSubclassOf<UDataLayerInstance> DataLayerType) const final { return false; }
 private:
 	virtual bool ActorTypeIsMainWorldOnly() const { return false; }
 	ENGINE_API TArray<const UDataLayerAsset*> ResolveDataLayerAssets(const TArray<TSoftObjectPtr<UDataLayerAsset>>& InDataLayerAssets) const;
 public:
 	ENGINE_API bool AddDataLayer(const UDataLayerInstance* DataLayerInstance);
 	ENGINE_API bool RemoveDataLayer(const UDataLayerInstance* DataLayerInstance);
-	ENGINE_API bool CanAddDataLayer(const UDataLayerInstance* InDataLayerInstance) const;
+	ENGINE_API bool CanAddDataLayer(const UDataLayerInstance* InDataLayerInstance, FText* OutReason = nullptr) const;
 
 	ENGINE_API TArray<const UDataLayerInstance*> RemoveAllDataLayers();
 	ENGINE_API bool SupportsDataLayerType(TSubclassOf<UDataLayerInstance> DataLayerType) const;
@@ -1326,7 +1332,11 @@ public:
 	static const FName GetDataLayerAssetsPropertyName() { return GET_MEMBER_NAME_CHECKED(AActor, DataLayerAssets); }
 	static const FName GetDataLayerPropertyName() { return GET_MEMBER_NAME_CHECKED(AActor, DataLayers); }
 
-	ENGINE_API TArray<const UDataLayerAsset*> GetDataLayerAssets() const;
+	ENGINE_API TArray<const UDataLayerAsset*> GetDataLayerAssets(bool bIncludeExternalDataLayerAsset = true) const;
+	ENGINE_API bool HasExternalContent() const;
+
+	virtual bool ActorTypeSupportsDataLayer() const { return true; }
+	virtual bool ActorTypeSupportsExternalDataLayer() const { return true; }
 
 	//~ Begin Deprecated
 
@@ -1343,9 +1353,6 @@ public:
 	UE_DEPRECATED(5.3, "Use AActor::SupportsDataLayerType(TSubclassOf<UDataLayerInstance>) instead.")
 	ENGINE_API bool SupportsDataLayer() const;
 
-	UE_DEPRECATED(5.3, "Use AActor::IsDataLayerTypeSupported(TSubclassOf<UDataLayerInstance>) instead.")
-	virtual bool ActorTypeSupportsDataLayer() const { return true; }
-
 	UE_DEPRECATED(5.3, "Use UDataLayerInstance::AddActor(AActor*) instead.")
 	bool AddDataLayer(const UDataLayerAsset* DataLayerAsset) { return false; }
 
@@ -1379,6 +1386,7 @@ public:
 	ENGINE_API bool ContainsDataLayer(const UDataLayerInstance* DataLayerInstance) const;
 	ENGINE_API bool HasDataLayers() const;
 	ENGINE_API bool HasContentBundle() const;
+	ENGINE_API const UExternalDataLayerAsset* GetExternalDataLayerAsset() const;
 
 private:
 	ENGINE_API TArray<const UDataLayerInstance*> GetDataLayerInstancesInternal(bool bUseLevelContext, bool bIncludeParentDataLayers = true) const;
@@ -4415,6 +4423,7 @@ private:
 	friend UWorld;
 	friend class FFoliageHelper;
 	friend class ULevelInstanceSubsystem;
+	friend class UExternalDataLayerInstance;
 };
 
 struct FSetActorGuid
@@ -4480,6 +4489,7 @@ private:
 		InActor->ContentBundleGuid = InContentBundleGuid;
 	}
 	friend class FContentBundleEditor;
+	friend class UGameFeatureActionConvertContentBundleWorldPartitionBuilder;
 };
 
 struct FAssignActorDataLayer
@@ -4490,6 +4500,7 @@ private:
 
 	friend class UDataLayerInstanceWithAsset;
 	friend class UDataLayerInstancePrivate;
+	friend class UExternalDataLayerInstance;
 	friend class ULevelInstanceSubsystem;
 };
 
diff --git a/Engine/Source/Runtime/Engine/Classes/GameFramework/Info.h b/Engine/Source/Runtime/Engine/Classes/GameFramework/Info.h
--- a/Engine/Source/Runtime/Engine/Classes/GameFramework/Info.h
+++ b/Engine/Source/Runtime/Engine/Classes/GameFramework/Info.h
@@ -30,7 +30,7 @@ public:
 
 private:
 #if WITH_EDITOR
-	virtual bool IsDataLayerTypeSupported(TSubclassOf<UDataLayerInstance> DataLayerType) const override { return false; }
+	virtual bool ActorTypeSupportsDataLayer() const override { return false; }
 #endif
 
 #if WITH_EDITORONLY_DATA
diff --git a/Engine/Source/Runtime/Engine/Private/Actor.cpp b/Engine/Source/Runtime/Engine/Private/Actor.cpp
--- a/Engine/Source/Runtime/Engine/Private/Actor.cpp
+++ b/Engine/Source/Runtime/Engine/Private/Actor.cpp
@@ -69,6 +69,8 @@
 #include "WorldPartition/DataLayer/WorldDataLayers.h"
 #include "WorldPartition/WorldPartitionActorDescUtils.h"
 #include "WorldPartition/DataLayer/DataLayerAsset.h"
+#include "WorldPartition/DataLayer/ExternalDataLayerAsset.h"
+#include "WorldPartition/DataLayer/ExternalDataLayerInstance.h"
 #include "WorldPartition/DataLayer/DataLayerManager.h"
 #include "WorldPartition/ContentBundle/ContentBundlePaths.h"
 
@@ -6266,7 +6268,6 @@ TArray<const UDataLayerInstance*> AActor::GetDataLayerInstancesInternal(bool bUs
 	}
 
 #if WITH_EDITOR
-	if (SupportsDataLayerType(UDataLayerInstance::StaticClass()))
 	{
 		TArray<const UDataLayerInstance*> DataLayerInstances;
 		if (UDataLayerManager* DataLayerManager = bUseLevelContext ? UDataLayerManager::GetDataLayerManager(this) : UDataLayerManager::GetDataLayerManager(GetWorld()))
@@ -6324,7 +6325,7 @@ bool AActor::ContainsDataLayer(const UDataLayerAsset* DataLayerAsset) const
 	}
 
 #if WITH_EDITOR
-	if (DataLayerAssets.Contains(DataLayerAsset))
+	if (DataLayerAssets.Contains(DataLayerAsset) || (ExternalDataLayerAsset == DataLayerAsset))
 	{
 		return true;
 	}
@@ -6338,6 +6339,20 @@ bool AActor::ContainsDataLayer(const UDataLayerAsset* DataLayerAsset) const
 	return false;
 }
 
+const UExternalDataLayerAsset* AActor::GetExternalDataLayerAsset() const
+{
+	if (const IWorldPartitionCell* Cell = GetWorldPartitionRuntimeCell())
+	{
+		const UExternalDataLayerInstance* ExternalDataLayerInstance = Cell->GetExternalDataLayerInstance();
+		return ExternalDataLayerInstance ? ExternalDataLayerInstance->GetExternalDataLayerAsset() : nullptr;
+	}
+#if WITH_EDITOR
+	return ExternalDataLayerAsset;
+#else
+	return nullptr;
+#endif
+}
+
 bool AActor::HasDataLayers() const
 {
 	if (const IWorldPartitionCell* Cell = GetWorldPartitionRuntimeCell())
diff --git a/Engine/Source/Runtime/Engine/Private/ActorEditor.cpp b/Engine/Source/Runtime/Engine/Private/ActorEditor.cpp
--- a/Engine/Source/Runtime/Engine/Private/ActorEditor.cpp
+++ b/Engine/Source/Runtime/Engine/Private/ActorEditor.cpp
@@ -28,6 +28,8 @@
 #include "WorldPartition/DataLayer/IDataLayerEditorModule.h"
 #include "WorldPartition/DataLayer/DataLayerInstanceWithAsset.h"
 #include "WorldPartition/DataLayer/DeprecatedDataLayerInstance.h"
+#include "WorldPartition/DataLayer/ExternalDataLayerInstance.h"
+#include "WorldPartition/DataLayer/ExternalDataLayerAsset.h"
 #include "LevelInstance/LevelInstanceSubsystem.h"
 #include "ActorFolder.h"
 #include "WorldPersistentFolders.h"
@@ -117,7 +119,7 @@ bool AActor::CanEditChange(const FProperty* PropertyThatWillChange) const
 		return false;
 	}
 
-	if (bIsDataLayersProperty && (!SupportsDataLayerType(UDataLayerInstance::StaticClass()) || !IsUserManaged() || GetAttachParentActor()))
+	if (bIsDataLayersProperty && (!SupportsDataLayerType(UDataLayerInstanceWithAsset::StaticClass()) || !IsUserManaged() || GetAttachParentActor()))
 	{
 		return false;
 	}
@@ -1066,8 +1068,10 @@ bool AActor::IsMainWorldOnly() const
 	}
 }
 
-void AActor::SetPackageExternal(bool bExternal, bool bShouldDirty)
+void AActor::SetPackageExternal(bool bExternal, bool bShouldDirty, UPackage* ActorExternalPackage)
 {
+	check(bExternal || !ActorExternalPackage);
+
 	// @todo_ow: Call FExternalPackageHelper::SetPackagingMode and keep calling the actor specific code here (components). 
 	//           The only missing part is GetExternalObjectsPath defaulting to a different folder than the one used by external actors.
 	if (bExternal == IsPackageExternal())
@@ -1081,7 +1085,7 @@ void AActor::SetPackageExternal(bool bExternal, bool bShouldDirty)
 	UPackage* LevelPackage = GetLevel()->GetPackage(); 
 	if (bExternal)
 	{
-		UPackage* NewActorPackage = ULevel::CreateActorPackage(LevelPackage, GetLevel()->GetActorPackagingScheme(), GetPathName());
+		UPackage* NewActorPackage = ActorExternalPackage ? ActorExternalPackage : ULevel::CreateActorPackage(LevelPackage, GetLevel()->GetActorPackagingScheme(), GetPathName());
 		SetExternalPackage(NewActorPackage);
 	}
 	else 
@@ -1651,9 +1655,14 @@ TArray<const UDataLayerAsset*> AActor::ResolveDataLayerAssets(const TArray<TSoft
 	return ResolvedAssets;
 }
 
-TArray<const UDataLayerAsset*> AActor::GetDataLayerAssets() const
+TArray<const UDataLayerAsset*> AActor::GetDataLayerAssets(bool bIncludeExternalDataLayerAsset) const
 {
-	return ResolveDataLayerAssets(DataLayerAssets);
+	TArray<const UDataLayerAsset*> ResolveDataLayerAsset = ResolveDataLayerAssets(DataLayerAssets);
+	if (bIncludeExternalDataLayerAsset && ExternalDataLayerAsset)
+	{
+		ResolveDataLayerAsset.Add(ExternalDataLayerAsset);
+	}
+	return ResolveDataLayerAsset;
 }
 
 TArray<FName> AActor::GetDataLayerInstanceNames() const
@@ -1668,6 +1677,12 @@ TArray<FName> AActor::GetDataLayerInstanceNames() const
 	return DataLayerInstanceNames;
 }
 
+bool AActor::HasExternalContent() const
+{
+	check(!ExternalDataLayerAsset || ExternalDataLayerAsset->GetUID().IsValid());
+	return ExternalDataLayerAsset ? true : GetContentBundleGuid().IsValid();
+}
+
 TArray<const UDataLayerInstance*> AActor::GetDataLayerInstancesForLevel() const
 {
 	const bool bUseLevelContext = true;
@@ -1684,14 +1699,15 @@ void AActor::FixupDataLayers(bool bRevertChangesOnLockedDataLayer /*= false*/)
 		return;
 	}
 
-	if (!SupportsDataLayerType(UDeprecatedDataLayerInstance::StaticClass()))
+	if (!SupportsDataLayerType(UDataLayerInstance::StaticClass()))
 	{
 		DataLayers.Empty();
+		DataLayerAssets.Empty();
 	}
 
-	if (!SupportsDataLayerType(UDataLayerInstanceWithAsset::StaticClass()))
+	if (ExternalDataLayerAsset && !SupportsDataLayerType(UExternalDataLayerInstance::StaticClass()))
 	{
-		DataLayerAssets.Empty();
+		ExternalDataLayerAsset = nullptr;
 	}
 
 	if (DataLayers.IsEmpty() && DataLayerAssets.IsEmpty())
@@ -1729,7 +1745,7 @@ void AActor::FixupDataLayers(bool bRevertChangesOnLockedDataLayer /*= false*/)
 	if (bRevertChangesOnLockedDataLayer)
 	{
 		TArray<const UDataLayerAsset*> ResolvedPreEditChangeAssets = ResolveDataLayerAssets(PreEditChangeDataLayers);
-		// Since it's not possible to prevent changes of particular elements of an array, rollback change on locked DataLayers.
+		// Since it's not possible to prevent changes of particular elements of an array, rollback change on read-only DataLayers.
 		TSet<const UDataLayerAsset*> PreEdit(ResolvedPreEditChangeAssets);
 		TSet<const UDataLayerAsset*> PostEdit(ResolvedAssets);
 
@@ -1739,7 +1755,7 @@ void AActor::FixupDataLayers(bool bRevertChangesOnLockedDataLayer /*= false*/)
 			for (const UDataLayerAsset* DataLayerAsset : Diff)
 			{
 				const UDataLayerInstance* DataLayerInstance = DataLayerManager->GetDataLayerInstance(DataLayerAsset);
-				if (DataLayerInstance && DataLayerInstance->IsLocked())
+				if (DataLayerInstance && DataLayerInstance->IsReadOnly())
 				{
 					return true;
 				}
@@ -1821,35 +1837,132 @@ bool AActor::SupportsDataLayerType(TSubclassOf<UDataLayerInstance> InDataLayerTy
 	ULevel* Level = GetLevel();
 	const bool bIsLevelNotPartitioned = Level ? !Level->bIsPartitioned : false;
 	const bool bHasComponentForceActorNoDataLayers = HasComponentForceActorNoDataLayers(this);
-	return (!bIsLevelNotPartitioned && 
+	const bool bActorTypeSupportsDataLayerType = InDataLayerType->IsChildOf<UExternalDataLayerInstance>() ? ActorTypeSupportsExternalDataLayer() : ActorTypeSupportsDataLayer();
+	
+	return (!bIsLevelNotPartitioned &&
 		!bHasComponentForceActorNoDataLayers &&
-		IsDataLayerTypeSupported(InDataLayerType) &&
+		bActorTypeSupportsDataLayerType &&
 		!FActorEditorUtils::IsABuilderBrush(this) &&
 		!GetClass()->GetDefaultObject<AActor>()->bHiddenEd);
 }
 
-bool AActor::CanAddDataLayer(const UDataLayerInstance* InDataLayerInstance) const
+bool AActor::CanAddDataLayer(const UDataLayerInstance* InDataLayerInstance, FText* OutReason) const
 {
-	if (InDataLayerInstance && SupportsDataLayerType(InDataLayerInstance->GetClass()))
+	auto PassesAssetReferenceFiltering = [](const UObject * InReferencingObject, const UDataLayerAsset * InDataLayerAsset, FText* OutReason)
 	{
-		if (const UDataLayerAsset* DataLayerAsset = InDataLayerInstance->GetAsset())
+		FAssetReferenceFilterContext AssetReferenceFilterContext;
+		AssetReferenceFilterContext.ReferencingAssets.Add(FAssetData(InReferencingObject));
+		TSharedPtr<IAssetReferenceFilter> AssetReferenceFilter = GEditor->MakeAssetReferenceFilter(AssetReferenceFilterContext);
+		return AssetReferenceFilter.IsValid() ? AssetReferenceFilter->PassesFilter(FAssetData(InDataLayerAsset), OutReason) : true;
+	};
+
+	if (!InDataLayerInstance)
+	{
+		if (OutReason)
 		{
-			return !DataLayerAssets.Contains(DataLayerAsset);
+			*OutReason = LOCTEXT("CantAddDataLayerInvalidDataLayerInstance", "Invalid data layer instance.");
 		}
-		else if (const UDeprecatedDataLayerInstance* DataLayerInstance = Cast<UDeprecatedDataLayerInstance>(InDataLayerInstance))
+		return false;
+	}
+
+	if (!SupportsDataLayerType(InDataLayerInstance->GetClass()))
+	{
+		if (OutReason)
 		{
-			return !DataLayers.Contains(DataLayerInstance->GetActorDataLayer());
+			*OutReason = LOCTEXT("CantAddDataLayerActorDoesntSupportDataLayerType", "Actor doesn't support this data layer type.");
 		}
+		return false;
 	}
 
-	return false;
+	if (const UDataLayerAsset* DataLayerAsset = InDataLayerInstance->GetAsset())
+	{
+		if (!PassesAssetReferenceFiltering(this, DataLayerAsset, OutReason))
+		{
+			return false;
+		}
+
+		if (const UExternalDataLayerAsset* InExternalDataLayerAsset = Cast<UExternalDataLayerAsset>(DataLayerAsset))
+		{
+			if (ContentBundleGuid.IsValid())
+			{
+				if (OutReason)
+				{
+					*OutReason = LOCTEXT("CantAddDataLayerActorAlreadyAssignedToContentBundle", "Actor is already assigned to a content bundle.");
+				}
+				return false;
+			}
+			else if (ExternalDataLayerAsset)
+			{
+				if (OutReason)
+				{
+					if (ExternalDataLayerAsset == DataLayerAsset)
+					{
+						*OutReason = LOCTEXT("CantAddDataLayerActorAlreadyAssignedToExternalDataLayer", "Actor is already assigned to this external data layer.");
+					}
+					else
+					{
+						*OutReason = LOCTEXT("CantAddDataLayerActorAlreadyAssignedToAnotherExternalDataLayer", "Actor is already assigned to another external data layer.");
+					}
+				}
+				return false;
+			}
+		}
+		else
+		{
+			if (DataLayerAssets.Contains(DataLayerAsset))
+			{
+				if (OutReason)
+				{
+					*OutReason = LOCTEXT("CantAddDataLayerActorAlreadyAssignedToDataLayer", "Actor is already assigned to this data layer.");
+				}
+				return false;
+			}
+		}
+	}
+	else
+	{
+		if (const UDeprecatedDataLayerInstance* DataLayerInstance = Cast<UDeprecatedDataLayerInstance>(InDataLayerInstance))
+		{
+			if (DataLayers.Contains(DataLayerInstance->GetActorDataLayer()))
+			{
+				if (OutReason)
+				{
+					*OutReason = LOCTEXT("CantAddDataLayerActorAlreadyAssignedToDataLayer", "Actor is already assigned to this data layer.");
+				}
+				return false;
+			}
+			return true;
+		}
+		else
+		{
+			if (OutReason)
+			{
+				*OutReason = LOCTEXT("CantAddDataLayerInvalidDataLayerAsset", "Invalid data layer asset.");
+			}
+			return false;
+		}
+	}
+
+	return true;
 }
 
 bool FAssignActorDataLayer::AddDataLayerAsset(AActor* InActor, const UDataLayerAsset* InDataLayerAsset)
 {
 	check(InDataLayerAsset != nullptr);
 
-	if (!InActor->DataLayerAssets.Contains(InDataLayerAsset))
+	if (const UExternalDataLayerAsset* ExternalDataLayerAsset = Cast<UExternalDataLayerAsset>(InDataLayerAsset))
+	{
+		if (!InActor->ExternalDataLayerAsset)
+		{
+			InActor->Modify();
+			InActor->ExternalDataLayerAsset = ExternalDataLayerAsset;
+			return true;
+		}
+
+		UE_CLOG(InActor->ExternalDataLayerAsset != ExternalDataLayerAsset, LogActor, Warning, TEXT("Trying to assign external data layer %s on actor %s while %s is already assigned."), 
+			*ExternalDataLayerAsset->GetPathName(), *InActor->GetActorNameOrLabel(), *InActor->ExternalDataLayerAsset->GetPathName());
+	}
+	else if (!InActor->DataLayerAssets.Contains(InDataLayerAsset))
 	{
 		InActor->Modify();
 		InActor->DataLayerAssets.Add(InDataLayerAsset);
@@ -1863,7 +1976,13 @@ bool FAssignActorDataLayer::RemoveDataLayerAsset(AActor* InActor, const UDataLay
 {
 	check(InDataLayerAsset != nullptr);
 	
-	if (InActor->DataLayerAssets.Contains(InDataLayerAsset))
+	if (InActor->ExternalDataLayerAsset == InDataLayerAsset)
+	{
+		InActor->Modify();
+		InActor->ExternalDataLayerAsset = nullptr;
+		return true;
+	}
+	else if (InActor->DataLayerAssets.Contains(InDataLayerAsset))
 	{
 		InActor->Modify();
 		InActor->DataLayerAssets.Remove(InDataLayerAsset);
@@ -1877,11 +1996,6 @@ bool FAssignActorDataLayer::RemoveDataLayerAsset(AActor* InActor, const UDataLay
 
 PRAGMA_DISABLE_DEPRECATION_WARNINGS
 
-bool AActor::SupportsDataLayer() const
-{
-	return SupportsDataLayerType(UDataLayerInstance::StaticClass());
-}
-
 bool AActor::AddDataLayer(const FActorDataLayer& ActorDataLayer)
 {
 	if (SupportsDataLayerType(UDataLayerInstance::StaticClass()) && !DataLayers.Contains(ActorDataLayer))
diff --git a/Engine/Source/Runtime/Engine/Private/ActorPartition/ActorPartitionSubsystem.cpp b/Engine/Source/Runtime/Engine/Private/ActorPartition/ActorPartitionSubsystem.cpp
--- a/Engine/Source/Runtime/Engine/Private/ActorPartition/ActorPartitionSubsystem.cpp
+++ b/Engine/Source/Runtime/Engine/Private/ActorPartition/ActorPartitionSubsystem.cpp
@@ -222,7 +222,7 @@ public:
 				(PartitionActorDesc->GridIndexZ == InCellCoord.Z) &&
 				(PartitionActorDesc->GridSize == InGridSize) &&
 				(PartitionActorDesc->GridGuid == InActorPartitionId.GetGridGuid()) &&
-				(FActorPartitionContextHash::Get(ActorDescInstance->GetContentBundleGuid(), ThisWorld, ActorDescInstance->GetDataLayerInstanceNames()) == InActorPartitionId.GetContextHash()))
+				(FActorPartitionContextHash::Get(ActorDescInstance->GetContentBundleGuid(), ThisWorld, ActorDescInstance->GetDataLayerInstanceNames().ToArray()) == InActorPartitionId.GetContextHash()))
 			{
 				AActor* DescActor = ActorDescInstance->GetActor();
 
diff --git a/Engine/Source/Runtime/Engine/Private/ExternalPackageHelper.cpp b/Engine/Source/Runtime/Engine/Private/ExternalPackageHelper.cpp
--- a/Engine/Source/Runtime/Engine/Private/ExternalPackageHelper.cpp
+++ b/Engine/Source/Runtime/Engine/Private/ExternalPackageHelper.cpp
@@ -16,10 +16,12 @@ EPackageFlags FExternalPackageHelper::GetDefaultExternalPackageFlags()
 	return (PKG_EditorOnly | PKG_ContainsMapData | PKG_NewlyCreated);
 }
 
-UPackage* FExternalPackageHelper::CreateExternalPackage(UObject* InObjectOuter, const FString& InObjectPath, EPackageFlags InFlags)
+UPackage* FExternalPackageHelper::CreateExternalPackage(UObject* InObjectOuter, const FString& InObjectPath, EPackageFlags InFlags, const UExternalDataLayerAsset* InExternalDataLayerAsset)
 {
-	UPackage* OutermostPackage = InObjectOuter->IsA<UPackage>() ? CastChecked<UPackage>(InObjectOuter) : InObjectOuter->GetOutermostObject()->GetPackage();
-	UPackage* Package = CreatePackage(*FExternalPackageHelper::GetExternalPackageName(OutermostPackage->GetName(), InObjectPath));
+	const UPackage* OutermostPackage = InObjectOuter->IsA<UPackage>() ? CastChecked<UPackage>(InObjectOuter) : InObjectOuter->GetOutermostObject()->GetPackage();
+	const FString RootPath = InExternalDataLayerAsset ? FExternalDataLayerHelper::GetExternalDataLayerLevelRootPath(InExternalDataLayerAsset, OutermostPackage->GetName()) : OutermostPackage->GetName();
+	const FString ExternalObjectPackageName = FExternalPackageHelper::GetExternalPackageName(RootPath, InObjectPath);
+	UPackage* Package = CreatePackage(*ExternalObjectPackageName);
 	Package->SetPackageFlags(InFlags);
 	return Package;
 }
@@ -36,7 +38,9 @@ void FExternalPackageHelper::SetPackagingMode(UObject* InObject, UObject* InObje
 
 	if (bInIsPackageExternal)
 	{
-		UPackage* NewObjectPackage = FExternalPackageHelper::CreateExternalPackage(InObjectOuter, InObject->GetPathName(), InExternalPackageFlags);
+		const IDataLayerInstanceProvider* DataLayerInstanceProvider = InObject->GetImplementingOuter<IDataLayerInstanceProvider>();
+		const UExternalDataLayerAsset* ExternalDataLayerAsset = DataLayerInstanceProvider ? DataLayerInstanceProvider->GetRootExternalDataLayerAsset() : nullptr;
+		UPackage* NewObjectPackage = FExternalPackageHelper::CreateExternalPackage(InObjectOuter, InObject->GetPathName(), InExternalPackageFlags, ExternalDataLayerAsset);
 		InObject->SetExternalPackage(NewObjectPackage);
 	}
 	else
diff --git a/Engine/Source/Runtime/Engine/Private/Level.cpp b/Engine/Source/Runtime/Engine/Private/Level.cpp
--- a/Engine/Source/Runtime/Engine/Private/Level.cpp
+++ b/Engine/Source/Runtime/Engine/Private/Level.cpp
@@ -4057,10 +4057,20 @@ TArray<UPackage*> ULevel::GetLoadedExternalObjectPackages() const
 		return InPath;
 	};
 
-	TArray<FString> ExternalObjectsPaths;
-	ExternalObjectsPaths.Add(SanitizeExternalPath(ULevel::GetExternalActorsPath(World->GetPackage(), (World->OriginalWorldName == NAME_None) ? World->GetName() : World->OriginalWorldName.ToString())));
-	ExternalObjectsPaths.Add(SanitizeExternalPath(FExternalPackageHelper::GetExternalObjectsPath(World->GetPackage(), (World->OriginalWorldName == NAME_None) ? World->GetName() : World->OriginalWorldName.ToString())));
-	ExternalObjectsPaths = ExternalObjectsPaths.FilterByPredicate([](const FString& Path) {
+	TSet<FString> ExternalObjectsPathSet;
+	ExternalObjectsPathSet.Add(SanitizeExternalPath(ULevel::GetExternalActorsPath(World->GetPackage(), (World->OriginalWorldName == NAME_None) ? World->GetName() : World->OriginalWorldName.ToString())));
+	ExternalObjectsPathSet.Add(SanitizeExternalPath(FExternalPackageHelper::GetExternalObjectsPath(World->GetPackage(), (World->OriginalWorldName == NAME_None) ? World->GetName() : World->OriginalWorldName.ToString())));
+
+	if (UWorldPartition* WorldPartition = GetWorldPartition())
+	{
+		WorldPartition->ForEachActorDescContainerInstance([&SanitizeExternalPath, &ExternalObjectsPathSet](UActorDescContainerInstance* ActorDescContainerInstance)
+		{
+			ExternalObjectsPathSet.Add(SanitizeExternalPath(ActorDescContainerInstance->GetExternalActorPath()));
+			ExternalObjectsPathSet.Add(SanitizeExternalPath(ActorDescContainerInstance->GetExternalObjectPath()));
+		});
+	}
+
+	TArray<FString> ExternalObjectsPaths = ExternalObjectsPathSet.Array().FilterByPredicate([](const FString& Path) {
 			FString Filename;
 			if (!Path.IsEmpty() && FPackageName::TryConvertLongPackageNameToFilename(Path, Filename))
 			{
@@ -4106,18 +4116,24 @@ TArray<UPackage*> ULevel::GetLoadedExternalObjectPackages() const
 	return ExternalObjectPackages.Array();
 }
 
-UPackage* ULevel::CreateActorPackage(UPackage* InLevelPackage, EActorPackagingScheme ActorPackagingScheme, const FString& InActorPath)
+static UPackage* CreateActorPackageInternal(const FString& InPackageName, const FString& InActorPath)
 {
-	const FString PackageName = GetActorPackageName(InLevelPackage, ActorPackagingScheme, InActorPath);
-	if (UPackage* ExistingActorPackage = FindObject<UPackage>(nullptr, *PackageName))
+	if (UPackage* ExistingActorPackage = FindObject<UPackage>(nullptr, *InPackageName))
 	{
 		check(ExistingActorPackage->HasAllPackagesFlags(PKG_EditorOnly | PKG_ContainsMapData));
 		return ExistingActorPackage;
 	}
 
-	UPackage* ActorPackage = CreatePackage(*PackageName);
+	UPackage* ActorPackage = CreatePackage(*InPackageName);
 	ActorPackage->SetPackageFlags(PKG_EditorOnly | PKG_ContainsMapData | PKG_NewlyCreated);
 
+	return ActorPackage;
+};
+
+UPackage* ULevel::CreateActorPackage(UPackage* InLevelPackage, EActorPackagingScheme InActorPackagingScheme, const FString& InActorPath)
+{
+	const FString PackageName = GetActorPackageName(InLevelPackage, InActorPackagingScheme, InActorPath);
+	UPackage* ActorPackage = CreateActorPackageInternal(PackageName, InActorPath);
 	// Should be prevented upstream but we propagate the flag to prevent issues in asset enumeration
 	if (!ensureMsgf(!(InLevelPackage->GetPackageFlags() & PKG_PlayInEditor), TEXT("Actor packages should not be created on PlayInEditor levels")))
 	{
@@ -4126,6 +4142,12 @@ UPackage* ULevel::CreateActorPackage(UPackage* InLevelPackage, EActorPackagingSc
 	return ActorPackage;
 }
 
+UPackage* ULevel::CreateActorPackage(const FString& InBaseDir, EActorPackagingScheme InActorPackagingScheme, const FString& InActorPath)
+{
+	const FString PackageName = GetActorPackageName(InBaseDir, InActorPackagingScheme, InActorPath);
+	return CreateActorPackageInternal(PackageName, InActorPath);
+}
+
 void ULevel::DetachAttachAllActorsPackages(bool bReattach)
 {
 	if (bReattach)
diff --git a/Engine/Source/Runtime/Engine/Private/LevelInstance/LevelInstanceSubsystem.cpp b/Engine/Source/Runtime/Engine/Private/LevelInstance/LevelInstanceSubsystem.cpp
--- a/Engine/Source/Runtime/Engine/Private/LevelInstance/LevelInstanceSubsystem.cpp
+++ b/Engine/Source/Runtime/Engine/Private/LevelInstance/LevelInstanceSubsystem.cpp
@@ -14,6 +14,7 @@
 
 #if WITH_EDITOR
 #include "Settings/LevelEditorMiscSettings.h"
+#include "ActorEditorContext/ScopedActorEditorContextSetExternalDataLayerAsset.h"
 #include "LevelInstance/LevelInstanceEditorLevelStreaming.h"
 #include "LevelInstance/ILevelInstanceEditorModule.h"
 #include "LevelInstance/LevelInstanceEditorInstanceActor.h"
@@ -28,6 +29,9 @@
 #include "WorldPartition/WorldPartitionActorDescInstance.h"
 #include "WorldPartition/DataLayer/DataLayerManager.h"
 #include "WorldPartition/DataLayer/DataLayerInstanceWithAsset.h"
+#include "WorldPartition/DataLayer/ExternalDataLayerAsset.h"
+#include "WorldPartition/DataLayer/ExternalDataLayerInstance.h"
+#include "WorldPartition/DataLayer/ExternalDataLayerManager.h"
 #include "WorldPartition/DataLayer/WorldDataLayersActorDesc.h"
 #include "WorldPartition/WorldPartitionMiniMap.h"
 #include "Misc/ScopedSlowTask.h"
@@ -1125,7 +1129,11 @@ ULevelStreamingLevelInstanceEditor* ULevelInstanceSubsystem::CreateNewStreamingL
 								// Validate that there's a valid Data Layer Instance for this asset in the source level and that this isn't a private Data Layer
 								if (!DataLayerAsset->IsPrivate())
 								{
-									SourceDataLayerAssets.Add(DataLayerAsset);
+									// @todo_ow: Add LevelInstance EDL support (For now skip all Data Layer Instances part of an EDL)
+									if (!DataLayerInstance->GetRootExternalDataLayerInstance())
+									{
+										SourceDataLayerAssets.Add(DataLayerAsset);
+									}
 								}
 							}
 						}
@@ -1163,7 +1171,8 @@ ILevelInstanceInterface* ULevelInstanceSubsystem::CreateLevelInstanceFrom(const
 		UE_LOG(LogLevelInstance, Warning, TEXT("Failed to create Level Instance from empty actor array"));
 		return nullptr;
 	}
-		
+	
+	TOptional<const UExternalDataLayerAsset*> CommonExternalDataLayerAsset;
 	FBox ActorLocationBox(ForceInit);
 	for (const AActor* ActorToMove : ActorsToMove)
 	{
@@ -1183,6 +1192,15 @@ ILevelInstanceInterface* ULevelInstanceSubsystem::CreateLevelInstanceFrom(const
 			UE_LOG(LogLevelInstance, Warning, TEXT("%s"), *Reason.ToString());
 			return nullptr;
 		}
+
+		if (!CommonExternalDataLayerAsset.IsSet())
+		{
+			CommonExternalDataLayerAsset = ActorToMove->GetExternalDataLayerAsset();
+		}
+		else if (CommonExternalDataLayerAsset.GetValue() != ActorToMove->GetExternalDataLayerAsset())
+		{
+			CommonExternalDataLayerAsset = nullptr;
+		}
 	}
 
 	FVector LevelInstanceLocation;
@@ -1268,12 +1286,21 @@ ILevelInstanceInterface* ULevelInstanceSubsystem::CreateLevelInstanceFrom(const
 	ULevel* LoadedLevel = LevelStreaming->GetLoadedLevel();
 	check(LoadedLevel);
 		
-	// @todo_ow : Decide if we want to re-create the same hierarchy as the source level.
 	for (AActor* Actor : LoadedLevel->Actors)
 	{
 		if (Actor)
 		{
+			// @todo_ow : Decide if we want to re-create the same hierarchy as the source level.
 			Actor->SetFolderPath_Recursively(NAME_None);
+
+			// @todo_ow: Add LevelInstance EDL support (For now, remove all Data Layers part of an EDL)
+			for (const UDataLayerInstance* DataLayerInstance : Actor->GetDataLayerInstances())
+			{
+				if (DataLayerInstance->GetRootExternalDataLayerInstance())
+				{
+					DataLayerInstance->RemoveActor(Actor);
+				}
+			}
 		}
 	}
 
@@ -1295,6 +1322,16 @@ ILevelInstanceInterface* ULevelInstanceSubsystem::CreateLevelInstanceFrom(const
 
 	check(ActorClass->ImplementsInterface(ULevelInstanceInterface::StaticClass()));
 
+	const UExternalDataLayerAsset* ExternalDataLayerAsset = CommonExternalDataLayerAsset.Get(nullptr);
+	UExternalDataLayerManager* ExternalDataLayerManager = UExternalDataLayerManager::GetExternalDataLayerManager(GetWorld());
+	UExternalDataLayerInstance* ExternalDataLayerInstance = ExternalDataLayerManager ? ExternalDataLayerManager->GetExternalDataLayerInstance(ExternalDataLayerAsset) : nullptr;
+	// @todo_ow: We temporarily allow adding the ExternalDataLayerInstance to the ActorEditorContext or else it wouldn't allow it to be added since the current level is not the persistent level.
+	if (ExternalDataLayerInstance)
+	{
+		ExternalDataLayerInstance->bSkipCheckReadOnlyForSubLevels = true;
+	}
+	FScopedActorEditorContextSetExternalDataLayerAsset EDLScope(ExternalDataLayerAsset);
+
 	if (!ActorClass->IsChildOf<APackedLevelActor>())
 	{
 		NewLevelInstanceActor = GetWorld()->SpawnActor<AActor>(ActorClass, SpawnParams);
@@ -1327,6 +1364,10 @@ ILevelInstanceInterface* ULevelInstanceSubsystem::CreateLevelInstanceFrom(const
 			NewLevelInstanceActor = GetWorld()->SpawnActor<APackedLevelActor>(APackedLevelActor::StaticClass(), SpawnParams);
 		}
 	}
+	if (ExternalDataLayerInstance)
+	{
+		ExternalDataLayerInstance->bSkipCheckReadOnlyForSubLevels = false;
+	}
 	
 	check(NewLevelInstanceActor);
 	check(NewLevelInstanceActor->GetActorGuid() == LevelInstanceActorGuid);
@@ -1335,7 +1376,7 @@ ILevelInstanceInterface* ULevelInstanceSubsystem::CreateLevelInstanceFrom(const
 	NewLevelInstance->SetWorldAsset(WorldPtr);
 	NewLevelInstanceActor->SetActorLocation(LevelInstanceLocation);
 	NewLevelInstanceActor->SetActorLabel(WorldPtr.GetAssetName());
-	
+
 	// Actors were moved and kept their World positions so when saving we want their positions to actually be relative to the LevelInstance Actor
 	// so we set the LevelTransform and we mark the level as having moved its actors. 
 	// On Level save FLevelUtils::RemoveEditorTransform will fixup actor transforms to make them relative to the LevelTransform.
@@ -1515,9 +1556,9 @@ void ULevelInstanceSubsystem::BreakLevelInstance_Impl(ILevelInstanceInterface* L
 					}
 
 					// Apply the same data layer settings to the actors to move out
-					if (Actor->SupportsDataLayerType(UDataLayerInstance::StaticClass()))
+					for (const UDataLayerInstance* DataLayerInstance : LevelInstanceDataLayerInstances)
 					{
-						for (const UDataLayerInstance* DataLayerInstance : LevelInstanceDataLayerInstances)
+						if (Actor->SupportsDataLayerType(DataLayerInstance->GetClass()))
 						{
 							if (const UDataLayerAsset* DataLayerAsset = DataLayerInstance->GetAsset())
 							{
@@ -1710,7 +1751,7 @@ bool ULevelInstanceSubsystem::CanMoveActorToLevel(const AActor* Actor, FText* Ou
 			{
 				if (OutReason != nullptr)
 				{
-					*OutReason = LOCTEXT("CanMoveActorLevelEditing", "Can't move Level Instance actor while it is being edited");
+					*OutReason = LOCTEXT("CantMoveActorLevelEditing", "Can't move Level Instance actor while it is being edited");
 				}
 				return false;
 			}
@@ -1730,7 +1771,7 @@ bool ULevelInstanceSubsystem::CanMoveActorToLevel(const AActor* Actor, FText* Ou
 			{
 				if (OutReason != nullptr)
 				{
-					*OutReason = LOCTEXT("CanMoveActorToLevelChildEditing", "Can't move Level Instance actor while one of its child Level Instance is being edited");
+					*OutReason = LOCTEXT("CantMoveActorToLevelChildEditing", "Can't move Level Instance actor while one of its child Level Instance is being edited");
 				}
 				return false;
 			}
diff --git a/Engine/Source/Runtime/Engine/Private/Tests/WorldPartition/WorldPartitionPathConversionTests.cpp b/Engine/Source/Runtime/Engine/Private/Tests/WorldPartition/WorldPartitionPathConversionTests.cpp
--- a/Engine/Source/Runtime/Engine/Private/Tests/WorldPartition/WorldPartitionPathConversionTests.cpp
+++ b/Engine/Source/Runtime/Engine/Private/Tests/WorldPartition/WorldPartitionPathConversionTests.cpp
@@ -51,13 +51,15 @@ namespace WorldPartitionTests
 		}
 		
 		check(WorldPartition);
-		UActorDescContainerInstance* ActorDescMainContainerInstance = WorldPartition->GetActorDescContainerInstance();
-		if (!TestNotNull(TEXT("Missing World Partition Container"), ActorDescMainContainerInstance))
+		UActorDescContainerInstance* BaseActorDescContainerInstance = WorldPartition->GetActorDescContainerInstance();
+		if (!TestNotNull(TEXT("Missing World Partition Container"), BaseActorDescContainerInstance))
 		{
 			return false;
 		}
 
-		UWorldPartition::FGenerateStreamingParams Params = UWorldPartition::FGenerateStreamingParams().SetActorDescContainerInstance(ActorDescMainContainerInstance);
+		FActorDescContainerInstanceCollection Collection({ TObjectPtr<UActorDescContainerInstance>(BaseActorDescContainerInstance) });
+		UWorldPartition::FGenerateStreamingParams Params = UWorldPartition::FGenerateStreamingParams()
+			.SetContainerInstanceCollection(Collection, FStreamingGenerationContainerInstanceCollection::ECollectionType::BaseAndEDLs);
 		UWorldPartition::FGenerateStreamingContext Context;
 
 		if (!TestTrue(TEXT("World Partition Generate Streaming"), WorldPartition->GenerateContainerStreaming(Params, Context)))
@@ -65,7 +67,7 @@ namespace WorldPartitionTests
 			return false;
 		}
 
-		FWorldPartitionReference ActorRef(ActorDescMainContainerInstance, FGuid(TEXT("5D9F93BA407A811AFDDDAAB4F1CECC6A")));
+		FWorldPartitionReference ActorRef(BaseActorDescContainerInstance, FGuid(TEXT("5D9F93BA407A811AFDDDAAB4F1CECC6A")));
 		if (!TestTrue(TEXT("Invalid Actor Reference"), ActorRef.IsValid()))
 		{
 			return false;
@@ -77,7 +79,7 @@ namespace WorldPartitionTests
 			return false;
 		}
 
-		FWorldPartitionHandle ActorHandle(ActorDescMainContainerInstance, FGuid(TEXT("0D2B04D240BE5DE58FE437A8D2DBF5C9")));
+		FWorldPartitionHandle ActorHandle(BaseActorDescContainerInstance, FGuid(TEXT("0D2B04D240BE5DE58FE437A8D2DBF5C9")));
 		if (!TestTrue(TEXT("Invalid Actor Handle"), ActorHandle.IsValid()))
 		{
 			return false;
diff --git a/Engine/Source/Runtime/Engine/Private/WorldPartition/ActorDescContainer.cpp b/Engine/Source/Runtime/Engine/Private/WorldPartition/ActorDescContainer.cpp
--- a/Engine/Source/Runtime/Engine/Private/WorldPartition/ActorDescContainer.cpp
+++ b/Engine/Source/Runtime/Engine/Private/WorldPartition/ActorDescContainer.cpp
@@ -17,6 +17,7 @@
 #include "WorldPartition/WorldPartitionActorDescUtils.h"
 #include "WorldPartition/WorldPartitionClassDescRegistry.h"
 #include "WorldPartition/DataLayer/DataLayerManager.h"
+#include "WorldPartition/DataLayer/ExternalDataLayerAsset.h"
 
 UActorDescContainer::FActorDescContainerInitializeDelegate UActorDescContainer::OnActorDescContainerInitialized;
 #endif
@@ -35,8 +36,17 @@ void UActorDescContainer::Initialize(const FInitializeParams& InitParams)
 
 	check(!bContainerInitialized);
 	ContainerPackageName = InitParams.PackageName;
-	TArray<FAssetData> Assets;
+	if (InitParams.ExternalDataLayerAsset)
+	{
+		ensure(!InitParams.ContentBundleGuid.IsValid());
+		ExternalDataLayerAsset = InitParams.ExternalDataLayerAsset;
+	}
+	else if (InitParams.ContentBundleGuid.IsValid())
+	{
+		ContentBundleGuid = InitParams.ContentBundleGuid;
+	}
 
+	TArray<FAssetData> Assets;
 	if (!ContainerPackageName.IsNone())
 	{
 		const FString ContainerExternalActorsPath = GetExternalActorPath();
@@ -189,11 +199,29 @@ FString UActorDescContainer::GetExternalActorPath() const
 	return ULevel::GetExternalActorsPath(ContainerPackageName.ToString());
 }
 
-bool UActorDescContainer::IsActorDescHandled(const AActor* Actor) const
+FString UActorDescContainer::GetExternalObjectPath() const
+{
+	return FExternalPackageHelper::GetExternalObjectsPath(ContainerPackageName.ToString());
+}
+
+bool UActorDescContainer::HasExternalContent() const
+{
+	check(!ExternalDataLayerAsset || ExternalDataLayerAsset->GetUID().IsValid());
+	return ExternalDataLayerAsset ? true : GetContentBundleGuid().IsValid();
+}
+
+bool UActorDescContainer::IsActorDescHandled(const AActor* InActor) const
 {
-	if (Actor->GetContentBundleGuid() == GetContentBundleGuid())
+	// Actor External Content Guid must match Container's External Content Guid to be considered
+	// AWorldDataLayers actors are an exception as they don't have an External Content Guid
+	const bool bIsCandidateActor = InActor->IsA<AWorldDataLayers>() ||
+		(!HasExternalContent() && !InActor->HasExternalContent()) ||
+		(ExternalDataLayerAsset && (ExternalDataLayerAsset == InActor->GetExternalDataLayerAsset())) ||
+		(ContentBundleGuid.IsValid() && (ContentBundleGuid == InActor->GetContentBundleGuid()));
+	
+	if (bIsCandidateActor)
 	{
-		const FString ActorPackageName = Actor->GetPackage()->GetName();
+		const FString ActorPackageName = InActor->GetPackage()->GetName();
 		const FString ExternalActorPath = GetExternalActorPath() / TEXT("");
 		return ActorPackageName.StartsWith(ExternalActorPath);
 	}
diff --git a/Engine/Source/Runtime/Engine/Private/WorldPartition/ActorDescContainerInstance.cpp b/Engine/Source/Runtime/Engine/Private/WorldPartition/ActorDescContainerInstance.cpp
--- a/Engine/Source/Runtime/Engine/Private/WorldPartition/ActorDescContainerInstance.cpp
+++ b/Engine/Source/Runtime/Engine/Private/WorldPartition/ActorDescContainerInstance.cpp
@@ -34,6 +34,9 @@ void UActorDescContainerInstance::RegisterContainer(const FInitializeParams& InP
 	{
 		return !InParams.FilterActorDescFunc || InParams.FilterActorDescFunc(InActorDesc);
 	};
+	ContainerInitParams.ContentBundleGuid = InParams.ContentBundleGuid;
+	ContainerInitParams.ExternalDataLayerAsset = InParams.ExternalDataLayerAsset;
+
 	SetContainer(UActorDescContainerSubsystem::GetChecked().RegisterContainer(ContainerInitParams));
 }
 
@@ -158,6 +161,11 @@ void UActorDescContainerInstance::Initialize(const FInitializeParams& InParams)
 		AddActorDescInstance(MoveTemp(ActorDescInstance));
 	}
 
+	if (InParams.OnInitializedFunc)
+	{
+		InParams.OnInitializedFunc(this);
+	}
+
 	OnActorDescContainerInstanceInitialized.Broadcast(this);
 		
 	// If Container Instance is required to create hierarchy go ahead
@@ -379,11 +387,26 @@ FGuid UActorDescContainerInstance::GetContentBundleGuid() const
 	return Container->GetContentBundleGuid();
 }
 
+const UExternalDataLayerAsset* UActorDescContainerInstance::GetExternalDataLayerAsset() const
+{
+	return Container->GetExternalDataLayerAsset();
+}
+
+bool UActorDescContainerInstance::HasExternalContent() const
+{
+	return Container->HasExternalContent();
+}
+
 FString UActorDescContainerInstance::GetExternalActorPath() const
 {
 	return Container->GetExternalActorPath();
 }
 
+FString UActorDescContainerInstance::GetExternalObjectPath() const
+{
+	return Container->GetExternalObjectPath();
+}
+
 FWorldPartitionActorDescInstance* UActorDescContainerInstance::AddActor(FWorldPartitionActorDesc* InActorDesc)
 {
 	// We don't support adding actors when instanced
diff --git a/Engine/Source/Runtime/Engine/Private/WorldPartition/ContentBundle/ContentBundleClient.cpp b/Engine/Source/Runtime/Engine/Private/WorldPartition/ContentBundle/ContentBundleClient.cpp
--- a/Engine/Source/Runtime/Engine/Private/WorldPartition/ContentBundle/ContentBundleClient.cpp
+++ b/Engine/Source/Runtime/Engine/Private/WorldPartition/ContentBundle/ContentBundleClient.cpp
@@ -10,7 +10,9 @@
 
 TSharedPtr<FContentBundleClient> FContentBundleClient::CreateClient(const UContentBundleDescriptor* InContentBundleDescriptor, FString const& InDisplayName)
 {
-	return GEngine->GetEngineSubsystem<UContentBundleEngineSubsystem>()->RegisterContentBundle(InContentBundleDescriptor, InDisplayName);
+	TSharedPtr<FContentBundleClient> ContentBundleClient = GEngine->GetEngineSubsystem<UContentBundleEngineSubsystem>()->RegisterContentBundle(InContentBundleDescriptor, InDisplayName);
+	UE_CLOG(ContentBundleClient == nullptr, LogContentBundle, Error, TEXT("FContentBundleClient::CreateClient failed to create a content bundle client for %s using %s"), *InDisplayName, *InContentBundleDescriptor->GetDisplayName());
+	return ContentBundleClient;
 }
 
 FContentBundleClient::FContentBundleClient(const UContentBundleDescriptor* InContentBundleDescriptor, FString const& InDisplayName)
diff --git a/Engine/Source/Runtime/Engine/Private/WorldPartition/ContentBundle/ContentBundleEditor.cpp b/Engine/Source/Runtime/Engine/Private/WorldPartition/ContentBundle/ContentBundleEditor.cpp
--- a/Engine/Source/Runtime/Engine/Private/WorldPartition/ContentBundle/ContentBundleEditor.cpp
+++ b/Engine/Source/Runtime/Engine/Private/WorldPartition/ContentBundle/ContentBundleEditor.cpp
@@ -65,6 +65,7 @@ void FContentBundleEditor::DoInjectContent()
 		UActorDescContainerInstance::FInitializeParams InitParams(*ActorDescContainerPackage);
 		
 		const FTopLevelAssetPath InjectedWorldAssetPath = FSoftObjectPath(GetInjectedWorld()).GetAssetPath();
+		InitParams.ContentBundleGuid = GetDescriptor()->GetGuid();
 		InitParams.FilterActorDescFunc = [&](const FWorldPartitionActorDesc* ActorDesc)
 		{
 			if (ActorDesc->GetActorSoftPath().GetAssetPath() != InjectedWorldAssetPath)
@@ -80,13 +81,17 @@ void FContentBundleEditor::DoInjectContent()
 			return true;
 		};
 
+		InitParams.OnInitializedFunc = [this](UActorDescContainerInstance* InActorDescContainerInstance)
+		{ 
+			check(InActorDescContainerInstance->GetContentBundleGuid() == GetDescriptor()->GetGuid());
+		};
+
 		ActorDescContainerInstance = WorldPartition->RegisterActorDescContainerInstance(InitParams);
 		if (ActorDescContainerInstance.IsValid() && ActorDescContainerInstance->IsInitialized())
 		{
 			UE_LOG(LogContentBundle, Log, TEXT("%s ExternalActors in %s found. %u actors were injected"), *ContentBundle::Log::MakeDebugInfoString(*this), *ActorDescContainerInstance->GetExternalActorPath(), ActorDescContainerInstance->GetActorDescInstanceCount());
 
 			check(GetDescriptor()->GetGuid().IsValid());
-			ActorDescContainerInstance->GetContainer()->SetContentBundleGuid(GetDescriptor()->GetGuid());
 
 			if (!ActorDescContainerInstance->GetContainer()->IsEmpty())
 			{
@@ -338,11 +343,11 @@ void FContentBundleEditor::GenerateStreaming(TArray<FString>* OutPackageToGenera
 		return;
 	}
 
+	FActorDescContainerInstanceCollection Collection({ TObjectPtr<UActorDescContainerInstance>(ActorDescContainerInstance.Get()) });
 	UWorldPartition::FGenerateStreamingParams Params = UWorldPartition::FGenerateStreamingParams()
-		.SetActorDescContainerInstance(ActorDescContainerInstance.Get());
-
+		.SetContainerInstanceCollection(Collection, FStreamingGenerationContainerInstanceCollection::ECollectionType::BaseAsContentBundle);
 	UWorldPartition::FGenerateStreamingContext Context = UWorldPartition::FGenerateStreamingContext()
-		.SetPackagesToGenerate(OutPackageToGenerate);
+		.SetLevelPackagesToGenerate(OutPackageToGenerate);
 
 	UWorldPartition* WorldPartition = GetInjectedWorld()->GetWorldPartition();
 	WorldPartition->GenerateContainerStreaming(Params, Context);
@@ -440,7 +445,7 @@ bool FContentBundleEditor::PopulateGeneratorPackageForCook(class IWorldPartition
 					// Make sure the cell outer is set to the  ExternalStreamingObject so it will be saved in the right package at the end of the cook.
 					check(Cell->GetOuter() == ExternalStreamingObject);
 
-					if (!Cell->PrepareCellForCook(CookPackage->GetPackage()))
+					if (!Cell->OnPopulateGeneratorPackageForCook(CookPackage->GetPackage()))
 					{
 						UE_LOG(LogContentBundle, Error, TEXT("%s[Cook] Failed to prepare cell with package %s for cook."), *ContentBundle::Log::MakeDebugInfoString(*this), *CookPackage->RelativePath);
 						bIsSuccess = false;
@@ -454,7 +459,7 @@ bool FContentBundleEditor::PopulateGeneratorPackageForCook(class IWorldPartition
 			}
 		}
 
-		ExternalStreamingObject->PopulateGeneratorPackageForCook();
+		ExternalStreamingObject->OnPopulateGeneratorPackageForCook(nullptr);
 	}
 	
 
@@ -477,7 +482,7 @@ bool FContentBundleEditor::PopulateGeneratedPackageForCook(class IWorldPartition
 				if (!Cell->IsAlwaysLoaded())
 				{
 					TArray<UPackage*> ModifiedPackages;
-					if (Cell->PopulateGeneratedPackageForCook(PackageToCook.GetPackage(), OutModifiedPackages))
+					if (Cell->OnPopulateGeneratedPackageForCook(PackageToCook.GetPackage(), OutModifiedPackages))
 					{
 						UWorld* CellWorld = FindObject<UWorld>(PackageToCook.GetPackage(), *GetInjectedWorld()->GetName());
 						if (CellWorld != nullptr)
