External Data Layer V1

diff --git a/Engine/Source/Runtime/Engine/Private/WorldPartition/DataLayer/DataLayerAsset.cpp b/Engine/Source/Runtime/Engine/Private/WorldPartition/DataLayer/DataLayerAsset.cpp
--- a/Engine/Source/Runtime/Engine/Private/WorldPartition/DataLayer/DataLayerAsset.cpp
+++ b/Engine/Source/Runtime/Engine/Private/WorldPartition/DataLayer/DataLayerAsset.cpp
@@ -1,6 +1,8 @@
 // Copyright Epic Games, Inc. All Rights Reserved.
 
 #include "WorldPartition/DataLayer/DataLayerAsset.h"
+
+#include "Math/Color.h"
 #include "WorldPartition/DataLayer/DataLayerInstance.h"
 #include "GameFramework/Actor.h"
 #include "Engine/Level.h"
@@ -10,6 +12,7 @@
 UDataLayerAsset::UDataLayerAsset(const FObjectInitializer& ObjectInitializer)
 	: Super(ObjectInitializer)
 	, DataLayerType(EDataLayerType::Editor)
+	, bSupportsActorFilters(false)
 	, DebugColor(FColor::Black)
 {}
 
@@ -56,8 +59,7 @@ bool UDataLayerAsset::CanEditChange(const FProperty* InProperty) const
 	{
 		if (InProperty->GetFName() == GET_MEMBER_NAME_CHECKED(UDataLayerAsset, DataLayerType))
 		{
-			// UDataLayerAsset outered to a UDataLayerInstance do not support Runtime type
-			return !IsPrivate();
+			return CanEditDataLayerType();
 		}
 		else if (InProperty->GetFName() == GET_MEMBER_NAME_CHECKED(UDataLayerAsset, LoadFilter))
 		{
@@ -68,4 +70,29 @@ bool UDataLayerAsset::CanEditChange(const FProperty* InProperty) const
 
 	return true;
 }
+
+void UDataLayerAsset::OnCreated()
+{
+	SetDebugColor(FColor::MakeRandomSeededColor(GetTypeHash(GetFullName())));
+	if (!IsPrivate())
+	{
+		SetType(EDataLayerType::Runtime);
+	}
+}
+
+bool UDataLayerAsset::CanEditDataLayerType() const
+{
+    // UDataLayerAsset outered to a UDataLayerInstance do not support Runtime type
+	return !IsPrivate();
+}
+
+void UDataLayerAsset::SetType(EDataLayerType InDataLayerType)
+{
+	check(DataLayerType == EDataLayerType::Editor || !IsPrivate());
+	if (CanEditDataLayerType())
+	{
+		DataLayerType = InDataLayerType;
+	}
+}
+
 #endif
diff --git a/Engine/Source/Runtime/Engine/Private/WorldPartition/DataLayer/DataLayerInstance.cpp b/Engine/Source/Runtime/Engine/Private/WorldPartition/DataLayer/DataLayerInstance.cpp
--- a/Engine/Source/Runtime/Engine/Private/WorldPartition/DataLayer/DataLayerInstance.cpp
+++ b/Engine/Source/Runtime/Engine/Private/WorldPartition/DataLayer/DataLayerInstance.cpp
@@ -11,6 +11,7 @@
 #include "WorldPartition/DataLayer/DataLayerAsset.h"
 #include "WorldPartition/DataLayer/WorldDataLayers.h"
 #include "WorldPartition/DataLayer/WorldDataLayersActorDesc.h"
+#include "WorldPartition/DataLayer/ExternalDataLayerInstance.h"
 #include "WorldPartition/ErrorHandling/WorldPartitionStreamingGenerationErrorHandler.h"
 #include "AssetRegistry/AssetData.h"
 #include "AssetRegistry/AssetRegistryHelpers.h"
@@ -53,9 +54,17 @@ UWorld* UDataLayerInstance::GetOuterWorld() const
 
 AWorldDataLayers* UDataLayerInstance::GetOuterWorldDataLayers() const
 {
+	check(GetOuterWorld());
 	return GetOuterWorld()->GetWorldDataLayers();
 }
 
+AWorldDataLayers* UDataLayerInstance::GetDirectOuterWorldDataLayers() const
+{
+	// To retrieve the direct outer WorldDataLayers use UObject interface since UDataLayerInstance::GetTypedOuter<AWorldDataLayers>() won't compile.
+	AWorldDataLayers* DirectOuterWorldDataLayers = this->UObject::GetTypedOuter<AWorldDataLayers>();
+	return DirectOuterWorldDataLayers;
+}
+
 #if WITH_EDITOR
 bool UDataLayerInstance::IsAsset() const
 {
@@ -237,20 +246,52 @@ bool UDataLayerInstance::CanEditChange(const FProperty* InProperty) const
 	return true;
 }
 
-bool UDataLayerInstance::IsLocked() const
+bool UDataLayerInstance::IsLocked(FText* OutReason) const
 {
 	if (bIsLocked)
 	{
+		if (OutReason)
+		{
+			*OutReason = LOCTEXT("DataLayerInstanceLocked", "Data layer instance is locked.");
+		}
+		return true;
+	}
+
+	if (IsRuntime() && !GetOuterWorldDataLayers()->GetAllowRuntimeDataLayerEditing())
+	{
+		if (OutReason)
+		{
+			*OutReason = LOCTEXT("DataLayerRuntimeDataLayerEditingNotAllowed", "Runtime data layer editing is not allowed.");
+		}
 		return true;
 	}
 
-	return IsRuntime() && !GetOuterWorldDataLayers()->GetAllowRuntimeDataLayerEditing();
+	return false;
 }
 
-bool UDataLayerInstance::IsReadOnly() const
+bool UDataLayerInstance::IsReadOnly(FText* OutReason) const
 {
 	const UWorld* World = GetWorld();
-	return !World || World->IsGameWorld();
+	if (!World || World->IsGameWorld())
+	{
+		if (OutReason)
+		{
+			*OutReason = LOCTEXT("DataLayerInstanceReadOnlyInGameWorld", "Data layer instance is read-only in game world.");
+		}
+		return true;
+	}
+
+	// Check if root External Data Layer is read-only
+	const UExternalDataLayerInstance* ExternalDataLayerInstance = GetRootExternalDataLayerInstance();
+	if (ExternalDataLayerInstance && (ExternalDataLayerInstance != this))
+	{
+		if (ExternalDataLayerInstance->IsReadOnly(OutReason))
+		{
+			return true;
+		}
+	}
+
+	return IsLocked(OutReason);
 }
 
 const TCHAR* UDataLayerInstance::GetDataLayerIconName() const
@@ -258,14 +299,38 @@ const TCHAR* UDataLayerInstance::GetDataLayerIconName() const
 	return FDataLayerUtils::GetDataLayerIconName(GetType());
 }
 
-bool UDataLayerInstance::CanUserAddActors() const
+bool UDataLayerInstance::CanUserAddActors(FText* OutReason) const
 {
-	return !IsLocked();
+	return !IsReadOnly(OutReason);
 }
 
-bool UDataLayerInstance::CanAddActor(AActor* InActor) const
+bool UDataLayerInstance::CanUserRemoveActors(FText* OutReason) const
 {
-	return InActor != nullptr && InActor->CanAddDataLayer(this);
+	return !IsReadOnly(OutReason);
+}
+
+bool UDataLayerInstance::CanAddActor(AActor* InActor, FText* OutReason) const
+{
+	if (!InActor)
+	{
+		if (OutReason)
+		{
+			*OutReason = LOCTEXT("CantAddActorInvalidActor", "Invalid actor.");
+		}
+		return false;
+	}
+
+	if (IsReadOnly(OutReason))
+	{
+		return false;
+	}
+
+	if (!InActor->CanAddDataLayer(this, OutReason))
+	{
+		return false;
+	}
+
+	return true;
 }
 
 bool UDataLayerInstance::AddActor(AActor* InActor) const
@@ -278,14 +343,27 @@ bool UDataLayerInstance::AddActor(AActor* InActor) const
 	return false;
 }
 
-bool UDataLayerInstance::CanUserRemoveActors() const
+bool UDataLayerInstance::CanRemoveActor(AActor* InActor, FText* OutReason) const
 {
-	return !IsLocked();
-}
+	if (!InActor)
+	{
+		if (OutReason)
+		{
+			*OutReason = LOCTEXT("CantRemoveActorInvalidActor", "Invalid actor.");
+		}
+		return false;
+	}
 
-bool UDataLayerInstance::CanRemoveActor(AActor* InActor) const
-{
-	return InActor->GetDataLayerInstances().Contains(this) || InActor->GetDataLayerInstancesForLevel().Contains(this);
+	if (!InActor->GetDataLayerInstances().Contains(this) && !InActor->GetDataLayerInstancesForLevel().Contains(this))
+	{
+		if (OutReason)
+		{
+			*OutReason = LOCTEXT("CantRemoveActorNotAssignedToDataLayer", "Actor is not assigned to data layer.");
+		}
+		return false;
+	}
+
+	return true;
 }
 
 bool UDataLayerInstance::RemoveActor(AActor* InActor) const
@@ -300,17 +378,21 @@ bool UDataLayerInstance::RemoveActor(AActor* InActor) const
 
 bool UDataLayerInstance::CanBeChildOf(const UDataLayerInstance* InParent, FText* OutReason) const
 {
-	auto AssignReason = [OutReason](FText&& Reason)
+	if (this == InParent)
 	{
-		if(OutReason != nullptr)
+		if (OutReason)
 		{
-			*OutReason = MoveTemp(Reason);
+			*OutReason = LOCTEXT("ParentIsThis", "Can't parent to itself");
 		}
-	};
+		return false;
+	}
 
-	if (this == InParent || Parent == InParent)
+	if (Parent == InParent)
 	{
-		AssignReason(LOCTEXT("SameParentOrSameDataLayer", "Data Layer already has this parent"));
+		if (OutReason)
+		{
+			*OutReason = LOCTEXT("SameParent", "Data Layer already has this parent");
+		}
 		return false;
 	}
 
@@ -320,34 +402,42 @@ bool UDataLayerInstance::CanBeChildOf(const UDataLayerInstance* InParent, FText*
 		return true;
 	}
 
-	if (!CanHaveParentDataLayer())
+	if (!CanHaveParentDataLayerInstance())
 	{
-		AssignReason(FText::Format(LOCTEXT("ParentDataLayerUnsupported", "Data Layer \"{0}\" does not support parent data layers"), FText::FromString(GetDataLayerShortName())));
+		if (OutReason)
+		{
+			*OutReason = FText::Format(LOCTEXT("ParentDataLayerUnsupported", "Data Layer \"{0}\" does not support parent data layers"), FText::FromString(GetDataLayerShortName()));
+		}
 		return false;
 	}
 
-	if (!InParent->CanHaveChildDataLayers())
+	if (!InParent->CanHaveChildDataLayerInstance(this))
 	{
-		AssignReason(FText::Format(LOCTEXT("ChildDataLayerUnsuported", "Data Layer \"{0}\" does not support child data layers"), FText::FromString(InParent->GetDataLayerShortName())));
+		if (OutReason)
+		{
+			*OutReason = FText::Format(LOCTEXT("ChildDataLayerUnsuported", "Data Layer \"{0}\" does not support child data layers"), FText::FromString(InParent->GetDataLayerShortName()));
+		}
 		return false;
 	}
 
-	if (!IsParentDataLayerTypeCompatible(InParent))
+	if (!IsParentDataLayerTypeCompatible(InParent, OutReason))
 	{
-		AssignReason(FText::Format(LOCTEXT("IncompatibleChildType", "{0} Data Layer cannot have {1} child Data Layers"), UEnum::GetDisplayValueAsText(GetType()), UEnum::GetDisplayValueAsText(InParent->GetType())));
 		return false;
 	}
 
-	if (InParent->GetOuterWorldDataLayers() != GetOuterWorldDataLayers())
+	if (InParent->GetDirectOuterWorldDataLayers() != GetDirectOuterWorldDataLayers())
 	{
-		AssignReason(LOCTEXT("DifferentOuterWorldDataLayer", "Parent WorldDataLayers is a different from child WorldDataLayers"));
+		if (OutReason)
+		{
+			*OutReason = LOCTEXT("DifferentOuterWorldDataLayer", "Parent WorldDataLayers is a different from child WorldDataLayers");
+		}
 		return false;
 	}
 
 	return true;
 }
 
-bool UDataLayerInstance::IsParentDataLayerTypeCompatible(const UDataLayerInstance* InParent) const
+bool UDataLayerInstance::IsParentDataLayerTypeCompatible(const UDataLayerInstance* InParent, FText* OutReason) const
 {
 	if (InParent == nullptr)
 	{
@@ -356,14 +446,27 @@ bool UDataLayerInstance::IsParentDataLayerTypeCompatible(const UDataLayerInstanc
 
 	if (IsClientOnly() || IsServerOnly())
 	{
+		if (OutReason)
+		{
+			*OutReason = FText::Format(LOCTEXT("ClientOrServerOnlyCantHaveParent", "{0} Data Layer cannot be a child Data Layer"), IsClientOnly() ? FText::FromString(TEXT("Client-Only")) : FText::FromString(TEXT("Server-Only")));
+		}
 		return false;
 	}
 
-	EDataLayerType ParentDataLayerType = InParent->GetType();
+	const EDataLayerType ParentType = InParent->GetType();
+	const EDataLayerType ChildType = GetType();
 
-	return GetType() != EDataLayerType::Unknown
-		&& ParentDataLayerType != EDataLayerType::Unknown
-		&& (ParentDataLayerType == EDataLayerType::Editor || GetType() == EDataLayerType::Runtime);
+	if ((ChildType == EDataLayerType::Unknown) ||
+		(ParentType == EDataLayerType::Unknown) || 
+		(ParentType != EDataLayerType::Editor && ChildType != EDataLayerType::Runtime))
+	{
+		if (OutReason)
+		{
+			*OutReason = FText::Format(LOCTEXT("IncompatibleChildType", "{0} Data Layer cannot have {1} child Data Layers"), UEnum::GetDisplayValueAsText(InParent->GetType()), UEnum::GetDisplayValueAsText(GetType()));
+		}
+		return false;
+	}
+	return true;
 }
 
 bool UDataLayerInstance::SetParent(UDataLayerInstance* InParent)
@@ -373,7 +476,7 @@ bool UDataLayerInstance::SetParent(UDataLayerInstance* InParent)
 		return false;
 	}
 
-	check(CanHaveParentDataLayer());
+	check(CanHaveParentDataLayerInstance());
 
 	Modify();
 
@@ -443,7 +546,12 @@ bool UDataLayerInstance::Validate(IStreamingGenerationErrorHandler* ErrorHandler
 
 bool UDataLayerInstance::CanBeInActorEditorContext() const
 {
-	return !IsLocked();
+	return !IsReadOnly();
+}
+
+bool UDataLayerInstance::IsActorEditorContextCurrentColorized() const
+{
+	return GetOuterWorldDataLayers()->IsActorEditorContextCurrentColorized(this);
 }
 
 bool UDataLayerInstance::IsInActorEditorContext() const
@@ -510,8 +618,8 @@ void UDataLayerInstance::ForEachChild(TFunctionRef<bool(const UDataLayerInstance
 
 void UDataLayerInstance::AddChild(UDataLayerInstance* InDataLayer)
 {
-	check(InDataLayer->GetOuterWorldDataLayers() == GetOuterWorldDataLayers())
-	check(CanHaveChildDataLayers());
+	check(InDataLayer->GetDirectOuterWorldDataLayers() == GetDirectOuterWorldDataLayers())
+	check(CanHaveChildDataLayerInstance(InDataLayer));
 	Modify(false);
 	checkSlow(!Children.Contains(InDataLayer));
 	Children.Add(InDataLayer);
diff --git a/Engine/Source/Runtime/Engine/Private/WorldPartition/DataLayer/DataLayerInstancePrivate.cpp b/Engine/Source/Runtime/Engine/Private/WorldPartition/DataLayer/DataLayerInstancePrivate.cpp
--- a/Engine/Source/Runtime/Engine/Private/WorldPartition/DataLayer/DataLayerInstancePrivate.cpp
+++ b/Engine/Source/Runtime/Engine/Private/WorldPartition/DataLayer/DataLayerInstancePrivate.cpp
@@ -9,6 +9,8 @@
 
 #include UE_INLINE_GENERATED_CPP_BY_NAME(DataLayerInstancePrivate)
 
+#define LOCTEXT_NAMESPACE "UDataLayerInstancePrivate"
+
 UDataLayerInstancePrivate::UDataLayerInstancePrivate(const FObjectInitializer& ObjectInitializer)
 	: Super(ObjectInitializer)
 #if WITH_EDITORONLY_DATA
@@ -29,7 +31,7 @@ FName UDataLayerInstancePrivate::MakeName()
 
 void UDataLayerInstancePrivate::OnCreated()
 {
-	check(!GetOuterWorldDataLayers()->HasDeprecatedDataLayers() || IsRunningCommandlet());
+	check(!GetDirectOuterWorldDataLayers()->HasDeprecatedDataLayers() || IsRunningCommandlet());
 
 	Modify(/*bAlwaysMarkDirty*/false);
 	FDataLayerUtils::SetDataLayerShortName(this, TEXT("DataLayer"));
@@ -37,40 +39,49 @@ void UDataLayerInstancePrivate::OnCreated()
 	SetVisible(true);
 }
 
-bool UDataLayerInstancePrivate::IsReadOnly() const
-{
-	if (Super::IsReadOnly())
-	{
-		return true;
-	}
-	return GetOuterWorldDataLayers()->IsReadOnly();
-}
-
-bool UDataLayerInstancePrivate::IsLocked() const
+bool UDataLayerInstancePrivate::IsReadOnly(FText* OutReason) const
 {
-	if (Super::IsLocked())
+	if (Super::IsReadOnly(OutReason))
 	{
 		return true;
 	}
-	return IsReadOnly();
+	return GetDirectOuterWorldDataLayers()->IsReadOnly(OutReason);
 }
 
 bool UDataLayerInstancePrivate::PerformAddActor(AActor* InActor) const
 {
-	check(GetOuterWorldDataLayers() == InActor->GetLevel()->GetWorldDataLayers()); // Make sure the instance is part of the same WorldDataLayers as the actor's level WorldDataLayer.
+	check(GetDirectOuterWorldDataLayers() == InActor->GetLevel()->GetWorldDataLayers()); // Make sure the instance is part of the same WorldDataLayers as the actor's level WorldDataLayer.
 	check(UDataLayerManager::GetDataLayerManager(InActor)->GetDataLayerInstance(DataLayerAsset) != nullptr); // Make sure the DataLayerInstance exists for this level
 	return FAssignActorDataLayer::AddDataLayerAsset(InActor, DataLayerAsset);
 }
 
-bool UDataLayerInstancePrivate::CanAddActor(AActor* InActor) const
+bool UDataLayerInstancePrivate::CanAddActor(AActor* InActor, FText* OutReason) const
 {
-	if (!Super::CanAddActor(InActor))
+	if (!Super::CanAddActor(InActor, OutReason))
 	{
 		return false;
 	}
 
-	return GetOuterWorldDataLayers() == InActor->GetLevel()->GetWorldDataLayers() && // Make sure the instance is part of the same WorldDataLayers as the actor's level WorldDataLayer.
-		   UDataLayerManager::GetDataLayerManager(InActor)->GetDataLayerInstance(DataLayerAsset) != nullptr; // Make sure the DataLayerInstance exists for this level
+	// Make sure the instance is part of the same WorldDataLayers as the actor's level WorldDataLayer
+	if (GetDirectOuterWorldDataLayers() != InActor->GetLevel()->GetWorldDataLayers())
+	{
+		if (OutReason)
+		{
+			*OutReason = LOCTEXT("CantAddDataLayerPrivateDataLayerNotPartActorLevelWorldDataLayers", "Actor is not in the same level as the private data layer.");
+		}
+		return false;
+	}
+	// Make sure the DataLayerInstance exists for this level
+	UDataLayerManager* DataLayerManager = UDataLayerManager::GetDataLayerManager(InActor);
+	if (!DataLayerManager || !DataLayerManager->GetDataLayerInstance(DataLayerAsset))
+	{
+		if (OutReason)
+		{
+			*OutReason = LOCTEXT("CantAddDataLayerPrivateDataLayerInstanceNotFoundInActorLevel", "Private data layer instance doesn't exist in the actor level.");
+		}
+		return false;
+	}
+	return true;
 }
 
 bool UDataLayerInstancePrivate::PerformRemoveActor(AActor* InActor) const
@@ -94,3 +105,5 @@ bool UDataLayerInstancePrivate::CanEditChange(const FProperty* InProperty) const
 }
 
 #endif
+
+#undef LOCTEXT_NAMESPACE
\ No newline at end of file
diff --git a/Engine/Source/Runtime/Engine/Private/WorldPartition/DataLayer/DataLayerInstanceWithAsset.cpp b/Engine/Source/Runtime/Engine/Private/WorldPartition/DataLayer/DataLayerInstanceWithAsset.cpp
--- a/Engine/Source/Runtime/Engine/Private/WorldPartition/DataLayer/DataLayerInstanceWithAsset.cpp
+++ b/Engine/Source/Runtime/Engine/Private/WorldPartition/DataLayer/DataLayerInstanceWithAsset.cpp
@@ -3,6 +3,9 @@
 #include "WorldPartition/DataLayer/DataLayerInstanceWithAsset.h"
 #include "WorldPartition/DataLayer/WorldDataLayers.h"
 #include "WorldPartition/DataLayer/DataLayerManager.h"
+#include "WorldPartition/DataLayer/ExternalDataLayerAsset.h"
+#include "WorldPartition/DataLayer/ExternalDataLayerInstance.h"
+#include "WorldPartition/WorldPartitionRuntimeHash.h"
 #include "Engine/Level.h"
 #include "Misc/StringFormatArg.h"
 #include "UObject/UnrealType.h"
@@ -10,17 +13,42 @@
 
 #include UE_INLINE_GENERATED_CPP_BY_NAME(DataLayerInstanceWithAsset)
 
+#define LOCTEXT_NAMESPACE "DataLayerInstanceWithAsset"
+
 UDataLayerInstanceWithAsset::UDataLayerInstanceWithAsset(const FObjectInitializer& ObjectInitializer)
 	: Super(ObjectInitializer)
 #if WITH_EDITORONLY_DATA
 	, bIsIncludedInActorFilterDefault(true)
 #endif
 {
+#if WITH_EDITOR
+	bSkipCheckReadOnlyForSubLevels = false;
+#endif
+}
 
+UWorld* UDataLayerInstanceWithAsset::GetOuterWorld() const
+{
+	if (URuntimeHashExternalStreamingObjectBase* ExternalStreamingObject = GetTypedOuter<URuntimeHashExternalStreamingObjectBase>())
+	{
+		return ExternalStreamingObject->GetOuterWorld();
+	}
+	check(GetDirectOuterWorldDataLayers()->GetTypedOuter<UWorld>() == Super::GetOuterWorld());
+	return Super::GetOuterWorld();
 }
 
+
 #if WITH_EDITOR
-FName UDataLayerInstanceWithAsset::MakeName(const UDataLayerAsset* DeprecatedDataLayer)
+const UExternalDataLayerInstance* UDataLayerInstanceWithAsset::GetRootExternalDataLayerInstance() const
+{
+	IDataLayerInstanceProvider* DataLayerInstanceProvider = GetImplementingOuter<IDataLayerInstanceProvider>();
+	if (ensure(DataLayerInstanceProvider))
+	{
+		return DataLayerInstanceProvider->GetRootExternalDataLayerInstance();
+	}
+	return nullptr;
+}
+
+FName UDataLayerInstanceWithAsset::MakeName(const UDataLayerAsset* InDataLayerAsset)
 {
 	return FName(FString::Format(TEXT("DataLayer_{0}"), { FGuid::NewGuid().ToString() }));
 }
@@ -32,7 +60,8 @@ TSubclassOf<UDataLayerInstanceWithAsset> UDataLayerInstanceWithAsset::GetDataLay
 
 void UDataLayerInstanceWithAsset::OnCreated(const UDataLayerAsset* Asset)
 {
-	check(!GetOuterWorldDataLayers()->HasDeprecatedDataLayers() || IsRunningCommandlet());
+	checkf(Asset->IsA<UExternalDataLayerAsset>() == this->IsA<UExternalDataLayerInstance>(), TEXT("Only ExternalDataLayerInstance can reference ExternalDataLayerAsset"));
+	check(!GetDirectOuterWorldDataLayers()->HasDeprecatedDataLayers() || IsRunningCommandlet());
 
 	Modify(/*bAlwaysMarkDirty*/false);
 
@@ -42,27 +71,51 @@ void UDataLayerInstanceWithAsset::OnCreated(const UDataLayerAsset* Asset)
 	SetVisible(true);
 }
 
-bool UDataLayerInstanceWithAsset::IsReadOnly() const
+bool UDataLayerInstanceWithAsset::IsReadOnly(FText* OutReason) const
 {
-	if (Super::IsReadOnly())
+	if (Super::IsReadOnly(OutReason))
 	{
 		return true;
 	}
-	return GetOuterWorldDataLayers()->IsReadOnly();
+	return !bSkipCheckReadOnlyForSubLevels && GetDirectOuterWorldDataLayers()->IsReadOnly(OutReason);
 }
 
-bool UDataLayerInstanceWithAsset::IsLocked() const
+bool UDataLayerInstanceWithAsset::CanBeInActorEditorContext() const
 {
-	if (Super::IsLocked())
+	if (!Super::CanBeInActorEditorContext())
 	{
-		return true;
+		return false;
 	}
-	return IsReadOnly();
+
+	return DataLayerAsset != nullptr;
 }
 
-bool UDataLayerInstanceWithAsset::CanAddActor(AActor* InActor) const
+bool UDataLayerInstanceWithAsset::CanAddActor(AActor* InActor, FText* OutReason) const
 {
-	return DataLayerAsset != nullptr && DataLayerAsset->CanBeReferencedByActor(InActor) && Super::CanAddActor(InActor);
+	if (!Super::CanAddActor(InActor, OutReason))
+	{
+		return false;
+	}
+
+	if (!DataLayerAsset)
+	{
+		if (OutReason)
+		{
+			*OutReason = LOCTEXT("CantAddActorInvalidDataLayerAsset", "Invalid data layer asset.");
+		}
+		return false;
+	}
+
+	if (!DataLayerAsset->CanBeReferencedByActor(InActor))
+	{
+		if (OutReason)
+		{
+			*OutReason = LOCTEXT("CantAddActorDataLayerAssetCantBeReferencedByActor", "Data layer asset can't be referenced by actor.");
+		}
+		return false;
+	}
+
+	return true;
 }
 
 bool UDataLayerInstanceWithAsset::PerformAddActor(AActor* InActor) const
@@ -72,9 +125,23 @@ bool UDataLayerInstanceWithAsset::PerformAddActor(AActor* InActor) const
 	return FAssignActorDataLayer::AddDataLayerAsset(InActor, DataLayerAsset);
 }
 
-bool UDataLayerInstanceWithAsset::CanRemoveActor(AActor* InActor) const
+bool UDataLayerInstanceWithAsset::CanRemoveActor(AActor* InActor, FText* OutReason) const
 {
-	return DataLayerAsset != nullptr && Super::CanRemoveActor(InActor);
+	if (!Super::CanRemoveActor(InActor, OutReason))
+	{
+		return false;
+	}
+
+	if (!DataLayerAsset)
+	{
+		if (OutReason)
+		{
+			*OutReason = LOCTEXT("CantRemoveActorInvalidDataLayerAsset", "Invalid data layer asset.");
+		}
+		return false;
+	}
+
+	return true;
 }
 
 bool UDataLayerInstanceWithAsset::PerformRemoveActor(AActor* InActor) const
@@ -92,6 +159,12 @@ bool UDataLayerInstanceWithAsset::Validate(IStreamingGenerationErrorHandler* Err
 		return false;
 	}
 
+	if (GetAsset()->IsA<UExternalDataLayerAsset>() != this->IsA<UExternalDataLayerInstance>())
+	{
+		ErrorHandler->OnInvalidDataLayerAssetType(this, GetAsset());
+		return false;
+	}
+
 	// Get the DataLayerManager for this DataLayerInstance which will be the one of its outer world
 	UDataLayerManager* DataLayerManager = UDataLayerManager::GetDataLayerManager(this);
 	if (ensure(DataLayerManager))
@@ -144,6 +217,11 @@ bool UDataLayerInstanceWithAsset::CanEditChange(const FProperty* InProperty) con
 		return SupportsActorFilters();
 	}
 
+	if (InProperty && (InProperty->GetFName() == GET_MEMBER_NAME_CHECKED(UDataLayerInstanceWithAsset, DataLayerAsset)))
+	{
+		return CanEditDataLayerAsset();
+	}
+
 	return true;
 }
 
@@ -173,3 +251,5 @@ void UDataLayerInstanceWithAsset::PostEditUndo()
 	CachedDataLayerAsset = nullptr;
 }
 #endif
+
+#undef LOCTEXT_NAMESPACE 
\ No newline at end of file
diff --git a/Engine/Source/Runtime/Engine/Private/WorldPartition/DataLayer/DataLayerManager.cpp b/Engine/Source/Runtime/Engine/Private/WorldPartition/DataLayer/DataLayerManager.cpp
--- a/Engine/Source/Runtime/Engine/Private/WorldPartition/DataLayer/DataLayerManager.cpp
+++ b/Engine/Source/Runtime/Engine/Private/WorldPartition/DataLayer/DataLayerManager.cpp
@@ -5,6 +5,8 @@
 #include "WorldPartition/DataLayer/DataLayerInstanceWithAsset.h"
 #include "WorldPartition/DataLayer/DataLayerSubsystem.h"
 #include "WorldPartition/DataLayer/DataLayerLoadingPolicy.h"
+#include "WorldPartition/DataLayer/ExternalDataLayerAsset.h"
+#include "WorldPartition/DataLayer/ExternalDataLayerInstance.h"
 #include "WorldPartition/WorldPartitionDebugHelper.h"
 #include "WorldPartition/WorldPartitionSubsystem.h"
 #include "WorldPartition/WorldPartitionLog.h"
@@ -18,6 +20,7 @@
 #include "WorldPartition/DataLayer/WorldDataLayersActorDesc.h"
 #include "WorldPartition/DataLayer/DataLayerUtils.h"
 #include "WorldPartition/DataLayer/IDataLayerEditorModule.h"
+#include "WorldPartition/DataLayer/ExternalDataLayerHelper.h"
 #include "WorldPartition/ActorDescContainerInstance.h"
 #include "WorldPartition/WorldPartitionEditorPerProjectUserSettings.h"
 #include "WorldPartition/WorldPartitionActorLoaderInterface.h"
@@ -29,6 +32,8 @@
 #include "UObject/UObjectGlobals.h"
 #include "Engine/LevelStreaming.h"
 #include "EngineUtils.h"
+#include "Algo/Transform.h"
+#include "Algo/AnyOf.h"
 #include "LevelUtils.h"
 #include "Editor.h"
 #else
@@ -174,14 +179,7 @@ void UDataLayerManager::Initialize()
 		if (ActorDescContainerInstance)
 		{
 			// Try to find and load AWorldDataLayers actor
-			for (UActorDescContainerInstance::TIterator<> Iterator(ActorDescContainerInstance); Iterator; ++Iterator)
-			{
-				if (Iterator->GetActorNativeClass()->IsChildOf<AWorldDataLayers>())
-				{
-					WorldDataLayersActor = FWorldPartitionReference(ActorDescContainerInstance, Iterator->GetGuid());
-					break;
-				}
-			}
+			WorldDataLayersActor = UDataLayerManager::LoadWorldDataLayersActor(ActorDescContainerInstance);
 		}
 
 		WorldDataLayers = GetWorldDataLayers();
@@ -365,7 +363,7 @@ const TSet<FName>& UDataLayerManager::GetEffectiveLoadedDataLayerNames() const
 	return WorldDataLayers ? WorldDataLayers->GetEffectiveLoadedDataLayerNames() : EmptySet;
 }
 
-bool UDataLayerManager::IsAnyDataLayerInEffectiveRuntimeState(const TArray<FName>& InDataLayerNames, EDataLayerRuntimeState InState) const
+bool UDataLayerManager::IsAnyDataLayerInEffectiveRuntimeState(TArrayView<const FName> InDataLayerNames, EDataLayerRuntimeState InState) const
 {
 	if (InState == EDataLayerRuntimeState::Activated)
 	{
@@ -392,7 +390,7 @@ bool UDataLayerManager::IsAnyDataLayerInEffectiveRuntimeState(const TArray<FName
 	return false;
 }
 
-bool UDataLayerManager::IsAllDataLayerInEffectiveRuntimeState(const TArray<FName>& InDataLayerNames, EDataLayerRuntimeState InState) const
+bool UDataLayerManager::IsAllDataLayerInEffectiveRuntimeState(TArrayView<const FName> InDataLayerNames, EDataLayerRuntimeState InState) const
 {
 	const TSet<FName>& Activated = GetEffectiveActiveDataLayerNames();
 
@@ -763,12 +761,12 @@ TArray<AWorldDataLayers*> UDataLayerManager::GetActorEditorContextWorldDataLayer
 	return WorldDataLayersArray;
 }
 
-void UDataLayerManager::PushActorEditorContext() const
+void UDataLayerManager::PushActorEditorContext(bool bDuplicateContext) const
 {
 	++DataLayerActorEditorContextID;
 	for (AWorldDataLayers* WorldDataLayers : GetActorEditorContextWorldDataLayers())
 	{
-		WorldDataLayers->PushActorEditorContext(DataLayerActorEditorContextID);
+		WorldDataLayers->PushActorEditorContext(DataLayerActorEditorContextID, bDuplicateContext);
 	}
 }
 
@@ -804,7 +802,6 @@ uint32 UDataLayerManager::GetDataLayerEditorContextHash() const
 
 bool UDataLayerManager::CanResolveDataLayers() const
 {
-	check(GetWorldDataLayers());
 	return (GetWorldDataLayers() != nullptr) && bCanResolveDataLayers;
 }
 
@@ -825,6 +822,20 @@ void UDataLayerManager::ResolveActorDescContainersDataLayers() const
 	}
 }
 
+FWorldPartitionReference UDataLayerManager::LoadWorldDataLayersActor(UActorDescContainerInstance* InActorDescContainerInstance)
+{
+	FWorldPartitionReference WDLReference;
+	for (UActorDescContainerInstance::TIterator<> Iterator(InActorDescContainerInstance); Iterator; ++Iterator)
+	{
+		if (Iterator->GetActorNativeClass()->IsChildOf<AWorldDataLayers>())
+		{
+			WDLReference = FWorldPartitionReference(InActorDescContainerInstance, Iterator->GetGuid());
+			break;
+		}
+	}
+	return WDLReference;
+}
+
 void UDataLayerManager::ResolveActorDescContainerInstanceDataLayers(UActorDescContainerInstance* InActorDescContainerInstance) const
 {
 	ResolveActorDescContainerInstanceDataLayersInternal(InActorDescContainerInstance, nullptr);
@@ -862,21 +873,42 @@ void UDataLayerManager::ResolveActorDescContainerInstanceDataLayersInternal(UAct
 	check(CanResolveDataLayers());
 	if (InActorDescInstance)
 	{
-		InActorDescInstance->SetDataLayerInstanceNames(FDataLayerUtils::ResolvedDataLayerInstanceNames(this, InActorDescInstance->GetActorDesc()));
+		InActorDescInstance->SetDataLayerInstanceNames(FDataLayerUtils::ResolveDataLayerInstanceNames(this, InActorDescInstance->GetActorDesc()));
 	}
 	else
 	{
 		for (UActorDescContainerInstance::TIterator<> Iterator(InActorDescContainerInstance); Iterator; ++Iterator)
 		{
-			Iterator->SetDataLayerInstanceNames(FDataLayerUtils::ResolvedDataLayerInstanceNames(this, Iterator->GetActorDesc()));
+			Iterator->SetDataLayerInstanceNames(FDataLayerUtils::ResolveDataLayerInstanceNames(this, Iterator->GetActorDesc()));
 		}
 	}
 }
 
 bool UDataLayerManager::ResolveIsLoadedInEditor(const TArray<FName>& InDataLayerInstanceNames) const
 {
-	TArray<const UDataLayerInstance*> DataLayerInstances = GetDataLayerInstances(InDataLayerInstanceNames);
-	return DataLayerInstances.Num() ? DataLayerLoadingPolicy->ResolveIsLoadedInEditor(DataLayerInstances) : true;
+	TArray<const UDataLayerInstance*> AllDataLayerInstances = GetDataLayerInstances(InDataLayerInstanceNames);
+	TArray<const UDataLayerInstance*> DataLayerInstances;
+	TArray<const UExternalDataLayerInstance*> ExternalDataLayerInstances;
+	for (const UDataLayerInstance* DataLayerInstance : AllDataLayerInstances)
+	{
+		if (const UExternalDataLayerInstance* ExternalDataLayerInstance = Cast<UExternalDataLayerInstance>(DataLayerInstance))
+		{
+			ExternalDataLayerInstances.Add(ExternalDataLayerInstance);
+		}
+		else
+		{
+			DataLayerInstances.Add(DataLayerInstance);
+		}
+	}
+
+	// If part of an External Data Layer, the External Data Layer must be loaded in editor
+	const bool bAnyExternalDataLayerInstanceLoaded = ExternalDataLayerInstances.Num() ? Algo::AnyOf(ExternalDataLayerInstances, [](const UExternalDataLayerInstance* ExternalDataLayerInstance) { return ExternalDataLayerInstance->IsEffectiveLoadedInEditor(); }) : true;
+	if (bAnyExternalDataLayerInstanceLoaded)
+	{
+		const bool bAnyDataLayerInstanceLoaded = DataLayerInstances.Num() ? DataLayerLoadingPolicy->ResolveIsLoadedInEditor(DataLayerInstances) : true;
+		return bAnyDataLayerInstanceLoaded;
+	}
+	return false;
 }
 
 TArray<const UDataLayerInstance*> UDataLayerManager::GetRuntimeDataLayerInstances(const TArray<FName>& InDataLayerInstanceNames) const
diff --git a/Engine/Source/Runtime/Engine/Private/WorldPartition/DataLayer/DataLayerUtils.cpp b/Engine/Source/Runtime/Engine/Private/WorldPartition/DataLayer/DataLayerUtils.cpp
--- a/Engine/Source/Runtime/Engine/Private/WorldPartition/DataLayer/DataLayerUtils.cpp
+++ b/Engine/Source/Runtime/Engine/Private/WorldPartition/DataLayer/DataLayerUtils.cpp
@@ -11,90 +11,119 @@
 #include "WorldPartition/DataLayer/DataLayerManager.h"
 #include "Algo/AllOf.h"
 
-TArray<FName> FDataLayerUtils::ResolvedDataLayerInstanceNames(const UDataLayerManager* InDataLayerManager, const FWorldPartitionActorDesc* InActorDesc, const TArray<const FWorldDataLayersActorDesc*>& InWorldDataLayersActorDescs)
+FDataLayerInstanceNames FDataLayerUtils::ResolveDataLayerInstanceNames(const UDataLayerManager* InDataLayerManager, const FWorldPartitionActorDesc* InActorDesc, const TArray<const FWorldDataLayersActorDesc*>& InWorldDataLayersActorDescs)
 {
-	const TArray<FName>& ActorDescDataLayers = InActorDesc->GetDataLayers();
-
-	if (!ActorDescDataLayers.IsEmpty())
+	const bool bIncludeExternalDataLayer = false;
+	const FName ActorDescExternalDataLayer = InActorDesc->GetExternalDataLayer();
+	const TArray<FName> ActorDescNonExternalDataLayers = InActorDesc->GetDataLayers(bIncludeExternalDataLayer);
+	if (ActorDescNonExternalDataLayers.IsEmpty() && ActorDescExternalDataLayer.IsNone())
+	{
+		return FDataLayerInstanceNames();
+	}
+	
+	// DataLayers not using DataLayer Assets represent DataLayerInstanceNames
+	if (!InActorDesc->IsUsingDataLayerAsset())
 	{
-		// DataLayers not using DataLayer Assets represent DataLayerInstanceNames
-		if (!InActorDesc->IsUsingDataLayerAsset())
+		if (InDataLayerManager && InDataLayerManager->CanResolveDataLayers())
 		{
-			if (InDataLayerManager && InDataLayerManager->CanResolveDataLayers())
+			TArray<FName> Result;
+			for (const FName& DataLayerInstanceName : ActorDescNonExternalDataLayers)
 			{
-				TArray<FName> Result;
-				for (const FName& DataLayerInstanceName : ActorDescDataLayers)
+				if (const UDataLayerInstance* DataLayerInstance = InDataLayerManager->GetDataLayerInstance(DataLayerInstanceName))
 				{
-					if (const UDataLayerInstance* DataLayerInstance = InDataLayerManager->GetDataLayerInstance(DataLayerInstanceName))
-					{
-						Result.Add(DataLayerInstanceName);
-					}
+					Result.Add(DataLayerInstanceName);
 				}
-				return Result;
 			}
-			// Fallback on FWorldDataLayersActorDesc
-			else if (!InWorldDataLayersActorDescs.IsEmpty() && AreWorldDataLayersActorDescsSane(InWorldDataLayersActorDescs))
+			return FDataLayerInstanceNames(Result, false);
+		}
+		// Fallback on FWorldDataLayersActorDesc
+		else if (!InWorldDataLayersActorDescs.IsEmpty() && AreWorldDataLayersActorDescsSane(InWorldDataLayersActorDescs))
+		{
+			TArray<FName> Result;
+			for (const FName& DataLayerInstanceName : ActorDescNonExternalDataLayers)
 			{
-				TArray<FName> Result;
-				for (const FName& DataLayerInstanceName : ActorDescDataLayers)
+				if (GetDataLayerInstanceDescFromInstanceName(InWorldDataLayersActorDescs, DataLayerInstanceName) != nullptr)
 				{
-					if (GetDataLayerInstanceDescFromInstanceName(InWorldDataLayersActorDescs, DataLayerInstanceName) != nullptr)
-					{
-						Result.Add(DataLayerInstanceName);
-					}
+					Result.Add(DataLayerInstanceName);
 				}
-				return Result;
 			}
+			return FDataLayerInstanceNames(Result, false);
 		}
-		// ActorDesc DataLayers represents DataLayer Asset Paths
-		else
+	}
+	// ActorDesc DataLayers represents DataLayer Asset Paths
+	else
+	{
+		if (InDataLayerManager && InDataLayerManager->CanResolveDataLayers())
 		{
-			if (InDataLayerManager && InDataLayerManager->CanResolveDataLayers())
+			TArray<FName> ResolvedDLINames;
+			ResolvedDLINames.Reserve(ActorDescNonExternalDataLayers.Num() + 1);
+			const UDataLayerInstance* ExternalDataLayerInstance = InDataLayerManager->GetDataLayerInstanceFromAssetName(ActorDescExternalDataLayer);
+			FName ExternalDataLayer = ExternalDataLayerInstance ? ExternalDataLayerInstance->GetDataLayerFName() : NAME_None;
+			bool bIsFirstDataLayerIsExternal = !ExternalDataLayer.IsNone();
+			if (bIsFirstDataLayerIsExternal)
+			{
+				ResolvedDLINames.Add(ExternalDataLayer);
+			}
+			for (const FName& DataLayerAssetPath : ActorDescNonExternalDataLayers)
 			{
-				TArray<FName> Result;
-				for (const FName& DataLayerAssetPath : ActorDescDataLayers)
+				if (const UDataLayerInstance* DataLayerInstance = InDataLayerManager->GetDataLayerInstanceFromAssetName(DataLayerAssetPath))
 				{
-					if (const UDataLayerInstance* DataLayerInstance = InDataLayerManager->GetDataLayerInstanceFromAssetName(DataLayerAssetPath))
-					{
-						Result.Add(DataLayerInstance->GetDataLayerFName());
-					}
+					ResolvedDLINames.Add(DataLayerInstance->GetDataLayerFName());
 				}
-				return Result;
 			}
-			// Fallback on FWorldDataLayersActorDesc
-			else if (!InWorldDataLayersActorDescs.IsEmpty() && AreWorldDataLayersActorDescsSane(InWorldDataLayersActorDescs))
+			return FDataLayerInstanceNames(ResolvedDLINames, bIsFirstDataLayerIsExternal);
+		}
+		// Fallback on FWorldDataLayersActorDesc
+		else if (!InWorldDataLayersActorDescs.IsEmpty() && AreWorldDataLayersActorDescsSane(InWorldDataLayersActorDescs))
+		{
+			TArray<FName> ResolvedDLINames;
+			ResolvedDLINames.Reserve(ActorDescNonExternalDataLayers.Num() + 1);
+			const FDataLayerInstanceDesc* ExternalDataLayerInstanceDesc = GetDataLayerInstanceDescFromAssetPath(InWorldDataLayersActorDescs, ActorDescExternalDataLayer);
+			FName ExternalDataLayer = ExternalDataLayerInstanceDesc ? ExternalDataLayerInstanceDesc->GetName() : NAME_None;
+			bool bIsFirstDataLayerIsExternal = !ExternalDataLayer.IsNone();
+			if (bIsFirstDataLayerIsExternal)
 			{
-				TArray<FName> Result;
-				for (const FName& DataLayerAssetPath : ActorDescDataLayers)
+				ResolvedDLINames.Add(ExternalDataLayer);
+			}
+			for (const FName& DataLayerAssetPath : ActorDescNonExternalDataLayers)
+			{
+				if (const FDataLayerInstanceDesc* DataLayerInstanceDesc = GetDataLayerInstanceDescFromAssetPath(InWorldDataLayersActorDescs, DataLayerAssetPath))
 				{
-					if (const FDataLayerInstanceDesc* DataLayerInstanceDesc = GetDataLayerInstanceDescFromAssetPath(InWorldDataLayersActorDescs, DataLayerAssetPath))
-					{
-						Result.Add(DataLayerInstanceDesc->GetName());
-					}
+					ResolvedDLINames.Add(DataLayerInstanceDesc->GetName());
 				}
-				return Result;
 			}
+			return FDataLayerInstanceNames(ResolvedDLINames, bIsFirstDataLayerIsExternal);
 		}
 	}
-	return InActorDesc->GetDataLayers();
+	return FDataLayerInstanceNames(ActorDescNonExternalDataLayers, ActorDescExternalDataLayer);
 }
 
 // For performance reasons, this function assumes that InActorDesc's DataLayerInstanceNames was already resolved.
-bool FDataLayerUtils::ResolveRuntimeDataLayerInstanceNames(const UDataLayerManager* InDataLayerManager, const IWorldPartitionActorDescInstanceView& InActorDescView, const FStreamingGenerationActorDescViewMap& InActorDescViewMap, TArray<FName>& OutRuntimeDataLayerInstanceNames)
+bool FDataLayerUtils::ResolveRuntimeDataLayerInstanceNames(const UDataLayerManager* InDataLayerManager, const IWorldPartitionActorDescInstanceView& InActorDescView, const FStreamingGenerationActorDescViewMap& InActorDescViewMap, FDataLayerInstanceNames& OutRuntimeDataLayerInstanceNames)
 {
-	const TArray<FName>& ActorDescViewDataLayerInstanceNames = InActorDescView.GetDataLayerInstanceNames();
+	const FDataLayerInstanceNames& ActorDescViewDataLayerInstanceNames = InActorDescView.GetDataLayerInstanceNames();
+	const FName EDLInstanceName = ActorDescViewDataLayerInstanceNames.GetExternalDataLayer();
+	const TArrayView<const FName> NonEDLInstanceNames = ActorDescViewDataLayerInstanceNames.GetNonExternalDataLayers();
 
 	if (InDataLayerManager && InDataLayerManager->CanResolveDataLayers())
 	{
-		for (FName DataLayerInstanceName : ActorDescViewDataLayerInstanceNames)
+		TArray<FName> ResolvedRuntimeDLINames;
+		ResolvedRuntimeDLINames.Reserve(NonEDLInstanceNames.Num() + 1);
+		const UDataLayerInstance* ExternalDataLayerInstance = InDataLayerManager->GetDataLayerInstanceFromName(EDLInstanceName);
+		const bool bIsFirstDataLayerIsExternal = (ExternalDataLayerInstance && ExternalDataLayerInstance->IsRuntime());
+		if (bIsFirstDataLayerIsExternal)
+		{
+			ResolvedRuntimeDLINames.Add(EDLInstanceName);
+		}
+		for (FName DataLayerInstanceName : NonEDLInstanceNames)
 		{
 			const UDataLayerInstance* DataLayerInstance = InDataLayerManager->GetDataLayerInstanceFromName(DataLayerInstanceName);
 			if (DataLayerInstance && DataLayerInstance->IsRuntime())
 			{
-				OutRuntimeDataLayerInstanceNames.Add(DataLayerInstanceName);
+				ResolvedRuntimeDLINames.Add(DataLayerInstanceName);
 			}
 		}
-
+		OutRuntimeDataLayerInstanceNames = FDataLayerInstanceNames(ResolvedRuntimeDLINames, bIsFirstDataLayerIsExternal);
 		return true;
 	}
 
@@ -103,17 +132,25 @@ bool FDataLayerUtils::ResolveRuntimeDataLayerInstanceNames(const UDataLayerManag
 	if (WorldDataLayersActorDescs.Num())
 	{
 		check(AreWorldDataLayersActorDescsSane(WorldDataLayersActorDescs));
-		for (FName DataLayerInstanceName : ActorDescViewDataLayerInstanceNames)
+
+		TArray<FName> ResolvedRuntimeDLINames;
+		ResolvedRuntimeDLINames.Reserve(NonEDLInstanceNames.Num() + 1);
+		const FDataLayerInstanceDesc* ExternalDataLayerInstanceDesc = GetDataLayerInstanceDescFromInstanceName(WorldDataLayersActorDescs, EDLInstanceName);
+		const bool bIsFirstDataLayerIsExternal = (ExternalDataLayerInstanceDesc && (ExternalDataLayerInstanceDesc->GetDataLayerType() == EDataLayerType::Runtime));
+		if (bIsFirstDataLayerIsExternal)
 		{
-			if (const FDataLayerInstanceDesc* DataLayerInstanceDesc = GetDataLayerInstanceDescFromInstanceName(WorldDataLayersActorDescs, DataLayerInstanceName))
+			ResolvedRuntimeDLINames.Add(EDLInstanceName);
+		}
+
+		for (FName DataLayerInstanceName : NonEDLInstanceNames)
+		{
+			const FDataLayerInstanceDesc* DataLayerInstanceDesc = GetDataLayerInstanceDescFromInstanceName(WorldDataLayersActorDescs, DataLayerInstanceName);
+			if (DataLayerInstanceDesc && (DataLayerInstanceDesc->GetDataLayerType() == EDataLayerType::Runtime))
 			{
-				if (DataLayerInstanceDesc->GetDataLayerType() == EDataLayerType::Runtime)
-				{
-					OutRuntimeDataLayerInstanceNames.Add(DataLayerInstanceName);
-				}
+				ResolvedRuntimeDLINames.Add(DataLayerInstanceName);
 			}
 		}
-
+		OutRuntimeDataLayerInstanceNames = FDataLayerInstanceNames(ResolvedRuntimeDLINames, bIsFirstDataLayerIsExternal);
 		return true;
 	}
 
diff --git a/Engine/Source/Runtime/Engine/Private/WorldPartition/DataLayer/WorldDataLayers.cpp b/Engine/Source/Runtime/Engine/Private/WorldPartition/DataLayer/WorldDataLayers.cpp
--- a/Engine/Source/Runtime/Engine/Private/WorldPartition/DataLayer/WorldDataLayers.cpp
+++ b/Engine/Source/Runtime/Engine/Private/WorldPartition/DataLayer/WorldDataLayers.cpp
@@ -10,6 +10,8 @@
 #include "WorldPartition/DataLayer/DataLayerInstanceWithAsset.h"
 #include "WorldPartition/DataLayer/DeprecatedDataLayerInstance.h"
 #include "WorldPartition/DataLayer/DataLayerUtils.h"
+#include "WorldPartition/DataLayer/ExternalDataLayerInstance.h"
+#include "WorldPartition/DataLayer/ExternalDataLayerAsset.h"
 #include "WorldPartition/WorldPartition.h"
 #include "WorldPartition/WorldPartitionLog.h"
 #include "WorldPartition/WorldPartitionRuntimeCellInterface.h"
@@ -107,6 +109,11 @@ void AWorldDataLayers::RewindForReplay()
 
 void AWorldDataLayers::InitializeDataLayerRuntimeStates()
 {
+	if (IsExternalDataLayerWorldDataLayers())
+	{
+		return;
+	}
+
 	check(ActiveDataLayerNames.IsEmpty() && LoadedDataLayerNames.IsEmpty());
 
 	if (GetWorld()->IsGameWorld())
@@ -272,7 +279,7 @@ void AWorldDataLayers::SetDataLayerRuntimeState(const UDataLayerInstance* InData
 
 		++DataLayersStateEpoch;
 
-		UE_LOG(LogWorldPartition, Log, TEXT("Data Layer '%s' state changed: %s -> %s"),
+		UE_LOG(LogWorldPartition, Log, TEXT("Data Layer Instance '%s' state changed: %s -> %s"),
 			*InDataLayerInstance->GetDataLayerShortName(),
 			*StaticEnum<EDataLayerRuntimeState>()->GetDisplayNameTextByValue((int64)CurrentState).ToString(),
 			*StaticEnum<EDataLayerRuntimeState>()->GetDisplayNameTextByValue((int64)InState).ToString());
@@ -481,6 +488,67 @@ void AWorldDataLayers::ResolveEffectiveRuntimeState(const UDataLayerInstance* In
 	}
 }
 
+static TSet<TObjectPtr<UDataLayerInstance>>& GetDataLayerInstancesFromProvider(UDataLayerInstance* DataLayerInstance)
+{
+	check(DataLayerInstance);
+#if WITH_EDITOR
+	check(DataLayerInstance->GetRootExternalDataLayerInstance() == DataLayerInstance);
+#endif
+	IDataLayerInstanceProvider* DataLayerInstanceProvider = DataLayerInstance->GetImplementingOuter<IDataLayerInstanceProvider>();
+	if (ensure(DataLayerInstanceProvider))
+	{
+		return DataLayerInstanceProvider->GetDataLayerInstances();
+	}
+
+	static TSet<TObjectPtr<UDataLayerInstance>> Empty;
+	return Empty;
+}
+
+bool AWorldDataLayers::AddExternalDataLayerInstance(UExternalDataLayerInstance* ExternalDataLayerInstance)
+{
+#if WITH_EDITOR
+	Modify(/*bDirty*/false);
+#endif
+	check(!IsExternalDataLayerWorldDataLayers());
+	if (!TransientDataLayerInstances.Contains(ExternalDataLayerInstance))
+	{
+		TransientDataLayerInstances.Add(ExternalDataLayerInstance);
+		for (UDataLayerInstance* DataLayerInstance : GetDataLayerInstancesFromProvider(ExternalDataLayerInstance))
+		{
+			SetDataLayerRuntimeState(DataLayerInstance, DataLayerInstance->GetInitialRuntimeState());
+#if !WITH_EDITOR
+			UpdateAccelerationTable(DataLayerInstance, true);
+#endif
+		}
+		return true;
+	}
+	
+	return false;
+}
+
+bool AWorldDataLayers::RemoveExternalDataLayerInstance(UExternalDataLayerInstance* ExternalDataLayerInstance)
+{
+#if WITH_EDITOR
+	Modify(/*bDirty*/false);
+#endif
+	check(!IsExternalDataLayerWorldDataLayers());
+	uint32 Index = TransientDataLayerInstances.IndexOfByKey(ExternalDataLayerInstance);
+	if (Index != INDEX_NONE)
+	{
+		for (UDataLayerInstance* DataLayerInstance : GetDataLayerInstancesFromProvider(ExternalDataLayerInstance))
+		{
+			SetDataLayerRuntimeState(DataLayerInstance, EDataLayerRuntimeState::Unloaded);
+#if !WITH_EDITOR
+			UpdateAccelerationTable(DataLayerInstance, false);
+#endif
+		}
+		TransientDataLayerInstances.RemoveAtSwap(Index);
+		return true;
+	}
+
+	return false;
+}
+
 void AWorldDataLayers::DumpDataLayerRecursively(const UDataLayerInstance* DataLayer, FString Prefix, FOutputDevice& OutputDevice) const
 {
 	auto GetDataLayerRuntimeStateString = [this](const UDataLayerInstance* DataLayer)
@@ -642,7 +710,7 @@ TArray<const UDataLayerInstance*> AWorldDataLayers::GetDataLayerInstances(const
 
 bool AWorldDataLayers::IsEmpty() const
 {
-	return GetDataLayerInstances().IsEmpty();
+	return GetDataLayerInstances().IsEmpty() && TransientDataLayerInstances.IsEmpty();
 }
 
 void AWorldDataLayers::AddDataLayerInstance(UDataLayerInstance* InDataLayerInstance)
@@ -660,6 +728,12 @@ void AWorldDataLayers::AddDataLayerInstance(UDataLayerInstance* InDataLayerInsta
 	{
 		InDataLayerInstance->MarkPackageDirty();
 	}
+	if (UExternalDataLayerInstance* ExternalDataLayerInstance = Cast<UExternalDataLayerInstance>(InDataLayerInstance))
+	{
+		check(TransientDataLayerInstances.IsEmpty());
+		check(!RootExternalDataLayerInstance);
+		RootExternalDataLayerInstance = ExternalDataLayerInstance;
+	}
 }
 
 int32 AWorldDataLayers::RemoveDataLayers(const TArray<UDataLayerInstance*>& InDataLayerInstances, bool bInResolveActorDescContainers)
@@ -715,39 +789,138 @@ void AWorldDataLayers::SetAllowRuntimeDataLayerEditing(bool bInAllowRuntimeDataL
 	}
 }
 
+bool AWorldDataLayers::IsActorEditorContextCurrentColorized(const UDataLayerInstance* InDataLayerInstance) const
+{
+	return InDataLayerInstance && !CurrentDataLayers.CurrentColorizedDataLayerInstanceName.IsNone() && (InDataLayerInstance->GetDataLayerFName() == CurrentDataLayers.CurrentColorizedDataLayerInstanceName);
+}
+
 bool AWorldDataLayers::IsInActorEditorContext(const UDataLayerInstance* InDataLayerInstance) const
 {
 	for (const FName& DataLayerInstanceName : CurrentDataLayers.DataLayerInstanceNames)
 	{
 		const UDataLayerInstance* DataLayerInstance = GetDataLayerInstance(DataLayerInstanceName);
-		if (DataLayerInstance && (DataLayerInstance == InDataLayerInstance) && !DataLayerInstance->IsLocked())
+		if (DataLayerInstance && (DataLayerInstance == InDataLayerInstance) && !DataLayerInstance->IsReadOnly())
 		{
 			return true;
 		}
 	}
+
+	if (const UDataLayerInstance* ExternalDataLayerInstance = GetDataLayerInstance(CurrentDataLayers.ExternalDataLayerName))
+	{
+		check(ExternalDataLayerInstance->IsA<UExternalDataLayerInstance>());
+		if (ExternalDataLayerInstance && ExternalDataLayerInstance == InDataLayerInstance)
+		{
+			return true;
+		}
+	}
+
 	return false;
 }
 
+void AWorldDataLayers::UpdateCurrentColorizedDataLayerInstance()
+{
+	Modify(/*bDirty*/false);
+	CurrentDataLayers.CurrentColorizedDataLayerInstanceName = NAME_None;
+	TArray<UDataLayerInstance*> ContextDataLayerInstances = GetActorEditorContextDataLayers();
+	if (ContextDataLayerInstances.Num() == 1)
+	{
+		CurrentDataLayers.CurrentColorizedDataLayerInstanceName = ContextDataLayerInstances[0]->GetDataLayerFName();
+	}
+	else if (ContextDataLayerInstances.Num() == 2)
+	{
+		for (UDataLayerInstance* DataLayerInstance : ContextDataLayerInstances)
+		{
+			const UExternalDataLayerInstance* ExternalDataLayerInstance = DataLayerInstance->GetRootExternalDataLayerInstance();
+			if (ExternalDataLayerInstance && (ExternalDataLayerInstance != DataLayerInstance) && ContextDataLayerInstances.Contains(ExternalDataLayerInstance))
+			{
+				CurrentDataLayers.CurrentColorizedDataLayerInstanceName = DataLayerInstance->GetDataLayerFName();
+				break;
+			}
+		}
+	}
+}
+
 bool AWorldDataLayers::AddToActorEditorContext(UDataLayerInstance* InDataLayerInstance)
 {
+	ON_SCOPE_EXIT { UpdateCurrentColorizedDataLayerInstance(); };
+
 	check(InDataLayerInstance->CanBeInActorEditorContext());
 	check(ContainsDataLayer(InDataLayerInstance));
+	
+	bool bSuccess = false;
+	if (UExternalDataLayerInstance* ExternalDataLayerInstance = Cast<UExternalDataLayerInstance>(InDataLayerInstance))
+	{
+		Modify(/*bDirty*/false);
+		CurrentDataLayers.ExternalDataLayerName = ExternalDataLayerInstance->GetDataLayerFName();
 
-	if (!CurrentDataLayers.DataLayerInstanceNames.Contains(InDataLayerInstance->GetDataLayerFName()))
+		// Adding an EDL Instance replaces the existing (if any) and removes all DataLayerInstances with a different root EDL Instance
+		TArray<FName> ToRemove;
+		for (const FName& DataLayerInstanceName : CurrentDataLayers.DataLayerInstanceNames)
+		{
+			const UDataLayerInstance* DataLayerInstance = GetDataLayerInstance(DataLayerInstanceName);
+			const UExternalDataLayerInstance* RootEDLInstance = DataLayerInstance ? DataLayerInstance->GetRootExternalDataLayerInstance() : nullptr;
+			if (!DataLayerInstance || (RootEDLInstance && RootEDLInstance != ExternalDataLayerInstance))
+			{
+				ToRemove.Add(DataLayerInstanceName);
+			}
+		}
+		for (const FName& DataLayerInstanceName : ToRemove)
+		{
+			CurrentDataLayers.DataLayerInstanceNames.Remove(DataLayerInstanceName);
+		}
+
+		bSuccess = true;
+	}
+	else if (!CurrentDataLayers.DataLayerInstanceNames.Contains(InDataLayerInstance->GetDataLayerFName()))
 	{
 		Modify(/*bDirty*/false);
 		CurrentDataLayers.DataLayerInstanceNames.Add(InDataLayerInstance->GetDataLayerFName());
-		return true;
+		bSuccess = true;
+
+		// Adding a Data Layer with a RootExternalDataLayerInstance will set this Root EDL Instance in the context
+		if (const UExternalDataLayerInstance* RootEDLInstance = InDataLayerInstance->GetRootExternalDataLayerInstance())
+		{
+			if (!AddToActorEditorContext(const_cast<UExternalDataLayerInstance*>(RootEDLInstance)))
+			{
+				bSuccess = false;
+			}
+		}
 	}
-	return false;
+
+	return bSuccess;
 }
 
 bool AWorldDataLayers::RemoveFromActorEditorContext(UDataLayerInstance* InDataLayerInstance)
 {
+	ON_SCOPE_EXIT { UpdateCurrentColorizedDataLayerInstance(); };
+
 	check(InDataLayerInstance->CanBeInActorEditorContext());
 	check(ContainsDataLayer(InDataLayerInstance));
 
-	if (CurrentDataLayers.DataLayerInstanceNames.Contains(InDataLayerInstance->GetDataLayerFName()))
+	const UExternalDataLayerInstance* ExternalDataLayerInstance = Cast<UExternalDataLayerInstance>(InDataLayerInstance);
+	const FName ExternalDataLayerName = ExternalDataLayerInstance ? ExternalDataLayerInstance->GetDataLayerFName() : NAME_None;
+	if (!ExternalDataLayerName.IsNone() && (CurrentDataLayers.ExternalDataLayerName == ExternalDataLayerName))
+	{
+		Modify(/*bDirty*/false);
+		CurrentDataLayers.ExternalDataLayerName = NAME_None;
+		// Removing an EDL Instance removes all DataLayerInstances with a matching root EDL Instance
+		TArray<FName> ToRemove;
+		for (const FName& DataLayerInstanceName : CurrentDataLayers.DataLayerInstanceNames)
+		{
+			const UDataLayerInstance* DataLayerInstance = GetDataLayerInstance(DataLayerInstanceName);
+			const UExternalDataLayerInstance* RootEDLInstance = DataLayerInstance ? DataLayerInstance->GetRootExternalDataLayerInstance() : nullptr;
+			if (!DataLayerInstance || (RootEDLInstance && RootEDLInstance == ExternalDataLayerInstance))
+			{
+				ToRemove.Add(DataLayerInstanceName);
+			}
+		}
+		for (const FName& DataLayerInstanceName : ToRemove)
+		{
+			CurrentDataLayers.DataLayerInstanceNames.Remove(DataLayerInstanceName);
+		}
+		return true;
+	}
+	else if (CurrentDataLayers.DataLayerInstanceNames.Contains(InDataLayerInstance->GetDataLayerFName()))
 	{
 		Modify(/*bDirty*/false);
 		CurrentDataLayers.DataLayerInstanceNames.Remove(InDataLayerInstance->GetDataLayerFName());
@@ -756,12 +929,15 @@ bool AWorldDataLayers::RemoveFromActorEditorContext(UDataLayerInstance* InDataLa
 	return false;
 }
 
-void AWorldDataLayers::PushActorEditorContext(int32 InContextID)
+void AWorldDataLayers::PushActorEditorContext(int32 InContextID, bool bDuplicateContext)
 {
 	Modify(/*bDirty*/false);
 	CurrentDataLayers.ContextID = InContextID;
 	CurrentDataLayersStack.Push(CurrentDataLayers);
-	CurrentDataLayers.Reset();
+	if (!bDuplicateContext)
+	{
+		CurrentDataLayers.Reset();
+	}
 }
 
 void AWorldDataLayers::PopActorEditorContext(int32 InContextID)
@@ -786,11 +962,18 @@ TArray<UDataLayerInstance*> AWorldDataLayers::GetActorEditorContextDataLayers()
 	for (const FName& DataLayerInstanceName : CurrentDataLayers.DataLayerInstanceNames)
 	{
 		const UDataLayerInstance* DataLayerInstance = GetDataLayerInstance(DataLayerInstanceName);
-		if (DataLayerInstance && !DataLayerInstance->IsLocked())
+		if (DataLayerInstance && !DataLayerInstance->IsReadOnly())
 		{
 			Result.Add(const_cast<UDataLayerInstance*>(DataLayerInstance));
 		}
 	}
+
+	if (const UDataLayerInstance* ExternalDataLayerInstance = GetDataLayerInstance(CurrentDataLayers.ExternalDataLayerName))
+	{
+		check(ExternalDataLayerInstance->IsA<UExternalDataLayerInstance>());
+		Result.Add(const_cast<UDataLayerInstance*>(ExternalDataLayerInstance));
+	}
+
 	return Result;
 }
 
@@ -798,7 +981,21 @@ TArray<UDataLayerInstance*> AWorldDataLayers::GetActorEditorContextDataLayers()
 
 bool AWorldDataLayers::ContainsDataLayer(const UDataLayerInstance* InDataLayerInstance) const 
 {
-	return GetDataLayerInstances().Contains(InDataLayerInstance);
+	if (GetDataLayerInstances().Contains(InDataLayerInstance) || TransientDataLayerInstances.Contains(InDataLayerInstance))
+	{
+		return true;
+	}
+	else if (!TransientDataLayerInstances.IsEmpty())
+	{
+		for (UDataLayerInstance* DataLayerInstance : TransientDataLayerInstances)
+		{
+			if (GetDataLayerInstancesFromProvider(DataLayerInstance).Contains(InDataLayerInstance))
+			{
+				return true;
+			}
+		}
+	}
+	return false;
 }
 
 const UDataLayerInstance* AWorldDataLayers::GetDataLayerInstance(const FName& InDataLayerInstanceName) const
@@ -888,6 +1085,26 @@ const UDataLayerInstance* AWorldDataLayers::GetDataLayerInstance(const UDataLaye
 #endif
 }
 
+bool AWorldDataLayers::IsExternalDataLayerWorldDataLayers() const
+{
+	return !!GetRootExternalDataLayerInstance();
+}
+
+UExternalDataLayerInstance* AWorldDataLayers::GetExternalDataLayerInstance(const UExternalDataLayerAsset* InExternalDataLayerAsset)
+{
+	if (UDataLayerInstance* DataLayerInstance = const_cast<UDataLayerInstance*>(GetDataLayerInstance(InExternalDataLayerAsset)))
+	{
+		return CastChecked<UExternalDataLayerInstance>(DataLayerInstance);
+	}
+
+	return nullptr;
+}
+
+const UExternalDataLayerInstance* AWorldDataLayers::GetExternalDataLayerInstance(const UExternalDataLayerAsset* InExternalDataLayerAsset) const
+{
+	return const_cast<AWorldDataLayers*>(this)->GetExternalDataLayerInstance(InExternalDataLayerAsset);
+}
+
 void AWorldDataLayers::ForEachDataLayerInstance(TFunctionRef<bool(UDataLayerInstance*)> Func)
 {
 	for (UDataLayerInstance* DataLayerInstance : GetDataLayerInstances())
@@ -897,6 +1114,17 @@ void AWorldDataLayers::ForEachDataLayerInstance(TFunctionRef<bool(UDataLayerInst
 			return;
 		}
 	}
+
+	for (UDataLayerInstance* DataLayerInstance : TransientDataLayerInstances)
+	{
+		for (UDataLayerInstance* TransientDataLayerInstance : GetDataLayerInstancesFromProvider(DataLayerInstance))
+		{
+			if (TransientDataLayerInstance && !Func(TransientDataLayerInstance))
+			{
+				return;
+			}
+		}
+	}
 }
 
 void AWorldDataLayers::ForEachDataLayerInstance(TFunctionRef<bool(UDataLayerInstance*)> Func) const
@@ -937,11 +1165,6 @@ TSet<TObjectPtr<UDataLayerInstance>>& AWorldDataLayers::GetDataLayerInstances()
 	return DataLayerInstances;
 }
 
-const TSet<TObjectPtr<UDataLayerInstance>>& AWorldDataLayers::GetDataLayerInstances() const
-{
-	return const_cast<AWorldDataLayers*>(this)->GetDataLayerInstances();
-}
-
 void AWorldDataLayers::OnDataLayerManagerInitialized()
 {
 #if WITH_EDITOR
@@ -955,21 +1178,26 @@ void AWorldDataLayers::OnDataLayerManagerInitialized()
 		SetUseExternalPackageDataLayerInstances(false);
 	}
 
-	ConvertDataLayerToInstancces();
-
-	// Remove all Editor Data Layers when cooking or when in a game world
-	if (IsRunningCookCommandlet() || GetWorld()->IsGameWorld())
+	if (RootExternalDataLayerInstance)
 	{
-		TArray<UDataLayerInstance*> EditorDataLayers;
-		ForEachDataLayerInstance([&EditorDataLayers](UDataLayerInstance* DataLayerInstance)
+		TArray<UDataLayerInstance*> ToDelete;
+		ForEachDataLayerInstance([&ToDelete, this](UDataLayerInstance* DataLayerInstance)
 		{
-			if (!DataLayerInstance->IsRuntime())
+			if (DataLayerInstance->IsA<UExternalDataLayerInstance>() && RootExternalDataLayerInstance != DataLayerInstance)
 			{
-				EditorDataLayers.Add(DataLayerInstance);
+				ToDelete.Add(DataLayerInstance);
 			}
 			return true;
 		});
-		RemoveDataLayers(EditorDataLayers);
+		RemoveDataLayers(ToDelete);
+	}
+
+	ConvertDataLayerToInstances();
+
+	// Remove all Editor Data Layers when cooking or when in a game world
+	if (IsRunningCookCommandlet() || GetWorld()->IsGameWorld())
+	{
+		RemoveEditorDataLayers();
 	}
 
 	// Setup defaults before overriding with user settings
@@ -1024,54 +1252,98 @@ void AWorldDataLayers::PostLoad()
 {
 	Super::PostLoad();
 
-	GetLevel()->ConditionalPostLoad();
-
-	// Patch WorldDataLayer in UWorld.
-	// Only the "main" world data Layer is named AWorldDataLayers::StaticClass()->GetFName() for a given world.
-	if ((GetTypedOuter<UWorld>()->GetWorldDataLayers() == nullptr) && (GetFName() == GetWorldPartionWorldDataLayersName()))
+	ULevel* Level = GetLevel();
+#if WITH_EDITOR
+	// When duplicating the EDL WorldDataLayers for PIE/Cook, the outer is the GObjTransientPkg.
+	// In this case, there's nothing to do in the PostLoad called by DuplicateObject.
+	// (see UExternalDataLayerManager::CreateExternalStreamingObjectUsingStreamingGeneration for details)
+	check(Level || IsExternalDataLayerWorldDataLayers());
+	if (Level)
+#endif
 	{
-		GetTypedOuter<UWorld>()->SetWorldDataLayers(this);
-	}
+		Level->ConditionalPostLoad();
+
+		// Patch WorldDataLayer in UWorld.
+		// Only the "main" world data Layer is named AWorldDataLayers::StaticClass()->GetFName() for a given world.
+		if ((GetTypedOuter<UWorld>()->GetWorldDataLayers() == nullptr) && (GetFName() == GetWorldPartionWorldDataLayersName()))
+		{
+			GetTypedOuter<UWorld>()->SetWorldDataLayers(this);
+		}
 
 #if WITH_EDITOR
-	ULevel* Level = GetLevel();
-	if (!Level->bWasDuplicated && IsUsingExternalPackageDataLayerInstances())
-	{
-		// Load all folders for this level
-		FExternalPackageHelper::LoadObjectsFromExternalPackages<UDataLayerInstance>(this, [this](UDataLayerInstance* LoadedDataLayerInstance)
+		if (!Level->bWasDuplicated && IsUsingExternalPackageDataLayerInstances())
 		{
-			check(IsValid(LoadedDataLayerInstance));
-			LoadedExternalPackageDataLayerInstances.Add(LoadedDataLayerInstance);
-		});
+			// Load all folders for this level
+			FExternalPackageHelper::LoadObjectsFromExternalPackages<UDataLayerInstance>(this, [this](UDataLayerInstance* LoadedDataLayerInstance)
+			{
+				check(IsValid(LoadedDataLayerInstance));
+				LoadedExternalPackageDataLayerInstances.Add(LoadedDataLayerInstance);
+			});
+		}
+#endif
 	}
 
+#if WITH_EDITOR
 	bListedInSceneOutliner = true;
+#else
+	// Build acceleration tables
+	ForEachDataLayerInstance([this](const UDataLayerInstance* DataLayerInstance)
+	{
+		UpdateAccelerationTable(DataLayerInstance, true);
+		return true;
+	});
 #endif
+}
 
 #if !WITH_EDITOR
-	// Build acceleration tables
-	ForEachDataLayerInstance([this](const UDataLayerInstance* DataLayerInstance)
+void AWorldDataLayers::UpdateAccelerationTable(const UDataLayerInstance* DataLayerInstance, bool bIsAdding)
+{
+	if (bIsAdding)
 	{
 		InstanceNameToInstance.Add(DataLayerInstance->GetDataLayerFName(), DataLayerInstance);
+	}
+	else
+	{
+		InstanceNameToInstance.Remove(DataLayerInstance->GetDataLayerFName());
+	}
 
-		static_assert(DATALAYER_TO_INSTANCE_RUNTIME_CONVERSION_ENABLED, "Remove unnecessary cast. All DataLayerInstance now have assets");
-		if (const UDataLayerInstanceWithAsset* DataLayerInstanceWithAsset = Cast<UDataLayerInstanceWithAsset>(DataLayerInstance))
+	static_assert(DATALAYER_TO_INSTANCE_RUNTIME_CONVERSION_ENABLED, "Remove unnecessary cast. All DataLayerInstance now have assets");
+	if (const UDataLayerInstanceWithAsset* DataLayerInstanceWithAsset = Cast<UDataLayerInstanceWithAsset>(DataLayerInstance))
+	{
+		if (!DataLayerInstanceWithAsset->GetAsset())
 		{
-			if (!DataLayerInstanceWithAsset->GetAsset())
+			UE_LOG(LogWorldPartition, Warning, TEXT("DataLayerWithAsset %s has null asset."), *DataLayerInstanceWithAsset->GetPathName());
+		}
+		else
+		{
+			if (bIsAdding)
 			{
-				UE_LOG(LogWorldPartition, Warning, TEXT("DataLayerWithAsset %s has null asset."), *DataLayerInstanceWithAsset->GetPathName());
+				AssetNameToInstance.Add(DataLayerInstanceWithAsset->GetAsset()->GetFullName(), DataLayerInstance);
 			}
 			else
 			{
-				AssetNameToInstance.Add(DataLayerInstanceWithAsset->GetAsset()->GetFullName(), DataLayerInstance);
+				AssetNameToInstance.Remove(DataLayerInstanceWithAsset->GetAsset()->GetFullName());
 			}
 		}
+	}
+}
+#endif
+
+#if WITH_EDITOR
+void AWorldDataLayers::RemoveEditorDataLayers()
+{
+	TArray<UDataLayerInstance*> EditorDataLayers;
+	ForEachDataLayerInstance([&EditorDataLayers](UDataLayerInstance* DataLayerInstance)
+	{
+		if (!DataLayerInstance->IsRuntime())
+		{
+			EditorDataLayers.Add(DataLayerInstance);
+		}
 		return true;
 	});
-#endif
+	RemoveDataLayers(EditorDataLayers);
 }
 
-#if WITH_EDITOR
 bool AWorldDataLayers::IsSubWorldDataLayers() const
 {
 	UWorld* ActorWorld = GetWorld();
@@ -1079,41 +1351,22 @@ bool AWorldDataLayers::IsSubWorldDataLayers() const
 	return ActorWorld != nullptr && OuterWorld != nullptr && (OuterWorld->GetFName() != ActorWorld->GetFName());
 }
 
-bool AWorldDataLayers::IsReadOnly() const
+bool AWorldDataLayers::IsReadOnly(FText* OutReason) const
 {
 	if (IsSubWorldDataLayers())
 	{
 		const UWorld* ActorWorld = GetWorld();
 		const ULevel* CurrentLevel = (ActorWorld && ActorWorld->GetCurrentLevel() && !ActorWorld->GetCurrentLevel()->IsPersistentLevel()) ? ActorWorld->GetCurrentLevel() : nullptr;
 		const bool bIsCurrentLevelWorldDataLayers = CurrentLevel && (CurrentLevel == GetLevel());
-		return !bIsCurrentLevelWorldDataLayers;
-	}
-	return false;
-}
-
-bool AWorldDataLayers::IsRuntimeRelevant() const
-{
-	// This function is only meant to be called on Game worlds
-	check(GetWorld()->IsGameWorld());
-
-	// WorldDataLayers referenced by LevelInstances are not relevant at runtime
-	UWorld* OuterWorld = GetTypedOuter<UWorld>();
-	if (OuterWorld && (OuterWorld->GetWorldDataLayers() == this))
-	{
-		if (!UWorld::IsPartitionedWorld(GetWorld()))
+		if (!bIsCurrentLevelWorldDataLayers)
 		{
+			if (OutReason)
+			{
+				*OutReason = LOCTEXT("WorldDataLayersIsReadOnly", "WorldDataLayers actor is read-only, it's not the Current Level's WorldDataLayers.");
+			}
 			return true;
 		}
-		else
-		{
-			check(OuterWorld->IsPartitionedWorld());
-			ULevelInstanceSubsystem* LevelInstanceSubsystem = UWorld::GetSubsystem<ULevelInstanceSubsystem>(GetWorld());
-			ILevelInstanceInterface* LevelInstance = LevelInstanceSubsystem ? LevelInstanceSubsystem->GetOwningLevelInstance(OuterWorld->PersistentLevel) : nullptr;
-			return (LevelInstance == nullptr);
-		}
 	}
-
-	// @todo_ow: revisit that logic for content bundle data layers
 	return false;
 }
 
@@ -1170,11 +1423,12 @@ bool AWorldDataLayers::SetUseExternalPackageDataLayerInstances(bool bInNewValue,
 	}
 
 	Modify();
-	ForEachDataLayerInstance([this, bInNewValue](UDataLayerInstance* DataLayerInstance)
+	// Only change packaging for owned data layer instances
+	for (UDataLayerInstance* DataLayerInstance : GetDataLayerInstances())
 	{
+		check(DataLayerInstance->GetDirectOuterWorldDataLayers() == this);
 		FExternalPackageHelper::SetPackagingMode(DataLayerInstance, this, bInNewValue);
-		return true;
-	});
+	}
 	bUseExternalPackageDataLayerInstances = bInNewValue;
 	Swap(DataLayerInstances, ExternalPackageDataLayerInstances);
 
@@ -1184,9 +1438,9 @@ bool AWorldDataLayers::SetUseExternalPackageDataLayerInstances(bool bInNewValue,
 	return true;
 }
 
-void AWorldDataLayers::ConvertDataLayerToInstancces()
+void AWorldDataLayers::ConvertDataLayerToInstances()
 {
-	static_assert(DATALAYER_TO_INSTANCE_RUNTIME_CONVERSION_ENABLED, "AWorldDataLayers::ConvertDataLayerToInstancces function is deprecated and needs to be deleted.");
+	static_assert(DATALAYER_TO_INSTANCE_RUNTIME_CONVERSION_ENABLED, "AWorldDataLayers::ConvertDataLayerToInstances function is deprecated and needs to be deleted.");
 	
 	PRAGMA_DISABLE_DEPRECATION_WARNINGS
 	bHasDeprecatedDataLayers = !WorldDataLayers_DEPRECATED.IsEmpty();
@@ -1276,6 +1530,16 @@ void AWorldDataLayers::PostEditUndo()
 	CachedDataLayerInstances.Empty();
 }
 
+bool AWorldDataLayers::ShouldLevelKeepRefIfExternal() const
+{
+	return !IsExternalDataLayerWorldDataLayers();
+}
+
+bool AWorldDataLayers::IsEditorOnly() const
+{
+	return Super::IsEditorOnly() || IsExternalDataLayerWorldDataLayers();
+}
+
 #endif
 
 //~ Begin Deprecated
diff --git a/Engine/Source/Runtime/Engine/Private/WorldPartition/DataLayer/WorldDataLayersActorDesc.cpp b/Engine/Source/Runtime/Engine/Private/WorldPartition/DataLayer/WorldDataLayersActorDesc.cpp
--- a/Engine/Source/Runtime/Engine/Private/WorldPartition/DataLayer/WorldDataLayersActorDesc.cpp
+++ b/Engine/Source/Runtime/Engine/Private/WorldPartition/DataLayer/WorldDataLayersActorDesc.cpp
@@ -10,9 +10,13 @@
 #include "WorldPartition/WorldPartitionHelpers.h"
 #include "WorldPartition/DataLayer/WorldDataLayers.h"
 #include "WorldPartition/DataLayer/DataLayerInstanceWithAsset.h"
+#include "WorldPartition/DataLayer/ExternalDataLayerInstance.h"
+#include "WorldPartition/WorldPartitionActorDescInstance.h"
 #include "WorldPartition/WorldPartitionActorContainerID.h"
 #include "UObject/FortniteSeasonBranchObjectVersion.h"
 #include "UObject/FortniteMainBranchObjectVersion.h"
+#include "UObject/UE5ReleaseStreamObjectVersion.h"
+#include "UObject/MetaData.h"
 #include "ExternalPackageHelper.h"
 
 FArchive& operator<<(FArchive& Ar, FDataLayerInstanceDesc& Desc)
@@ -203,6 +207,7 @@ void FDataLayerInstanceDesc::Init(UDataLayerInstance* InDataLayerInstance)
 
 FWorldDataLayersActorDesc::FWorldDataLayersActorDesc()
 : bIsValid(false)
+, bIsExternalDataLayerWorldDataLayers(false)
 , bUseExternalPackageDataLayerInstances(false)
 {}
 
@@ -212,6 +217,7 @@ void FWorldDataLayersActorDesc::Init(const AActor* InActor)
 
 	const AWorldDataLayers* WorldDataLayers = CastChecked<AWorldDataLayers>(InActor);
 	bUseExternalPackageDataLayerInstances = WorldDataLayers->IsUsingExternalPackageDataLayerInstances();
+	bIsExternalDataLayerWorldDataLayers = WorldDataLayers->IsExternalDataLayerWorldDataLayers();
 	if (!bUseExternalPackageDataLayerInstances)
 	{
 		WorldDataLayers->ForEachDataLayerInstance([this](UDataLayerInstance* DataLayerInstance)
@@ -231,7 +237,8 @@ bool FWorldDataLayersActorDesc::Equals(const FWorldPartitionActorDesc* Other) co
 	{
 		const FWorldDataLayersActorDesc* OtherDesc = (FWorldDataLayersActorDesc*)Other;
 		return (bUseExternalPackageDataLayerInstances == OtherDesc->bUseExternalPackageDataLayerInstances) && 
-			CompareUnsortedArrays(DataLayerInstances, OtherDesc->DataLayerInstances);
+			   (bIsExternalDataLayerWorldDataLayers == OtherDesc->bIsExternalDataLayerWorldDataLayers) &&
+			   CompareUnsortedArrays(DataLayerInstances, OtherDesc->DataLayerInstances);
 	}
 	return false;
 }
@@ -240,6 +247,7 @@ void FWorldDataLayersActorDesc::Serialize(FArchive& Ar)
 {
 	Ar.UsingCustomVersion(FFortniteSeasonBranchObjectVersion::GUID);
 	Ar.UsingCustomVersion(FFortniteMainBranchObjectVersion::GUID);
+	Ar.UsingCustomVersion(FUE5ReleaseStreamObjectVersion::GUID);
 
 	FWorldPartitionActorDesc::Serialize(Ar);
 
@@ -254,7 +262,47 @@ void FWorldDataLayersActorDesc::Serialize(FArchive& Ar)
 		{
 			Ar << bUseExternalPackageDataLayerInstances;
 		}
+		if (Ar.CustomVer(FUE5ReleaseStreamObjectVersion::GUID) >= FUE5ReleaseStreamObjectVersion::WorldPartitionExternalDataLayers)
+		{
+			Ar << bIsExternalDataLayerWorldDataLayers;
+		}
+	}
+}
+
+void FWorldDataLayersActorDesc::OnUnloadingInstance(const FWorldPartitionActorDescInstance* InActorDescInstance) const
+{
+	if (AWorldDataLayers* WorldDataLayers = Cast<AWorldDataLayers>(InActorDescInstance->GetActor()))
+	{
+		if (WorldDataLayers->IsUsingExternalPackageDataLayerInstances())
+		{
+			WorldDataLayers->ForEachDataLayerInstance([this](UDataLayerInstance* DataLayerInstance)
+			{
+				check(DataLayerInstance->IsPackageExternal())
+				ForEachObjectWithPackage(DataLayerInstance->GetPackage(), [](UObject* Object)
+				{
+					if (Object->HasAnyFlags(RF_Public | RF_Standalone))
+					{
+						CastChecked<UMetaData>(Object)->ClearFlags(RF_Public | RF_Standalone);
+					}
+					return true;
+				}, false);
+
+				return true;
+			});
+		}
+	}
+	FWorldPartitionActorDesc::OnUnloadingInstance(InActorDescInstance);
+}
+
+bool FWorldDataLayersActorDesc::IsRuntimeRelevant(const FWorldPartitionActorDescInstance* InActorDescInstance) const 
+{
+	if (!FWorldPartitionActorDesc::IsRuntimeRelevant(InActorDescInstance))
+	{
+		return false;
 	}
+
+	// ExternalDataLayer WorldDataLayers (used to store data layers for an ExternalDataLayer) are not used at runtime.
+	return !bIsExternalDataLayerWorldDataLayers;
 }
 
 const TArray<FDataLayerInstanceDesc>& FWorldDataLayersActorDesc::GetExternalPackageDataLayerInstances() const
diff --git a/Engine/Source/Runtime/Engine/Private/WorldPartition/ErrorHandling/WorldPartitionStreamingGenerationLogErrorHandler.cpp b/Engine/Source/Runtime/Engine/Private/WorldPartition/ErrorHandling/WorldPartitionStreamingGenerationLogErrorHandler.cpp
--- a/Engine/Source/Runtime/Engine/Private/WorldPartition/ErrorHandling/WorldPartitionStreamingGenerationLogErrorHandler.cpp
+++ b/Engine/Source/Runtime/Engine/Private/WorldPartition/ErrorHandling/WorldPartitionStreamingGenerationLogErrorHandler.cpp
@@ -26,9 +26,20 @@ void FStreamingGenerationLogErrorHandler::OnInvalidReferenceGridPlacement(const
 	UE_ASSET_LOG_ACTORDESCVIEW(LogWorldPartition, Log, ActorDescView, TEXT("%s %s reference %s %s"), ActorDescView.GetIsSpatiallyLoaded() ? *SpatiallyLoadedActor : *NonSpatiallyLoadedActor, *GetActorName(ActorDescView), ReferenceActorDescView.GetIsSpatiallyLoaded() ? *SpatiallyLoadedActor : *NonSpatiallyLoadedActor, *GetActorName(ReferenceActorDescView));
 }
 
-void FStreamingGenerationLogErrorHandler::OnInvalidReferenceDataLayers(const IWorldPartitionActorDescInstanceView& ActorDescView, const IWorldPartitionActorDescInstanceView& ReferenceActorDescView)
+void FStreamingGenerationLogErrorHandler::OnInvalidReferenceDataLayers(const IWorldPartitionActorDescInstanceView& ActorDescView, const IWorldPartitionActorDescInstanceView& ReferenceActorDescView, EDataLayerInvalidReason Reason)
 {
-	UE_ASSET_LOG_ACTORDESCVIEW(LogWorldPartition, Log, ActorDescView, TEXT("Actor %s references an actor in a different set of runtime data layers %s"), *GetActorName(ActorDescView), *GetActorName(ReferenceActorDescView));
+	const FString ActorName = GetActorName(ActorDescView);
+	const FString ReferenceActorName = GetActorName(ReferenceActorDescView);
+
+	switch (Reason)
+	{
+	case EDataLayerInvalidReason::ReferencedActorDifferentRuntimeDataLayers:
+		UE_ASSET_LOG_ACTORDESCVIEW(LogWorldPartition, Log, ActorDescView, TEXT("Actor %s references an actor %s in a different set of runtime data layers"), *ActorName, *ReferenceActorName);
+		break;
+	case EDataLayerInvalidReason::ReferencedActorDifferentExternalDataLayer:
+		UE_ASSET_LOG_ACTORDESCVIEW(LogWorldPartition, Log, ActorDescView, TEXT("Actor %s references an actor %s with a different external data layer"), *ActorName, *ReferenceActorName);
+		break;
+	}
 }
 
 void FStreamingGenerationLogErrorHandler::OnInvalidReferenceRuntimeGrid(const IWorldPartitionActorDescInstanceView& ActorDescView, const IWorldPartitionActorDescInstanceView& ReferenceActorDescView)
@@ -51,6 +62,11 @@ void FStreamingGenerationLogErrorHandler::OnInvalidReferenceDataLayerAsset(const
 	UE_ASSET_LOG(LogWorldPartition, Log, DataLayerInstance, TEXT("Data Layer does not have a Data Layer asset"));
 }
 
+void FStreamingGenerationLogErrorHandler::OnInvalidDataLayerAssetType(const UDataLayerInstanceWithAsset* DataLayerInstance, const UDataLayerAsset* DataLayerAsset)
+{
+	UE_ASSET_LOG(LogWorldPartition, Log, DataLayerInstance, TEXT("Data Layer is not compatible with Data Layer asset %s type %s"), *DataLayerAsset->GetName(), *DataLayerAsset->GetClass()->GetName());
+}
+
 void FStreamingGenerationLogErrorHandler::OnDataLayerHierarchyTypeMismatch(const UDataLayerInstance* DataLayerInstance, const UDataLayerInstance* Parent)
 {
 	UE_ASSET_LOG(LogWorldPartition, Log, DataLayerInstance, TEXT("Data Layer %s is of Type %s and its parent %s is of type %s"), *DataLayerInstance->GetDataLayerFullName(), *UEnum::GetValueAsString(DataLayerInstance->GetType()), *Parent->GetDataLayerFullName(), *UEnum::GetValueAsString(Parent->GetType()));
diff --git a/Engine/Source/Runtime/Engine/Private/WorldPartition/ErrorHandling/WorldPartitionStreamingGenerationProxyErrorHandler.cpp b/Engine/Source/Runtime/Engine/Private/WorldPartition/ErrorHandling/WorldPartitionStreamingGenerationProxyErrorHandler.cpp
--- a/Engine/Source/Runtime/Engine/Private/WorldPartition/ErrorHandling/WorldPartitionStreamingGenerationProxyErrorHandler.cpp
+++ b/Engine/Source/Runtime/Engine/Private/WorldPartition/ErrorHandling/WorldPartitionStreamingGenerationProxyErrorHandler.cpp
@@ -18,9 +18,9 @@ void FStreamingGenerationProxyErrorHandler::OnInvalidReferenceGridPlacement(cons
 	InnerErrorHandler->OnInvalidReferenceGridPlacement(ActorDescView, ReferenceActorDescView);
 }
 
-void FStreamingGenerationProxyErrorHandler::OnInvalidReferenceDataLayers(const IWorldPartitionActorDescInstanceView& ActorDescView, const IWorldPartitionActorDescInstanceView& ReferenceActorDescView)
+void FStreamingGenerationProxyErrorHandler::OnInvalidReferenceDataLayers(const IWorldPartitionActorDescInstanceView& ActorDescView, const IWorldPartitionActorDescInstanceView& ReferenceActorDescView, EDataLayerInvalidReason Reason)
 {
-	InnerErrorHandler->OnInvalidReferenceDataLayers(ActorDescView, ReferenceActorDescView);
+	InnerErrorHandler->OnInvalidReferenceDataLayers(ActorDescView, ReferenceActorDescView, Reason);
 }
 
 void FStreamingGenerationProxyErrorHandler::OnInvalidReferenceRuntimeGrid(const IWorldPartitionActorDescInstanceView& ActorDescView, const IWorldPartitionActorDescInstanceView& ReferenceActorDescView)
@@ -43,6 +43,11 @@ void FStreamingGenerationProxyErrorHandler::OnInvalidReferenceDataLayerAsset(con
 	InnerErrorHandler->OnInvalidReferenceDataLayerAsset(DataLayerInstance);
 }
 
+void FStreamingGenerationProxyErrorHandler::OnInvalidDataLayerAssetType(const UDataLayerInstanceWithAsset* DataLayerInstance, const UDataLayerAsset* DataLayerAsset) 
+{
+	InnerErrorHandler->OnInvalidDataLayerAssetType(DataLayerInstance, DataLayerAsset);
+}
+
 void FStreamingGenerationProxyErrorHandler::OnDataLayerHierarchyTypeMismatch(const UDataLayerInstance* DataLayerInstance, const UDataLayerInstance* Parent)
 {
 	InnerErrorHandler->OnDataLayerHierarchyTypeMismatch(DataLayerInstance, Parent);
diff --git a/Engine/Source/Runtime/Engine/Private/WorldPartition/ErrorHandling/WorldPartitionStreamingGenerationTokenizedMessageErrorHandler.cpp b/Engine/Source/Runtime/Engine/Private/WorldPartition/ErrorHandling/WorldPartitionStreamingGenerationTokenizedMessageErrorHandler.cpp
--- a/Engine/Source/Runtime/Engine/Private/WorldPartition/ErrorHandling/WorldPartitionStreamingGenerationTokenizedMessageErrorHandler.cpp
+++ b/Engine/Source/Runtime/Engine/Private/WorldPartition/ErrorHandling/WorldPartitionStreamingGenerationTokenizedMessageErrorHandler.cpp
@@ -6,6 +6,7 @@
 #include "Misc/PackageName.h"
 #include "Misc/UObjectToken.h"
 #include "Misc/MapErrors.h"
+#include "WorldPartition/DataLayer/DataLayerAsset.h"
 #include "WorldPartition/DataLayer/DataLayerInstanceWithAsset.h"
 
 #define LOCTEXT_NAMESPACE "WorldPartition"
@@ -58,14 +59,27 @@ void ITokenizedMessageErrorHandler::OnInvalidReferenceGridPlacement(const IWorld
 	HandleTokenizedMessage(MoveTemp(Message));
 }
 
-void ITokenizedMessageErrorHandler::OnInvalidReferenceDataLayers(const IWorldPartitionActorDescInstanceView& ActorDescView, const IWorldPartitionActorDescInstanceView& ReferenceActorDescView)
+void ITokenizedMessageErrorHandler::OnInvalidReferenceDataLayers(const IWorldPartitionActorDescInstanceView& ActorDescView, const IWorldPartitionActorDescInstanceView& ReferenceActorDescView, EDataLayerInvalidReason Reason)
 {
+	FText ReasonText;
+
+	switch (Reason)
+	{
+	case EDataLayerInvalidReason::ReferencedActorDifferentRuntimeDataLayers:
+		ReasonText = LOCTEXT("TokenMessage_WorldPartition_ReferenceActorInOtherDataLayers", "references an actor in a different set of runtime data layers");
+		break;
+	case EDataLayerInvalidReason::ReferencedActorDifferentExternalDataLayer:
+		ReasonText = LOCTEXT("TokenMessage_WorldPartition_ReferenceActorInOtherExternalDataLayer", "references an actor assigned to a different external data layer");
+		break;
+	};
+
 	TSharedRef<FTokenizedMessage> Message = FTokenizedMessage::Create(EMessageSeverity::Error);
+
 	Message->AddToken(FTextToken::Create(LOCTEXT("TokenMessage_WorldPartition_Actor", "Actor")))
-		->AddToken(FActorToken::Create(ActorDescView.GetActorSoftPath().ToString(), ActorDescView.GetGuid(),  FText::FromString(GetActorName(ActorDescView))))
-		->AddToken(FTextToken::Create(LOCTEXT("TokenMessage_WorldPartition_ReferenceActorInOtherDataLayers", "references an actor in a different set of runtime data layers")))
+		->AddToken(FActorToken::Create(ActorDescView.GetActorSoftPath().ToString(), ActorDescView.GetGuid(), FText::FromString(GetActorName(ActorDescView))))
+		->AddToken(FTextToken::Create(ReasonText))
 		->AddToken(FActorToken::Create(ReferenceActorDescView.GetActorSoftPath().ToString(), ReferenceActorDescView.GetGuid(), FText::FromString(GetActorName(ReferenceActorDescView))))
-		->AddToken(FMapErrorToken::Create(TEXT("WorldPartition_ActorReferenceActorInAnotherDataLayer_CheckForErrors")));
+		->AddToken(FMapErrorToken::Create(TEXT("WorldPartition_ActorInvalidReferenceDataLayers_CheckForErrors")));
 
 	HandleTokenizedMessage(MoveTemp(Message));
 }
@@ -115,6 +129,19 @@ void ITokenizedMessageErrorHandler::OnInvalidReferenceDataLayerAsset(const UData
 	HandleTokenizedMessage(MoveTemp(Message));
 }
 
+void ITokenizedMessageErrorHandler::OnInvalidDataLayerAssetType(const UDataLayerInstanceWithAsset* DataLayerInstance, const UDataLayerAsset* DataLayerAsset)
+{
+	TSharedRef<FTokenizedMessage> Message = FTokenizedMessage::Create(EMessageSeverity::Error);
+	Message->AddToken(FTextToken::Create(LOCTEXT("TokenMessage_DataLayers_DataLayer", "Data layer")))
+		->AddToken(FUObjectToken::Create(DataLayerInstance))
+		->AddToken(FTextToken::Create(LOCTEXT("TokenMessage_DataLayers_InvalidAssetType_NotCompatible", "is not compatible with its Data Layer Asset")))
+		->AddToken(FUObjectToken::Create(DataLayerAsset))
+		->AddToken(FTextToken::Create(LOCTEXT("TokenMessage_DataLayers_InvalidAssetType_type", "type")))
+		->AddToken(FMapErrorToken::Create(TEXT("DataLayers_InvalidDataLayerAssetType_CheckForErrors")));
+
+	HandleTokenizedMessage(MoveTemp(Message));
+}
+
 void ITokenizedMessageErrorHandler::OnDataLayerHierarchyTypeMismatch(const UDataLayerInstance* DataLayerInstance, const UDataLayerInstance* Parent)
 {
 	TSharedRef<FTokenizedMessage> Message = FTokenizedMessage::Create(EMessageSeverity::Error);
diff --git a/Engine/Source/Runtime/Engine/Private/WorldPartition/HLOD/HLODRuntimeSubsystem.cpp b/Engine/Source/Runtime/Engine/Private/WorldPartition/HLOD/HLODRuntimeSubsystem.cpp
--- a/Engine/Source/Runtime/Engine/Private/WorldPartition/HLOD/HLODRuntimeSubsystem.cpp
+++ b/Engine/Source/Runtime/Engine/Private/WorldPartition/HLOD/HLODRuntimeSubsystem.cpp
@@ -759,7 +759,7 @@ bool UWorldPartitionHLODRuntimeSubsystem::WriteHLODStatsCSV(UWorld* InWorld, con
 		{ "Name",				[](FWorldPartitionActorDescInstance* InActorDescInstance, const FHLODActorDesc& InActorDesc) { return InActorDescInstance->GetActorLabel().ToString(); } },
 		{ "HLODLayer",			[](FWorldPartitionActorDescInstance* InActorDescInstance, const FHLODActorDesc& InActorDesc) { return InActorDesc.GetSourceHLODLayer().GetAssetName().ToString(); }},
 		{ "SpatiallyLoaded",	[](FWorldPartitionActorDescInstance* InActorDescInstance, const FHLODActorDesc& InActorDesc) { return InActorDescInstance->GetIsSpatiallyLoaded() ? TEXT("true") : TEXT("false"); } },
-		{ "DataLayers",			[&GetDataLayerShortName](FWorldPartitionActorDescInstance* InActorDescInstance, const FHLODActorDesc& InActorDesc) { return FString::JoinBy(InActorDescInstance->GetDataLayerInstanceNames(), TEXT(" | "), GetDataLayerShortName); } },
+		{ "DataLayers",			[&GetDataLayerShortName](FWorldPartitionActorDescInstance* InActorDescInstance, const FHLODActorDesc& InActorDesc) { return FString::JoinBy(InActorDescInstance->GetDataLayerInstanceNames().ToArray(), TEXT(" | "), GetDataLayerShortName); }},
 
 		GetHLODStat(FWorldPartitionHLODStats::InputActorCount),
 		GetHLODStat(FWorldPartitionHLODStats::InputTriangleCount),
diff --git a/Engine/Source/Runtime/Engine/Private/WorldPartition/RuntimeHashSet/WorldPartitionRuntimeHashSet.cpp b/Engine/Source/Runtime/Engine/Private/WorldPartition/RuntimeHashSet/WorldPartitionRuntimeHashSet.cpp
--- a/Engine/Source/Runtime/Engine/Private/WorldPartition/RuntimeHashSet/WorldPartitionRuntimeHashSet.cpp
+++ b/Engine/Source/Runtime/Engine/Private/WorldPartition/RuntimeHashSet/WorldPartitionRuntimeHashSet.cpp
@@ -144,9 +144,9 @@ void UWorldPartitionRuntimeHashSet::SetDefaultValues()
 	}
 }
 
-void UWorldPartitionRuntimeHashSet::FlushStreaming()
+void UWorldPartitionRuntimeHashSet::FlushStreamingContent()
 {
-	Super::FlushStreaming();
+	Super::FlushStreamingContent();
 	check(!PersistentPartitionDesc.Class);
 	RuntimeStreamingData.Empty();
 }
@@ -276,27 +276,32 @@ bool UWorldPartitionRuntimeHashSet::ParseGridName(FName GridName, TArray<FName>&
 	return true;
 }
 
-URuntimeHashExternalStreamingObjectBase* UWorldPartitionRuntimeHashSet::StoreToExternalStreamingObject(UObject* StreamingObjectOuter, FName StreamingObjectName)
+bool UWorldPartitionRuntimeHashSet::HasStreamingContent() const
+{
+	return !RuntimeStreamingData.IsEmpty();
+}
+
+void UWorldPartitionRuntimeHashSet::StoreStreamingContentToExternalStreamingObject(URuntimeHashExternalStreamingObjectBase* OutExternalStreamingObject)
 {
 	check(!RuntimeStreamingData.IsEmpty());
 
-	URuntimeHashSetExternalStreamingObject* NewStreamingObject = CreateExternalStreamingObject<URuntimeHashSetExternalStreamingObject>(StreamingObjectOuter, StreamingObjectName);
-	NewStreamingObject->RuntimeStreamingData = MoveTemp(RuntimeStreamingData);
+	Super::StoreStreamingContentToExternalStreamingObject(OutExternalStreamingObject);
 
-	for (FRuntimePartitionStreamingData& StreamingData : NewStreamingObject->RuntimeStreamingData)
+	URuntimeHashSetExternalStreamingObject* StreamingObject = CastChecked<URuntimeHashSetExternalStreamingObject>(OutExternalStreamingObject);
+	StreamingObject->RuntimeStreamingData = MoveTemp(RuntimeStreamingData);
+
+	for (FRuntimePartitionStreamingData& StreamingData : StreamingObject->RuntimeStreamingData)
 	{
 		for (UWorldPartitionRuntimeCell* Cell : StreamingData.StreamingCells)
 		{
-			Cell->Rename(nullptr, NewStreamingObject,  REN_DoNotDirty | REN_ForceNoResetLoaders);
+			Cell->Rename(nullptr, StreamingObject,  REN_DoNotDirty | REN_ForceNoResetLoaders);
 		}
 
 		for (UWorldPartitionRuntimeCell* Cell : StreamingData.NonStreamingCells)
 		{
-			Cell->Rename(nullptr, NewStreamingObject,  REN_DoNotDirty | REN_ForceNoResetLoaders);
+			Cell->Rename(nullptr, StreamingObject,  REN_DoNotDirty | REN_ForceNoResetLoaders);
 		}
 	}
-
-	return NewStreamingObject;
 }
 #endif
 
diff --git a/Engine/Source/Runtime/Engine/Private/WorldPartition/RuntimeHashSet/WorldPartitionRuntimeHashSetConversions.cpp b/Engine/Source/Runtime/Engine/Private/WorldPartition/RuntimeHashSet/WorldPartitionRuntimeHashSetConversions.cpp
--- a/Engine/Source/Runtime/Engine/Private/WorldPartition/RuntimeHashSet/WorldPartitionRuntimeHashSetConversions.cpp
+++ b/Engine/Source/Runtime/Engine/Private/WorldPartition/RuntimeHashSet/WorldPartitionRuntimeHashSetConversions.cpp
@@ -17,8 +17,9 @@ UWorldPartitionRuntimeHashSet* UWorldPartitionRuntimeHashSet::CreateFrom(const U
 	check(WorldPartition);
 
 	FStreamingGenerationNullErrorHandler NullErrorHandler;
+	FActorDescContainerInstanceCollection Collection({ TObjectPtr<UActorDescContainerInstance>(WorldPartition->GetActorDescContainerInstance()) });
 	UWorldPartition::FGenerateStreamingParams Params = UWorldPartition::FGenerateStreamingParams()
-		.SetActorDescContainerInstance(WorldPartition->GetActorDescContainerInstance())
+		.SetContainerInstanceCollection(Collection, FStreamingGenerationContainerInstanceCollection::ECollectionType::BaseAndEDLs)
 		.SetErrorHandler(&NullErrorHandler);
 
 	UWorldPartition::FGenerateStreamingContext Context;
diff --git a/Engine/Source/Runtime/Engine/Private/WorldPartition/RuntimeHashSet/WorldPartitionRuntimeHashSetHLODGeneration.cpp b/Engine/Source/Runtime/Engine/Private/WorldPartition/RuntimeHashSet/WorldPartitionRuntimeHashSetHLODGeneration.cpp
--- a/Engine/Source/Runtime/Engine/Private/WorldPartition/RuntimeHashSet/WorldPartitionRuntimeHashSetHLODGeneration.cpp
+++ b/Engine/Source/Runtime/Engine/Private/WorldPartition/RuntimeHashSet/WorldPartitionRuntimeHashSetHLODGeneration.cpp
@@ -36,7 +36,7 @@ public:
 		return WorldBounds;
 	}
 
-	virtual const FActorSetContainerInstance* GetMainWorldContainerInstance() const override
+	virtual const FActorSetContainerInstance* GetActorSetContainerForContextBaseContainerInstance() const override
 	{
 		return &ActorSetContainerInstance;
 	}
@@ -166,16 +166,19 @@ bool UWorldPartitionRuntimeHashSet::SetupHLODActors(const IStreamingGenerationCo
 
 	UWorldPartition* WorldPartition = GetOuterUWorldPartition();
 	const UDataLayerManager* DataLayerManager = WorldPartition->GetDataLayerManager();
-	IStreamingGenerationContext::FActorSetContainerInstance* MainActorSetContainer = const_cast<IStreamingGenerationContext::FActorSetContainerInstance*>(StreamingGenerationContext->GetMainWorldContainerInstance());
-	const FStreamingGenerationContainerInstanceCollection* MainContainerCollection = MainActorSetContainer->ContainerInstanceCollection;
+	IStreamingGenerationContext::FActorSetContainerInstance* BaseActorSetContainerInstance = const_cast<IStreamingGenerationContext::FActorSetContainerInstance*>(StreamingGenerationContext->GetActorSetContainerForContextBaseContainerInstance());
+	const FStreamingGenerationContainerInstanceCollection* BaseContainerInstanceCollection = BaseActorSetContainerInstance->ContainerInstanceCollection;
 
 	// Create the HLOD creation context
 	FHLODCreationContext HLODCreationContext;
-	for (UActorDescContainerInstance::TConstIterator<AWorldPartitionHLOD> HLODIterator(MainContainerCollection->GetMainContainer()); HLODIterator; ++HLODIterator)
+	BaseContainerInstanceCollection->ForEachActorDescContainerInstance([&HLODCreationContext, WorldPartition](const UActorDescContainerInstance* ActorDescContainerInstance)
 	{
-		FWorldPartitionHandle HLODActorHandle(WorldPartition, HLODIterator->GetGuid());
-		HLODCreationContext.HLODActorDescs.Emplace(HLODIterator->GetActorName(), MoveTemp(HLODActorHandle));
-	}
+		for (UActorDescContainerInstance::TConstIterator<AWorldPartitionHLOD> HLODIterator(ActorDescContainerInstance); HLODIterator; ++HLODIterator)
+		{
+			FWorldPartitionHandle HLODActorHandle(WorldPartition, HLODIterator->GetGuid());
+			HLODCreationContext.HLODActorDescs.Emplace(HLODIterator->GetActorName(), MoveTemp(HLODActorHandle));
+		}
+	});
 
 	TUniquePtr<IStreamingGenerationContext> CurrentHLODStreamingGenerationContext = MakeUnique<FStreamingGenerationContextProxy>(StreamingGenerationContext);
 
@@ -293,15 +296,15 @@ bool UWorldPartitionRuntimeHashSet::SetupHLODActors(const IStreamingGenerationCo
 				ActorSetInstance.Bounds = HLODActorDescView->GetRuntimeBounds();
 				ActorSetInstance.RuntimeGrid = HLODActorDescView->GetRuntimeGrid();
 				ActorSetInstance.bIsSpatiallyLoaded = HLODActorDescView->GetIsSpatiallyLoaded();
-				ActorSetInstance.ContentBundleID = MainContainerCollection->GetContentBundleGuid();
+				ActorSetInstance.ContentBundleID = BaseContainerInstanceCollection->GetContentBundleGuid();
 				ActorSetInstance.ActorSetContainerInstance = &HLODStreamingGenerationContext->ActorSetContainerInstance;
 				ActorSetInstance.ActorSet = ActorSet;
 
-				TArray<FName> RuntimeDataLayerInstanceNames;
-				if (FDataLayerUtils::ResolveRuntimeDataLayerInstanceNames(DataLayerManager, *HLODActorDescView, *MainActorSetContainer->ActorDescViewMap, RuntimeDataLayerInstanceNames))
+				FDataLayerInstanceNames RuntimeDataLayerInstanceNames;
+				if (FDataLayerUtils::ResolveRuntimeDataLayerInstanceNames(DataLayerManager, *HLODActorDescView, *BaseActorSetContainerInstance->ActorDescViewMap, RuntimeDataLayerInstanceNames))
 				{
 					HLODActorDescView->SetRuntimeDataLayerInstanceNames(RuntimeDataLayerInstanceNames);
-					ActorSetInstance.DataLayers = DataLayerManager->GetRuntimeDataLayerInstances(RuntimeDataLayerInstanceNames);
+					ActorSetInstance.DataLayers = DataLayerManager->GetRuntimeDataLayerInstances(RuntimeDataLayerInstanceNames.ToArray());
 				}
 
 				UE_LOG(LogWorldPartition, Log, TEXT("\t- %s"), *HLODActorDescInstance->ToString());
