External Data Layer V1

diff --git a/Engine/Source/Runtime/Engine/Private/WorldPartition/RuntimeSpatialHash/RuntimeSpatialHashHLOD.cpp b/Engine/Source/Runtime/Engine/Private/WorldPartition/RuntimeSpatialHash/RuntimeSpatialHashHLOD.cpp
--- a/Engine/Source/Runtime/Engine/Private/WorldPartition/RuntimeSpatialHash/RuntimeSpatialHashHLOD.cpp
+++ b/Engine/Source/Runtime/Engine/Private/WorldPartition/RuntimeSpatialHash/RuntimeSpatialHashHLOD.cpp
@@ -6,7 +6,6 @@
 #if WITH_EDITOR
 #include "HAL/PlatformFile.h"
 #include "WorldPartition/RuntimeSpatialHash/RuntimeSpatialHashGridHelper.h"
-
 #include "WorldPartition/HLOD/HLODLayer.h"
 #include "WorldPartition/HLOD/HLODActor.h"
 #include "WorldPartition/HLOD/HLODActorDesc.h"
@@ -18,7 +17,7 @@
 #include "WorldPartition/WorldPartition.h"
 #include "WorldPartition/WorldPartitionHelpers.h"
 #include "WorldPartition/ContentBundle/ContentBundleActivationScope.h"
-#include "WorldPartition/WorldPartitionActorDescInstance.h"
+#include "ActorEditorContext/ScopedActorEditorContextSetExternalDataLayerAsset.h"
 
 #include "UObject/GCObjectScopeGuard.h"
 #include "UObject/SavePackage.h"
@@ -114,9 +113,7 @@ static void GameTick(UWorld* InWorld)
 
 static TArray<FGuid> GenerateHLODActorsForGrid(UWorldPartition* WorldPartition, const IStreamingGenerationContext* StreamingGenerationContext, const UWorldPartition::FSetupHLODActorsParams& Params, const FSpatialHashRuntimeGrid& RuntimeGrid, uint32 HLODLevel, FHLODCreationContext& Context, ISourceControlHelper* SourceControlHelper, const TArray<const IStreamingGenerationContext::FActorSetInstance*>& ActorSetInstances, TArray<FName>& NewActors, TArray<FName>& DirtyActors, const FSpatialHashSettings& Settings)
 {
-	const IStreamingGenerationContext::FActorSetContainerInstance* MainActorSetContainer = StreamingGenerationContext->GetMainWorldContainerInstance();
 	const FBox WorldBounds = StreamingGenerationContext->GetWorldBounds();
-
 	const FSquare2DGridHelper PartitionedActors = GetPartitionedActors(WorldBounds, RuntimeGrid, ActorSetInstances, Settings);
 	const FSquare2DGridHelper::FGridLevel::FGridCell& AlwaysLoadedCell = PartitionedActors.GetAlwaysLoadedCell();
 
@@ -371,13 +368,11 @@ static void UpdateHLODGridsActors(UWorld* World, const UActorDescContainerInstan
 				if (HLODGrid && GridActor->GridSettings.Priority && GridActor->GridSettings.HLODLayer)
 				{
 					check(GridActor->GetContentBundleGuid() == ContainerInstance->GetContentBundleGuid());
-
 					ExistingGridActors.Emplace(GridActor->GridSettings.GridName, GridActor);
 				}
 				else
 				{
 					World->DestroyActor(GridActor);
-
 					DeletePackage(GridActor->GetPackage(), SourceControlHelper);
 				}
 			}
@@ -396,6 +391,7 @@ static void UpdateHLODGridsActors(UWorld* World, const UActorDescContainerInstan
 		if (!GridActor)
 		{
 			FContentBundleActivationScope ContentBndleScope(ContainerInstance->GetContentBundleGuid());
+			FScopedActorEditorContextSetExternalDataLayerAsset EDLScope(ContainerInstance->GetExternalDataLayerAsset());
 
 			FActorSpawnParameters SpawnParams;
 			SpawnParams.bCreateActorPackage = true;
@@ -404,6 +400,7 @@ static void UpdateHLODGridsActors(UWorld* World, const UActorDescContainerInstan
 			bDirty = true;
 
 			check(GridActor->GetContentBundleGuid() == ContainerInstance->GetContentBundleGuid());
+			check(GridActor->GetExternalDataLayerAsset() == ContainerInstance->GetExternalDataLayerAsset());
 		}
 
 		const FString ActorLabel = GridSettings.GridName.ToString();
@@ -448,8 +445,8 @@ bool UWorldPartitionRuntimeSpatialHash::SetupHLODActors(const IStreamingGenerati
 	}
 
 	UWorldPartition* WorldPartition = GetOuterUWorldPartition();
-	IStreamingGenerationContext::FActorSetContainerInstance* MainActorSetContainer = const_cast<IStreamingGenerationContext::FActorSetContainerInstance*>(StreamingGenerationContext->GetMainWorldContainerInstance());
-	const FStreamingGenerationContainerInstanceCollection* MainContainerCollection = MainActorSetContainer->ContainerInstanceCollection;
+	IStreamingGenerationContext::FActorSetContainerInstance* BaseActorSetContainerInstance = const_cast<IStreamingGenerationContext::FActorSetContainerInstance*>(StreamingGenerationContext->GetActorSetContainerForContextBaseContainerInstance());
+	const FStreamingGenerationContainerInstanceCollection* BaseContainerInstanceCollection = BaseActorSetContainerInstance->ContainerInstanceCollection;
 
 	// Find all used HLOD layers
 	TMap<UHLODLayer*, int32> HLODLayersLevels = GatherHLODLayers(StreamingGenerationContext, WorldPartition);
@@ -472,11 +469,14 @@ bool UWorldPartitionRuntimeSpatialHash::SetupHLODActors(const IStreamingGenerati
 
 	// HLOD creation context
 	FHLODCreationContext Context;
-	for (UActorDescContainerInstance::TConstIterator<AWorldPartitionHLOD> HLODIterator(MainContainerCollection->GetMainContainer()); HLODIterator; ++HLODIterator)
+	BaseContainerInstanceCollection->ForEachActorDescContainerInstance([&Context, WorldPartition](const UActorDescContainerInstance* ActorDescContainerInstance)
 	{
-		FWorldPartitionHandle HLODActorHandle(WorldPartition, HLODIterator->GetGuid());
-		Context.HLODActorDescs.Emplace(HLODIterator->GetActorName(), MoveTemp(HLODActorHandle));
-	}
+		for (UActorDescContainerInstance::TConstIterator<AWorldPartitionHLOD> HLODIterator(ActorDescContainerInstance); HLODIterator; ++HLODIterator)
+		{
+			FWorldPartitionHandle HLODActorHandle(WorldPartition, HLODIterator->GetGuid());
+			Context.HLODActorDescs.Emplace(HLODIterator->GetActorName(), MoveTemp(HLODActorHandle));
+		}
+	});
 
 	TArray<TArray<const IStreamingGenerationContext::FActorSetInstance*>> GridActorSetInstances;
 	GridActorSetInstances.InsertDefaulted(0, Grids.Num());
@@ -503,7 +503,8 @@ bool UWorldPartitionRuntimeSpatialHash::SetupHLODActors(const IStreamingGenerati
 	FSpatialHashSettings HashSettings = Settings;	
 	HashSettings.UpdateSettings(*this);
 
-	auto GenerateHLODActors = [&](const FSpatialHashRuntimeGrid& RuntimeGrid, uint32 HLODLevel, const TArray<const IStreamingGenerationContext::FActorSetInstance*>& ActorSetInstances)
+	auto GenerateHLODActors = [&GridsHLODActors, &BaseActorSetContainerInstance, StreamingGenerationContext, &Params, WorldPartition, DataLayerManager, &Context, SourceControlHelper, &NewActors, &DirtyActors, &HashSettings, this]
+	(const FSpatialHashRuntimeGrid& RuntimeGrid, uint32 HLODLevel, const TArray<const IStreamingGenerationContext::FActorSetInstance*>& ActorSetInstances)
 	{
 		// Generate HLODs for this grid
 		TArray<FGuid> HLODActors = GenerateHLODActorsForGrid(WorldPartition, StreamingGenerationContext, Params, RuntimeGrid, HLODLevel, Context, SourceControlHelper, ActorSetInstances, NewActors, DirtyActors, HashSettings);
@@ -514,11 +515,11 @@ bool UWorldPartitionRuntimeSpatialHash::SetupHLODActors(const IStreamingGenerati
 			{
 				FWorldPartitionActorDescInstance* HLODActorDescInstance = WorldPartition->GetActorDescInstance(HLODActorGuid);
 				check(HLODActorDescInstance);
-				FStreamingGenerationActorDescViewMap* NonConstActorDescViewMap = const_cast<FStreamingGenerationActorDescViewMap*>(MainActorSetContainer->ActorDescViewMap);
+				FStreamingGenerationActorDescViewMap* NonConstActorDescViewMap = const_cast<FStreamingGenerationActorDescViewMap*>(BaseActorSetContainerInstance->ActorDescViewMap);
 				FStreamingGenerationActorDescView* ActorDescView = NonConstActorDescViewMap->Emplace(HLODActorDescInstance);
 				
-				TArray<FName> RuntimeDataLayerInstanceNames;
-				if (FDataLayerUtils::ResolveRuntimeDataLayerInstanceNames(DataLayerManager, *ActorDescView, *MainActorSetContainer->ActorDescViewMap, RuntimeDataLayerInstanceNames))
+				FDataLayerInstanceNames RuntimeDataLayerInstanceNames;
+				if (FDataLayerUtils::ResolveRuntimeDataLayerInstanceNames(DataLayerManager, *ActorDescView, *BaseActorSetContainerInstance->ActorDescViewMap, RuntimeDataLayerInstanceNames))
 				{
 					ActorDescView->SetRuntimeDataLayerInstanceNames(RuntimeDataLayerInstanceNames);
 				}
@@ -552,11 +553,11 @@ bool UWorldPartitionRuntimeSpatialHash::SetupHLODActors(const IStreamingGenerati
 		TArray<const IStreamingGenerationContext::FActorSetInstance*> HLODActorSetInstancePtrs;
 		HLODActorSetInstances.Reserve(GridsHLODActors[HLODGridName].Num());
 		HLODActorSetInstancePtrs.Reserve(GridsHLODActors[HLODGridName].Num());
-		MainActorSetContainer->ActorSets.Reserve(MainActorSetContainer->ActorSets.Num() + GridsHLODActors[HLODGridName].Num());
+		BaseActorSetContainerInstance->ActorSets.Reserve(BaseActorSetContainerInstance->ActorSets.Num() + GridsHLODActors[HLODGridName].Num());
 
 		for (const FGuid& HLODActorGuid : GridsHLODActors[HLODGridName])
 		{
-			const FStreamingGenerationActorDescView& HLODActorDescView = MainActorSetContainer->ActorDescViewMap->FindByGuidChecked(HLODActorGuid);
+			const FStreamingGenerationActorDescView& HLODActorDescView = BaseActorSetContainerInstance->ActorDescViewMap->FindByGuidChecked(HLODActorGuid);
 			
 			IStreamingGenerationContext::FActorSetInstance& NewHLODActorSetInstance = HLODActorSetInstances.Emplace_GetRef();
 			
@@ -565,11 +566,11 @@ bool UWorldPartitionRuntimeSpatialHash::SetupHLODActors(const IStreamingGenerati
 			NewHLODActorSetInstance.bIsSpatiallyLoaded = HLODActorDescView.GetIsSpatiallyLoaded();
 			if (DataLayerManager)
 			{
-				NewHLODActorSetInstance.DataLayers = DataLayerManager->GetRuntimeDataLayerInstances(HLODActorDescView.GetRuntimeDataLayerInstanceNames());
+				NewHLODActorSetInstance.DataLayers = DataLayerManager->GetRuntimeDataLayerInstances(HLODActorDescView.GetRuntimeDataLayerInstanceNames().ToArray());
 			}
-			NewHLODActorSetInstance.ContentBundleID = MainContainerCollection->GetContentBundleGuid();
-			NewHLODActorSetInstance.ActorSetContainerInstance = MainActorSetContainer;
-			NewHLODActorSetInstance.ActorSet = MainActorSetContainer->ActorSets.Add_GetRef(MakeUnique<IStreamingGenerationContext::FActorSet>()).Get();
+			NewHLODActorSetInstance.ContentBundleID = BaseContainerInstanceCollection->GetContentBundleGuid();
+			NewHLODActorSetInstance.ActorSetContainerInstance = BaseActorSetContainerInstance;
+			NewHLODActorSetInstance.ActorSet = BaseActorSetContainerInstance->ActorSets.Add_GetRef(MakeUnique<IStreamingGenerationContext::FActorSet>()).Get();
 			const_cast<IStreamingGenerationContext::FActorSet*>(NewHLODActorSetInstance.ActorSet)->Actors.Add(HLODActorDescView.GetGuid());
 
 			HLODActorSetInstancePtrs.Add(&NewHLODActorSetInstance);
@@ -601,8 +602,11 @@ bool UWorldPartitionRuntimeSpatialHash::SetupHLODActors(const IStreamingGenerati
 
 	if (!Params.bReportOnly)
 	{
+		BaseContainerInstanceCollection->ForEachActorDescContainerInstance([this, &HLODGrids, &SourceControlHelper](const UActorDescContainerInstance* ActorDescContainerInstance)
+		{
 		// Create/destroy HLOD grid actors
-		UpdateHLODGridsActors(GetWorld(), MainContainerCollection->GetMainContainer(), HLODGrids, SourceControlHelper);
+			UpdateHLODGridsActors(GetWorld(), ActorDescContainerInstance, HLODGrids, SourceControlHelper);
+		});
 	}
 
 	auto DumpActorsStats = [](const TCHAR* StatType, const TArray<FName>& Actors)
diff --git a/Engine/Source/Runtime/Engine/Private/WorldPartition/WorldPartition.cpp b/Engine/Source/Runtime/Engine/Private/WorldPartition/WorldPartition.cpp
--- a/Engine/Source/Runtime/Engine/Private/WorldPartition/WorldPartition.cpp
+++ b/Engine/Source/Runtime/Engine/Private/WorldPartition/WorldPartition.cpp
@@ -15,6 +15,7 @@
 #include "WorldPartition/WorldPartitionSubsystem.h"
 #include "WorldPartition/HLOD/HLODRuntimeSubsystem.h"
 #include "WorldPartition/DataLayer/DataLayerManager.h"
+#include "WorldPartition/DataLayer/ExternalDataLayerManager.h"
 #include "WorldPartition/WorldPartitionSettings.h"
 #include "GameFramework/WorldSettings.h"
 #include "ProfilingDebugging/ScopedTimers.h"
@@ -166,7 +167,7 @@ FString GetActorDescDumpString(const FWorldPartitionActorDescInstance* ActorDesc
 	return FString::Printf(
 		TEXT("%s DataLayerNames:%s") LINE_TERMINATOR, 
 		*ActorDescInstance->ToString(FWorldPartitionActorDesc::EToStringMode::Full),
-		*GetDataLayerString(ActorDescInstance->GetDataLayerInstanceNames())
+		*GetDataLayerString(ActorDescInstance->GetDataLayerInstanceNames().ToArray())
 	);
 }
 
@@ -393,7 +394,7 @@ bool UWorldPartition::IsValidPackageName(const FString& InPackageName)
 		// Remove PIE prefix
 		FString PackageName = UWorld::RemovePIEPrefix(InPackageName);
 		// Test if package is a valid world partition PIE package
-		return GeneratedStreamingPackageNames.Contains(PackageName);
+		return GeneratedLevelStreamingPackageNames.Contains(PackageName);
 	}
 	return false;
 }
@@ -420,22 +421,24 @@ void UWorldPartition::OnBeginPlay()
 	FGenerateStreamingParams Params = FGenerateStreamingParams()
 		.SetErrorHandler(bIsPIE ? (IStreamingGenerationErrorHandler*)&MapCheckErrorHandler : (IStreamingGenerationErrorHandler*)&LogErrorHandler);
 
-	TArray<FString> OutGeneratedStreamingPackageNames;
+	TArray<FString> OutGeneratedLevelStreamingPackageNames;
 	FGenerateStreamingContext Context = FGenerateStreamingContext()
-		.SetPackagesToGenerate((bIsPIE || IsRunningGame()) ? &OutGeneratedStreamingPackageNames : nullptr);
+		.SetLevelPackagesToGenerate((bIsPIE || IsRunningGame()) ? &OutGeneratedLevelStreamingPackageNames : nullptr);
 
 	GenerateStreaming(Params, Context);
 
 	// Prepare GeneratedStreamingPackages
-	check(GeneratedStreamingPackageNames.IsEmpty());
-	for (const FString& PackageName : OutGeneratedStreamingPackageNames)
+	check(GeneratedLevelStreamingPackageNames.IsEmpty());
+	for (const FString& PackageName : OutGeneratedLevelStreamingPackageNames)
 	{
 		// Set as memory package to avoid wasting time in UWorldPartition::IsValidPackageName (GenerateStreaming for PIE runs on the editor world)
 		FString Package = FPaths::RemoveDuplicateSlashes(FPackageName::IsMemoryPackage(PackageName) ? PackageName : TEXT("/Memory/") + PackageName);
-		GeneratedStreamingPackageNames.Add(Package);
+		GeneratedLevelStreamingPackageNames.Add(Package);
 	}
 
 	RuntimeHash->OnBeginPlay();
+
+	ExternalDataLayerManager->OnBeginPlay();
 }
 
 void UWorldPartition::OnCancelPIE()
@@ -450,6 +453,7 @@ void UWorldPartition::OnEndPlay()
 	if (bIsPIE)
 	{
 		FlushStreaming();
+		ExternalDataLayerManager->OnEndPlay();
 		RuntimeHash->OnEndPlay();
 		bIsPIE = false;
 	}
@@ -567,6 +571,13 @@ void UWorldPartition::Initialize(UWorld* InWorld, const FTransform& InTransform)
 			UseMakingInvisibleTransactionRequests() ? 1 : 0);
 	}
 
+	auto CreateAndInitializeDataLayerManager = [this]()
+	{
+		check(!DataLayerManager);
+		DataLayerManager = NewObject<UDataLayerManager>(this, TEXT("DataLayerManager"), RF_Transient);
+		DataLayerManager->Initialize();
+	};
+
 #if WITH_EDITOR
 	if (bEnableStreaming)
 	{
@@ -605,20 +616,30 @@ void UWorldPartition::Initialize(UWorld* InWorld, const FTransform& InTransform)
 	if (bIsEditor || bIsGame || bIsPIEWorldTravel || bIsDedicatedServer)
 	{
 		FName ContainerPackageName = UActorDescContainerInstance::GetContainerPackageNameFromWorld(OuterWorld);
-
 		ActorDescContainerInstance = RegisterActorDescContainerInstance(UActorDescContainerInstance::FInitializeParams(ContainerPackageName));
-
-		ForEachActorDescContainerInstance([this, bIsEditor, bIsCooking](UActorDescContainerInstance* InActorDescContainerInstance)
-		{
-			InitializeActorDescContainerEditorStreaming(InActorDescContainerInstance, bIsEditor && !bIsCooking);
-		});
+		CreateAndInitializeDataLayerManager();
+		InitializeActorDescContainerEditorStreaming(ActorDescContainerInstance);
 	}
 #endif
 
-	// Here's it's safe to initialize the DataLayerManager
-	DataLayerManager = NewObject<UDataLayerManager>(this, TEXT("DataLayerManager"), RF_Transient);
-	DataLayerManager->Initialize();
+#if !WITH_EDITOR
+	check(!DataLayerManager);
+	check(!ExternalDataLayerManager);
+#endif
+
+	// Create and initialize the DataLayerManager (When WorldPartition's ActorDescContainerInstance is created, we create/initialize the DataLayerManager before calling InitializeActorDescContainerEditorStreaming)
+	if (!DataLayerManager)
+	{
+		CreateAndInitializeDataLayerManager();
+	}
 
+	// Create and initialize the ExternalDataLayerManager (In PIE, we use the exiting/duplicated ExternalDataLayerManager containing the duplicated ExternalStreamingObjects)
+	if (!ExternalDataLayerManager)
+	{
+		ExternalDataLayerManager = NewObject<UExternalDataLayerManager>(this, TEXT("ExternalDataLayerManager"), RF_Transient | RF_Transactional);
+	}
+	ExternalDataLayerManager->Initialize();
+	
 #if WITH_EDITOR
 	if (bIsEditor)
 	{
@@ -749,12 +770,13 @@ void UWorldPartition::Uninitialize()
 			RegisteredEditorLoaderAdapters.Empty();
 		}
 
-		UninitializeActorDescContainers();
-		ActorDescContainerInstance = nullptr;
+#endif
 
-		EditorHash = nullptr;
-		bIsPIE = false;
-#endif		
+		if (ExternalDataLayerManager)
+		{
+			ExternalDataLayerManager->DeInitialize();
+			ExternalDataLayerManager = nullptr;
+		}
 
 		if (DataLayerManager)
 		{
@@ -762,6 +784,14 @@ void UWorldPartition::Uninitialize()
 			DataLayerManager = nullptr;
 		}
 
+#if WITH_EDITOR
+		UninitializeActorDescContainers();
+		ActorDescContainerInstance = nullptr;
+
+		EditorHash = nullptr;
+		bIsPIE = false;
+#endif		
+
 		InitState = EWorldPartitionInitState::Uninitialized;
 
 		FWorldPartitionEvents::BroadcastWorldPartitionUninitialized(World, this);
@@ -787,6 +817,11 @@ UDataLayerManager* UWorldPartition::GetResolvingDataLayerManager() const
 	return GetDataLayerManager();
 }
 
+UExternalDataLayerManager* UWorldPartition::GetExternalDataLayerManager() const
+{
+	return ExternalDataLayerManager;
+}
+
 bool UWorldPartition::IsInitialized() const
 {
 	return InitState == EWorldPartitionInitState::Initialized;
@@ -1271,10 +1306,19 @@ void UWorldPartition::OnActorDescInstanceUpdated(FWorldPartitionActorDescInstanc
 	}
 }
 
-void UWorldPartition::InitializeActorDescContainerEditorStreaming(UActorDescContainerInstance* InActorDescContainerInstance, bool bInHashActorDescs)
+bool UWorldPartition::ShouldHashUnhashActorDescInstances() const
+{
+	const bool bIsEditor = !GetWorld()->IsGameWorld();
+	const bool bIsCooking = IsRunningCookCommandlet();
+	const bool bHashActorDescs = bIsEditor && !bIsCooking;
+	return bHashActorDescs;
+}
+
+void UWorldPartition::InitializeActorDescContainerEditorStreaming(UActorDescContainerInstance* InActorDescContainerInstance)
 {
 	TRACE_CPUPROFILER_EVENT_SCOPE(InitializeActorDescContainerEditorStreaming);
 
+	const bool bHashActorDescs = ShouldHashUnhashActorDescInstances();
 	const bool bIsStreamingEnabled = IsStreamingEnabledInEditor();
 
 	TArray<FGuid> ForceLoadedActorGuids;
@@ -1287,7 +1331,7 @@ void UWorldPartition::InitializeActorDescContainerEditorStreaming(UActorDescCont
 			ForceLoadedActorGuids.Add(It->GetGuid());
 		}
 
-		if (bInHashActorDescs)
+		if (bHashActorDescs)
 		{
 			HashActorDescInstance(*It);
 		}
@@ -1406,19 +1450,16 @@ void UWorldPartition::Serialize(FArchive& Ar)
 
 	Super::Serialize(Ar);
 
+#if WITH_EDITOR
 	if (Ar.GetPortFlags() & PPF_DuplicateForPIE)
 	{
+		Ar << ExternalDataLayerManager;
 		Ar << StreamingPolicy;
-
-#if WITH_EDITORONLY_DATA
-		Ar << GeneratedStreamingPackageNames;
-#endif
-
-#if WITH_EDITOR
+		Ar << GeneratedLevelStreamingPackageNames;
 		Ar << bIsPIE;
-#endif
 	}
 	else
+#endif
 	{
 		if (Ar.CustomVer(FUE5MainStreamObjectVersion::GUID) >= FUE5MainStreamObjectVersion::WorldPartitionSerializeStreamingPolicyOnCook)
 		{
@@ -1584,6 +1625,11 @@ void UWorldPartition::Tick(float DeltaSeconds)
 #endif
 }
 
+bool UWorldPartition::IsExternalStreamingObjectInjected(URuntimeHashExternalStreamingObjectBase* InExternalStreamingObject) const
+{
+	return RuntimeHash->IsExternalStreamingObjectInjected(InExternalStreamingObject);
+}
+
 bool UWorldPartition::InjectExternalStreamingObject(URuntimeHashExternalStreamingObjectBase* InExternalStreamingObject)
 {
 	bool bInjected = RuntimeHash->InjectExternalStreamingObject(InExternalStreamingObject);
@@ -1749,67 +1795,6 @@ void UWorldPartition::DrawRuntimeHashPreview()
 	RuntimeHash->DrawPreview();
 }
 
-void UWorldPartition::BeginCook(IWorldPartitionCookPackageContext& CookContext)
-{
-	OnBeginCook.Broadcast(CookContext);
-
-	CookContext.RegisterPackageCookPackageGenerator(this);
-}
-
-void UWorldPartition::EndCook(IWorldPartitionCookPackageContext& CookContext)
-{
-	OnEndCook.Broadcast(CookContext);
-
-	CookContext.UnregisterPackageCookPackageGenerator(this);
-}
-
-bool UWorldPartition::GatherPackagesToCook(IWorldPartitionCookPackageContext& CookContext)
-{
-	FGenerateStreamingParams Params = FGenerateStreamingParams()
-		.SetActorDescContainerInstance(ActorDescContainerInstance);
-
-	TArray<FString> PackagesToCook;
-	FGenerateStreamingContext Context = FGenerateStreamingContext()
-		.SetPackagesToGenerate(&PackagesToCook);
-
-	if (GenerateContainerStreaming(Params, Context))
-	{
-		FString PackageName = GetPackage()->GetName();
-		for (const FString& PackageToCook : PackagesToCook)
-		{
-			CookContext.AddLevelStreamingPackageToGenerate(this, PackageName, PackageToCook);
-		}
-	
-		return true;
-	}
-
-	return false;
-}
-
-bool UWorldPartition::PrepareGeneratorPackageForCook(IWorldPartitionCookPackageContext& CookContext, TArray<UPackage*>& OutModifiedPackages)
-{
-	check(RuntimeHash);
-	return RuntimeHash->PrepareGeneratorPackageForCook(OutModifiedPackages);
-}
-
-bool UWorldPartition::PopulateGeneratorPackageForCook(IWorldPartitionCookPackageContext& CookContext, const TArray<FWorldPartitionCookPackage*>& InPackagesToCook, TArray<UPackage*>& OutModifiedPackages)
-{
-	check(RuntimeHash);
-	return RuntimeHash->PopulateGeneratorPackageForCook(InPackagesToCook, OutModifiedPackages);
-}
-
-bool UWorldPartition::PopulateGeneratedPackageForCook(IWorldPartitionCookPackageContext& CookContext, const FWorldPartitionCookPackage& InPackagesToCook, TArray<UPackage*>& OutModifiedPackages)
-{
-	check(RuntimeHash);
-	return RuntimeHash->PopulateGeneratedPackageForCook(InPackagesToCook, OutModifiedPackages);
-}
-
-UWorldPartitionRuntimeCell* UWorldPartition::GetCellForPackage(const FWorldPartitionCookPackage& PackageToCook) const
-{
-	check(RuntimeHash);
-	return RuntimeHash->GetCellForPackage(PackageToCook);
-}
-
 TArray<FBox> UWorldPartition::GetUserLoadedEditorRegions() const
 {
 	TArray<FBox> Result;
@@ -1937,6 +1922,8 @@ UActorDescContainerInstance* UWorldPartition::RegisterActorDescContainerInstance
 		// Initialize ContainerInstance hierarchy if we are the main world partition or if we are a game streamed world partition which means we have our own generate streaming
 		const bool bCreateContainerInstanceHierarchy = IsMainWorldPartition() || (bIsGameWorld && bIsStreamedLevel) || InParams.bCreateContainerInstanceHierarchy;
 		UActorDescContainerInstance::FInitializeParams InitParams(InParams.ContainerPackageName, bCreateContainerInstanceHierarchy);
+		InitParams.ContentBundleGuid = InParams.ContentBundleGuid;
+		InitParams.ExternalDataLayerAsset = InParams.ExternalDataLayerAsset;
 		
 		const FWorldDataLayersActorDesc* WorldDataLayerActorsDesc = nullptr;
 		InitParams.FilterActorDescFunc = [this, &WorldDataLayerActorsDesc, &InParams](const FWorldPartitionActorDesc* ActorDesc)
@@ -1974,6 +1961,14 @@ UActorDescContainerInstance* UWorldPartition::RegisterActorDescContainerInstance
 			return true;
 		};
 
+		InitParams.OnInitializedFunc = [&InParams](UActorDescContainerInstance* InActorDescContainerInstance)
+		{
+			if (InParams.OnInitializedFunc)
+			{
+				InParams.OnInitializedFunc(InActorDescContainerInstance);
+			}
+		};
+
 		UActorDescContainerInstance* ContainerInstanceToRegister = NewObject<UActorDescContainerInstance>(this, UActorDescContainerInstance::StaticClass(), NAME_None, RF_Transient);
 		
 		OnActorDescContainerInstancePreInitialize.ExecuteIfBound(InitParams, ContainerInstanceToRegister);
@@ -1982,20 +1977,11 @@ UActorDescContainerInstance* UWorldPartition::RegisterActorDescContainerInstance
 			
 		AddContainer(ContainerInstanceToRegister);
 
-		if (IsInitialized() && EditorHash != nullptr)
+		if (ActorDescContainerInstance && EditorHash)
 		{
-			FWorldPartitionReference WDLReference;
-			for (UActorDescContainerInstance::TIterator<> Iterator(ContainerInstanceToRegister); Iterator; ++Iterator)
-			{
-				if (Iterator->GetActorNativeClass()->IsChildOf<AWorldDataLayers>())
-				{
-					WDLReference = FWorldPartitionReference(this, Iterator->GetGuid());
-					break;
-				}
-			}
-
-			const bool bHashActorDescs = bIsEditor && !bIsCooking;
-			InitializeActorDescContainerEditorStreaming(ContainerInstanceToRegister, bHashActorDescs);
+			check(ActorDescContainerInstance->IsInitialized());
+			// When world partition is already initialized, it's safe to call InitializeActorDescContainerEditorStreaming as the DataLayerManager is created
+			InitializeActorDescContainerEditorStreaming(ContainerInstanceToRegister);
 		}
 
 		OnActorDescContainerInstanceRegistered.Broadcast(ContainerInstanceToRegister);
@@ -2029,7 +2015,8 @@ bool UWorldPartition::UnregisterActorDescContainerInstance(UActorDescContainerIn
 
 		OnActorDescContainerInstanceUnregistered.Broadcast(InActorDescContainerInstance);
 
-		if (IsInitialized() && EditorHash != nullptr)
+		// Un-hashing needs to be done for an initialized container instance that was previously hashed
+		if (EditorHash && (IsInitialized() || (InActorDescContainerInstance->IsInitialized() && ShouldHashUnhashActorDescInstances())))
 		{
 			for (UActorDescContainerInstance::TIterator<> It(InActorDescContainerInstance); It; ++It)
 			{
diff --git a/Engine/Source/Runtime/Engine/Private/WorldPartition/WorldPartitionActorDesc.cpp b/Engine/Source/Runtime/Engine/Private/WorldPartition/WorldPartitionActorDesc.cpp
--- a/Engine/Source/Runtime/Engine/Private/WorldPartition/WorldPartitionActorDesc.cpp
+++ b/Engine/Source/Runtime/Engine/Private/WorldPartition/WorldPartitionActorDesc.cpp
@@ -24,6 +24,7 @@
 #include "WorldPartition/DataLayer/WorldDataLayers.h"
 #include "WorldPartition/DataLayer/DataLayerManager.h"
 #include "WorldPartition/DataLayer/DataLayerAsset.h"
+#include "WorldPartition/DataLayer/ExternalDataLayerAsset.h"
 #include "WorldPartition/DataLayer/DataLayerUtils.h"
 #include "WorldPartition/ContentBundle/ContentBundlePaths.h"
 #include "WorldPartition/ErrorHandling/WorldPartitionStreamingGenerationErrorHandler.h"
@@ -107,7 +108,8 @@ void FWorldPartitionActorDesc::Init(const AActor* InActor)
 
 		if (UWorldPartition* ActorWorldPartition = FWorldPartitionHelpers::GetWorldPartition(InActor))
 		{
-			TArray<const UDataLayerAsset*> DataLayerAssets = InActor->GetDataLayerAssets();
+			const bool bIncludeExternalDataLayerAsset = false;
+			TArray<const UDataLayerAsset*> DataLayerAssets = InActor->GetDataLayerAssets(bIncludeExternalDataLayerAsset);
 			LocalDataLayerAssetPaths.Reserve(DataLayerAssets.Num());
 			for (const UDataLayerAsset* DataLayerAsset : DataLayerAssets)
 			{
@@ -141,6 +143,9 @@ void FWorldPartitionActorDesc::Init(const AActor* InActor)
 			bIsUsingDataLayerAsset = true;
 			DataLayers.Empty();
 		}
+
+		// Initialize ExternalDataLayerAsset
+		ExternalDataLayerAsset = InActor->GetExternalDataLayerAsset() ? FSoftObjectPath(InActor->GetExternalDataLayerAsset()->GetPathName()) : FSoftObjectPath();
 	}
 
 	Tags = InActor->Tags;
@@ -290,7 +295,8 @@ bool FWorldPartitionActorDesc::Equals(const FWorldPartitionActorDesc* Other) con
 		CompareUnsortedArrays(References, Other->References) &&
 		CompareUnsortedArrays(EditorOnlyReferences, Other->EditorOnlyReferences) &&
 		CompareUnsortedArrays(Tags, Other->Tags) &&
-		Properties == Other->Properties;
+		Properties == Other->Properties &&
+		ExternalDataLayerAsset == Other->ExternalDataLayerAsset;
 }
 
 bool FWorldPartitionActorDesc::ShouldResave(const FWorldPartitionActorDesc* Other) const
@@ -311,7 +317,8 @@ bool FWorldPartitionActorDesc::ShouldResave(const FWorldPartitionActorDesc* Othe
 		!CompareUnsortedArrays(DataLayers, Other->DataLayers) ||
 		!CompareUnsortedArrays(References, Other->References) ||
 		!CompareUnsortedArrays(EditorOnlyReferences, Other->EditorOnlyReferences) ||
-		Properties != Other->Properties)
+		Properties != Other->Properties ||
+		ExternalDataLayerAsset != Other->ExternalDataLayerAsset)
 	{
 		return true;
 	}
@@ -360,6 +367,27 @@ void FWorldPartitionActorDesc::SerializeTo(TArray<uint8>& OutData) const
 	OutData.Append(PayloadData);
 }
 
+TArray<FName> FWorldPartitionActorDesc::GetDataLayers(bool bIncludeExternalDataLayer) const
+{
+	const FName ExternalDataLayer = GetExternalDataLayer();
+	if (!bIncludeExternalDataLayer || ExternalDataLayer.IsNone())
+	{
+		return DataLayers;
+	}
+
+	TArray<FName> AllDataLayers;
+	AllDataLayers.Reserve(DataLayers.Num() + 1);
+	AllDataLayers.Add(ExternalDataLayer);
+	AllDataLayers.Append(DataLayers);
+	return AllDataLayers;
+}
+
+FName FWorldPartitionActorDesc::GetExternalDataLayer() const
+{
+	const bool bHasExternalDataLayerAsset = bIsUsingDataLayerAsset && ExternalDataLayerAsset.IsValid();
+	return bHasExternalDataLayerAsset ? FName(ExternalDataLayerAsset.GetAssetPath().ToString()) : NAME_None;
+}
+
 void FWorldPartitionActorDesc::TransferFrom(const FWorldPartitionActorDesc* From)
 {
 	Container = From->Container;
@@ -465,6 +493,11 @@ FString FWorldPartitionActorDesc::ToString(EToStringMode Mode) const
 			{
 				Result.Appendf(TEXT(" DataLayers:%s"), *FString::JoinBy(DataLayers, TEXT(","), [&](const FName& DataLayerName) { return DataLayerName.ToString(); }));
 			}
+
+			if (ExternalDataLayerAsset.IsValid())
+			{
+				Result.Appendf(TEXT(" ExternalDataLayerAsset:%s"), *ExternalDataLayerAsset.ToString());
+			}
 		}
 	}
 
@@ -680,6 +713,11 @@ PRAGMA_ENABLE_DEPRECATION_WARNINGS
 		Ar << TDeltaSerialize<bool>(bActorIsListedInSceneOutliner);
 	}
 
+    if (Ar.CustomVer(FUE5ReleaseStreamObjectVersion::GUID) >= FUE5ReleaseStreamObjectVersion::WorldPartitionExternalDataLayers)
+	{
+		Ar << TDeltaSerialize<FSoftObjectPath>(ExternalDataLayerAsset);
+	}
+
 	// Fixup redirected data layer asset paths
 	if (Ar.IsLoading() && bIsUsingDataLayerAsset)
 	{
diff --git a/Engine/Source/Runtime/Engine/Private/WorldPartition/WorldPartitionActorDescInstance.cpp b/Engine/Source/Runtime/Engine/Private/WorldPartition/WorldPartitionActorDescInstance.cpp
--- a/Engine/Source/Runtime/Engine/Private/WorldPartition/WorldPartitionActorDescInstance.cpp
+++ b/Engine/Source/Runtime/Engine/Private/WorldPartition/WorldPartitionActorDescInstance.cpp
@@ -169,9 +169,9 @@ void FWorldPartitionActorDescInstance::Invalidate()
 	ContainerInstance = nullptr;
 }
 
-const TArray<FName>& FWorldPartitionActorDescInstance::GetDataLayerInstanceNames() const
+const FDataLayerInstanceNames& FWorldPartitionActorDescInstance::GetDataLayerInstanceNames() const
 {
-	static TArray<FName> EmptyDataLayers;
+	static FDataLayerInstanceNames EmptyDataLayers;
 	if (ensure(HasResolvedDataLayerInstanceNames()))
 	{
 		return ResolvedDataLayerInstanceNames.GetValue();
diff --git a/Engine/Source/Runtime/Engine/Private/WorldPartition/WorldPartitionLevelHelper.cpp b/Engine/Source/Runtime/Engine/Private/WorldPartition/WorldPartitionLevelHelper.cpp
--- a/Engine/Source/Runtime/Engine/Private/WorldPartition/WorldPartitionLevelHelper.cpp
+++ b/Engine/Source/Runtime/Engine/Private/WorldPartition/WorldPartitionLevelHelper.cpp
@@ -332,10 +332,8 @@ bool FWorldPartitionLevelHelper::RemapLevelCellPathInContentBundle(ULevel* Level
 	CellPath += ContentBundleEditor->GetExternalStreamingObjectName();
 	CellPath += TEXT(".");
 	CellPath += Cell->GetName();
-	
-	Level->WorldPartitionRuntimeCell = FSoftObjectPath(CellPath);
-
-	return !Level->WorldPartitionRuntimeCell.GetUniqueID().IsNull();
+	FSetWorldPartitionRuntimeCell SetWorldPartitionRuntimeCell(Level, FSoftObjectPath(CellPath));
+	return Level->IsWorldPartitionRuntimeCell();
 }
 
 /**
diff --git a/Engine/Source/Runtime/Engine/Private/WorldPartition/WorldPartitionLevelStreamingDynamic.cpp b/Engine/Source/Runtime/Engine/Private/WorldPartition/WorldPartitionLevelStreamingDynamic.cpp
--- a/Engine/Source/Runtime/Engine/Private/WorldPartition/WorldPartitionLevelStreamingDynamic.cpp
+++ b/Engine/Source/Runtime/Engine/Private/WorldPartition/WorldPartitionLevelStreamingDynamic.cpp
@@ -15,6 +15,7 @@
 #include "Engine/Level.h"
 #include "Misc/PathViews.h"
 #include "ContentStreaming.h"
+#include "WorldPartition/DataLayer/ExternalDataLayerHelper.h"
 #include "WorldPartition/ContentBundle/ContentBundlePaths.h"
 #endif
 
@@ -366,7 +367,8 @@ bool UWorldPartitionLevelStreamingDynamic::IssueLoadRequests()
 
 		for (const FWorldPartitionRuntimeCellObjectMapping& CellObjectMapping : ChildPackages)
 		{
-			if (ContentBundlePaths::IsAContentBundlePath(CellObjectMapping.ContainerPackage.ToString()))
+			if (ContentBundlePaths::IsAContentBundlePath(CellObjectMapping.ContainerPackage.ToString()) ||
+				FExternalDataLayerHelper::IsExternalDataLayerPath(CellObjectMapping.ContainerPackage.ToString()))
 			{
 				check(CellObjectMapping.ContainerPackage != CellObjectMapping.WorldPackage);
 				bool bIsContainerPackageAlreadyRemapped = OutLinkInstancingContext.RemapPackage(CellObjectMapping.ContainerPackage) != CellObjectMapping.ContainerPackage;
diff --git a/Engine/Source/Runtime/Engine/Private/WorldPartition/WorldPartitionLevelStreamingPolicy.cpp b/Engine/Source/Runtime/Engine/Private/WorldPartition/WorldPartitionLevelStreamingPolicy.cpp
--- a/Engine/Source/Runtime/Engine/Private/WorldPartition/WorldPartitionLevelStreamingPolicy.cpp
+++ b/Engine/Source/Runtime/Engine/Private/WorldPartition/WorldPartitionLevelStreamingPolicy.cpp
@@ -120,9 +120,9 @@ void UWorldPartitionLevelStreamingPolicy::RemapSoftObjectPath(FSoftObjectPath& O
 	ConvertEditorPathToRuntimePath(SrcPath, ObjectPath);
 }
 
-bool UWorldPartitionLevelStreamingPolicy::StoreToExternalStreamingObject(URuntimeHashExternalStreamingObjectBase& OutExternalStreamingObject)
+bool UWorldPartitionLevelStreamingPolicy::StoreStreamingContentToExternalStreamingObject(URuntimeHashExternalStreamingObjectBase& OutExternalStreamingObject)
 {
-	if (Super::StoreToExternalStreamingObject(OutExternalStreamingObject))
+	if (Super::StoreStreamingContentToExternalStreamingObject(OutExternalStreamingObject))
 	{
 		OutExternalStreamingObject.SubObjectsToCellRemapping = MoveTemp(SubObjectsToCellRemapping);
 		OutExternalStreamingObject.ContainerResolver = MoveTemp(ContainerResolver);
diff --git a/Engine/Source/Runtime/Engine/Private/WorldPartition/WorldPartitionRuntimeCell.cpp b/Engine/Source/Runtime/Engine/Private/WorldPartition/WorldPartitionRuntimeCell.cpp
--- a/Engine/Source/Runtime/Engine/Private/WorldPartition/WorldPartitionRuntimeCell.cpp
+++ b/Engine/Source/Runtime/Engine/Private/WorldPartition/WorldPartitionRuntimeCell.cpp
@@ -3,9 +3,12 @@
 #include "WorldPartition/WorldPartitionRuntimeCell.h"
 #include "UObject/Package.h"
 #include "Engine/Level.h"
+#include "Algo/Transform.h"
+#include "Algo/Count.h"
 #include "Misc/HierarchicalLogArchive.h"
 #include "WorldPartition/WorldPartitionDebugHelper.h"
 #include "WorldPartition/DataLayer/DataLayerManager.h"
+#include "WorldPartition/DataLayer/ExternalDataLayerInstance.h"
 #include "WorldPartition/DataLayer/DataLayersID.h"
 #include "WorldPartition/ContentBundle/ContentBundleDescriptor.h"
 #include "WorldPartition/DataLayer/WorldDataLayers.h"
@@ -80,13 +83,43 @@ bool UWorldPartitionRuntimeCell::NeedsActorToCellRemapping() const
 void UWorldPartitionRuntimeCell::SetDataLayers(const TArray<const UDataLayerInstance*>& InDataLayerInstances)
 {
 	check(DataLayers.IsEmpty());
-	DataLayers.Reserve(InDataLayerInstances.Num());
-	for (const UDataLayerInstance* DataLayerInstance : InDataLayerInstances)
+	check(!ExternalDataLayerAsset);
+
+	if (InDataLayerInstances.IsEmpty())
+	{
+		return;
+	}
+
+	// Validate that we have maximum 1 External Data Layer
+	check(Algo::CountIf(InDataLayerInstances, [](const UDataLayerInstance* DataLayerInstance) { return DataLayerInstance->IsA<UExternalDataLayerInstance>(); }) <= 1);
+	// Validate that all Data Layers are Runtime
+	check(Algo::CountIf(InDataLayerInstances, [](const UDataLayerInstance* DataLayerInstance) { return DataLayerInstance->IsRuntime(); }) == InDataLayerInstances.Num());
+
+	// Sort Data Layers by FName except for External Data Layer that will always be the first in the list
+	TArray<const UDataLayerInstance*> SortedDataLayerInstances(InDataLayerInstances);
+	Algo::Sort(SortedDataLayerInstances, [](const UDataLayerInstance* A, const UDataLayerInstance* B)
+	{
+		if (A->IsA<UExternalDataLayerInstance>() && !B->IsA<UExternalDataLayerInstance>())
+		{
+			return true;
+		}
+		else if (!A->IsA<UExternalDataLayerInstance>() && B->IsA<UExternalDataLayerInstance>())
+		{
+			return false;
+		}
+
+		return A->GetDataLayerFName().ToString() < B->GetDataLayerFName().ToString();
+	});
+
+	TArray<FName> SortedDataLayerInstanceNames;
+	bool bIsFirstDataLayerIsExternal = false;
+	Algo::Transform(SortedDataLayerInstances, SortedDataLayerInstanceNames, [](const UDataLayerInstance* DataLayerInstance) { return DataLayerInstance->GetDataLayerFName(); });
+	if (const UExternalDataLayerInstance* ExternalDataLayerInstance = Cast<UExternalDataLayerInstance>(SortedDataLayerInstances[0]))
 	{
-		check(DataLayerInstance->IsRuntime());
-		DataLayers.Add(DataLayerInstance->GetDataLayerFName());
+		bIsFirstDataLayerIsExternal = true;
+		ExternalDataLayerAsset = ExternalDataLayerInstance->GetExternalDataLayerAsset();
 	}
-	DataLayers.Sort([](const FName& A, const FName& B) { return A.ToString() < B.ToString(); });
+	DataLayers = FDataLayerInstanceNames(SortedDataLayerInstanceNames, bIsFirstDataLayerIsExternal);
 }
 
 void UWorldPartitionRuntimeCell::DumpStateLog(FHierarchicalLogArchive& Ar) const
@@ -103,7 +136,7 @@ int32 UWorldPartitionRuntimeCell::SortCompare(const UWorldPartitionRuntimeCell*
 bool UWorldPartitionRuntimeCell::IsDebugShown() const
 {
 	return FWorldPartitionDebugHelper::IsDebugStreamingStatusShown(GetStreamingStatus()) &&
-	       FWorldPartitionDebugHelper::AreDebugDataLayersShown(DataLayers) &&
+	       FWorldPartitionDebugHelper::AreDebugDataLayersShown(GetDataLayers()) &&
 		   (FWorldPartitionDebugHelper::CanDrawContentBundles() || !ContentBundleID.IsValid()) &&
 			RuntimeCellData->IsDebugShown();
 }
@@ -131,30 +164,41 @@ EDataLayerRuntimeState UWorldPartitionRuntimeCell::GetCellEffectiveWantedState()
 			UWorldPartition* WorldPartition = OuterWorld->GetWorldPartition();
 			if (const UDataLayerManager* DataLayerManager = WorldPartition->GetDataLayerManager())
 			{
-				switch (WorldPartition->GetDataLayersLogicOperator())
+				if (!DataLayers.HasExternalDataLayer() || DataLayerManager->IsAllDataLayerInEffectiveRuntimeState({ GetExternalDataLayer() }, EDataLayerRuntimeState::Activated))
 				{
-				case EWorldPartitionDataLayersLogicOperator::Or:
-					if (DataLayerManager->IsAnyDataLayerInEffectiveRuntimeState(GetDataLayers(), EDataLayerRuntimeState::Activated))
-					{
-						EffectiveWantedState = EDataLayerRuntimeState::Activated;
-					}
-					else if (DataLayerManager->IsAnyDataLayerInEffectiveRuntimeState(GetDataLayers(), EDataLayerRuntimeState::Loaded))
-					{
-						EffectiveWantedState = EDataLayerRuntimeState::Loaded;
-					}
-					break;
-				case EWorldPartitionDataLayersLogicOperator::And:
-					if (DataLayerManager->IsAllDataLayerInEffectiveRuntimeState(GetDataLayers(), EDataLayerRuntimeState::Activated))
+					TArrayView<const FName> NonExternalDataLayers = DataLayers.GetNonExternalDataLayers();
+					if (NonExternalDataLayers.IsEmpty())
 					{
 						EffectiveWantedState = EDataLayerRuntimeState::Activated;
 					}
-					else if (DataLayerManager->IsAllDataLayerInEffectiveRuntimeState(GetDataLayers(), EDataLayerRuntimeState::Loaded))
+					else
 					{
-						EffectiveWantedState = EDataLayerRuntimeState::Loaded;
+						switch (WorldPartition->GetDataLayersLogicOperator())
+						{
+						case EWorldPartitionDataLayersLogicOperator::Or:
+							if (DataLayerManager->IsAnyDataLayerInEffectiveRuntimeState(NonExternalDataLayers, EDataLayerRuntimeState::Activated))
+							{
+								EffectiveWantedState = EDataLayerRuntimeState::Activated;
+							}
+							else if (DataLayerManager->IsAnyDataLayerInEffectiveRuntimeState(NonExternalDataLayers, EDataLayerRuntimeState::Loaded))
+							{
+								EffectiveWantedState = EDataLayerRuntimeState::Loaded;
+							}
+							break;
+						case EWorldPartitionDataLayersLogicOperator::And:
+							if (DataLayerManager->IsAllDataLayerInEffectiveRuntimeState(NonExternalDataLayers, EDataLayerRuntimeState::Activated))
+							{
+								EffectiveWantedState = EDataLayerRuntimeState::Activated;
+							}
+							else if (DataLayerManager->IsAllDataLayerInEffectiveRuntimeState(NonExternalDataLayers, EDataLayerRuntimeState::Loaded))
+							{
+								EffectiveWantedState = EDataLayerRuntimeState::Loaded;
+							}
+							break;
+						default:
+							checkNoEntry();
+						}
 					}
-					break;
-				default:
-					checkNoEntry();
 				}
 			}
 
@@ -171,6 +215,12 @@ TArray<const UDataLayerInstance*> UWorldPartitionRuntimeCell::GetDataLayerInstan
 	return DataLayerManager ? DataLayerManager->GetDataLayerInstances(GetDataLayers()) : TArray<const UDataLayerInstance*>();
 }
 
+const UExternalDataLayerInstance* UWorldPartitionRuntimeCell::GetExternalDataLayerInstance() const
+{
+	const UDataLayerManager* DataLayerManager = !GetExternalDataLayer().IsNone() ? GetDataLayerManager() : nullptr;
+	return DataLayerManager ? Cast<UExternalDataLayerInstance>(DataLayerManager->GetDataLayerInstance(GetExternalDataLayer())) : nullptr;
+}
+
 bool UWorldPartitionRuntimeCell::ContainsDataLayer(const UDataLayerAsset* DataLayerAsset) const
 {
 	const UDataLayerManager* DataLayerManager = HasDataLayers() ? GetDataLayerManager() : nullptr;
diff --git a/Engine/Source/Runtime/Engine/Private/WorldPartition/WorldPartitionRuntimeHash.cpp b/Engine/Source/Runtime/Engine/Private/WorldPartition/WorldPartitionRuntimeHash.cpp
--- a/Engine/Source/Runtime/Engine/Private/WorldPartition/WorldPartitionRuntimeHash.cpp
+++ b/Engine/Source/Runtime/Engine/Private/WorldPartition/WorldPartitionRuntimeHash.cpp
@@ -8,6 +8,8 @@
 #include "WorldPartition/WorldPartitionStreamingSource.h"
 #include "WorldPartition/WorldPartitionRuntimeLevelStreamingCell.h"
 #include "WorldPartition/WorldPartitionLevelStreamingDynamic.h"
+#include "WorldPartition/DataLayer/ExternalDataLayerAsset.h"
+#include "WorldPartition/DataLayer/ExternalDataLayerHelper.h"
 #include "Misc/ArchiveMD5.h"
 #if WITH_EDITOR
 #include "WorldPartition/Cook/WorldPartitionCookPackage.h"
@@ -51,6 +53,11 @@ void URuntimeHashExternalStreamingObjectBase::ForEachStreamingCells(TFunctionRef
 	}
 }
 
+TSet<TObjectPtr<UDataLayerInstance>>& URuntimeHashExternalStreamingObjectBase::GetDataLayerInstances()
+{
+	return DataLayerInstances;
+}
+
 void URuntimeHashExternalStreamingObjectBase::OnStreamingObjectLoaded(UWorld* InjectedWorld)
 {
 	bool bIsACookedObject = !CellToLevelStreamingPackage.IsEmpty();
@@ -68,7 +75,29 @@ void URuntimeHashExternalStreamingObjectBase::OnStreamingObjectLoaded(UWorld* In
 }
 
 #if WITH_EDITOR
-void URuntimeHashExternalStreamingObjectBase::PopulateGeneratorPackageForCook()
+UWorldPartitionRuntimeCell* URuntimeHashExternalStreamingObjectBase::GetCellForCookPackage(const FString& InCookPackageName) const
+{
+	if (UWorldPartitionRuntimeCell* const* MatchingCell = PackagesToGenerateForCook.Find(InCookPackageName))
+	{
+		if (ensure(*MatchingCell))
+		{
+			return const_cast<UWorldPartitionRuntimeCell*>(*MatchingCell);
+		}
+	}
+	return nullptr;
+}
+
+FString URuntimeHashExternalStreamingObjectBase::GetPackageNameToCreate() const
+{
+	// GetPackageNameToCreate should not be called for ExternalStreamingObjects
+	if (ensure(ExternalDataLayerAsset))
+	{
+		return TEXT("/") + FExternalDataLayerHelper::GetExternalStreamingObjectPackageName(ExternalDataLayerAsset);
+	}
+	return FString();
+}
+
+bool URuntimeHashExternalStreamingObjectBase::OnPopulateGeneratorPackageForCook(UPackage* InPackage)
 {
 	ForEachStreamingCells([this](UWorldPartitionRuntimeCell& Cell)
 	{
@@ -80,6 +109,19 @@ void URuntimeHashExternalStreamingObjectBase::PopulateGeneratorPackageForCook()
 		// Do not save them, instead they will be created once the external streaming object is loaded at runtime. 
 		LevelStreamingDynamic->SetFlags(RF_Transient);
 	});
+	return true;
+}
+
+bool URuntimeHashExternalStreamingObjectBase::OnPopulateGeneratedPackageForCook(UPackage* InPackage, TArray<UPackage*>& OutModifiedPackages)
+{
+	return Rename(nullptr, InPackage, REN_DontCreateRedirectors);
+}
+
+void URuntimeHashExternalStreamingObjectBase::DumpStateLog(FHierarchicalLogArchive& Ar)
+{
+	Ar.Printf(TEXT("----------------------------------------------------------------------------------------------------------------"));
+	Ar.Printf(TEXT("%s%s"), *GetWorld()->GetName(), ExternalDataLayerAsset ? *FString::Printf(TEXT(" - External Data Layer - %s"), *ExternalDataLayerAsset->GetName()) : TEXT(""));
+	Ar.Printf(TEXT("----------------------------------------------------------------------------------------------------------------"));
 }
 
 TMap<TPair<const UClass*, const UClass*>, UWorldPartitionRuntimeHash::FRuntimeHashConvertFunc> UWorldPartitionRuntimeHash::WorldPartitionRuntimeHashConverters;
@@ -172,7 +214,7 @@ bool UWorldPartitionRuntimeHash::GenerateStreaming(class UWorldPartitionStreamin
 	return PackagesToGenerateForCook.IsEmpty();
 }
 
-void UWorldPartitionRuntimeHash::FlushStreaming()
+void UWorldPartitionRuntimeHash::FlushStreamingContent()
 {
 	PackagesToGenerateForCook.Empty();
 }
@@ -299,24 +341,28 @@ void UWorldPartitionRuntimeHash::PopulateRuntimeCell(UWorldPartitionRuntimeCell*
 	}
 }
 
-bool UWorldPartitionRuntimeHash::PopulateGeneratedPackageForCook(const FWorldPartitionCookPackage& InPackagesToCook, TArray<UPackage*>& OutModifiedPackages)
+UWorldPartitionRuntimeCell* UWorldPartitionRuntimeHash::GetCellForCookPackage(const FString& InCookPackageName) const
 {
-	OutModifiedPackages.Reset();
-	if (UWorldPartitionRuntimeCell** MatchingCell = PackagesToGenerateForCook.Find(InPackagesToCook.RelativePath))
+	if (UWorldPartitionRuntimeCell* const* MatchingCell = PackagesToGenerateForCook.Find(InCookPackageName))
 	{
-		UWorldPartitionRuntimeCell* Cell = *MatchingCell;
-		if (ensure(Cell))
+		if (ensure(*MatchingCell))
 		{
-			return Cell->PopulateGeneratedPackageForCook(InPackagesToCook.GetPackage(), OutModifiedPackages);
+			return const_cast<UWorldPartitionRuntimeCell*>(*MatchingCell);
 		}
 	}
-	return false;
+	return nullptr;
 }
 
-UWorldPartitionRuntimeCell* UWorldPartitionRuntimeHash::GetCellForPackage(const FWorldPartitionCookPackage& PackageToCook) const
+URuntimeHashExternalStreamingObjectBase* UWorldPartitionRuntimeHash::StoreStreamingContentToExternalStreamingObject(FName InStreamingObjectName)
 {
-	UWorldPartitionRuntimeCell** MatchingCell = const_cast<UWorldPartitionRuntimeCell**>(PackagesToGenerateForCook.Find(PackageToCook.RelativePath));
-	return MatchingCell ? *MatchingCell : nullptr;
+	URuntimeHashExternalStreamingObjectBase* NewExternalStreamingObject = CreateExternalStreamingObject(GetExternalStreamingObjectClass(), GetOuterUWorldPartition(), InStreamingObjectName, GetWorld(), GetTypedOuter<UWorld>());
+	StoreStreamingContentToExternalStreamingObject(NewExternalStreamingObject);
+	return NewExternalStreamingObject;
+}
+
+void UWorldPartitionRuntimeHash::StoreStreamingContentToExternalStreamingObject(URuntimeHashExternalStreamingObjectBase* OutExternalStreamingObject)
+{
+	OutExternalStreamingObject->PackagesToGenerateForCook = MoveTemp(PackagesToGenerateForCook);
 }
 
 TArray<UWorldPartitionRuntimeCell*> UWorldPartitionRuntimeHash::GetAlwaysLoadedCells() const
@@ -333,39 +379,6 @@ TArray<UWorldPartitionRuntimeCell*> UWorldPartitionRuntimeHash::GetAlwaysLoadedC
 	return Result;
 }
 
-bool UWorldPartitionRuntimeHash::PrepareGeneratorPackageForCook(TArray<UPackage*>& OutModifiedPackages)
-{
-	check(IsRunningCookCommandlet());
-
-	for (UWorldPartitionRuntimeCell* Cell : GetAlwaysLoadedCells())
-	{
-		check(Cell->IsAlwaysLoaded());
-		if (!Cell->PopulateGeneratorPackageForCook(OutModifiedPackages))
-		{
-			return false;
-		}
-	}
-
-	//@todo_ow: here we can safely remove always loaded cells as they are not part of the OutPackagesToGenerate
-	return true;
-}
-
-bool UWorldPartitionRuntimeHash::PopulateGeneratorPackageForCook(const TArray<FWorldPartitionCookPackage*>& InPackagesToCook, TArray<UPackage*>& OutModifiedPackages)
-{
-	check(IsRunningCookCommandlet());
-
-	for (const FWorldPartitionCookPackage* CookPackage : InPackagesToCook)
-	{
-		UWorldPartitionRuntimeCell** MatchingCell = PackagesToGenerateForCook.Find(CookPackage->RelativePath);
-		UWorldPartitionRuntimeCell* Cell = MatchingCell ? *MatchingCell : nullptr;
-		if (!Cell || !Cell->PrepareCellForCook(CookPackage->GetPackage()))
-		{
-			return false;
-		}
-	}
-	return true;
-}
-
 void UWorldPartitionRuntimeHash::DumpStateLog(FHierarchicalLogArchive& Ar) const
 {
 	Ar.Printf(TEXT("----------------------------------------------------------------------------------------------------------------"));
@@ -483,6 +496,11 @@ EWorldPartitionStreamingPerformance UWorldPartitionRuntimeHash::GetStreamingPerf
 	return StreamingPerformance;
 }
 
+bool UWorldPartitionRuntimeHash::IsExternalStreamingObjectInjected(URuntimeHashExternalStreamingObjectBase* InExternalStreamingObject) const
+{
+	return InjectedExternalStreamingObjects.Contains(InExternalStreamingObject);
+}
+
 bool UWorldPartitionRuntimeHash::InjectExternalStreamingObject(URuntimeHashExternalStreamingObjectBase* InExternalStreamingObject)
 {
 	if (!ensure(InExternalStreamingObject))
diff --git a/Engine/Source/Runtime/Engine/Private/WorldPartition/WorldPartitionRuntimeLevelStreamingCell.cpp b/Engine/Source/Runtime/Engine/Private/WorldPartition/WorldPartitionRuntimeLevelStreamingCell.cpp
--- a/Engine/Source/Runtime/Engine/Private/WorldPartition/WorldPartitionRuntimeLevelStreamingCell.cpp
+++ b/Engine/Source/Runtime/Engine/Private/WorldPartition/WorldPartitionRuntimeLevelStreamingCell.cpp
@@ -264,20 +264,40 @@ void UWorldPartitionRuntimeLevelStreamingCell::AddActorToCell(const FStreamingGe
 
 	for (const FGuid& EditorReferenceGuid : ActorDescView.GetEditorReferences())
 	{
-		const FWorldPartitionActorDescInstance& ReferenceActorDesc = ContainerInstance->GetActorDescInstanceChecked(EditorReferenceGuid);
-
-		Packages.Emplace(
-			ReferenceActorDesc.GetActorPackage(),
-			*ReferenceActorDesc.GetActorSoftPath().ToString(),
-			ReferenceActorDesc.GetBaseClass(),
-			ReferenceActorDesc.GetNativeClass(),
-			ContainerID,
-			ContainerTransform,
-			ContainerPackage,
-			GetWorld()->GetPackage()->GetFName(),
-			ContainerID.GetActorGuid(EditorReferenceGuid),
-			true
-		);
+		// Special case where ActorDescView has invalid references : Use InvalidReference information as the Actor Guid isn't 
+		// necessarily in the ContainerInstance.
+		FStreamingGenerationActorDescView::FInvalidReference InvalidRef;
+		if (ActorDescView.IsInvalidReference(EditorReferenceGuid, &InvalidRef))
+		{
+			Packages.Emplace(
+				InvalidRef.ActorPackage,
+				*InvalidRef.ActorSoftPath.ToString(),
+				InvalidRef.BaseClass,
+				InvalidRef.NativeClass,
+				ContainerID,
+				ContainerTransform,
+				ContainerPackage,
+				GetWorld()->GetPackage()->GetFName(),
+				ContainerID.GetActorGuid(EditorReferenceGuid),
+				true
+			);
+		}
+		else
+		{
+			const FWorldPartitionActorDescInstance& ReferenceActorDesc = ContainerInstance->GetActorDescInstanceChecked(EditorReferenceGuid);
+			Packages.Emplace(
+				ReferenceActorDesc.GetActorPackage(),
+				*ReferenceActorDesc.GetActorSoftPath().ToString(),
+				ReferenceActorDesc.GetBaseClass(),
+				ReferenceActorDesc.GetNativeClass(),
+				ContainerID,
+				ContainerTransform,
+				ContainerPackage,
+				GetWorld()->GetPackage()->GetFName(),
+				ContainerID.GetActorGuid(EditorReferenceGuid),
+				true
+			);
+		}
 	}
 
 	Packages.Emplace(
@@ -311,7 +331,12 @@ void UWorldPartitionRuntimeLevelStreamingCell::Fixup()
 	}
 }
 
-bool UWorldPartitionRuntimeLevelStreamingCell::PopulateGeneratorPackageForCook(TArray<UPackage*>& OutModifiedPackages)
+FString UWorldPartitionRuntimeLevelStreamingCell::GetPackageNameToCreate() const
+{
+	return UWorldPartitionLevelStreamingPolicy::GetCellPackagePath(GetFName(), GetOuterWorld());
+}
+
+bool UWorldPartitionRuntimeLevelStreamingCell::OnPrepareGeneratorPackageForCook(TArray<UPackage*>& OutModifiedPackages)
 {
 	check(IsAlwaysLoaded());
 
@@ -364,7 +389,12 @@ bool UWorldPartitionRuntimeLevelStreamingCell::PrepareCellForCook(UPackage* InPa
 	return true;
 }
 
-bool UWorldPartitionRuntimeLevelStreamingCell::PopulateGeneratedPackageForCook(UPackage* InPackage, TArray<UPackage*>& OutModifiedPackage)
+bool UWorldPartitionRuntimeLevelStreamingCell::OnPopulateGeneratorPackageForCook(UPackage* InPackage)
+{
+	return PrepareCellForCook(InPackage);
+}
+
+bool UWorldPartitionRuntimeLevelStreamingCell::OnPopulateGeneratedPackageForCook(UPackage* InPackage, TArray<UPackage*>& OutModifiedPackages)
 {
 	check(!IsAlwaysLoaded());
 	if (!InPackage)
@@ -399,7 +429,7 @@ bool UWorldPartitionRuntimeLevelStreamingCell::PopulateGeneratedPackageForCook(U
 		// Create a level and move these actors in it
 		ULevel* NewLevel = FWorldPartitionLevelHelper::CreateEmptyLevelForRuntimeCell(this, OuterWorld, LevelStreaming->GetWorldAsset().ToString(), InPackage);
 		check(NewLevel->GetPackage() == InPackage);
-		FWorldPartitionLevelHelper::MoveExternalActorsToLevel(Packages, NewLevel, OutModifiedPackage);
+		FWorldPartitionLevelHelper::MoveExternalActorsToLevel(Packages, NewLevel, OutModifiedPackages);
 
 		// Remap Level's SoftObjectPaths
 		FWorldPartitionLevelHelper::RemapLevelSoftObjectPaths(NewLevel, WorldPartition);
@@ -412,11 +442,6 @@ int32 UWorldPartitionRuntimeLevelStreamingCell::GetActorCount() const
 	return Packages.Num();
 }
 
-FString UWorldPartitionRuntimeLevelStreamingCell::GetPackageNameToCreate() const
-{
-	return UWorldPartitionLevelStreamingPolicy::GetCellPackagePath(GetFName(), GetOuterWorld());
-}
-
 void UWorldPartitionRuntimeLevelStreamingCell::DumpStateLog(FHierarchicalLogArchive& Ar) const
 {
 	Super::DumpStateLog(Ar);
diff --git a/Engine/Source/Runtime/Engine/Private/WorldPartition/WorldPartitionRuntimeSpatialHash.cpp b/Engine/Source/Runtime/Engine/Private/WorldPartition/WorldPartitionRuntimeSpatialHash.cpp
--- a/Engine/Source/Runtime/Engine/Private/WorldPartition/WorldPartitionRuntimeSpatialHash.cpp
+++ b/Engine/Source/Runtime/Engine/Private/WorldPartition/WorldPartitionRuntimeSpatialHash.cpp
@@ -20,6 +20,8 @@
 #include "WorldPartition/ContentBundle/ContentBundleBase.h"
 #include "WorldPartition/DataLayer/DataLayerManager.h"
 #include "WorldPartition/DataLayer/DataLayerInstance.h"
+#include "WorldPartition/DataLayer/ExternalDataLayerAsset.h"
+#include "WorldPartition/DataLayer/ExternalDataLayerInstance.h"
 #include "WorldPartition/HLOD/HLODLayer.h"
 #include "Components/LineBatchComponent.h"
 #include "Engine/Engine.h"
@@ -187,6 +189,93 @@ int64 FSpatialHashStreamingGrid::GetCellSize(int32 Level) const
 	return GetGridHelper().Levels[Level].CellSize;
 }
 
+#if WITH_EDITOR
+void FSpatialHashStreamingGrid::DumpStateLog(FHierarchicalLogArchive& Ar) const
+{
+	FHierarchicalLogArchive::FIndentScope GridIndentScope = Ar.PrintfIndent(TEXT("%s"), *GridName.ToString());
+	Ar.Printf(TEXT("            Origin: %s"), *Origin.ToString());
+	Ar.Printf(TEXT("         Cell Size: %d"), CellSize);
+	Ar.Printf(TEXT("      World Bounds: %s"), *WorldBounds.ToString());
+	Ar.Printf(TEXT("     Loading Range: %3.2f"), LoadingRange);
+	Ar.Printf(TEXT("Block Slow Loading: %s"), bBlockOnSlowStreaming ? TEXT("Yes") : TEXT("No"));
+	Ar.Printf(TEXT(" ClientOnlyVisible: %s"), bClientOnlyVisible ? TEXT("Yes") : TEXT("No"));
+	Ar.Printf(TEXT(""));
+	if (HLODLayer)
+	{
+		Ar.Printf(TEXT("    HLOD Layer: %s"), *HLODLayer->GetName());
+	}
+
+	struct FGridLevelStats
+	{
+		int32 CellCount;
+		int64 CellSize;
+		int32 ActorCount;
+	};
+
+	TArray<FGridLevelStats> LevelsStats;
+	int32 TotalActorCount = 0;
+
+	{
+		int32 Level = 0;
+		for (const FSpatialHashStreamingGridLevel& GridLevel : GridLevels)
+		{
+			int32 LevelCellCount = 0;
+			int32 LevelActorCount = 0;
+			for (const FSpatialHashStreamingGridLayerCell& LayerCell : GridLevel.LayerCells)
+			{
+				LevelCellCount += LayerCell.GridCells.Num();
+				for (const UWorldPartitionRuntimeCell* Cell : LayerCell.GridCells)
+				{
+					LevelActorCount += Cell->GetActorCount();
+				}
+			}
+			LevelsStats.Add({ LevelCellCount, ((int64)CellSize << (int64)Level), LevelActorCount });
+			TotalActorCount += LevelActorCount;
+			++Level;
+		}
+		TotalActorCount = (TotalActorCount > 0) ? TotalActorCount : 1;
+	}
+
+	{
+		FHierarchicalLogArchive::FIndentScope LevelIndentScope = Ar.PrintfIndent(TEXT("Grid Levels: %d"), GridLevels.Num());
+		for (int Level = 0; Level < LevelsStats.Num(); ++Level)
+		{
+			Ar.Printf(TEXT("Level %2d: Cell Count %4d | Cell Size %7lld | Actor Count %4d (%3.1f%%)"), Level, LevelsStats[Level].CellCount, LevelsStats[Level].CellSize, LevelsStats[Level].ActorCount, (100.f * LevelsStats[Level].ActorCount) / TotalActorCount);
+		}
+	}
+
+	{
+		Ar.Printf(TEXT(""));
+		int32 Level = 0;
+		for (const FSpatialHashStreamingGridLevel& GridLevel : GridLevels)
+		{
+			FHierarchicalLogArchive::FIndentScope LevelIndentScope = Ar.PrintfIndent(TEXT("Content of Grid Level %d"), Level++);
+
+			TArray<UWorldPartitionRuntimeCell*> Cells;
+			for (const FSpatialHashStreamingGridLayerCell& LayerCell : GridLevel.LayerCells)
+			{
+				for (const TObjectPtr<UWorldPartitionRuntimeCell>& Cell : LayerCell.GridCells)
+				{
+					if (!Cell->IsAlwaysLoaded())
+					{
+						Cells.Add(Cell);
+					}
+				}
+			}
+
+			Cells.Sort([this](const UWorldPartitionRuntimeCell& A, const UWorldPartitionRuntimeCell& B) { return A.GetFName().LexicalLess(B.GetFName()); });
+
+			for (UWorldPartitionRuntimeCell* Cell : Cells)
+			{
+				FHierarchicalLogArchive::FIndentScope CellIndentScope = Ar.PrintfIndent(TEXT("Content of Cell %s (%s)"), *Cell->GetDebugName(), *Cell->GetName());
+				Cell->DumpStateLog(Ar);
+			}
+		}
+	}
+	Ar.Printf(TEXT(""));
+}
+#endif
+
 void FSpatialHashStreamingGrid::GetCells(const FWorldPartitionStreamingQuerySource& QuerySource, TSet<const UWorldPartitionRuntimeCell*>& OutCells, bool bEnableZCulling, FWorldPartitionQueryCache* QueryCache) const
 {
 	TRACE_CPUPROFILER_EVENT_SCOPE(FSpatialHashStreamingGrid::GetCells_QuerySource);
@@ -754,8 +843,12 @@ void FSpatialHashStreamingGrid::Draw3D(const UWorldPartitionRuntimeSpatialHash*
 				{
 					GetFilteredCellsForDebugDraw(LayerCell, FilteredCells);
 				});
+				
+				int32 FilteredCellIndex = 0;
 				for (const UWorldPartitionRuntimeCell* Cell : FilteredCells)
 				{
+					++FilteredCellIndex;
+
 					bool bIsAlreadyInSet = false;
 					DrawnCells.Add(Cell, &bIsAlreadyInSet);
 					if (bIsAlreadyInSet)
@@ -778,6 +871,10 @@ void FSpatialHashStreamingGrid::Draw3D(const UWorldPartitionRuntimeSpatialHash*
 						FBox DataLayerColoredBox = CellBox;
 						double DataLayerSizeX = DataLayerColoredBox.GetSize().X / (5 * Cell->GetDataLayers().Num()); // Use 20% of cell's width
 						DataLayerColoredBox.Max.X = DataLayerColoredBox.Min.X + DataLayerSizeX;
+						double YFilteredCellSize = (CellBox.GetSize().Y / FilteredCells.Num());
+						double YOffset = YFilteredCellSize * (FilteredCellIndex - 1);
+						DataLayerColoredBox.Min.Y = CellBox.Min.Y + YOffset;
+						DataLayerColoredBox.Max.Y = DataLayerColoredBox.Min.Y + YFilteredCellSize;
 						FTranslationMatrix DataLayerOffsetMatrix(FVector(DataLayerSizeX, 0, 0));
 						for (const FName& DataLayer : Cell->GetDataLayers())
 						{
@@ -795,6 +892,7 @@ void FSpatialHashStreamingGrid::Draw3D(const UWorldPartitionRuntimeSpatialHash*
 							check(ContentBundle->GetDescriptor());
 							FBox ContentBundleColoredBox = CellBox;
 							double ContentBundleSizeX = ContentBundleColoredBox.GetSize().X / 5; // Use 20% of cell's width
+							ContentBundleColoredBox.Min.X = ContentBundleColoredBox.Max.X - ContentBundleSizeX;
 							ContentBundleColoredBox.Max.X = ContentBundleColoredBox.Min.X + ContentBundleSizeX;
 							FTranslationMatrix ContentBundleOffsetMatrix(FVector(ContentBundleColoredBox.GetSize().X - ContentBundleSizeX, 0, 0));
 							ContentBundleColoredBox = ContentBundleColoredBox.TransformBy(ContentBundleOffsetMatrix);
@@ -1092,10 +1190,25 @@ void UWorldPartitionRuntimeSpatialHash::DrawPreview() const
 	GridPreviewer.Draw(GetWorld(), Grids, bPreviewGrids, PreviewGridLevel, bUseAlignedGridLevels);
 }
 
-URuntimeHashExternalStreamingObjectBase* UWorldPartitionRuntimeSpatialHash::StoreToExternalStreamingObject(UObject* StreamingObjectOuter, FName StreamingObjectName)
+bool UWorldPartitionRuntimeSpatialHash::HasStreamingContent() const
 {
-	URuntimeSpatialHashExternalStreamingObject* StreamingObject = CreateExternalStreamingObject<URuntimeSpatialHashExternalStreamingObject>(StreamingObjectOuter, StreamingObjectName);
+	bool bHasStreamingContent = false;
+	ForEachStreamingCells([&bHasStreamingContent](const UWorldPartitionRuntimeCell* Cell)
+	{
+		if (Cell)
+		{
+			bHasStreamingContent = true;
+		}
+		return !bHasStreamingContent;
+	});
+	return bHasStreamingContent;
+}
 
+void UWorldPartitionRuntimeSpatialHash::StoreStreamingContentToExternalStreamingObject(URuntimeHashExternalStreamingObjectBase* OutExternalStreamingObject)
+{
+	Super::StoreStreamingContentToExternalStreamingObject(OutExternalStreamingObject);
+
+	URuntimeSpatialHashExternalStreamingObject* StreamingObject = CastChecked<URuntimeSpatialHashExternalStreamingObject>(OutExternalStreamingObject);
 	StreamingObject->StreamingGrids = MoveTemp(StreamingGrids);
 
 	for (FSpatialHashStreamingGrid& StreamingGrid : StreamingObject->StreamingGrids)
@@ -1116,8 +1229,6 @@ URuntimeHashExternalStreamingObjectBase* UWorldPartitionRuntimeSpatialHash::Stor
 			}
 		}
 	}
-
-	return StreamingObject;
 }
 #endif
 
@@ -1228,22 +1339,25 @@ bool UWorldPartitionRuntimeSpatialHash::GenerateStreaming(UWorldPartitionStreami
 	TArray<FSpatialHashRuntimeGrid> AllGrids;
 	AllGrids.Append(Grids);
 
-	TArray<const FStreamingGenerationActorDescView*> SpatialHashRuntimeGridInfos = StreamingGenerationContext->GetMainWorldContainerInstance()->ActorDescViewMap->FindByExactNativeClass<ASpatialHashRuntimeGridInfo>();
-	for (const FStreamingGenerationActorDescView* SpatialHashRuntimeGridInfo : SpatialHashRuntimeGridInfos)
+	StreamingGenerationContext->ForEachActorSetContainerInstance([&AllGrids, &WorldPartition](const IStreamingGenerationContext::FActorSetContainerInstance& ActorSetContainerInstance)
 	{
-		FWorldPartitionReference Ref(WorldPartition, SpatialHashRuntimeGridInfo->GetGuid());
-		ASpatialHashRuntimeGridInfo* RuntimeGridActor = CastChecked<ASpatialHashRuntimeGridInfo>(SpatialHashRuntimeGridInfo->GetActor());
-
-		if (const FSpatialHashRuntimeGrid* ExistingRuntimeGrid = AllGrids.FindByPredicate([RuntimeGridActor](const FSpatialHashRuntimeGrid& RuntimeGrid) { return RuntimeGrid.GridName == RuntimeGridActor->GridSettings.GridName; }))
+		TArray<const FStreamingGenerationActorDescView*> SpatialHashRuntimeGridInfos = ActorSetContainerInstance.ActorDescViewMap->FindByExactNativeClass<ASpatialHashRuntimeGridInfo>();
+		for (const FStreamingGenerationActorDescView* SpatialHashRuntimeGridInfo : SpatialHashRuntimeGridInfos)
 		{
-			UE_LOG(LogWorldPartition, Log, TEXT("Got duplicated runtime grid actor '%s' for grid '%s'"), *SpatialHashRuntimeGridInfo->GetActorLabelOrName().ToString(), *RuntimeGridActor->GridSettings.GridName.ToString());
-			check(*ExistingRuntimeGrid == RuntimeGridActor->GridSettings);
-		}
-		else
-		{
-			AllGrids.Add(RuntimeGridActor->GridSettings);
+			FWorldPartitionReference Ref(WorldPartition, SpatialHashRuntimeGridInfo->GetGuid());
+			ASpatialHashRuntimeGridInfo* RuntimeGridActor = CastChecked<ASpatialHashRuntimeGridInfo>(SpatialHashRuntimeGridInfo->GetActor());
+
+			if (const FSpatialHashRuntimeGrid* ExistingRuntimeGrid = AllGrids.FindByPredicate([RuntimeGridActor](const FSpatialHashRuntimeGrid& RuntimeGrid) { return RuntimeGrid.GridName == RuntimeGridActor->GridSettings.GridName; }))
+			{
+				UE_LOG(LogWorldPartition, Log, TEXT("Got duplicated runtime grid actor '%s' for grid '%s'"), *SpatialHashRuntimeGridInfo->GetActorLabelOrName().ToString(), *RuntimeGridActor->GridSettings.GridName.ToString());
+				check(*ExistingRuntimeGrid == RuntimeGridActor->GridSettings);
+			}
+			else
+			{
+				AllGrids.Add(RuntimeGridActor->GridSettings);
+			}
 		}
-	}
+	});
 
 	TMap<FName, int32> GridsMapping;
 	GridsMapping.Add(NAME_None, 0);
@@ -1269,9 +1383,9 @@ bool UWorldPartitionRuntimeSpatialHash::GenerateStreaming(UWorldPartitionStreami
 		int32 GridIndex = FoundIndex ? *FoundIndex : 0;
 		GridActorSetInstances[GridIndex].Add(&ActorSetInstance);
 	});
-	
+
 	const FBox WorldBounds = StreamingGenerationContext->GetWorldBounds();
-	for (int32 GridIndex=0; GridIndex < AllGrids.Num(); GridIndex++)
+	for (int32 GridIndex = 0; GridIndex < AllGrids.Num(); GridIndex++)
 	{
 		const FSpatialHashRuntimeGrid& Grid = AllGrids[GridIndex];
 		const FSquare2DGridHelper PartionedActors = GetPartitionedActors(WorldBounds, Grid, GridActorSetInstances[GridIndex], Settings);
@@ -1288,92 +1402,13 @@ void UWorldPartitionRuntimeSpatialHash::DumpStateLog(FHierarchicalLogArchive& Ar
 {
 	Super::DumpStateLog(Ar);
 
+	Ar.Printf(TEXT("----------------------------------------------------------------------------------------------------------------"));
+	Ar.Printf(TEXT("%s - Runtime Spatial Hash"), *GetWorld()->GetName());
+	Ar.Printf(TEXT("----------------------------------------------------------------------------------------------------------------"));
+	FHierarchicalLogArchive::FIndentScope IndentScope = Ar.PrintfIndent(TEXT("StreamingGrids"));
 	ForEachStreamingGrid([&Ar, this](const FSpatialHashStreamingGrid& StreamingGrid)
 	{
-		Ar.Printf(TEXT("----------------------------------------------------------------------------------------------------------------"));
-		Ar.Printf(TEXT("%s - Runtime Spatial Hash - Streaming Grid - %s"), *GetWorld()->GetName(), *StreamingGrid.GridName.ToString());
-		Ar.Printf(TEXT("----------------------------------------------------------------------------------------------------------------"));
-		Ar.Printf(TEXT("            Origin: %s"), *StreamingGrid.Origin.ToString());
-		Ar.Printf(TEXT("         Cell Size: %d"), StreamingGrid.CellSize);
-		Ar.Printf(TEXT("      World Bounds: %s"), *StreamingGrid.WorldBounds.ToString());
-		Ar.Printf(TEXT("     Loading Range: %3.2f"), StreamingGrid.LoadingRange);
-		Ar.Printf(TEXT("Block Slow Loading: %s"), StreamingGrid.bBlockOnSlowStreaming ? TEXT("Yes") : TEXT("No"));
-		Ar.Printf(TEXT(" ClientOnlyVisible: %s"), StreamingGrid.bClientOnlyVisible ? TEXT("Yes") : TEXT("No"));
-		Ar.Printf(TEXT("            Origin: %s"), *StreamingGrid.Origin.ToString());
-		Ar.Printf(TEXT(""));
-		if (const UHLODLayer* HLODLayer = StreamingGrid.HLODLayer)
-		{
-			Ar.Printf(TEXT("    HLOD Layer: %s"), *HLODLayer->GetName());
-		}
-
-		struct FGridLevelStats
-		{
-			int32 CellCount;
-			int64 CellSize;
-			int32 ActorCount;
-		};
-
-		TArray<FGridLevelStats> LevelsStats;
-		int32 TotalActorCount = 0;
-			
-		{
-			int32 Level = 0;
-			for (const FSpatialHashStreamingGridLevel& GridLevel : StreamingGrid.GridLevels)
-			{
-				int32 LevelCellCount = 0;
-				int32 LevelActorCount = 0;
-				for (const FSpatialHashStreamingGridLayerCell& LayerCell : GridLevel.LayerCells)
-				{
-					LevelCellCount += LayerCell.GridCells.Num();
-					for (const TObjectPtr<UWorldPartitionRuntimeCell>& Cell : LayerCell.GridCells)
-					{
-						LevelActorCount += Cell->GetActorCount();
-					}
-				}
-				LevelsStats.Add({ LevelCellCount, ((int64)StreamingGrid.CellSize << (int64)Level), LevelActorCount });
-				TotalActorCount += LevelActorCount;
-				++Level;
-			}
-			TotalActorCount = (TotalActorCount > 0) ? TotalActorCount : 1;
-		}
-
-		{
-			FHierarchicalLogArchive::FIndentScope IndentScope = Ar.PrintfIndent(TEXT("Grid Levels: %d"), StreamingGrid.GridLevels.Num());
-			for (int Level=0; Level<LevelsStats.Num(); ++Level)
-			{
-				Ar.Printf(TEXT("Level %2d: Cell Count %4d | Cell Size %7lld | Actor Count %4d (%3.1f%%)"), Level, LevelsStats[Level].CellCount, LevelsStats[Level].CellSize, LevelsStats[Level].ActorCount, (100.f*LevelsStats[Level].ActorCount)/TotalActorCount);
-			}
-		}
-
-		{
-			Ar.Printf(TEXT(""));
-			int32 Level = 0;
-			for (const FSpatialHashStreamingGridLevel& GridLevel : StreamingGrid.GridLevels)
-			{
-				FHierarchicalLogArchive::FIndentScope LevelIndentScope = Ar.PrintfIndent(TEXT("Content of Grid Level %d"), Level++);
-
-				TArray<TObjectPtr<UWorldPartitionRuntimeCell>> Cells;
-				for (const FSpatialHashStreamingGridLayerCell& LayerCell : GridLevel.LayerCells)
-				{
-					for (const TObjectPtr<UWorldPartitionRuntimeCell>& Cell : LayerCell.GridCells)
-					{
-						if (!Cell->IsAlwaysLoaded())
-						{
-							Cells.Add(Cell);
-						}
-					}
-				}
-
-				Cells.Sort([this](const UWorldPartitionRuntimeCell& A, const UWorldPartitionRuntimeCell& B) { return A.GetFName().LexicalLess(B.GetFName()); });
-
-				for (const UWorldPartitionRuntimeCell* Cell : Cells)
-				{
-					FHierarchicalLogArchive::FIndentScope CellIndentScope = Ar.PrintfIndent(TEXT("Content of Cell %s (%s)"), *Cell->GetDebugName(), *Cell->GetName());
-					Cell->DumpStateLog(Ar);
-				}
-			}
-		}
-		Ar.Printf(TEXT(""));
+		StreamingGrid.DumpStateLog(Ar);
 	});
 }
 
@@ -1573,9 +1608,9 @@ bool UWorldPartitionRuntimeSpatialHash::CreateStreamingGrid(const FSpatialHashRu
 	return true;
 }
 
-void UWorldPartitionRuntimeSpatialHash::FlushStreaming()
+void UWorldPartitionRuntimeSpatialHash::FlushStreamingContent()
 {
-	Super::FlushStreaming();
+	Super::FlushStreamingContent();
 	bIsNameToGridMappingDirty = true;
 	StreamingGrids.Empty();
 }
@@ -1738,7 +1773,7 @@ bool UWorldPartitionRuntimeSpatialHash::RemoveExternalStreamingObject(URuntimeHa
 
 	for (const FSpatialHashStreamingGrid& ExternalStreamingGrid : SpatialHashExternalStreamingObject->StreamingGrids)
 	{
-		// In PIE, UWorldPartition::OnEndPlay() will call UWorldPartitionRuntimeSpatialHash::FlushStreaming() 
+		// In PIE, UWorldPartition::OnEndPlay() will call UWorldPartitionRuntimeSpatialHash::FlushStreamingContent() 
 		// which will empty StreamingGrids before RemoveExternalStreamingObject is called.
 		// In this case, we don't need to remove injected cells as the source grid was removed.
 		if (const FSpatialHashStreamingGrid* SourceGrid = GetStreamingGridByName(ExternalStreamingGrid.GridName))
@@ -2009,5 +2044,20 @@ TArray<const FSpatialHashStreamingGrid*> UWorldPartitionRuntimeSpatialHash::GetF
 	return FilteredStreamingGrids;
 }
 
+// ------------------------------------------------------------------------------------------------
+
+#if WITH_EDITOR
+void URuntimeSpatialHashExternalStreamingObject::DumpStateLog(FHierarchicalLogArchive& Ar)
+{
+	Super::DumpStateLog(Ar);
+
+	FHierarchicalLogArchive::FIndentScope IndentScope = Ar.PrintfIndent(TEXT("StreamingGrids"));
+	for (FSpatialHashStreamingGrid& StreamingGrid : StreamingGrids)
+	{
+		StreamingGrid.DumpStateLog(Ar);
+	}
+}
+#endif
+
 #undef LOCTEXT_NAMESPACE
 
