External Data Layer V1

diff --git a/Engine/Source/Runtime/Engine/Private/WorldPartition/WorldPartitionStreamingGeneration.cpp b/Engine/Source/Runtime/Engine/Private/WorldPartition/WorldPartitionStreamingGeneration.cpp
--- a/Engine/Source/Runtime/Engine/Private/WorldPartition/WorldPartitionStreamingGeneration.cpp
+++ b/Engine/Source/Runtime/Engine/Private/WorldPartition/WorldPartitionStreamingGeneration.cpp
@@ -7,9 +7,11 @@
 
 #include "Editor.h"
 #include "Algo/ForEach.h"
+#include "Algo/RemoveIf.h"
 #include "Algo/Sort.h"
 #include "Algo/Transform.h"
 #include "Algo/Unique.h"
+#include "Algo/Count.h"
 #include "Containers/ArrayView.h"
 #include "Engine/LevelScriptBlueprint.h"
 #include "ActorReferencesUtils.h"
@@ -24,6 +26,9 @@
 #include "WorldPartition/WorldPartitionLevelStreamingPolicy.h"
 #include "WorldPartition/DataLayer/DataLayerManager.h"
 #include "WorldPartition/DataLayer/DataLayerUtils.h"
+#include "WorldPartition/DataLayer/ExternalDataLayerHelper.h"
+#include "WorldPartition/DataLayer/ExternalDataLayerManager.h"
+#include "WorldPartition/DataLayer/DataLayerInstanceNames.h"
 #include "WorldPartition/ErrorHandling/WorldPartitionStreamingGenerationNullErrorHandler.h"
 #include "WorldPartition/ErrorHandling/WorldPartitionStreamingGenerationLogErrorHandler.h"
 #include "WorldPartition/HLOD/HLODActor.h"
@@ -207,24 +212,27 @@ FSoftObjectPath FStreamingGenerationActorDescView::GetHLODLayer() const
 	return Super::GetHLODLayer();
 }
 
-const TArray<FName>& FStreamingGenerationActorDescView::GetDataLayerInstanceNames() const
+const FDataLayerInstanceNames& FStreamingGenerationActorDescView::GetDataLayerInstanceNames() const
 {
-	if (bIsForcedNoDataLayers)
-	{
-		static TArray<FName> EmptyDataLayers;
-		return EmptyDataLayers;
-	}
-
-	if (ParentView)
+	if (!bIsForcedNoDataLayers && ParentView)
 	{
 		return ParentView->GetDataLayerInstanceNames();
 	}
 
 	if (ResolvedDataLayerInstanceNames.IsSet())
 	{
+		// ResolvedDataLayerInstanceNames contains the bIsForcedNoDataLayers information internally and will return an empty non-EDL array when requested.
+		check(ResolvedDataLayerInstanceNames->IsForcedEmptyNonExternalDataLayers() == bIsForcedNoDataLayers);
 		return ResolvedDataLayerInstanceNames.GetValue();
 	}
 
+	if (bIsForcedNoDataLayers)
+	{
+		// Build a FDataLayerInstanceNames containing only the EDL (if any) and cache the result as we need to return a ref
+		LastReturnedDataLayerInstanceNames = FDataLayerInstanceNames(TArray<FName>(), Super::GetDataLayerInstanceNames().GetExternalDataLayer());
+		return LastReturnedDataLayerInstanceNames;
+	}
+
 	return Super::GetDataLayerInstanceNames();
 }
 
@@ -233,6 +241,12 @@ const TArray<FGuid>& FStreamingGenerationActorDescView::GetReferences() const
 	return RuntimeReferences.IsSet() ? RuntimeReferences.GetValue() : Super::GetReferences();
 }
 
+bool FStreamingGenerationActorDescView::IsEditorOnlyReference(const FGuid& ReferenceGuid) const
+{
+	// We consider forced invalid references to be Editor-Only references as they will be skipped by streaming generation and PIE
+	return Super::IsEditorOnlyReference(ReferenceGuid) || ForcedInvalidReference.Contains(ReferenceGuid);
+}
+
 const TArray<FGuid>& FStreamingGenerationActorDescView::GetEditorReferences() const
 {
 	return EditorReferences;
@@ -245,10 +259,39 @@ void FStreamingGenerationActorDescView::SetParentView(const FStreamingGeneration
 	ParentView = InParentView;
 }
 
-void FStreamingGenerationActorDescView::SetDataLayerInstanceNames(const TArray<FName>& InDataLayerInstanceNames)
+void FStreamingGenerationActorDescView::SetDataLayerInstanceNames(const FDataLayerInstanceNames& InDataLayerInstanceNames)
 {
 	check(!Super::HasResolvedDataLayerInstanceNames());
 	ResolvedDataLayerInstanceNames = InDataLayerInstanceNames;
+	ResolvedDataLayerInstanceNames->SetIsForcedEmptyNonExternalDataLayers(bIsForcedNoDataLayers);
+}
+
+bool FStreamingGenerationActorDescView::IsInvalidReference(const FGuid& InGuid, FInvalidReference* OutInvalidReference) const
+{
+	if (const FInvalidReference* InvalidReference = ForcedInvalidReference.Find(InGuid))
+	{
+		if (OutInvalidReference)
+		{
+			*OutInvalidReference = *InvalidReference;
+		}
+		return true;
+	}
+	return false;
+}
+
+void FStreamingGenerationActorDescView::AddForcedInvalidReference(const FStreamingGenerationActorDescView* ReferenceView)
+{ 
+	const FGuid& ReferenceGuid = ReferenceView->GetGuid();
+	if (!ForcedInvalidReference.Contains(ReferenceGuid))
+	{
+		check(GetReferences().Contains(ReferenceGuid));
+		FInvalidReference& InvalidReference = ForcedInvalidReference.Add(ReferenceGuid);
+		InvalidReference.ActorPackage = ReferenceView->GetActorPackage();
+		InvalidReference.ActorSoftPath = ReferenceView->GetActorSoftPath();
+		InvalidReference.BaseClass = ReferenceView->GetBaseClass();
+		InvalidReference.NativeClass = ReferenceView->GetNativeClass();
+		UE_LOG(LogWorldPartition, Verbose, TEXT("Actor '%s' forced invalid reference %s"), *GetActorLabelOrName().ToString(), *ReferenceGuid.ToString());
+	}
 }
 
 void FStreamingGenerationActorDescView::SetForcedNonSpatiallyLoaded()
@@ -275,12 +318,23 @@ void FStreamingGenerationActorDescView::SetForcedNoDataLayers()
 	{
 		bIsForcedNoDataLayers = true;
 		UE_LOG(LogWorldPartition, Verbose, TEXT("Actor '%s' data layers invalidated"), *GetActorLabelOrName().ToString());
+
+		if (ResolvedDataLayerInstanceNames.IsSet())
+		{
+			ResolvedDataLayerInstanceNames->SetIsForcedEmptyNonExternalDataLayers(bIsForcedNoDataLayers);
+		}
+
+		if (RuntimeDataLayerInstanceNames.IsSet())
+		{
+			RuntimeDataLayerInstanceNames->SetIsForcedEmptyNonExternalDataLayers(bIsForcedNoDataLayers);
+		}
 	}
 }
 
-void FStreamingGenerationActorDescView::SetRuntimeDataLayerInstanceNames(const TArray<FName>& InRuntimeDataLayerInstanceNames)
+void FStreamingGenerationActorDescView::SetRuntimeDataLayerInstanceNames(const FDataLayerInstanceNames& InRuntimeDataLayerInstanceNames)
 {
 	RuntimeDataLayerInstanceNames = InRuntimeDataLayerInstanceNames;
+	RuntimeDataLayerInstanceNames->SetIsForcedEmptyNonExternalDataLayers(bIsForcedNoDataLayers);
 }
 
 void FStreamingGenerationActorDescView::SetRuntimeReferences(const TArray<FGuid>& InRuntimeReferences)
@@ -307,19 +361,21 @@ void FStreamingGenerationActorDescView::SetRuntimeHLODLayer(const FSoftObjectPat
 	RuntimedHLODLayer = InHLODLayer;
 }
 
-const TArray<FName>& FStreamingGenerationActorDescView::GetRuntimeDataLayerInstanceNames() const
+const FDataLayerInstanceNames& FStreamingGenerationActorDescView::GetRuntimeDataLayerInstanceNames() const
 {
-	if (bIsForcedNoDataLayers || !RuntimeDataLayerInstanceNames.IsSet())
+	if (!RuntimeDataLayerInstanceNames.IsSet())
 	{
-		static TArray<FName> EmptyDataLayers;
+		static FDataLayerInstanceNames EmptyDataLayers;
 		return EmptyDataLayers;
 	}
 
-	if (ParentView)
+	if (!bIsForcedNoDataLayers && ParentView)
 	{
 		return ParentView->GetRuntimeDataLayerInstanceNames();
 	}
 
+	// RuntimeDataLayerInstanceNames contains the bIsForcedNoDataLayers information internally and will return an empty non-EDL array when requested.
+	check(RuntimeDataLayerInstanceNames->IsForcedEmptyNonExternalDataLayers() == bIsForcedNoDataLayers);
 	return RuntimeDataLayerInstanceNames.GetValue();
 }
 
@@ -368,7 +424,7 @@ class FWorldPartitionStreamingGenerator
 		FStreamingGenerationContext(const FWorldPartitionStreamingGenerator* StreamingGenerator, const FStreamingGenerationContainerInstanceCollection& TopLevelActorDescCollection)
 		{
 			// Create the dataset required for IStreamingGenerationContext interface
-			MainWorldActorSetContainerIndex = INDEX_NONE;
+			ContextBaseContainerActorSetContainerInstanceIndex = INDEX_NONE;
 			ActorSetContainerInstances.Empty(StreamingGenerator->ContainerCollectionInstanceDescriptorsMap.Num());
 
 			TMap<TWeakPtr<FStreamingGenerationContainerInstanceCollection>, int32> ActorSetContainerMap;
@@ -388,12 +444,13 @@ class FWorldPartitionStreamingGenerator
 					ActorSet.Actors = Cluster;
 				}
 
-				if (ContainerDescriptor.ContainerInstanceCollection->GetMainContainerPackageName() == TopLevelActorDescCollection.GetMainContainerPackageName())
+				if (ContainerDescriptor.ContainerInstanceCollection->GetBaseContainerInstancePackageName() == TopLevelActorDescCollection.GetBaseContainerInstancePackageName())
 				{
-					check(MainWorldActorSetContainerIndex == INDEX_NONE);
-					MainWorldActorSetContainerIndex = ContainerIndex;
+					check(ContextBaseContainerActorSetContainerInstanceIndex == INDEX_NONE);
+					ContextBaseContainerActorSetContainerInstanceIndex = ContainerIndex;
 				}
 			}
+			check(StreamingGenerator->ContainerCollectionInstanceDescriptorsMap.IsEmpty() || (ContextBaseContainerActorSetContainerInstanceIndex != INDEX_NONE));
 
 			ActorSetInstances.Empty();
 			for (const auto& [ContainerID, ContainerCollectionInstanceDescriptor] : StreamingGenerator->ContainerCollectionInstanceDescriptorsMap)
@@ -462,9 +519,9 @@ class FWorldPartitionStreamingGenerator
 			return WorldBounds;
 		}
 
-		virtual const FActorSetContainerInstance* GetMainWorldContainerInstance() const override
+		virtual const FActorSetContainerInstance* GetActorSetContainerForContextBaseContainerInstance() const override
 		{
-			return ActorSetContainerInstances.IsValidIndex(MainWorldActorSetContainerIndex) ? &ActorSetContainerInstances[MainWorldActorSetContainerIndex] : nullptr;
+			return ActorSetContainerInstances.IsValidIndex(ContextBaseContainerActorSetContainerInstanceIndex) ? &ActorSetContainerInstances[ContextBaseContainerActorSetContainerInstanceIndex] : nullptr;
 		}
 
 		virtual void ForEachActorSetInstance(TFunctionRef<void(const FActorSetInstance&)> Func) const override
@@ -486,7 +543,10 @@ class FWorldPartitionStreamingGenerator
 
 	private:
 		FBox WorldBounds;
-		int32 MainWorldActorSetContainerIndex;
+		// Represents the index of the ActorSetContainerInstance (in the ActorSetContainerInstances array)
+		// that contains the a BaseContainerInstance matching this context 
+		// (the same BaseContainerInstance of the collection provided at the construction)
+		int32 ContextBaseContainerActorSetContainerInstanceIndex;
 		TArray<FActorSetContainerInstance> ActorSetContainerInstances;
 		TArray<FActorSetInstance> ActorSetInstances;
 	};
@@ -514,7 +574,7 @@ class FWorldPartitionStreamingGenerator
 
 		FContainerCollectionInstanceDescriptor& operator=(const FContainerCollectionInstanceDescriptor& Other) = delete;
 		FContainerCollectionInstanceDescriptor& operator=(FContainerCollectionInstanceDescriptor&& Other) = default;
-				
+
 		/** The actor descriptor views for for this descriptor (TUniqueObj so it is moveable without having to update FStreamingGenerationActorDescView::ActorDescViewMap pointer */
 		TUniqueObj<FStreamingGenerationActorDescViewMap> ActorDescViewMap;
 
@@ -539,9 +599,9 @@ class FWorldPartitionStreamingGenerator
 			inline bool operator==(const FPerInstanceData& Other) const
 			{
 				// Assumes DataLayers are sorted
-				return 
-					(bIsSpatiallyLoaded == Other.bIsSpatiallyLoaded) && 
-					(RuntimeGrid == Other.RuntimeGrid) && 
+				return
+					(bIsSpatiallyLoaded == Other.bIsSpatiallyLoaded) &&
+					(RuntimeGrid == Other.RuntimeGrid) &&
 					(DataLayers == Other.DataLayers);
 			}
 
@@ -619,11 +679,11 @@ class FWorldPartitionStreamingGenerator
 			// Build a WorldDataLayerActorDescs if DataLayerManager can't resolve Data Layers (i.e. when validating changelists and World is not loaded)
 			const bool bDataLayerManagerCanResolve = DataLayerManager && DataLayerManager->CanResolveDataLayers();
 			const TArray<const FWorldDataLayersActorDesc*> WorldDataLayerActorDescs = !bDataLayerManagerCanResolve ? FDataLayerUtils::FindWorldDataLayerActorDescs(ActorDescViewMap) : TArray<const FWorldDataLayersActorDesc*>();
-			const TArray<FName> DataLayerInstanceNames = FDataLayerUtils::ResolvedDataLayerInstanceNames(DataLayerManager, ActorDescView.GetActorDesc(), WorldDataLayerActorDescs);
+			const FDataLayerInstanceNames DataLayerInstanceNames = FDataLayerUtils::ResolveDataLayerInstanceNames(DataLayerManager, ActorDescView.GetActorDesc(), WorldDataLayerActorDescs);
 			ActorDescView.SetDataLayerInstanceNames(DataLayerInstanceNames);
 		}
 
-		TArray<FName> RuntimeDataLayerInstanceNames;
+		FDataLayerInstanceNames RuntimeDataLayerInstanceNames;
 		if (FDataLayerUtils::ResolveRuntimeDataLayerInstanceNames(DataLayerManager, ActorDescView, ActorDescViewMap, RuntimeDataLayerInstanceNames))
 		{
 			ActorDescView.SetRuntimeDataLayerInstanceNames(RuntimeDataLayerInstanceNames);
@@ -654,7 +714,7 @@ class FWorldPartitionStreamingGenerator
 		const bool bShouldHandleUnsavedActors = bHandleUnsavedActors && InContainerID.IsMainContainer();
 
 		// Consider all actors of a /Temp/ container package as Unsaved because loading them from disk will fail (Outer world name mismatch)
-		const bool bIsTempContainerPackage = FPackageName::IsTempPackage((InActorDescCollection.GetMainContainerPackageName().ToString()));
+		const bool bIsTempContainerPackage = FPackageName::IsTempPackage((InActorDescCollection.GetBaseContainerInstancePackageName().ToString()));
 
 		// Test whether an actor descriptor instance should be included in the ActorDescViewMap.
 		auto ShouldRegisterActorDesc = [this](const FWorldPartitionActorDescInstance* InActorDescInstance)
@@ -728,8 +788,11 @@ class FWorldPartitionStreamingGenerator
 		{
 			for (AActor* Actor : InActorDescCollection.GetWorld()->PersistentLevel->Actors)
 			{
+				// Here, FindHandlingContainer is used to make sure that the actor is handled by the collection
+				// The main reason is that UWorldPartition::CheckForErrors currently builds a collection per ActorDescContainer of the WorldPartition.
+				// This is probably a limitation introduced by ContentBundles. 
 				if (IsValid(Actor) && Actor->IsPackageExternal() && Actor->IsMainPackageActor() && !Actor->IsEditorOnly()
-					&& (Actor->GetContentBundleGuid() == InActorDescCollection.GetContentBundleGuid())
+					&& InActorDescCollection.FindHandlingContainerInstance(Actor)
 					&& !InActorDescCollection.GetActorDescInstance(Actor->GetActorGuid()))
 				{
 					TUniquePtr<FStreamingGenerationUnsavedDirtyActorDescInstance> UnsavedViewPtr = FStreamingGenerationUnsavedDirtyActorDescInstance::Create(Actor, InActorDescCollection);
@@ -758,7 +821,7 @@ class FWorldPartitionStreamingGenerator
 			ResultPerInstanceData.RuntimeGrid = (ParentContainerCollectionInstanceDescriptor.ID.IsMainContainer() || ParentContainerCollectionInstanceDescriptor.InstanceData.RuntimeGrid.IsNone()) ? InActorDescView.GetRuntimeGrid() : ParentContainerCollectionInstanceDescriptor.InstanceData.RuntimeGrid;
 
 			// Data layers are accumulated down the hierarchy chain, since level instances supports data layers assignation on actors
-			ResultPerInstanceData.DataLayers = InActorDescView.GetRuntimeDataLayerInstanceNames();
+			ResultPerInstanceData.DataLayers = InActorDescView.GetRuntimeDataLayerInstanceNames().ToArray();
 			ResultPerInstanceData.DataLayers.Append(ParentContainerCollectionInstanceDescriptor.InstanceData.DataLayers);
 			ResultPerInstanceData.DataLayers.Sort(FNameFastLess());
 
@@ -837,7 +900,11 @@ class FWorldPartitionStreamingGenerator
 				SubContainerInstanceDescriptor.ID = SubContainerInstance.ContainerInstance->GetContainerID();
 				check(SubContainerInstanceDescriptor.ID == FActorContainerID(ContainerCollectionInstanceDescriptor.ID, ContainerCollectionInstanceView.GetGuid()));
 
-				SubContainerInstanceDescriptor.ContainerInstanceCollection = MakeShared<FStreamingGenerationContainerInstanceCollection>(FStreamingGenerationContainerInstanceCollection{ SubContainerInstance.ContainerInstance });
+				// @todo_ow: LevelInstance EDL Support
+				// LevelInstance don't support Content Bundle containers nor EDL containers
+				ensure(!SubContainerInstance.ContainerInstance->HasExternalContent());
+				FStreamingGenerationContainerInstanceCollection SubContainerInstanceCollection({ SubContainerInstance.ContainerInstance }, FStreamingGenerationContainerInstanceCollection::ECollectionType::BaseAndEDLs);
+				SubContainerInstanceDescriptor.ContainerInstanceCollection = MakeShared<FStreamingGenerationContainerInstanceCollection>(SubContainerInstanceCollection);
 				SubContainerInstanceDescriptor.Transform = SubContainerInstance.ContainerInstance->GetTransform();
 
 				// @todo_ow: this is to validate that new parenting of container instance code is equivalent
@@ -906,13 +973,13 @@ class FWorldPartitionStreamingGenerator
 	 */
 	void CreateContainerResolver(const FStreamingGenerationContainerInstanceCollection& InContainerInstanceCollection)
 	{
-		ContainerResolver.SetMainContainerPackage(InContainerInstanceCollection.GetMainContainerPackageName());
+		ContainerResolver.SetMainContainerPackage(InContainerInstanceCollection.GetBaseContainerInstancePackageName());
 
 		for (const auto& [ContainerID, ContainerDescriptor] : ContainerCollectionInstanceDescriptorsMap)
 		{
-			if (!ContainerResolver.ContainsContainer(ContainerDescriptor.ContainerInstanceCollection->GetMainContainerPackageName()))
+			if (!ContainerResolver.ContainsContainer(ContainerDescriptor.ContainerInstanceCollection->GetBaseContainerInstancePackageName()))
 			{
-				FWorldPartitionRuntimeContainer& Container = ContainerResolver.AddContainer(ContainerDescriptor.ContainerInstanceCollection->GetMainContainerPackageName());
+				FWorldPartitionRuntimeContainer& Container = ContainerResolver.AddContainer(ContainerDescriptor.ContainerInstanceCollection->GetBaseContainerInstancePackageName());
 
 				for (const FStreamingGenerationActorDescView& ActorDescView : ContainerDescriptor.ContainerCollectionInstanceViews)
 				{
@@ -982,7 +1049,7 @@ class FWorldPartitionStreamingGenerator
 			}
 			else
 			{
-				ULevel::GetLevelScriptExternalActorsReferencesFromPackage(ContainerCollectionInstanceDescriptor.ContainerInstanceCollection->GetMainContainerPackageName(), LevelScriptReferences);
+				ULevel::GetLevelScriptExternalActorsReferencesFromPackage(ContainerCollectionInstanceDescriptor.ContainerInstanceCollection->GetBaseContainerInstancePackageName(), LevelScriptReferences);
 			}
 
 			for (const FGuid& LevelScriptReferenceActorGuid : LevelScriptReferences)
@@ -1044,13 +1111,21 @@ class FWorldPartitionStreamingGenerator
 					return bIsActorDescSpatiallyLoaded == bIsActorDescRefSpatiallyLoaded;
 				};
 
+				// Validate external data layer
+				auto IsReferenceExternalDataLayerValid = [](const FStreamingGenerationActorDescView& RefererActorDescView, const FStreamingGenerationActorDescView& ReferenceActorDescView)
+				{
+					return RefererActorDescView.GetRuntimeDataLayerInstanceNames().GetExternalDataLayer() == ReferenceActorDescView.GetRuntimeDataLayerInstanceNames().GetExternalDataLayer();
+				};
+
 				// Validate data layers
 				auto IsReferenceDataLayersValid = [](const FStreamingGenerationActorDescView& RefererActorDescView, const FStreamingGenerationActorDescView& ReferenceActorDescView)
 				{
-					if (RefererActorDescView.GetRuntimeDataLayerInstanceNames().Num() == ReferenceActorDescView.GetRuntimeDataLayerInstanceNames().Num())
+					if (RefererActorDescView.GetRuntimeDataLayerInstanceNames().GetNonExternalDataLayers().Num() == ReferenceActorDescView.GetRuntimeDataLayerInstanceNames().GetNonExternalDataLayers().Num())
 					{
-						const TSet<FName> RefererActorDescDataLayers(RefererActorDescView.GetRuntimeDataLayerInstanceNames());
-						const TSet<FName> ReferenceActorDescDataLayers(ReferenceActorDescView.GetRuntimeDataLayerInstanceNames());
+						TSet<FName> RefererActorDescDataLayers;
+						TSet<FName> ReferenceActorDescDataLayers;
+						RefererActorDescDataLayers.Append(RefererActorDescView.GetRuntimeDataLayerInstanceNames().GetNonExternalDataLayers());
+						ReferenceActorDescDataLayers.Append(ReferenceActorDescView.GetRuntimeDataLayerInstanceNames().GetNonExternalDataLayers());
 
 						return RefererActorDescDataLayers.Includes(ReferenceActorDescDataLayers);
 					}
@@ -1148,11 +1223,25 @@ class FWorldPartitionStreamingGenerator
 								NbErrorsDetected++;
 							}
 
+							if (!IsReferenceExternalDataLayerValid(*RefererActorDescView, *ReferenceActorDescView))
+							{
+								if (PassType == EPassType::ErrorReporting)
+								{
+									ErrorHandler->OnInvalidReferenceDataLayers(*RefererActorDescView, *ReferenceActorDescView, IStreamingGenerationErrorHandler::EDataLayerInvalidReason::ReferencedActorDifferentExternalDataLayer);
+								}
+								else
+								{
+									RefererActorDescView->AddForcedInvalidReference(ReferenceActorDescView);
+								}
+
+								NbErrorsDetected++;
+							}
+
 							if (!IsReferenceDataLayersValid(*RefererActorDescView, *ReferenceActorDescView))
 							{
 								if (PassType == EPassType::ErrorReporting)
 								{
-									ErrorHandler->OnInvalidReferenceDataLayers(*RefererActorDescView, *ReferenceActorDescView);
+									ErrorHandler->OnInvalidReferenceDataLayers(*RefererActorDescView, *ReferenceActorDescView, IStreamingGenerationErrorHandler::EDataLayerInvalidReason::ReferencedActorDifferentRuntimeDataLayers);
 								}
 								else
 								{
@@ -1448,6 +1537,7 @@ PRAGMA_ENABLE_DEPRECATION_WARNINGS
 					check(ActorDescView.GetRuntimeGrid() == ReferenceActorDescView.GetRuntimeGrid());
 					check(ActorDescView.GetIsSpatiallyLoaded() == ReferenceActorDescView.GetIsSpatiallyLoaded());
 					check(ActorDescView.GetContentBundleGuid() == ReferenceActorDescView.GetContentBundleGuid());
+					check(ActorDescView.GetExternalDataLayerAsset() == ReferenceActorDescView.GetExternalDataLayerAsset());
 				}
 			}
 		}
@@ -1553,7 +1643,7 @@ public:
 		UE_SCOPED_INDENT_LOG_ARCHIVE(Ar.PrintfIndent(TEXT("Containers:")));
 		for (auto& [ContainerID, ContainerCollectionInstanceDescriptor] : ContainerCollectionInstanceDescriptorsMap)
 		{
-			FString ContainerPackageName = ContainerCollectionInstanceDescriptor.ContainerInstanceCollection->GetMainContainerPackageName().ToString();
+			FString ContainerPackageName = ContainerCollectionInstanceDescriptor.ContainerInstanceCollection->GetBaseContainerInstancePackageName().ToString();
 			bool bIsAlreadySet = false;
 			UniqueContainerNames.Add(ContainerPackageName, &bIsAlreadySet);
 			if (bIsAlreadySet)
@@ -1614,7 +1704,7 @@ public:
 					Ar.Printf(TEXT(" ParentID: %s"), *ContainerInstanceDescriptor.ParentID.ToString());
 					Ar.Printf(TEXT("   Bounds: %s"), *ContainerInstanceDescriptor.Bounds.ToString());
 					Ar.Printf(TEXT("Transform: %s"), *ContainerInstanceDescriptor.Transform.ToString());
-					Ar.Printf(TEXT("Container: %s"), *ContainerInstanceDescriptor.ContainerInstanceCollection->GetMainContainerPackageName().ToString());
+					Ar.Printf(TEXT("Container: %s"), *ContainerInstanceDescriptor.ContainerInstanceCollection->GetBaseContainerInstancePackageName().ToString());
 				}
 
 				TArray<FActorContainerID> ChildContainersIDs;
@@ -1690,8 +1780,7 @@ private:
 
 bool UWorldPartition::GenerateStreaming(const FGenerateStreamingParams& InParams, FGenerateStreamingContext& InContext)
 {
-	check(InParams.ContainerInstanceCollection.IsEmpty());
-	FGenerateStreamingParams Params = FGenerateStreamingParams(InParams).SetActorDescContainerInstance(ActorDescContainerInstance);
+	FGenerateStreamingParams Params = FGenerateStreamingParams(InParams).SetContainerInstanceCollection(*this, FStreamingGenerationContainerInstanceCollection::ECollectionType::BaseAndEDLs);
 
 	OnPreGenerateStreaming.Broadcast(InContext.PackagesToGenerate);
 
@@ -1702,7 +1791,7 @@ bool UWorldPartition::GenerateContainerStreaming(const FGenerateStreamingParams&
 {
 	TRACE_CPUPROFILER_EVENT_SCOPE(UWorldPartition::GenerateContainerStreaming);
 
-	const FString ContainerPackageName = InParams.ContainerInstanceCollection.GetMainContainerPackageName().ToString();
+	const FString ContainerPackageName = InParams.ContainerInstanceCollection.GetBaseContainerInstancePackageName().ToString();
 	FString ContainerShortName = FPackageName::GetShortName(ContainerPackageName);
 	if (!ContainerPackageName.StartsWith(TEXT("/Game/")))
 	{
@@ -1759,27 +1848,63 @@ bool UWorldPartition::GenerateContainerStreaming(const FGenerateStreamingParams&
 		StreamingGenerator.DumpStateLog(*HierarchicalLogAr);
 	}
 
-	// Generate streaming
-	check(!StreamingPolicy);
-	StreamingPolicy = NewObject<UWorldPartitionStreamingPolicy>(const_cast<UWorldPartition*>(this), WorldPartitionStreamingPolicyClass.Get(), NAME_None, bIsPIE ? RF_Transient : RF_NoFlags);
-
-	check(RuntimeHash);
-	const IStreamingGenerationContext* StreamingGenerationContext = StreamingGenerator.GetStreamingGenerationContext(InParams.ContainerInstanceCollection);
-	check(StreamingGenerationContext);
-	if (RuntimeHash->GenerateStreaming(StreamingPolicy, StreamingGenerationContext, InContext.PackagesToGenerate))
+	auto GenerateRuntimeHash = [this, &HierarchicalLogAr, &StreamingGenerator, &InContext, &InParams](const UActorDescContainerInstance* InActorDescContainerInstance)
 	{
-		if (HierarchicalLogAr.IsValid())
+		check(!StreamingPolicy);
+		StreamingPolicy = NewObject<UWorldPartitionStreamingPolicy>(const_cast<UWorldPartition*>(this), WorldPartitionStreamingPolicyClass.Get(), NAME_None, bIsPIE ? RF_Transient : RF_NoFlags);
+
+		FStreamingGenerationContextProxy GenerationContextProxy(StreamingGenerator.GetStreamingGenerationContext(InParams.ContainerInstanceCollection));
+		GenerationContextProxy.SetActorSetInstanceFilter([ExternalDataLayerAsset = InActorDescContainerInstance->GetExternalDataLayerAsset()](const IStreamingGenerationContext::FActorSetInstance& InActorSetInstance) { return (InActorSetInstance.GetExternalDataLayerAsset() == ExternalDataLayerAsset); });
+
+		check(RuntimeHash);
+		if (RuntimeHash->GenerateStreaming(StreamingPolicy, &GenerationContextProxy, InContext.PackagesToGenerate))
 		{
-			RuntimeHash->DumpStateLog(*HierarchicalLogAr);
+			StreamingPolicy->SetContainerResolver(StreamingGenerator.GetContainerResolver());
+			StreamingPolicy->PrepareActorToCellRemapping();
+			StreamingPolicy->SetShouldMergeStreamingSourceInfo(RuntimeHash->GetShouldMergeStreamingSourceInfo());
+			return true;
 		}
+		return false;
+	};
 
-		StreamingPolicy->SetContainerResolver(StreamingGenerator.GetContainerResolver());
-		StreamingPolicy->PrepareActorToCellRemapping();
-		StreamingPolicy->SetShouldMergeStreamingSourceInfo(RuntimeHash->GetShouldMergeStreamingSourceInfo());
-		return true;
+	// Generate streaming for External Data Layer container instances
+	bool bStreamingGenerationSuccess = true;
+	for (const UActorDescContainerInstance* ExternalDataLayerContainerInstance : InParams.ContainerInstanceCollection.GetExternalDataLayerContainerInstances())
+	{
+		const UExternalDataLayerAsset* ExternalDataLayerAsset = ExternalDataLayerContainerInstance->GetExternalDataLayerAsset();
+		check(ExternalDataLayerAsset);
+		bool bExternalDataLayerGenerationSuccess = GenerateRuntimeHash(ExternalDataLayerContainerInstance);
+		// No need to create an ExternalStreamingObject and move the streaming content if it's empty
+		if (bExternalDataLayerGenerationSuccess && RuntimeHash->HasStreamingContent())
+		{
+			URuntimeHashExternalStreamingObjectBase* ExternalStreamingObject = ExternalDataLayerManager->CreateExternalStreamingObjectUsingStreamingGeneration(ExternalDataLayerAsset);
+			bExternalDataLayerGenerationSuccess = !!ExternalStreamingObject;
+			if (bExternalDataLayerGenerationSuccess)
+			{
+				if (HierarchicalLogAr.IsValid())
+				{
+					ExternalStreamingObject->DumpStateLog(*HierarchicalLogAr);
+				}
+
+				if (InContext.GeneratedExternalStreamingObjects)
+				{
+					InContext.GeneratedExternalStreamingObjects->Add(ExternalStreamingObject);
+				}
+			}
+		}
+		bStreamingGenerationSuccess &= bExternalDataLayerGenerationSuccess;
+		FlushStreaming();
 	}
 
-	return false;
+	// Generate streaming for the Base container instance
+	const UActorDescContainerInstance* BaseContainerInstance = InParams.ContainerInstanceCollection.GetBaseContainerInstance();
+	const bool bBaseContainerGenerationSuccess = GenerateRuntimeHash(BaseContainerInstance);
+	if (bBaseContainerGenerationSuccess && HierarchicalLogAr.IsValid())
+	{
+		RuntimeHash->DumpStateLog(*HierarchicalLogAr);
+	}
+	bStreamingGenerationSuccess &= bBaseContainerGenerationSuccess;
+	return bStreamingGenerationSuccess;
 }
 
 class FStreamingGenerationContextCopy : public FStreamingGenerationContextProxy
@@ -1826,28 +1951,51 @@ TUniquePtr<IStreamingGenerationContext> UWorldPartition::GenerateStreamingGenera
 
 void UWorldPartition::FlushStreaming()
 {
-	RuntimeHash->FlushStreaming();
+	RuntimeHash->FlushStreamingContent();
 	StreamingPolicy = nullptr;
-	GeneratedStreamingPackageNames.Empty();
+	GeneratedLevelStreamingPackageNames.Empty();
+}
+
+bool UWorldPartition::HasStreamingContent() const
+{
+	return RuntimeHash && RuntimeHash->HasStreamingContent();
 }
 
 URuntimeHashExternalStreamingObjectBase* UWorldPartition::FlushStreamingToExternalStreamingObject(const FString& ExternalStreamingObjectName)
 {
-	URuntimeHashExternalStreamingObjectBase* ExternalStreamingObject = RuntimeHash->StoreToExternalStreamingObject(this, *ExternalStreamingObjectName);
+	URuntimeHashExternalStreamingObjectBase* ExternalStreamingObject = RuntimeHash->StoreStreamingContentToExternalStreamingObject(*ExternalStreamingObjectName);
 	check(ExternalStreamingObject);
 
-	StreamingPolicy->StoreToExternalStreamingObject(*ExternalStreamingObject);
+	StreamingPolicy->StoreStreamingContentToExternalStreamingObject(*ExternalStreamingObject);
 
 	FlushStreaming();
 	return ExternalStreamingObject;
 }
 
+static void ExtractContentBundleContainerInstances(const FActorDescContainerInstanceCollection* InContainerInstanceCollection, TArray<const UActorDescContainerInstance*>& OutContentBundleContainerInstances, TArray<const UActorDescContainerInstance*>& OutNonContentBundleContainerInstances)
+{
+	InContainerInstanceCollection->ForEachActorDescContainerInstance([&OutContentBundleContainerInstances, &OutNonContentBundleContainerInstances](const UActorDescContainerInstance* InActorDescContainerInstance)
+	{
+		if (InActorDescContainerInstance->GetContentBundleGuid().IsValid())
+		{
+			OutContentBundleContainerInstances.Add(InActorDescContainerInstance);
+		}
+		else
+		{
+			OutNonContentBundleContainerInstances.Add(InActorDescContainerInstance);
+		}
+	});
+}
+
 void UWorldPartition::SetupHLODActors(const FSetupHLODActorsParams& Params)
 {
-	ForEachActorDescContainerInstance([this, &Params](UActorDescContainerInstance* InContainerInstance)
+	TArray<const UActorDescContainerInstance*> ContentBundleContainerInstances;
+	TArray<const UActorDescContainerInstance*> BaseAndEDLContainerInstances;
+	ExtractContentBundleContainerInstances(this, ContentBundleContainerInstances, BaseAndEDLContainerInstances);
+
+	auto SetupHLODActorsForCollection = [this, &Params](const FStreamingGenerationContainerInstanceCollection& InContainerInstanceCollection)
 	{
 		TErrorHandlerSelector<FStreamingGenerationLogErrorHandler> ErrorHandlerSelector;
-
 		FWorldPartitionStreamingGenerator::FWorldPartitionStreamingGeneratorParams StreamingGeneratorParams = FWorldPartitionStreamingGenerator::FWorldPartitionStreamingGeneratorParams()
 			.SetWorldPartitionContext(this)
 			.SetErrorHandler(ErrorHandlerSelector.Get())
@@ -1857,185 +2005,257 @@ void UWorldPartition::SetupHLODActors(const FSetupHLODActorsParams& Params)
 			.SetIsValidHLODLayer([this](FName GridName, const FSoftObjectPath& HLODLayerPath) { return RuntimeHash->IsValidHLODLayer(GridName, HLODLayerPath); });
 
 		FWorldPartitionStreamingGenerator StreamingGenerator(StreamingGeneratorParams);
-		FStreamingGenerationContainerInstanceCollection ActorDescCollection{ InContainerInstance };
-		StreamingGenerator.PreparationPhase(ActorDescCollection);
+		StreamingGenerator.PreparationPhase(InContainerInstanceCollection);
 
 		TUniquePtr<FArchive> LogFileAr = FWorldPartitionStreamingGenerator::CreateDumpStateLogArchive(TEXT("HLOD"));
-
 		if (LogFileAr.IsValid())
 		{
 			FHierarchicalLogArchive HierarchicalLogAr(*LogFileAr);
 			StreamingGenerator.DumpStateLog(HierarchicalLogAr);
 		}
 
-		RuntimeHash->SetupHLODActors(StreamingGenerator.GetStreamingGenerationContext(ActorDescCollection), Params);
-	});
+		RuntimeHash->SetupHLODActors(StreamingGenerator.GetStreamingGenerationContext(InContainerInstanceCollection), Params);
+	};
+
+	// Process all Content Bundle container instances
+	for (const UActorDescContainerInstance* ContentBundleContainerInstance : ContentBundleContainerInstances)
+	{
+		FStreamingGenerationContainerInstanceCollection ContentBundleCollection({ ContentBundleContainerInstance }, FStreamingGenerationContainerInstanceCollection::ECollectionType::BaseAsContentBundle);
+		SetupHLODActorsForCollection(ContentBundleCollection);
+	}
+
+	// Single pass for base and EDL container instances
+	if (!BaseAndEDLContainerInstances.IsEmpty())
+	{
+		FStreamingGenerationContainerInstanceCollection Collection(BaseAndEDLContainerInstances, FStreamingGenerationContainerInstanceCollection::ECollectionType::BaseAndEDLs);
+		SetupHLODActorsForCollection(Collection);
+	}
 }
 
-FStreamingGenerationContainerInstanceCollection::FStreamingGenerationContainerInstanceCollection(std::initializer_list<TObjectPtr<const UActorDescContainerInstance>> ActorDescContainerArray)
-	: TActorDescContainerInstanceCollection<TObjectPtr<const UActorDescContainerInstance>>(ActorDescContainerArray)
+FStreamingGenerationContainerInstanceCollection::FStreamingGenerationContainerInstanceCollection(std::initializer_list<TObjectPtr<const UActorDescContainerInstance>> ActorDescContainerInstanceArray, const ECollectionType& InCollectionType)
+	: TActorDescContainerInstanceCollection<TObjectPtr<const UActorDescContainerInstance>>(ActorDescContainerInstanceArray)
+	, CollectionType(InCollectionType)
 {
-	SortCollection();
+	InitializeCollection();
 }
 
-FStreamingGenerationContainerInstanceCollection::FStreamingGenerationContainerInstanceCollection(const TArray<const UActorDescContainerInstance*>& ActorDescContainers)
-	: TActorDescContainerInstanceCollection<TObjectPtr<const UActorDescContainerInstance>>(ActorDescContainers)
+FStreamingGenerationContainerInstanceCollection::FStreamingGenerationContainerInstanceCollection(const TArray<const UActorDescContainerInstance*>& ActorDescContainerInstances, const ECollectionType& InCollectionType)
+	: TActorDescContainerInstanceCollection<TObjectPtr<const UActorDescContainerInstance>>(ActorDescContainerInstances)
+	, CollectionType(InCollectionType)
 {
-	SortCollection();
+	InitializeCollection();
 }
 
 UWorld* FStreamingGenerationContainerInstanceCollection::GetWorld() const
 {
-	if (const UActorDescContainerInstance* MainContainer = GetMainContainer())
-	{
-		if (UWorldPartition* WorldPartition = MainContainer->GetWorldPartition())
-		{
-			return WorldPartition->GetWorld();
-		}
-	}
-
-	return nullptr;
+	UWorldPartition* WorldPartition = GetBaseContainerInstance()->GetWorldPartition();
+	check(WorldPartition);
+	UWorld* World = WorldPartition->GetWorld();
+	check(World);
+	return World;
 }
 
 FGuid FStreamingGenerationContainerInstanceCollection::GetContentBundleGuid() const
 {
-	if (const UActorDescContainerInstance* MainContainer = GetMainContainer())
+	FGuid ContentBundleGuid;
+	if (CollectionType == ECollectionType::BaseAsContentBundle)
 	{
-		return MainContainer->GetContentBundleGuid();
+		check(GetActorDescContainerCount() == 1);
+		ContentBundleGuid = GetBaseContainerInstance()->GetContentBundleGuid();
+		check(ContentBundleGuid.IsValid());
 	}
-
-	return FGuid();
+	return ContentBundleGuid;
 }
 
-const UActorDescContainerInstance* FStreamingGenerationContainerInstanceCollection::GetMainContainer() const
+//@todo_ow: Once ContentBundle code is removed, this function will always return a base ActorDescContainer that cannot have a valid content bundle
+const UActorDescContainerInstance* FStreamingGenerationContainerInstanceCollection::GetBaseContainerInstance() const
 {
-	if (!IsEmpty())
-	{
-		return ActorDescContainerInstanceCollection[MainContainerIdx];
-	}
+	check(CollectionType != ECollectionType::Invalid);
+	check(!IsEmpty());
+	const UActorDescContainerInstance* ActorDescContainerInstance = ActorDescContainerInstanceCollection[BaseContainerIdx];
+	check(ActorDescContainerInstance);
+	// This function is not designed to return a valid ActorDescContainer if it has external content (except for ContentBundle collection type)
+	check(!ActorDescContainerInstance->HasExternalContent() || (ActorDescContainerInstance->GetContentBundleGuid().IsValid() && (CollectionType == ECollectionType::BaseAsContentBundle || CollectionType == ECollectionType::BaseAndAny)));
+	return ActorDescContainerInstance;
+}
 
-	return nullptr;
+FName FStreamingGenerationContainerInstanceCollection::GetBaseContainerInstancePackageName() const
+{
+	return GetBaseContainerInstance()->GetContainerPackage();
 }
 
-FName FStreamingGenerationContainerInstanceCollection::GetMainContainerPackageName() const
+TArrayView<const UActorDescContainerInstance* const> FStreamingGenerationContainerInstanceCollection::GetExternalDataLayerContainerInstances() const
 {
-	if (const UActorDescContainerInstance* MainContainer = GetMainContainer())
+	check(CollectionType != ECollectionType::Invalid);
+	check(!IsEmpty());
+
+	if (ExternalDataLayerStartIdx != INDEX_NONE)
 	{
-		return MainContainer->GetContainerPackage();
+		int32 Size = ((ContentBundleStartIdx != INDEX_NONE) ? ContentBundleStartIdx : GetActorDescContainerCount()) - ExternalDataLayerStartIdx;
+		return MakeArrayView(&ActorDescContainerInstanceCollection[ExternalDataLayerStartIdx], Size);
 	}
 
-	return FName();
+	return TArrayView<const UActorDescContainerInstance*>();
 }
 
-TArrayView<const UActorDescContainerInstance* const> FStreamingGenerationContainerInstanceCollection::GetExternalDataLayerContainers()
+TArrayView<const UActorDescContainerInstance* const> FStreamingGenerationContainerInstanceCollection::GetContentBundleContainerInstances() const
 {
-	if (ActorDescContainerInstanceCollection.Num() <= 1)
+	check(CollectionType != ECollectionType::Invalid);
+	check(!IsEmpty());
+
+	if (ContentBundleStartIdx != INDEX_NONE)
 	{
-		return TArrayView<const UActorDescContainerInstance*>();
+		int32 Size = GetActorDescContainerCount() - ContentBundleStartIdx;
+		return MakeArrayView(&ActorDescContainerInstanceCollection[ContentBundleStartIdx], Size);
 	}
 
-	uint32 NumExternalDataLayerContainers = ActorDescContainerInstanceCollection.Num() - ExternalDataLayerContainerStartIdx;
-	return MakeArrayView(&ActorDescContainerInstanceCollection[ExternalDataLayerContainerStartIdx], NumExternalDataLayerContainers);
+	return TArrayView<const UActorDescContainerInstance*>();
 }
 
 void FStreamingGenerationContainerInstanceCollection::OnCollectionChanged()
 {
-	SortCollection();
+	InitializeCollection();
 }
 
-void FStreamingGenerationContainerInstanceCollection::SortCollection()
+void FStreamingGenerationContainerInstanceCollection::InitializeCollection()
 {
+	ExternalDataLayerStartIdx = INDEX_NONE;
+	ContentBundleStartIdx = INDEX_NONE;
+
+	check(!IsEmpty());
+	check(CollectionType != ECollectionType::Invalid);
 	if (IsEmpty())
 	{
 		return;
 	}
 
-	auto IsMainPartitionContainer = [this](const UActorDescContainerInstance* ActorDescContainerInstance)
+	if (CollectionType == ECollectionType::BaseAsContentBundle)
 	{
-		bool bIsContentBundleContainer = ActorDescContainerInstance->GetContentBundleGuid().IsValid() && GetActorDescContainerCount() == 1 && !ActorDescContainerInstance->GetContainerPackage().ToString().StartsWith(TEXT("/Game/"), ESearchCase::IgnoreCase);
-		return !ActorDescContainerInstance->GetContentBundleGuid().IsValid() || bIsContentBundleContainer;
-	};
+		check(GetActorDescContainerCount() == 1);
+		check(GetContentBundleGuid().IsValid());
+		ContentBundleStartIdx = 0;
+		return;
+	}
 
-	Algo::Sort(ActorDescContainerInstanceCollection, [&IsMainPartitionContainer](const UActorDescContainerInstance* A, const UActorDescContainerInstance* B)
+	check((CollectionType == ECollectionType::BaseAndEDLs) || (CollectionType == ECollectionType::BaseAndAny));
+	check(!GetContentBundleGuid().IsValid());
+
+	if (CollectionType == ECollectionType::BaseAndEDLs)
+	{
+		// When type is set to BaseAndEDL, we remove ContentBundle containers from the collection.
+		// BaseAndEDL type assumes ContentBundle containers are generated separately one at a time.
+		ActorDescContainerInstanceCollection.SetNum(Algo::RemoveIf(ActorDescContainerInstanceCollection, [](const UActorDescContainerInstance* ActorDescContainerInstance) { return ActorDescContainerInstance->GetContentBundleGuid().IsValid(); }));
+	}
+
+	int32 BaseContainerCount = Algo::CountIf(ActorDescContainerInstanceCollection, [](const UActorDescContainerInstance* ActorDescContainerInstance) { return !ActorDescContainerInstance->HasExternalContent(); });
+	check(BaseContainerCount == 1);
+
+	// Sort containers : Base, EDLs, ContentBundles
+	if (GetActorDescContainerCount() > 1)
 	{
-		if (IsMainPartitionContainer(A) && !IsMainPartitionContainer(B))
+		auto GetContainerSortValue = [](const UActorDescContainerInstance* ActorDescContainerInstance) { return (!!ActorDescContainerInstance->GetExternalDataLayerAsset() ? 1 : ActorDescContainerInstance->GetContentBundleGuid().IsValid() ? 2 : 0); };
+		Algo::Sort(ActorDescContainerInstanceCollection, [&GetContainerSortValue](const UActorDescContainerInstance* A, const UActorDescContainerInstance* B)
 		{
-			return true;
-		}
-		else if (IsMainPartitionContainer(B) && !IsMainPartitionContainer(A))
+			int32 AValue = GetContainerSortValue(A);
+			int32 BValue = GetContainerSortValue(B);
+			return (AValue == BValue) ? A->GetContainerPackage().LexicalLess(B->GetContainerPackage()) : (AValue < BValue);
+		});
+
+		int32 Index = 0;
+		for (const UActorDescContainerInstance* ContainerInstance : ActorDescContainerInstanceCollection)
 		{
-			return false;
+			if (ExternalDataLayerStartIdx == INDEX_NONE && ContainerInstance->GetExternalDataLayerAsset())
+			{
+				ExternalDataLayerStartIdx = Index;
+			}
+			else if (ContentBundleStartIdx == INDEX_NONE && ContainerInstance->GetContentBundleGuid().IsValid())
+			{
+				ContentBundleStartIdx = Index;
+			}
+			++Index;
 		}
 
-		return A->GetContainerPackage().Compare(B->GetContainerPackage()) < 0;
-	});
-
-	check(IsMainPartitionContainer(ActorDescContainerInstanceCollection[MainContainerIdx]));
-
 #if DO_CHECK
-	TArrayView<const UActorDescContainerInstance* const> ExternalDataLayerContainerView = GetExternalDataLayerContainers();
-	Algo::ForEach(ExternalDataLayerContainerView, [&IsMainPartitionContainer](const UActorDescContainerInstance* ActorDescContainer) { check(!IsMainPartitionContainer(ActorDescContainer)); });
+		// Validation
+		check((ContentBundleStartIdx == INDEX_NONE) || (ContentBundleStartIdx > ExternalDataLayerStartIdx));
+		check(GetBaseContainerInstance());
+		TArrayView<const UActorDescContainerInstance* const> ExternalDataLayerContainers = GetExternalDataLayerContainerInstances();
+		Algo::ForEach(ExternalDataLayerContainers, [](const UActorDescContainerInstance* ActorDescContainerInstance) { check(ActorDescContainerInstance->GetExternalDataLayerAsset()) });
+		TArrayView<const UActorDescContainerInstance* const> ContentBundleContainers = GetContentBundleContainerInstances();
+		Algo::ForEach(ContentBundleContainers, [](const UActorDescContainerInstance* ActorDescContainerInstance) { check(ActorDescContainerInstance->GetContentBundleGuid().IsValid()) });
 #endif
+	}
 }
 
 void UWorldPartition::CheckForErrors(IStreamingGenerationErrorHandler* ErrorHandler) const
 {
-	TErrorHandlerSelector<FStreamingGenerationLogErrorHandler> ErrorHandlerSelector(ErrorHandler);
-
-	FWorldPartitionStreamingGenerator::FWorldPartitionStreamingGeneratorParams StreamingGeneratorParams = FWorldPartitionStreamingGenerator::FWorldPartitionStreamingGeneratorParams()
-		.SetWorldPartitionContext(this)
-		.SetHandleUnsavedActors(true)
-		.SetErrorHandler(ErrorHandlerSelector.Get())
-		.SetIsValidGrid([this](FName GridName) { return RuntimeHash->IsValidGrid(GridName); })
-		.SetIsValidHLODLayer([this](FName GridName, const FSoftObjectPath& HLODLayerPath) { return RuntimeHash->IsValidHLODLayer(GridName, HLODLayerPath); })
+	FCheckForErrorsParams Params = FCheckForErrorsParams()
+		.SetErrorHandler(ErrorHandler)
+		.SetActorDescContainerInstanceCollection(this)
 		.SetEnableStreaming(IsStreamingEnabled());
 
-	ForEachActorDescContainerInstance([&StreamingGeneratorParams](const UActorDescContainerInstance* InContainerInstance)
+	CheckForErrors(Params);
+}
+
+/* Static version, mainly used by changelist validation */
+void UWorldPartition::CheckForErrors(const FCheckForErrorsParams& InParams)
+{
+	check(InParams.ErrorHandler);
+	check(InParams.ActorDescContainerInstanceCollection);
+
+	// Prepare ActorGuidsToContainerInstanceMap
+	TMap<FGuid, const UActorDescContainerInstance*> ActorGuidsToContainerInstanceMap;
+	InParams.ActorDescContainerInstanceCollection->ForEachActorDescContainerInstance([&ActorGuidsToContainerInstanceMap](const UActorDescContainerInstance* InContainerInstance)
 	{
 		for (UActorDescContainerInstance::TConstIterator<> Iterator(InContainerInstance); Iterator; ++Iterator)
 		{
-			check(!StreamingGeneratorParams.ActorGuidsToContainerInstanceMap.Contains(Iterator->GetGuid()));
-			StreamingGeneratorParams.ActorGuidsToContainerInstanceMap.Add(Iterator->GetGuid(), InContainerInstance);
+			check(!ActorGuidsToContainerInstanceMap.Contains(Iterator->GetGuid()));
+			ActorGuidsToContainerInstanceMap.Add(Iterator->GetGuid(), InContainerInstance);
 		}
 	});
 
-	ForEachActorDescContainerInstance([this, &StreamingGeneratorParams](const UActorDescContainerInstance* InContainerInstance)
+	// Changelist validation can pass Content Bundle containers that are not necessarily registered in the collection's BaseContainerInstance world partition.
+	// Because these containers are validated one at a time, thus represent the base container for the generator's collection,
+	// we need to setup the generator's WorldPartitionContext based on the Content Bundle container.
+	// (Unregistered Content Bundle containers will differ from a registered BaseContainerInstance of the provided collection)
+	auto ValidateCollection = [&InParams, &ActorGuidsToContainerInstanceMap](FStreamingGenerationContainerInstanceCollection& InCollection)
 	{
-		check(StreamingGeneratorParams.WorldPartitionContext == InContainerInstance->GetWorldPartition());
-		FWorldPartitionStreamingGenerator StreamingGenerator(StreamingGeneratorParams);
-		FStreamingGenerationContainerInstanceCollection Collection{ InContainerInstance };
-		StreamingGenerator.PreparationPhase(Collection);
-	});
-}
-
-/* Static version, mainly used by changelist validation */
-void UWorldPartition::CheckForErrors(const FCheckForErrorsParams& Params)
-{
-	check(Params.ErrorHandler);
-	check(Params.ActorDescContainerInstanceCollection);
-
-	FActorDescList ModifiedActorDescList;
+		TErrorHandlerSelector<FStreamingGenerationLogErrorHandler> ErrorHandlerSelector(InParams.ErrorHandler);
+		const UActorDescContainerInstance* BaseContainerInstance = InCollection.GetBaseContainerInstance();
+		const UWorldPartition* WorldPartition = BaseContainerInstance->GetWorldPartition();
+		const TObjectPtr<UWorldPartitionRuntimeHash>& WorldPartitionRuntimeHash = WorldPartition ? WorldPartition->RuntimeHash : nullptr;
 
-	TArray<const UActorDescContainerInstance*> Containers;
-	Params.ActorDescContainerInstanceCollection->ForEachActorDescContainerInstance([&Containers](UActorDescContainerInstance* InContainerInstance) { Containers.Add(InContainerInstance); });
+		FWorldPartitionStreamingGenerator::FWorldPartitionStreamingGeneratorParams StreamingGeneratorParams = FWorldPartitionStreamingGenerator::FWorldPartitionStreamingGeneratorParams()
+			.SetWorldPartitionContext(WorldPartition)
+			.SetHandleUnsavedActors(!!WorldPartition)
+			.SetErrorHandler(ErrorHandlerSelector.Get())
+			.SetIsValidGrid([WorldPartitionRuntimeHash](FName GridName) { return WorldPartitionRuntimeHash ? WorldPartitionRuntimeHash->IsValidGrid(GridName) : true; })
+			.SetIsValidHLODLayer([WorldPartitionRuntimeHash](FName GridName, const FSoftObjectPath& HLODLayerPath) { return WorldPartitionRuntimeHash ? WorldPartitionRuntimeHash->IsValidHLODLayer(GridName, HLODLayerPath) : true; })
+			.SetEnableStreaming(InParams.bEnableStreaming)
+			.SetActorGuidsToContainerInstanceMap(ActorGuidsToContainerInstanceMap);
 
-	FStreamingGenerationContainerInstanceCollection Collection{ Containers };
-	const UActorDescContainerInstance* MainContainerInstance = Collection.GetMainContainer();
+		FWorldPartitionStreamingGenerator StreamingGenerator(StreamingGeneratorParams);
+		StreamingGenerator.PreparationPhase(InCollection);
+	};
 
-	TErrorHandlerSelector<FStreamingGenerationLogErrorHandler> ErrorHandlerSelector(Params.ErrorHandler);
+	// @todo_ow : Once content bundles are remove, we will only do this validation in 1 pass
+	TArray<const UActorDescContainerInstance*> ContentBundleContainerInstances;
+	TArray<const UActorDescContainerInstance*> BaseAndEDLContainerInstances;
+	ExtractContentBundleContainerInstances(InParams.ActorDescContainerInstanceCollection, ContentBundleContainerInstances, BaseAndEDLContainerInstances);
 
-	FWorldPartitionStreamingGenerator::FWorldPartitionStreamingGeneratorParams StreamingGeneratorParams = FWorldPartitionStreamingGenerator::FWorldPartitionStreamingGeneratorParams()
-		.SetWorldPartitionContext(MainContainerInstance->GetWorldPartition())
-		.SetHandleUnsavedActors(MainContainerInstance->HasWorldPartition())
-		.SetErrorHandler(ErrorHandlerSelector.Get())
-		.SetIsValidGrid([](FName) { return true; })
-		.SetIsValidHLODLayer([](FName, const FSoftObjectPath&) { return true; })
-		.SetEnableStreaming(Params.bEnableStreaming)
-		.SetActorGuidsToContainerInstanceMap(Params.ActorGuidsToContainerInstanceMap);
+	if (!BaseAndEDLContainerInstances.IsEmpty())
+	{
+		FStreamingGenerationContainerInstanceCollection Collection(BaseAndEDLContainerInstances, FStreamingGenerationContainerInstanceCollection::ECollectionType::BaseAndEDLs);
+		ValidateCollection(Collection);
+	}
 
-	FWorldPartitionStreamingGenerator StreamingGenerator(StreamingGeneratorParams);
-	StreamingGenerator.PreparationPhase(Collection);
+	for (const UActorDescContainerInstance* ContentBundleContainer : ContentBundleContainerInstances)
+	{
+		FStreamingGenerationContainerInstanceCollection Collection({ ContentBundleContainer }, FStreamingGenerationContainerInstanceCollection::ECollectionType::BaseAsContentBundle);
+		ValidateCollection(Collection);
+	}
 }
+
 #endif // WITH_EDITOR
 
 #undef LOCTEXT_NAMESPACE
diff --git a/Engine/Source/Runtime/Engine/Private/WorldPartition/WorldPartitionStreamingGenerationContext.cpp b/Engine/Source/Runtime/Engine/Private/WorldPartition/WorldPartitionStreamingGenerationContext.cpp
--- a/Engine/Source/Runtime/Engine/Private/WorldPartition/WorldPartitionStreamingGenerationContext.cpp
+++ b/Engine/Source/Runtime/Engine/Private/WorldPartition/WorldPartitionStreamingGenerationContext.cpp
@@ -2,6 +2,8 @@
 
 #include "WorldPartition/WorldPartitionStreamingGenerationContext.h"
 #include "WorldPartition/WorldPartitionStreamingGeneration.h"
+#include "WorldPartition/DataLayer/ExternalDataLayerInstance.h"
+#include "WorldPartition/DataLayer/ExternalDataLayerAsset.h"
 
 #if WITH_EDITOR
 
@@ -24,4 +26,15 @@ const FBox IStreamingGenerationContext::FActorInstance::GetBounds() const
 {
 	return ActorSetInstance->Bounds;
 }
+
+const UExternalDataLayerAsset* IStreamingGenerationContext::FActorSetInstance::GetExternalDataLayerAsset() const
+{
+	auto IsAnExternalDataLayerPred = [](const UDataLayerInstance* DataLayerInstance) { return DataLayerInstance->IsA<UExternalDataLayerInstance>(); };
+	if (const UDataLayerInstance* const* ExternalDataLayerInstance = DataLayers.FindByPredicate(IsAnExternalDataLayerPred))
+	{
+		return CastChecked<UExternalDataLayerInstance>(*ExternalDataLayerInstance)->GetExternalDataLayerAsset();
+	}
+
+	return nullptr;
+}
 #endif
diff --git a/Engine/Source/Runtime/Engine/Public/ActorEditorContext/ScopedActorEditorContextSetExternalDataLayerAsset.h b/Engine/Source/Runtime/Engine/Public/ActorEditorContext/ScopedActorEditorContextSetExternalDataLayerAsset.h
new file mode 100644
--- /dev/null
+++ b/Engine/Source/Runtime/Engine/Public/ActorEditorContext/ScopedActorEditorContextSetExternalDataLayerAsset.h
@@ -0,0 +1,21 @@
+// Copyright Epic Games, Inc. All Rights Reserved.
+
+#pragma once
+
+#if WITH_EDITOR
+
+#include "CoreFwd.h"
+
+class UExternalDataLayerAsset;
+
+/**
+ * Pushes a copy of the existing context and overrides the current External Data Layer Asset.
+ */
+class ENGINE_API FScopedActorEditorContextSetExternalDataLayerAsset
+{
+public:
+	FScopedActorEditorContextSetExternalDataLayerAsset(const UExternalDataLayerAsset* InExternalDataLayerAsset);
+	~FScopedActorEditorContextSetExternalDataLayerAsset();
+};
+
+#endif
\ No newline at end of file
diff --git a/Engine/Source/Runtime/Engine/Public/ExternalPackageHelper.h b/Engine/Source/Runtime/Engine/Public/ExternalPackageHelper.h
--- a/Engine/Source/Runtime/Engine/Public/ExternalPackageHelper.h
+++ b/Engine/Source/Runtime/Engine/Public/ExternalPackageHelper.h
@@ -14,6 +14,10 @@
 #include "Misc/PackageName.h"
 #include "Misc/PackagePath.h"
 #include "Delegates/DelegateCombinations.h"
+#include "WorldPartition/DataLayer/DataLayerInstanceProviderInterface.h"
+#include "WorldPartition/DataLayer/ExternalDataLayerInstance.h"
+#include "WorldPartition/DataLayer/ExternalDataLayerAsset.h"
+#include "WorldPartition/DataLayer/ExternalDataLayerHelper.h"
 
 class FExternalPackageHelper
 {
@@ -29,7 +33,7 @@ public:
 	 * @param InFlags the package flags to apply
 	 * @return the created package
 	 */
-	static ENGINE_API UPackage* CreateExternalPackage(UObject* InObjectOuter, const FString& InObjectPath, EPackageFlags InFlags = FExternalPackageHelper::GetDefaultExternalPackageFlags());
+	static ENGINE_API UPackage* CreateExternalPackage(UObject* InObjectOuter, const FString& InObjectPath, EPackageFlags InFlags = FExternalPackageHelper::GetDefaultExternalPackageFlags(), const UExternalDataLayerAsset* InExternalDataLayerAsset = nullptr);
 
 	/** Returns default external package flags used to create external packages. */
 	static ENGINE_API EPackageFlags GetDefaultExternalPackageFlags();
@@ -107,7 +111,11 @@ void FExternalPackageHelper::LoadObjectsFromExternalPackages(UObject* InOuter, T
 {
 	check(InOuter);
 	UPackage* OutermostPackage = InOuter->IsA<UPackage>() ? CastChecked<UPackage>(InOuter) : InOuter->GetOutermostObject()->GetPackage();
-	const FString ExternalObjectsPath = FExternalPackageHelper::GetExternalObjectsPath(OutermostPackage, FString(), /*bTryUsingPackageLoadedPath*/ true);
+	const FString OutermostPackageName = !OutermostPackage->GetLoadedPath().IsEmpty() ? OutermostPackage->GetLoadedPath().GetPackageName() : OutermostPackage->GetName();
+	const IDataLayerInstanceProvider* DataLayerInstanceProvider = Cast<IDataLayerInstanceProvider>(InOuter);
+	const UExternalDataLayerAsset* ExternalDataLayerAsset = DataLayerInstanceProvider ? DataLayerInstanceProvider->GetRootExternalDataLayerAsset() : nullptr;
+	const FString RootPath = ExternalDataLayerAsset ? FExternalDataLayerHelper::GetExternalDataLayerLevelRootPath(ExternalDataLayerAsset, OutermostPackageName) : OutermostPackageName;
+	const FString ExternalObjectsPath = FExternalPackageHelper::GetExternalObjectsPath(RootPath);
 	TArray<FString> ObjectPackageNames;
 
 	// Do a synchronous scan of the world external objects path.			
diff --git a/Engine/Source/Runtime/Engine/Public/LevelInstance/LevelInstanceEditorPivotActor.h b/Engine/Source/Runtime/Engine/Public/LevelInstance/LevelInstanceEditorPivotActor.h
--- a/Engine/Source/Runtime/Engine/Public/LevelInstance/LevelInstanceEditorPivotActor.h
+++ b/Engine/Source/Runtime/Engine/Public/LevelInstance/LevelInstanceEditorPivotActor.h
@@ -22,6 +22,7 @@ public:
 	ENGINE_API virtual void PostEditUndo() override;
 	virtual bool ShouldExport() override { return false; }
 private:
-	virtual bool IsDataLayerTypeSupported(TSubclassOf<UDataLayerInstance> DataLayerType) const override { return false; }
+	virtual bool ActorTypeSupportsDataLayer() const override { return false; }
+	virtual bool ActorTypeSupportsExternalDataLayer() const override { return false; }
 #endif
 };
diff --git a/Engine/Source/Runtime/Engine/Public/LocationVolume.h b/Engine/Source/Runtime/Engine/Public/LocationVolume.h
--- a/Engine/Source/Runtime/Engine/Public/LocationVolume.h
+++ b/Engine/Source/Runtime/Engine/Public/LocationVolume.h
@@ -60,7 +60,7 @@ public:
 
 #if WITH_EDITOR
 private:
-	virtual bool IsDataLayerTypeSupported(TSubclassOf<UDataLayerInstance> DataLayerType) const override { return false; }
+	virtual bool ActorTypeSupportsDataLayer() const override { return false; }
 	FLoaderAdapterActor* WorldPartitionActorLoader;
 #endif
 };
diff --git a/Engine/Source/Runtime/Engine/Public/WorldPartition/ActorDescContainer.h b/Engine/Source/Runtime/Engine/Public/WorldPartition/ActorDescContainer.h
--- a/Engine/Source/Runtime/Engine/Public/WorldPartition/ActorDescContainer.h
+++ b/Engine/Source/Runtime/Engine/Public/WorldPartition/ActorDescContainer.h
@@ -50,6 +50,12 @@ public:
 		/* The unique name for the container : defaults to PackageName */
 		FString ContainerName;
 
+		/** The associated Content Bundle Guid */
+		FGuid ContentBundleGuid;
+
+		/** The associated External Data Layer Asset */
+		const UExternalDataLayerAsset* ExternalDataLayerAsset = nullptr;
+
 		/* Custom filter function used to filter actors descriptors. */
 		TUniqueFunction<bool(const FWorldPartitionActorDesc*)> FilterActorDesc;
 	};
@@ -67,10 +73,13 @@ public:
 	FName GetContainerPackage() const { return ContainerPackageName; }
 	void SetContainerPackage(const FName& InContainerPackageName) { ContainerPackageName = InContainerPackageName; }
 
+	const UExternalDataLayerAsset* GetExternalDataLayerAsset() const { return ExternalDataLayerAsset; }
+	bool HasExternalContent() const;
+
 	FGuid GetContentBundleGuid() const { return ContentBundleGuid; }
-	void SetContentBundleGuid(const FGuid& InGetContentBundleGuid) { ContentBundleGuid = InGetContentBundleGuid; }
 
 	ENGINE_API FString GetExternalActorPath() const;
+	ENGINE_API FString GetExternalObjectPath() const;
 
 	/** Removes an actor desc without the need to load a package */
 	ENGINE_API bool RemoveActor(const FGuid& ActorGuid);
@@ -160,5 +169,7 @@ private:
 	ENGINE_API bool ShouldRegisterDelegates();
 	ENGINE_API void RegisterEditorDelegates();
 	ENGINE_API void UnregisterEditorDelegates();
+
+	TObjectPtr<const UExternalDataLayerAsset> ExternalDataLayerAsset;
 #endif
 };
diff --git a/Engine/Source/Runtime/Engine/Public/WorldPartition/ActorDescContainerInstance.h b/Engine/Source/Runtime/Engine/Public/WorldPartition/ActorDescContainerInstance.h
--- a/Engine/Source/Runtime/Engine/Public/WorldPartition/ActorDescContainerInstance.h
+++ b/Engine/Source/Runtime/Engine/Public/WorldPartition/ActorDescContainerInstance.h
@@ -13,6 +13,7 @@
 class UActorDescContainer;
 class FWorldPartitionActorDesc;
 class UWorldPartition;
+class UExternalDataLayerAsset;
 struct FWorldPartitionRuntimeCellPropertyOverride;
 
 class FActorDescInstanceList : public TActorDescList<FWorldPartitionActorDescInstance> { };
@@ -60,6 +61,12 @@ public:
 		FName ContainerPackageName;
 
 		FGuid ContainerActorGuid;
+
+		/** The associated Content Bundle Guid */
+		FGuid ContentBundleGuid;
+
+		/** The associated External Data Layer Asset */
+		const UExternalDataLayerAsset* ExternalDataLayerAsset = nullptr;
 		
 		const UActorDescContainerInstance* ParentContainerInstance = nullptr;
 
@@ -69,6 +76,9 @@ public:
 
 		/* Custom filter function used to filter actors descriptors. */
 		TUniqueFunction<bool(const FWorldPartitionActorDesc*)> FilterActorDescFunc;
+
+		/** Function called when container is initialized. */
+		TUniqueFunction<void(UActorDescContainerInstance*)> OnInitializedFunc;
 	};
 
 
@@ -100,7 +110,10 @@ public:
 	ENGINE_API static FName GetContainerPackageNameFromWorld(UWorld* InWorld);
 	ENGINE_API FName GetContainerPackage() const;
 	ENGINE_API FGuid GetContentBundleGuid() const;
+	ENGINE_API const UExternalDataLayerAsset* GetExternalDataLayerAsset() const;
+	ENGINE_API bool HasExternalContent() const;
 	ENGINE_API FString GetExternalActorPath() const;
+	ENGINE_API FString GetExternalObjectPath() const;
 		
 	ENGINE_API TUniquePtr<FWorldPartitionActorDescInstance>* GetActorDescInstancePtr(const FGuid& InActorGuid) const;
 	ENGINE_API FWorldPartitionActorDescInstance* GetActorDescInstance(const FGuid& InActorGuid) const;
@@ -137,7 +150,7 @@ private:
 	bool ShouldRegisterDelegates() const;
 	
 	FWorldPartitionActorDescInstance* AddActor(FWorldPartitionActorDesc* InActorDesc);
-	void RemoveActor(const FGuid& InActorGuid);
+	ENGINE_API void RemoveActor(const FGuid& InActorGuid);
 
 	FWorldPartitionActorDescInstance* AddActorDescInstance(FWorldPartitionActorDescInstance&& InActorDescInstance);
 	void RemoveActorDescInstance(TUniquePtr<FWorldPartitionActorDescInstance>* InActorDescInstance);
@@ -151,6 +164,8 @@ private:
 
 	void OnRegisterChildContainerInstance(const FGuid& InActorGuid, UActorDescContainerInstance* InChildContainerInstance);
 	void OnUnregisterChildContainerInstance(const FGuid& InActorGuid);
+
+	friend class UGameFeatureActionConvertContentBundleWorldPartitionBuilder;
 #endif
 
 private:
diff --git a/Engine/Source/Runtime/Engine/Public/WorldPartition/ContentBundle/ContentBundleEditorSubsystemInterface.h b/Engine/Source/Runtime/Engine/Public/WorldPartition/ContentBundle/ContentBundleEditorSubsystemInterface.h
--- a/Engine/Source/Runtime/Engine/Public/WorldPartition/ContentBundle/ContentBundleEditorSubsystemInterface.h
+++ b/Engine/Source/Runtime/Engine/Public/WorldPartition/ContentBundle/ContentBundleEditorSubsystemInterface.h
@@ -24,7 +24,9 @@ public:
 	virtual bool IsEditingContentBundle() const = 0;
 	virtual bool ActivateContentBundleEditing(TSharedPtr<FContentBundleEditor>& ContentBundleEditor) const = 0;
 	virtual bool DeactivateContentBundleEditing(TSharedPtr<FContentBundleEditor>& ContentBundleEditor) const = 0;
-	virtual void PushContentBundleEditing() = 0;
+	virtual bool DeactivateCurrentContentBundleEditing() const = 0;
+	virtual void PushContentBundleEditing() { PushContentBundleEditing(false); }
+	virtual void PushContentBundleEditing(bool bDuplicateContext) = 0;
 	virtual void PopContentBundleEditing() = 0;
 
 	DECLARE_EVENT_OneParam(UContentBundleEditorSubsystem, FOnContentBundleAdded, const FContentBundleEditor*);
diff --git a/Engine/Source/Runtime/Engine/Public/WorldPartition/Cook/WorldPartitionCookPackage.h b/Engine/Source/Runtime/Engine/Public/WorldPartition/Cook/WorldPartitionCookPackage.h
--- a/Engine/Source/Runtime/Engine/Public/WorldPartition/Cook/WorldPartitionCookPackage.h
+++ b/Engine/Source/Runtime/Engine/Public/WorldPartition/Cook/WorldPartitionCookPackage.h
@@ -2,15 +2,17 @@
 #pragma once
 
 #if WITH_EDITOR
-
 #include "CoreMinimal.h"
 #include "Hash/CityHash.h"
 #include "Misc/Paths.h"
 #include "WorldPartition/WorldPartitionLog.h"
 #include "UObject/Package.h"
+#endif
 
 struct FWorldPartitionCookPackage
 {
+#if WITH_EDITOR
+
 	enum class EType
 	{
 		Unknown,
@@ -52,7 +54,7 @@ struct FWorldPartitionCookPackage
 		TStringBuilderWithBuffer<TCHAR, NAME_SIZE> FullPath;
 		FullPath += TEXT("/");
 		FullPath += InRoot;
-		FullPath += TEXT("/_GENERATED_/");
+		FullPath += GeneratedFolder;
 		FullPath += InRelativeFilename;
 
 		return FPaths::RemoveDuplicateSlashes(*FullPath);
@@ -74,6 +76,11 @@ struct FWorldPartitionCookPackage
 	const FString RelativePath;
 	const IDType PackageId;
 	const EType Type;
-};
-
 #endif
+
+	/** Returns "_Generated_" */
+	static FStringView GetGeneratedFolderName() { return FStringView(&GeneratedFolder[1], GeneratedFolder.Len() - 2); }
+
+private:
+	static constexpr FStringView GeneratedFolder = TEXTVIEW("/_Generated_/");
+};
\ No newline at end of file
diff --git a/Engine/Source/Runtime/Engine/Public/WorldPartition/DataLayer/DataLayerAsset.h b/Engine/Source/Runtime/Engine/Public/WorldPartition/DataLayer/DataLayerAsset.h
--- a/Engine/Source/Runtime/Engine/Public/WorldPartition/DataLayer/DataLayerAsset.h
+++ b/Engine/Source/Runtime/Engine/Public/WorldPartition/DataLayer/DataLayerAsset.h
@@ -3,6 +3,7 @@
 #pragma once
 
 #include "CoreMinimal.h"
+#include "Engine/DataAsset.h"
 #include "UObject/ObjectMacros.h"
 #include "UObject/UObjectGlobals.h"
 #include "UObject/Object.h"
@@ -22,33 +23,32 @@ enum class EDataLayerLoadFilter : uint8
 };
 
 UCLASS(BlueprintType, editinlinenew, MinimalAPI)
-class UDataLayerAsset : public UObject
+class UDataLayerAsset : public UDataAsset
 {
 	GENERATED_UCLASS_BODY()
 
 	friend class UDataLayerConversionInfo;
 
 #if WITH_EDITOR
-	//~ Begin UObject Interface
+	//~ Begin UObject interface
 	ENGINE_API virtual void PostLoad() override;
 	ENGINE_API virtual bool CanEditChange(const FProperty* InProperty) const override;
-	//~ End UObject Interface
+	//~ End UObject interface
 
 public:
-	void SetType(EDataLayerType Type) 
-	{ 
-		check(Type == EDataLayerType::Editor || !IsPrivate());
-		DataLayerType = Type; 
-	}
+	virtual bool CanEditDataLayerType() const;
+	ENGINE_API void SetType(EDataLayerType Type);
 	void SetDebugColor(FColor InDebugColor) { DebugColor = InDebugColor; }
 	ENGINE_API bool CanBeReferencedByActor(AActor* InActor) const;
 	static ENGINE_API bool CanBeReferencedByActor(const TSoftObjectPtr<UDataLayerAsset>& InDataLayerAsset, AActor* InActor);
+
+	virtual void OnCreated();
 #endif
 
 	ENGINE_API bool IsPrivate() const;
 
 	UFUNCTION(Category = "Data Layer", BlueprintCallable)
-	EDataLayerType GetType() const { return DataLayerType; }
+	virtual EDataLayerType GetType() const { return DataLayerType; }
 
 	UFUNCTION(Category = "Data Layer|Runtime", BlueprintCallable)
 	bool IsRuntime() const { return !IsPrivate() && DataLayerType == EDataLayerType::Runtime; }
@@ -64,14 +64,15 @@ public:
 	UFUNCTION(Category = "Data Layer|Runtime", BlueprintCallable)
 	bool IsServerOnly() const { return IsRuntime() && LoadFilter == EDataLayerLoadFilter::ServerOnly; }
 
-private:
+protected:
 	/** Whether the Data Layer affects actor runtime loading */
 	UPROPERTY(Category = "Data Layer", EditAnywhere)
 	EDataLayerType DataLayerType;
 
 	UPROPERTY(Category = "Actor Filter", EditAnywhere)
 	bool bSupportsActorFilters;
-		
+
+private:
 	UPROPERTY(Category = "Runtime", EditAnywhere)
 	FColor DebugColor;
 
diff --git a/Engine/Source/Runtime/Engine/Public/WorldPartition/DataLayer/DataLayerInstance.h b/Engine/Source/Runtime/Engine/Public/WorldPartition/DataLayer/DataLayerInstance.h
--- a/Engine/Source/Runtime/Engine/Public/WorldPartition/DataLayer/DataLayerInstance.h
+++ b/Engine/Source/Runtime/Engine/Public/WorldPartition/DataLayer/DataLayerInstance.h
@@ -17,6 +17,7 @@ class FText;
 class IStreamingGenerationErrorHandler;
 class UWorld;
 class UDataLayerAsset;
+class UExternalDataLayerInstance;
 class FDataLayerInstanceDesc;
 
 UENUM(BlueprintType)
@@ -60,72 +61,77 @@ public:
 	{
 		static_assert(!std::is_same<T, UWorld>::value, "Use GetOuterWorld instead");
 		static_assert(!std::is_same<T, ULevel>::value, "Use GetOuterWorld()->PersistentLevel instead");
-		static_assert(!std::is_same<T, AWorldDataLayers>::value, "Use GetOuterWorldDataLayers instead");
+		static_assert(!std::is_same<T, AWorldDataLayers>::value, "Use GetOuterWorldDataLayers or GetDirectOuterWorldDataLayers instead");
 		return Super::GetTypedOuter<T>();
 	}
 
-	ENGINE_API UWorld* GetOuterWorld() const;
+	ENGINE_API virtual UWorld* GetOuterWorld() const;
 	ENGINE_API AWorldDataLayers* GetOuterWorldDataLayers() const;
+	ENGINE_API AWorldDataLayers* GetDirectOuterWorldDataLayers() const;
 
 #if WITH_EDITOR
-	//~ Begin UObject
+	//~ Begin UObject interface
 	ENGINE_API virtual bool IsAsset() const override;
 	ENGINE_API virtual void GetAssetRegistryTags(FAssetRegistryTagsContext Context) const override;
-	UE_DEPRECATED(5.4, "Implement the version that takes FAssetRegistryTagsContext instead.")
-	ENGINE_API virtual void GetAssetRegistryTags(TArray<FAssetRegistryTag>& OutTags) const override;
 	ENGINE_API virtual void PreEditUndo() override;
 	ENGINE_API virtual void PostEditUndo() override;
-	//~ End UObject
+	ENGINE_API virtual bool CanEditChange(const FProperty* InProperty) const override;
+	//~ End UObject interface
 
+	// Getters/Setters
+	ENGINE_API bool SetParent(UDataLayerInstance* InParent);
 	ENGINE_API void SetVisible(bool bIsVisible);
 	ENGINE_API void SetIsInitiallyVisible(bool bIsInitiallyVisible);
 	ENGINE_API void SetIsLoadedInEditor(bool bIsLoadedInEditor, bool bFromUserChange);
 	void SetIsLocked(bool bInIsLocked) { bIsLocked = bInIsLocked; }
 
-	bool IsInitiallyLoadedInEditor() const { return bIsInitiallyLoadedInEditor; }
-	bool IsLoadedInEditor() const { return bIsLoadedInEditor; }
 	ENGINE_API bool IsEffectiveLoadedInEditor() const;
-	bool IsLoadedInEditorChangedByUserOperation() const { return bIsLoadedInEditorChangedByUserOperation; }
-	void ClearLoadedInEditorChangedByUserOperation() { bIsLoadedInEditorChangedByUserOperation = false; }
-
-	ENGINE_API const TCHAR* GetDataLayerIconName() const;
+	ENGINE_API bool IsInitiallyLoadedInEditor() const { return bIsInitiallyLoadedInEditor; }
+	ENGINE_API bool IsLoadedInEditor() const { return bIsLoadedInEditor; }
+	ENGINE_API bool IsLoadedInEditorChangedByUserOperation() const { return bIsLoadedInEditorChangedByUserOperation; }
+	ENGINE_API virtual bool IsReadOnly(FText* OutReason = nullptr) const;
+	virtual bool IsIncludedInActorFilterDefault() const { return false; }
 
+	// Data Layer Instance features support
 	ENGINE_API bool CanBeChildOf(const UDataLayerInstance* InParent, FText* OutReason = nullptr) const;
-	ENGINE_API bool SetParent(UDataLayerInstance* InParent);
-
-	UE_DEPRECATED(5.4, "SetChildParent was removed")
-	ENGINE_API void SetChildParent(UDataLayerInstance* InParent) {}
-
-	static ENGINE_API FText GetDataLayerText(const UDataLayerInstance* InDataLayer);
+	ENGINE_API virtual bool CanUserAddActors(FText* OutReason = nullptr) const;
+	ENGINE_API virtual bool CanUserRemoveActors(FText* OutReason = nullptr) const;
+	ENGINE_API virtual bool CanAddActor(AActor* Actor, FText* OutReason = nullptr) const;
+	ENGINE_API virtual bool CanRemoveActor(AActor* Actor, FText* OutReason = nullptr) const;
+	ENGINE_API virtual bool CanBeInActorEditorContext() const;
+	virtual bool CanBeRemoved() const { return true; }
+	virtual bool CanEditDataLayerShortName() const { return false; }
+	virtual bool SupportsActorFilters() const { return false; }
+	virtual const UExternalDataLayerInstance* GetRootExternalDataLayerInstance() const { return nullptr; }
 
-	ENGINE_API virtual bool IsLocked() const;
-	ENGINE_API virtual bool IsReadOnly() const;
-	ENGINE_API virtual bool CanEditChange(const FProperty* InProperty) const;
-	
-	ENGINE_API virtual bool CanUserAddActors() const;
-	ENGINE_API virtual bool CanAddActor(AActor* Actor) const;
+	// Actor assignation and removal
 	ENGINE_API virtual bool AddActor(AActor* Actor) const;
-	ENGINE_API virtual bool CanUserRemoveActors() const;
-	ENGINE_API virtual bool CanRemoveActor(AActor* Actor) const;
 	ENGINE_API virtual bool RemoveActor(AActor* Actor) const;
 
-	ENGINE_API virtual bool CanBeInActorEditorContext() const;
+	// Actor Editor Context
+	ENGINE_API bool IsActorEditorContextCurrentColorized() const;
 	ENGINE_API bool IsInActorEditorContext() const;
 	ENGINE_API bool AddToActorEditorContext();
 	ENGINE_API bool RemoveFromActorEditorContext();
 
-	virtual bool CanEditDataLayerShortName() const { return false; }
-
-	virtual bool SupportsActorFilters() const { return false; }
-	virtual bool IsIncludedInActorFilterDefault() const { return false; }
-
-	// Whether the DataLayer was created by a user and can be deleted by a user.
-	virtual bool IsUserManaged() const { return true; }
+	// Validation
+	ENGINE_API virtual bool Validate(IStreamingGenerationErrorHandler* ErrorHandler) const;
 
+	// Helpers
+	ENGINE_API virtual const TCHAR* GetDataLayerIconName() const;
+	static ENGINE_API FText GetDataLayerText(const UDataLayerInstance* InDataLayer);
 	static bool GetAssetRegistryInfoFromPackage(FName InDataLayerInstancePackageName, FDataLayerInstanceDesc& OutDataLayerInstanceDesc);
 	static bool GetAssetRegistryInfoFromPackage(const FAssetData& InAsset, FDataLayerInstanceDesc& OutDataLayerInstanceDesc);
 
-	ENGINE_API virtual bool Validate(IStreamingGenerationErrorHandler* ErrorHandler) const;
+	//~Begin deprecation
+	UE_DEPRECATED(5.4, "Implement the version that takes FAssetRegistryTagsContext instead.")
+	ENGINE_API virtual void GetAssetRegistryTags(TArray<FAssetRegistryTag>& OutTags) const override;
+
+	UE_DEPRECATED(5.4, "SetChildParent was removed.")
+	ENGINE_API void SetChildParent(UDataLayerInstance* InParent);
+
+	UE_DEPRECATED(5.4, "Use IsReadOnly instead")
+	ENGINE_API virtual bool IsLocked() const { FText Reason; return IsLocked(&Reason); }
 
 	UE_DEPRECATED(5.3, "Use CanEditShortName instead")
 	virtual bool SupportRelabeling() const { return false; }
@@ -135,6 +141,7 @@ public:
 
 	UE_DEPRECATED(5.3, "Use CanBeChildOf instead")
 	bool CanParent(const UDataLayerInstance* InParent) const { return CanBeChildOf(InParent); }
+	//~End deprecation
 #endif // WITH_EDITOR
 
 	UFUNCTION(Category = "Data Layer", BlueprintCallable)
@@ -170,12 +177,17 @@ public:
 	UFUNCTION(Category = "Data Layer", BlueprintCallable)
 	virtual FString GetDataLayerFullName() const { return TEXT("Invalid Data Layer"); }
 
+	virtual bool CanHaveChildDataLayerInstance(const UDataLayerInstance* InChildDataLayerInstance) const { return InChildDataLayerInstance != nullptr; }
+	virtual bool CanHaveParentDataLayerInstance() const { return true; }
+	UE_DEPRECATED(5.4, "Use CanHaveParentDataLayerInstance instead")
+	virtual bool CanHaveParentDataLayer() const { return CanHaveParentDataLayerInstance(); }
+	UE_DEPRECATED(5.4, "Use CanHaveChildDataLayer instead")
 	virtual bool CanHaveChildDataLayers() const { return true; }
-	virtual bool CanHaveParentDataLayer() const { return true; }
 
 	const UDataLayerInstance* GetParent() const { return Parent; }
 	UDataLayerInstance* GetParent() { return Parent; }
 
+	virtual bool CanEditInitialRuntimeState() const { return IsRuntime(); }
 	ENGINE_API EDataLayerRuntimeState GetRuntimeState() const;
 	ENGINE_API EDataLayerRuntimeState GetEffectiveRuntimeState() const;
 
@@ -186,27 +198,18 @@ public:
 	ENGINE_API void ForEachChild(TFunctionRef<bool(const UDataLayerInstance*)> Operation) const;
 
 #if DATALAYER_TO_INSTANCE_RUNTIME_CONVERSION_ENABLED
+public:
 	virtual FName GetDataLayerFName() const { return GetFName(); }
 
 private:
 	FName GetFName() const { return Super::GetFName(); }
 	FString GetName() const { return Super::GetName();  }
-public:
 #endif // DATALAYER_TO_INSTANCE_RUNTIME_CONVERSION_ENABLED
 
-private:
-	ENGINE_API void AddChild(UDataLayerInstance* DataLayer);
-	friend class UDataLayerManager;
-	friend class FDataLayerUtils;
-
-#if WITH_EDITOR
-	ENGINE_API void OnRemovedFromWorldDataLayers();
-	ENGINE_API void RemoveChild(UDataLayerInstance* DataLayer);
-#endif
-
 protected:
 #if WITH_EDITOR
-	ENGINE_API bool IsParentDataLayerTypeCompatible(const UDataLayerInstance* InParent) const;
+	ENGINE_API bool IsLocked(FText* OutReason) const;
+	ENGINE_API bool IsParentDataLayerTypeCompatible(const UDataLayerInstance* InParent, FText* OutReason = nullptr) const;
 
 	virtual bool PerformAddActor(AActor* InActor) const { return false; }
 	virtual bool PerformRemoveActor(AActor* InActor) const { return false;  }
@@ -245,6 +248,14 @@ protected:
 	EDataLayerRuntimeState InitialRuntimeState;
 
 private:
+
+#if WITH_EDITOR
+	ENGINE_API void OnRemovedFromWorldDataLayers();
+	void ClearLoadedInEditorChangedByUserOperation() { bIsLoadedInEditorChangedByUserOperation = false; }
+	ENGINE_API void RemoveChild(UDataLayerInstance* DataLayer);
+#endif
+	ENGINE_API void AddChild(UDataLayerInstance* DataLayer);
+
 	UPROPERTY()
 	TObjectPtr<UDataLayerInstance> Parent;
 
@@ -252,4 +263,6 @@ private:
 	TArray<TObjectPtr<UDataLayerInstance>> Children;
 
 	friend class AWorldDataLayers;
+	friend class UDataLayerManager;
+	friend class FDataLayerUtils;
 };
diff --git a/Engine/Source/Runtime/Engine/Public/WorldPartition/DataLayer/DataLayerInstanceNames.h b/Engine/Source/Runtime/Engine/Public/WorldPartition/DataLayer/DataLayerInstanceNames.h
new file mode 100644
--- /dev/null
+++ b/Engine/Source/Runtime/Engine/Public/WorldPartition/DataLayer/DataLayerInstanceNames.h
@@ -0,0 +1,109 @@
+// Copyright Epic Games, Inc. All Rights Reserved.
+#pragma once
+
+#include "CoreMinimal.h"
+
+#include "DataLayerInstanceNames.generated.h"
+
+USTRUCT()
+struct FDataLayerInstanceNames
+{
+	GENERATED_USTRUCT_BODY()
+
+	FDataLayerInstanceNames()
+	: bIsFirstDataLayerIsExternal(false)
+	{
+#if WITH_EDITOR
+		bIsForcedEmptyNonExternalDataLayers = false;
+#endif
+	}
+
+	FDataLayerInstanceNames(const TArray<FName>& InDataLayers, bool bInIsFirstDataLayerIsExternal)
+	: bIsFirstDataLayerIsExternal(bInIsFirstDataLayerIsExternal)
+	, DataLayers(InDataLayers)
+	{
+#if WITH_EDITOR
+		bIsForcedEmptyNonExternalDataLayers = false;
+#endif
+	}
+
+	FDataLayerInstanceNames(const TArray<FName>& InNonExternalDataLayers, FName InExternalDataLayer)
+	{
+		bIsFirstDataLayerIsExternal = !InExternalDataLayer.IsNone();
+		if (bIsFirstDataLayerIsExternal)
+		{
+			DataLayers.Add(InExternalDataLayer);
+		}
+		DataLayers.Append(InNonExternalDataLayers);
+	}
+
+	const FName GetExternalDataLayer() const
+	{
+		check(!bIsFirstDataLayerIsExternal || !DataLayers.IsEmpty());
+		return bIsFirstDataLayerIsExternal ? DataLayers[0] : NAME_None;
+	}
+
+	TArrayView<const FName> GetNonExternalDataLayers() const
+	{
+		static TArray<FName> EmptyArray;
+
+#if WITH_EDITOR
+		if (bIsForcedEmptyNonExternalDataLayers)
+		{
+			return EmptyArray;
+		}
+#endif
+		check(!bIsFirstDataLayerIsExternal || !DataLayers.IsEmpty());
+		const int32 Offset = bIsFirstDataLayerIsExternal ? 1 : 0;
+		const int32 NonExternalDataLayersCount = DataLayers.Num() - Offset;
+		return NonExternalDataLayersCount > 0 ? MakeArrayView(&DataLayers[Offset], NonExternalDataLayersCount) : EmptyArray;
+	}
+
+#if WITH_EDITOR
+	TArray<FName> ToArray() const
+	{
+		if (bIsForcedEmptyNonExternalDataLayers)
+		{
+			static TArray<FName> EmptyArray;
+			return HasExternalDataLayer() ? TArray<FName>({ GetExternalDataLayer() }) : EmptyArray;
+		}
+		return DataLayers;
+	}
+
+	int32 Num() const
+	{
+		if (bIsForcedEmptyNonExternalDataLayers)
+		{
+			return HasExternalDataLayer() ? 1 : 0;
+		}
+		return DataLayers.Num();
+	}
+
+	bool IsEmpty() const
+	{
+		return Num() == 0;
+	}
+#endif
+
+	bool HasExternalDataLayer() const { return bIsFirstDataLayerIsExternal; }
+
+	const TArray<FName>& GetRawArray() const
+	{
+		return DataLayers;
+	}
+
+private:
+#if WITH_EDITOR
+	bool IsForcedEmptyNonExternalDataLayers() const { return bIsForcedEmptyNonExternalDataLayers; }
+	void SetIsForcedEmptyNonExternalDataLayers(bool bInNewValue) { bIsForcedEmptyNonExternalDataLayers = bInNewValue; }
+	bool bIsForcedEmptyNonExternalDataLayers;
+#endif
+
+	UPROPERTY()
+	bool bIsFirstDataLayerIsExternal;
+
+	UPROPERTY()
+	TArray<FName> DataLayers;
+
+	friend class FStreamingGenerationActorDescView;
+};
\ No newline at end of file
diff --git a/Engine/Source/Runtime/Engine/Public/WorldPartition/DataLayer/DataLayerInstancePrivate.h b/Engine/Source/Runtime/Engine/Public/WorldPartition/DataLayer/DataLayerInstancePrivate.h
--- a/Engine/Source/Runtime/Engine/Public/WorldPartition/DataLayer/DataLayerInstancePrivate.h
+++ b/Engine/Source/Runtime/Engine/Public/WorldPartition/DataLayer/DataLayerInstancePrivate.h
@@ -13,32 +13,38 @@ class UDataLayerInstancePrivate : public UDataLayerInstance
 
 public:
 #if WITH_EDITOR
-	static ENGINE_API FName MakeName();
+	ENGINE_API static FName MakeName();
 	ENGINE_API void OnCreated();
 
-	ENGINE_API virtual bool CanEditChange(const FProperty* InProperty) const;
-	ENGINE_API virtual bool IsLocked() const override;
-	ENGINE_API virtual bool IsReadOnly() const override;
-		
+	//~ Begin UObject Interface
+	ENGINE_API virtual bool CanEditChange(const FProperty* InProperty) const override;
+	//~ End UObject Interface
+
+	//~ Begin UDataLayerInstance Interface
+	ENGINE_API virtual bool IsReadOnly(FText* OutReason = nullptr) const override;
 	virtual bool SupportsActorFilters() const override { return GetAsset()->SupportsActorFilters(); }
 	virtual bool IsIncludedInActorFilterDefault() const override { return bIsIncludedInActorFilterDefault; }
 	virtual bool CanEditDataLayerShortName() const override { return true; }
-	virtual bool CanAddActor(AActor* InActor) const override;
+	virtual bool CanAddActor(AActor* InActor, FText* OutReason = nullptr) const override;
+	//~ End UDataLayerInstance Interface
 #endif
 
+	//~ Begin UDataLayerInstance Interface
 	virtual const UDataLayerAsset* GetAsset() const override { return DataLayerAsset; }
-
 	virtual EDataLayerType GetType() const override { return DataLayerAsset->GetType(); }
 	virtual bool IsRuntime() const override { return DataLayerAsset->IsRuntime(); }
 	virtual FColor GetDebugColor() const override { return DataLayerAsset->GetDebugColor(); }
-
 	virtual FString GetDataLayerShortName() const override { return ShortName; }
 	virtual FString GetDataLayerFullName() const override { return DataLayerAsset->GetPathName(); }
+	//~ End UDataLayerInstance Interface
+
 protected:
 #if WITH_EDITOR
+	//~ Begin UDataLayerInstance Interface
 	ENGINE_API virtual bool PerformAddActor(AActor* InActor) const override;
 	ENGINE_API virtual bool PerformRemoveActor(AActor* InActor) const override;
 	virtual void PerformSetDataLayerShortName(const FString& InNewShortName) override { ShortName = InNewShortName; }
+	//~ End UDataLayerInstance Interface
 #endif
 
 private:
diff --git a/Engine/Source/Runtime/Engine/Public/WorldPartition/DataLayer/DataLayerInstanceWithAsset.h b/Engine/Source/Runtime/Engine/Public/WorldPartition/DataLayer/DataLayerInstanceWithAsset.h
--- a/Engine/Source/Runtime/Engine/Public/WorldPartition/DataLayer/DataLayerInstanceWithAsset.h
+++ b/Engine/Source/Runtime/Engine/Public/WorldPartition/DataLayer/DataLayerInstanceWithAsset.h
@@ -11,51 +11,55 @@ class UDataLayerInstanceWithAsset : public UDataLayerInstance
 {
 	GENERATED_UCLASS_BODY()
 
-		friend class UDataLayerConversionInfo;
-
 public:
 #if WITH_EDITOR
-	static ENGINE_API FName MakeName(const UDataLayerAsset* DeprecatedDataLayer);
-	static ENGINE_API TSubclassOf<UDataLayerInstanceWithAsset> GetDataLayerInstanceClass();
-	ENGINE_API void OnCreated(const UDataLayerAsset* Asset);
-
+	//~ Begin UObject interface
 	ENGINE_API virtual bool CanEditChange(const FProperty* InProperty) const;
 	ENGINE_API virtual void PreEditUndo() override;
 	ENGINE_API virtual void PostEditUndo() override;
-	ENGINE_API virtual bool IsLocked() const override;
-	ENGINE_API virtual bool IsReadOnly() const override;
-	ENGINE_API virtual bool CanAddActor(AActor* InActor) const override;
-	ENGINE_API virtual bool CanRemoveActor(AActor* InActor) const override;
-
-	ENGINE_API virtual bool Validate(IStreamingGenerationErrorHandler* ErrorHandler) const override;
 	ENGINE_API virtual void PostEditChangeProperty(FPropertyChangedEvent& PropertyChangedEvent) override;
-
-	ENGINE_API bool SupportsActorFilters() const;
-	ENGINE_API bool IsIncludedInActorFilterDefault() const;
+	//~ End UObject interface
+
+	//~ Begin UDataLayerInstance interface
+	ENGINE_API virtual bool IsReadOnly(FText* OutReason = nullptr) const override;
+	ENGINE_API virtual bool CanAddActor(AActor* InActor, FText* OutReason = nullptr) const override;
+	ENGINE_API virtual bool CanRemoveActor(AActor* InActor, FText* OutReason = nullptr) const override;
+	ENGINE_API virtual bool CanBeInActorEditorContext() const override;
+	ENGINE_API virtual bool SupportsActorFilters() const override;
+	ENGINE_API virtual bool IsIncludedInActorFilterDefault() const override;
+	ENGINE_API virtual bool Validate(IStreamingGenerationErrorHandler* ErrorHandler) const override;
+	ENGINE_API virtual const UExternalDataLayerInstance* GetRootExternalDataLayerInstance() const override;
+	//~ End UDataLayerInstance interface
+
+	ENGINE_API static const UDataLayerAsset* GetDataLayerAsset(const UDataLayerAsset* Asset) { return Asset; } // Used when creating a DataLayerInstance
+	ENGINE_API static FName MakeName(const UDataLayerAsset* InDataLayerAsset);
+	ENGINE_API static TSubclassOf<UDataLayerInstanceWithAsset> GetDataLayerInstanceClass();
+	ENGINE_API virtual void OnCreated(const UDataLayerAsset* Asset);
+	virtual bool CanEditDataLayerAsset() const { return true; }
 #endif
 
-	const UDataLayerAsset* GetAsset() const override { return DataLayerAsset; }
-
+	//~ Begin UDataLayerInstance interface
+	ENGINE_API virtual UWorld* GetOuterWorld() const override;
+	virtual const UDataLayerAsset* GetAsset() const override { return DataLayerAsset; }
 	virtual EDataLayerType GetType() const override { return DataLayerAsset ? DataLayerAsset->GetType() : EDataLayerType::Unknown; }
-
 	virtual bool IsRuntime() const override { return DataLayerAsset && DataLayerAsset->IsRuntime(); }
-
 	virtual bool IsClientOnly() const override { return DataLayerAsset && DataLayerAsset->IsClientOnly(); }
 	virtual bool IsServerOnly() const override { return DataLayerAsset && DataLayerAsset->IsServerOnly(); }
-
 	virtual FColor GetDebugColor() const override { return DataLayerAsset ? DataLayerAsset->GetDebugColor() : FColor::Black; }
-
 	virtual FString GetDataLayerShortName() const override { return DataLayerAsset ? DataLayerAsset->GetName() : GetDataLayerFName().ToString(); }
 	virtual FString GetDataLayerFullName() const override { return DataLayerAsset ? DataLayerAsset->GetPathName() : GetDataLayerFName().ToString(); }
+	//~ End UDataLayerInstance interface
 
 protected:
 #if WITH_EDITOR
 	ENGINE_API virtual bool PerformAddActor(AActor* InActor) const;
 	ENGINE_API virtual bool PerformRemoveActor(AActor* InActor) const;
+
+	bool bSkipCheckReadOnlyForSubLevels;
 #endif
 
 private:
-	UPROPERTY(Category = "Data Layer", EditAnywhere)
+	UPROPERTY(Category = "Data Layer", EditAnywhere, meta = (DisallowedClasses = "/Script/Engine.ExternalDataLayerAsset"))
 	TObjectPtr<const UDataLayerAsset> DataLayerAsset;
 
 #if WITH_EDITORONLY_DATA
@@ -67,4 +71,7 @@ private:
 	// Used to compare state pre/post undo
 	TObjectPtr<const UDataLayerAsset> CachedDataLayerAsset;
 #endif
+
+	friend class UDataLayerConversionInfo;
+	friend class ULevelInstanceSubsystem;
 };
diff --git a/Engine/Source/Runtime/Engine/Public/WorldPartition/DataLayer/DataLayerManager.h b/Engine/Source/Runtime/Engine/Public/WorldPartition/DataLayer/DataLayerManager.h
--- a/Engine/Source/Runtime/Engine/Public/WorldPartition/DataLayer/DataLayerManager.h
+++ b/Engine/Source/Runtime/Engine/Public/WorldPartition/DataLayer/DataLayerManager.h
@@ -9,6 +9,7 @@
 #include "WorldPartition/DataLayer/DataLayerInstance.h"
 #include "WorldPartition/DataLayer/DataLayerEditorContext.h"
 #include "WorldPartition/DataLayer/WorldDataLayers.h"
+#include "WorldPartition/DataLayer/ExternalDataLayerInstance.h"
 #include "Templates/SubclassOf.h"
 #include "DataLayerManager.generated.h"
 
@@ -16,6 +17,8 @@ class UDataLayerAsset;
 class UDataLayerInstanceWithAsset;
 class UDataLayerInstance;
 class UDataLayerLoadingPolicy;
+class UExternalDataLayerAsset;
+class ExternalDataLayerInstance;
 class UActorDescContainer;
 class UActorDescContainerInstance;
 class UCanvas;
@@ -94,8 +97,8 @@ public:
 	//~ Begin Runtime State
 	ENGINE_API const TSet<FName>& GetEffectiveActiveDataLayerNames() const;
 	ENGINE_API const TSet<FName>& GetEffectiveLoadedDataLayerNames() const;
-	ENGINE_API bool IsAnyDataLayerInEffectiveRuntimeState(const TArray<FName>& InDataLayerNames, EDataLayerRuntimeState InState) const;
-	ENGINE_API bool IsAllDataLayerInEffectiveRuntimeState(const TArray<FName>& InDataLayerNames, EDataLayerRuntimeState InState) const;
+	ENGINE_API bool IsAnyDataLayerInEffectiveRuntimeState(TArrayView<const FName> InDataLayerNames, EDataLayerRuntimeState InState) const;
+	ENGINE_API bool IsAllDataLayerInEffectiveRuntimeState(TArrayView<const FName> InDataLayerNames, EDataLayerRuntimeState InState) const;
 	//~ End Runtime State
 
 private:
@@ -142,6 +145,8 @@ private:
 	friend class AWorldDataLayers;
 	friend class UDataLayerSubsystem;
 	friend class UDataLayerInstanceWithAsset;
+	friend class UExternalDataLayerManager;
+	friend class UExternalDataLayerInstance;
 	friend class UWorldPartitionStreamingPolicy;
 	friend class UWorldPartitionRuntimeCell;
 	friend class UWorldPartitionSubsystem;
@@ -151,7 +156,7 @@ private:
 #if WITH_EDITOR
 private:
 	//~ Begin Editor Context
-	ENGINE_API void PushActorEditorContext() const;
+	ENGINE_API void PushActorEditorContext(bool bDuplicateContext) const;
 	ENGINE_API void PopActorEditorContext() const;
 	ENGINE_API TArray<UDataLayerInstance*> GetActorEditorContextDataLayers() const;
 	ENGINE_API TArray<AWorldDataLayers*> GetActorEditorContextWorldDataLayers() const;
@@ -166,6 +171,7 @@ private:
 	void ResolveActorDescContainerInstanceDataLayers(UActorDescContainerInstance* InActorDescContainerInstance) const;
 	void ResolveActorDescInstanceDataLayers(FWorldPartitionActorDescInstance* InActorDescInstance) const;
 	void ResolveActorDescContainerInstanceDataLayersInternal(UActorDescContainerInstance* InActorDescContainerInstance, FWorldPartitionActorDescInstance* InActorDescInstance) const;
+	ENGINE_API static FWorldPartitionReference LoadWorldDataLayersActor(UActorDescContainerInstance* InActorDescContainerInstance);
 
 	//~ End
 
diff --git a/Engine/Source/Runtime/Engine/Public/WorldPartition/DataLayer/DataLayerUtils.h b/Engine/Source/Runtime/Engine/Public/WorldPartition/DataLayer/DataLayerUtils.h
--- a/Engine/Source/Runtime/Engine/Public/WorldPartition/DataLayer/DataLayerUtils.h
+++ b/Engine/Source/Runtime/Engine/Public/WorldPartition/DataLayer/DataLayerUtils.h
@@ -7,6 +7,7 @@
 #include "UObject/UObjectGlobals.h"
 #include "WorldPartition/DataLayer/DataLayerType.h"
 #include "WorldPartition/DataLayer/DataLayerInstance.h"
+#include "WorldPartition/DataLayer/DataLayerInstanceNames.h"
 #include "WorldPartition/WorldPartitionStreamingGeneration.h"
 
 class UDataLayerManager;
@@ -26,9 +27,11 @@ public:
 		return IconNameByType[static_cast<uint32>(DataLayerType)];
 	}
 
-	static ENGINE_API TArray<FName> ResolvedDataLayerInstanceNames(const UDataLayerManager* InDataLayerManager, const FWorldPartitionActorDesc* InActorDesc, const TArray<const FWorldDataLayersActorDesc*>& InWorldDataLayersActorDescs = TArray<const FWorldDataLayersActorDesc*>());
-
-	static ENGINE_API bool ResolveRuntimeDataLayerInstanceNames(const UDataLayerManager* InDataLayerManager, const IWorldPartitionActorDescInstanceView& InActorDescView, const FStreamingGenerationActorDescViewMap& ActorDescViewMap, TArray<FName>& OutRuntimeDataLayerInstanceNames);
+	UE_DEPRECATED(5.4, "Use ResolveDataLayerInstanceNames instead")
+	static ENGINE_API TArray<FName> ResolvedDataLayerInstanceNames(const UDataLayerManager* InDataLayerManager, const FWorldPartitionActorDesc* InActorDesc, const TArray<const FWorldDataLayersActorDesc*>& InWorldDataLayersActorDescs = TArray<const FWorldDataLayersActorDesc*>())
+	{
+		return ResolveDataLayerInstanceNames(InDataLayerManager, InActorDesc, InWorldDataLayersActorDescs).ToArray();
+	}
 
 	PRAGMA_DISABLE_DEPRECATION_WARNINGS
 
@@ -37,6 +40,10 @@ public:
 
 	PRAGMA_ENABLE_DEPRECATION_WARNINGS
 
+	static ENGINE_API FDataLayerInstanceNames ResolveDataLayerInstanceNames(const UDataLayerManager* InDataLayerManager, const FWorldPartitionActorDesc* InActorDesc, const TArray<const FWorldDataLayersActorDesc*>& InWorldDataLayersActorDescs = TArray<const FWorldDataLayersActorDesc*>());
+
+	static ENGINE_API bool ResolveRuntimeDataLayerInstanceNames(const UDataLayerManager* InDataLayerManager, const IWorldPartitionActorDescInstanceView& InActorDescView, const FStreamingGenerationActorDescViewMap& ActorDescViewMap, FDataLayerInstanceNames& OutRuntimeDataLayerInstanceNames);
+
 	static ENGINE_API const FDataLayerInstanceDesc* GetDataLayerInstanceDescFromInstanceName(const TArray<const FWorldDataLayersActorDesc*>& InWorldDataLayersActorDescs, const FName& DataLayerInstanceName);
 
 	static ENGINE_API const FDataLayerInstanceDesc* GetDataLayerInstanceDescFromAssetPath(const TArray<const FWorldDataLayersActorDesc*>& InWorldDataLayersActorDescs, const FName& DataLayerAssetPath);
diff --git a/Engine/Source/Runtime/Engine/Public/WorldPartition/DataLayer/DeprecatedDataLayerInstance.h b/Engine/Source/Runtime/Engine/Public/WorldPartition/DataLayer/DeprecatedDataLayerInstance.h
--- a/Engine/Source/Runtime/Engine/Public/WorldPartition/DataLayer/DeprecatedDataLayerInstance.h
+++ b/Engine/Source/Runtime/Engine/Public/WorldPartition/DataLayer/DeprecatedDataLayerInstance.h
@@ -17,11 +17,6 @@ class UDeprecatedDataLayerInstance final : public UDataLayerInstance
 
 	GENERATED_UCLASS_BODY()
 
-	friend class AWorldDataLayers;
-	friend class UDataLayerConversionInfo;
-	friend class UDataLayerToAssetCommandletContext;
-	friend class UDataLayerToAssetCommandlet;
-
 #if WITH_EDITOR
 	//~ Begin UObject Interface
 	ENGINE_API virtual void PostLoad() override;
@@ -34,8 +29,8 @@ public:
 	ENGINE_API void OnCreated();
 	FActorDataLayer GetActorDataLayer() const;
 
-	ENGINE_API PRAGMA_DISABLE_DEPRECATION_WARNINGS
-	static FName MakeName(const UDEPRECATED_DataLayer* DeprecatedDataLayer);
+	PRAGMA_DISABLE_DEPRECATION_WARNINGS
+	ENGINE_API static FName MakeName(const UDEPRECATED_DataLayer* DeprecatedDataLayer);
 	ENGINE_API void OnCreated(const UDEPRECATED_DataLayer* DeprecatedDataLayer);
 	PRAGMA_ENABLE_DEPRECATION_WARNINGS
 #endif
@@ -44,25 +39,24 @@ public:
 
 private:
 #if WITH_EDITOR
+	//~ Begin UDataLayerInstance Interface
 	ENGINE_API virtual bool PerformAddActor(AActor* InActor) const override;
 	ENGINE_API virtual bool PerformRemoveActor(AActor* InActor) const override;
-
-	virtual bool SupportRelabeling() const override { return true; }
 	ENGINE_API virtual bool RelabelDataLayer(FName NewDataLayerLabel) override;
+	virtual bool SupportRelabeling() const override { return true; }
 	virtual bool CanEditDataLayerShortName() const override { return true; }
 	virtual void PerformSetDataLayerShortName(const FString& InNewShortName) { Label = *InNewShortName; }
+	//~ Endif UDataLayerInstance Interface
 #endif
 
+	//~ Begin UDataLayerInstance Interface
 	virtual FName GetDataLayerFName() const override { return !DeprecatedDataLayerFName.IsNone() ? DeprecatedDataLayerFName : Super::GetDataLayerFName(); }
-
 	virtual EDataLayerType GetType() const override { return DataLayerType; }
-
 	virtual bool IsRuntime() const override { return DataLayerType == EDataLayerType::Runtime; }
-
 	virtual FColor GetDebugColor() const override { return DebugColor; }
-
 	virtual FString GetDataLayerShortName() const override { return Label.ToString(); }
 	virtual FString GetDataLayerFullName() const override { return TEXT("Deprecated_") + Label.ToString(); }
+	//~ End UDataLayerInstance Interface
 
 private:
 	UPROPERTY()
@@ -76,4 +70,9 @@ private:
 
 	UPROPERTY(Category = "Data Layer|Runtime", EditAnywhere, meta = (EditConditionHides, EditCondition = "DataLayerType == EDataLayerType::Runtime"))
 	FColor DebugColor;
+
+	friend class AWorldDataLayers;
+	friend class UDataLayerConversionInfo;
+	friend class UDataLayerToAssetCommandletContext;
+	friend class UDataLayerToAssetCommandlet;
 };
diff --git a/Engine/Source/Runtime/Engine/Public/WorldPartition/DataLayer/ExternalDataLayerAsset.h b/Engine/Source/Runtime/Engine/Public/WorldPartition/DataLayer/ExternalDataLayerAsset.h
new file mode 100644
--- /dev/null
+++ b/Engine/Source/Runtime/Engine/Public/WorldPartition/DataLayer/ExternalDataLayerAsset.h
@@ -0,0 +1,44 @@
+// Copyright Epic Games, Inc. All Rights Reserved.
+
+#pragma once
+
+#include "WorldPartition/DataLayer/DataLayerAsset.h"
+#include "WorldPartition/DataLayer/ExternalDataLayerUID.h"
+
+#include "ExternalDataLayerAsset.generated.h"
+
+struct FAssetData;
+
+UCLASS(BlueprintType, editinlinenew)
+class ENGINE_API UExternalDataLayerAsset : public UDataLayerAsset
+{
+	GENERATED_UCLASS_BODY()
+
+public:
+#if WITH_EDITOR
+	static constexpr FColor EditorUXColor = FColor(255, 167, 26);
+#endif
+
+#if WITH_EDITOR
+	//~ Begin UObject Interface
+	virtual void PostDuplicate(bool bDuplicateForPIE) override;
+	virtual void GetAssetRegistryTags(FAssetRegistryTagsContext Context) const override;
+	//~ End UObject Interface
+
+	static bool GetAssetRegistryInfoFromPackage(const FAssetData& InAsset, FExternalDataLayerUID& OutExternalDataLayerUID);
+#endif
+
+	//~ Begin UDataLayerAsset Interface
+#if WITH_EDITOR
+	virtual void OnCreated() override;
+	virtual bool CanEditDataLayerType() const override { return false; }
+#endif
+	virtual EDataLayerType GetType() const override { return EDataLayerType::Runtime; }
+	//~ End UDataLayerAsset Interface
+
+	const FExternalDataLayerUID& GetUID() const { return UID; }
+
+private:
+	UPROPERTY(VisibleAnywhere, Category = "Data Layer", AdvancedDisplay, DuplicateTransient, meta = (DisplayName = "External Data Layer UID"))
+	FExternalDataLayerUID UID;
+};
\ No newline at end of file
diff --git a/Engine/Source/Runtime/Engine/Public/WorldPartition/DataLayer/ExternalDataLayerInjectionPolicy.h b/Engine/Source/Runtime/Engine/Public/WorldPartition/DataLayer/ExternalDataLayerInjectionPolicy.h
new file mode 100644
--- /dev/null
+++ b/Engine/Source/Runtime/Engine/Public/WorldPartition/DataLayer/ExternalDataLayerInjectionPolicy.h
@@ -0,0 +1,21 @@
+// Copyright Epic Games, Inc. All Rights Reserved.
+#pragma once
+
+#include "CoreFwd.h"
+#include "UObject/ObjectMacros.h"
+
+#include "ExternalDataLayerInjectionPolicy.generated.h"
+
+class UWorld;
+class UExternalDataLayerAsset;
+
+UCLASS()
+class ENGINE_API UExternalDataLayerInjectionPolicy : public UObject
+{
+	GENERATED_BODY()
+
+#if WITH_EDITOR
+public:
+	virtual bool CanInject(const UWorld* InWorld, const UExternalDataLayerAsset* InExternalDataLayerAsset, const UObject* InClient, FText* OutFailureReason = nullptr) const;
+#endif
+};
