External Data Layer V1

diff --git a/Engine/Source/Runtime/Engine/Public/WorldPartition/DataLayer/ExternalDataLayerManager.h b/Engine/Source/Runtime/Engine/Public/WorldPartition/DataLayer/ExternalDataLayerManager.h
new file mode 100644
--- /dev/null
+++ b/Engine/Source/Runtime/Engine/Public/WorldPartition/DataLayer/ExternalDataLayerManager.h
@@ -0,0 +1,129 @@
+// Copyright Epic Games, Inc. All Rights Reserved.
+#pragma once
+
+#include "CoreFwd.h"
+#include "UObject/ObjectMacros.h"
+#include "UObject/Object.h"
+#include "WorldPartition/WorldPartition.h"
+#include "WorldPartition/WorldPartitionHandle.h"
+#include "WorldPartition/WorldPartitionHelpers.h"
+#include "WorldPartition/DataLayer/WorldDataLayersActorDesc.h"
+#include "WorldPartition/DataLayer/ExternalDataLayerEngineSubsystem.h"
+
+#include "ExternalDataLayerManager.generated.h"
+
+class UActorDescContainerInstance;
+class UDataLayerManager;
+class UExternalDataLayerAsset;
+class UExternalDataLayerInstance;
+class UWorldPartitionRuntimeCell;
+class URuntimeHashExternalStreamingObjectBase;
+class AWorldDataLayers;
+
+UCLASS(Within = WorldPartition)
+class ENGINE_API UExternalDataLayerManager : public UObject
+{
+	GENERATED_BODY()
+
+public:
+	template <class T>
+	static UExternalDataLayerManager* GetExternalDataLayerManager(const T* InObject)
+	{
+		UWorldPartition* WorldPartition = IsValid(InObject) ? FWorldPartitionHelpers::GetWorldPartition(InObject) : nullptr;
+		return WorldPartition ? WorldPartition->GetExternalDataLayerManager() : nullptr;
+	}
+
+private:
+	//~ Begin Initialization/Deinitialization
+	UExternalDataLayerManager();
+	void Initialize();
+	bool IsInitialized() const { return bIsInitialized; }
+	void DeInitialize();
+	//~ End Initialization/Deinitialization
+
+	//~ Begin UObject Interface
+	virtual void Serialize(FArchive& Ar) override;
+	//~ End UObject Interface
+
+	//~ Begin Injection/Removal
+	void UpdateExternalDataLayerInjectionState(const UExternalDataLayerAsset* InExternalDataLayerAsset);
+	bool CanInjectExternalDataLayerAsset(const UExternalDataLayerAsset* InExternalDataLayerAsset, FText* OutReason = nullptr) const;
+	bool IsExternalDataLayerInjected(const UExternalDataLayerAsset* InExternalDataLayerAsset) const;
+	bool InjectExternalDataLayer(const UExternalDataLayerAsset* InExternalDataLayerAsset);
+	bool RemoveExternalDataLayer(const UExternalDataLayerAsset* InExternalDataLayerAsset);
+	bool InjectIntoGameWorld(const UExternalDataLayerAsset* InExternalDataLayerAsset);
+	bool RemoveFromGameWorld(const UExternalDataLayerAsset* InExternalDataLayerAsset);
+	bool RegisterExternalStreamingObjectForGameWorld(const UExternalDataLayerAsset* InExternalDataLayerAsset);
+	bool UnregisterExternalStreamingObjectForGameWorld(const UExternalDataLayerAsset* InExternalDataLayerAsset);
+	//~ End Injection/Removal
+
+	// Monitor when EDL state changes
+	void OnExternalDataLayerAssetRegistrationStateChanged(const UExternalDataLayerAsset* InExternalDataLayerAsset, EExternalDataLayerRegistrationState InOldState, EExternalDataLayerRegistrationState InNewState);
+
+	FString GetExternalDataLayerLevelRootPath(const UExternalDataLayerAsset* InExternalDataLayerAsset) const;
+	FString GetExternalStreamingObjectPackagePath(const UExternalDataLayerAsset* InExternalDataLayerAsset) const;
+	const UExternalDataLayerInstance* GetExternalDataLayerInstance(const UExternalDataLayerAsset* InExternalDataLayerAsset) const;
+	UExternalDataLayerInstance* GetExternalDataLayerInstance(const UExternalDataLayerAsset* InExternalDataLayerAsset);
+	UDataLayerManager& GetDataLayerManager() const;
+
+#if WITH_EDITOR
+	//~ Begin UObject Interface
+	virtual void PreEditUndo() override;
+	virtual void PostEditUndo() override;
+	//~ End UObject Interface
+
+	// Used in editor
+	UActorDescContainerInstance* RegisterExternalDataLayerActorDescContainer(const UExternalDataLayerAsset* InExternalDataLayerAsset);
+	bool UnregisterExternalDataLayerActorDescContainer(const UExternalDataLayerAsset* InExternalDataLayerAsset);
+	bool OnActorExternalDataLayerAssetChanged(AActor* InActor);
+	bool RegisterExternalDataLayerInstance(UExternalDataLayerInstance* InExternalDataLayerInstance);
+	bool UnregisterExternalDataLayerInstance(UExternalDataLayerInstance* InExternalDataLayerInstance);
+	const UExternalDataLayerAsset* GetMatchingExternalDataLayerAssetForObjectPath(const FSoftObjectPath& InObjectPath);
+	const UExternalDataLayerAsset* GetActorEditorContextCurrentExternalDataLayer() const;
+	AWorldDataLayers* GetWorldDataLayers(const UExternalDataLayerAsset* InExternalDataLayerAsset, bool bInAllowCreate = false) const;
+	bool OnActorPreSpawnInitialization(AActor* InActor, const UExternalDataLayerAsset* InExternalDataLayerAsset);
+	FString GetActorPackageName(const UExternalDataLayerAsset* InExternalDataLayerAsset, const ULevel* InDestinationLevel, const FString& ActorName) const;
+	bool SetupActorPackageForExternalDataLayerAsset(AActor* InActor, const UExternalDataLayerAsset* InExternalDataLayerAsset);
+	URuntimeHashExternalStreamingObjectBase* CreateExternalStreamingObjectUsingStreamingGeneration(const UExternalDataLayerAsset* InExternalDataLayerAsset);
+
+	// Used for PIE/-game
+	void OnBeginPlay();
+	void OnEndPlay();
+
+	//~ Begin Cooking
+	UWorldPartitionRuntimeCell* GetCellForCookPackage(const FString& InCookPackageName) const;
+	URuntimeHashExternalStreamingObjectBase* GetExternalStreamingObjectForCookPackage(const FString& InCookPackageName) const;
+	void ForEachExternalStreamingObjects(TFunctionRef<bool(URuntimeHashExternalStreamingObjectBase*)> Func) const;
+	//~ End Cooking
+#endif
+
+private:
+	bool IsRunningGameOrInstancedWorldPartition() const { return bIsRunningGameOrInstancedWorldPartition; }
+
+	bool bIsInitialized;
+	bool bIsRunningGameOrInstancedWorldPartition;
+
+	UPROPERTY(Transient)
+	TMap<TObjectPtr<const UExternalDataLayerAsset>, TObjectPtr<URuntimeHashExternalStreamingObjectBase>> ExternalStreamingObjects;
+
+#if WITH_EDITORONLY_DATA
+	UPROPERTY(Transient)
+	TSet<TObjectPtr<const UExternalDataLayerAsset>> InjectedExternalDataLayerAssets;
+#endif
+
+#if WITH_EDITOR
+	using FExternalDataLayerContainerMap = TMap<TObjectPtr<const UExternalDataLayerAsset>, TObjectPtr<UActorDescContainerInstance>>;
+	FExternalDataLayerContainerMap EDLContainerMap;
+	TMap<TObjectPtr<const UExternalDataLayerAsset>, FWorldPartitionReference> EDLWorldDataLayersMap;
+	TSet<TObjectPtr<const UExternalDataLayerAsset>> PreEditUndoExternalDataLayerAssets;
+#endif
+
+	friend class FDataLayerMode;
+	friend class AWorldDataLayers;
+	friend class UWorldPartition;
+	friend class UExternalDataLayerInstance;
+	friend class UDataLayerEditorSubsystem;
+	friend class ULevelInstanceSubsystem;
+	friend class UWorldPartitionRuntimeLevelStreamingCell;
+	friend class UGameFeatureActionConvertContentBundleWorldPartitionBuilder;
+};
\ No newline at end of file
diff --git a/Engine/Source/Runtime/Engine/Public/WorldPartition/DataLayer/ExternalDataLayerUID.h b/Engine/Source/Runtime/Engine/Public/WorldPartition/DataLayer/ExternalDataLayerUID.h
new file mode 100644
--- /dev/null
+++ b/Engine/Source/Runtime/Engine/Public/WorldPartition/DataLayer/ExternalDataLayerUID.h
@@ -0,0 +1,33 @@
+// Copyright Epic Games, Inc. All Rights Reserved.
+
+#pragma once
+
+#include "ExternalDataLayerUID.generated.h"
+
+struct FAssetData;
+
+USTRUCT()
+struct ENGINE_API FExternalDataLayerUID
+{
+	GENERATED_BODY()
+
+public:
+	bool IsValid() const;
+	operator uint32() const { return Value; }
+
+	FString ToString() const;
+#if WITH_EDITOR
+	static bool Parse(const FString& InUIDString, FExternalDataLayerUID& OutUID);
+#endif
+
+private:
+#if WITH_EDITOR
+	static FExternalDataLayerUID NewUID();
+#endif
+
+	UPROPERTY(VisibleAnywhere, Category = "External Data Layer", AdvancedDisplay)
+	uint32 Value = 0;
+
+	friend uint32 GetTypeHash(const FExternalDataLayerUID& InUID) { return InUID.Value; }
+	friend class UExternalDataLayerAsset;
+};
\ No newline at end of file
diff --git a/Engine/Source/Runtime/Engine/Public/WorldPartition/DataLayer/IDataLayerEditorModule.h b/Engine/Source/Runtime/Engine/Public/WorldPartition/DataLayer/IDataLayerEditorModule.h
--- a/Engine/Source/Runtime/Engine/Public/WorldPartition/DataLayer/IDataLayerEditorModule.h
+++ b/Engine/Source/Runtime/Engine/Public/WorldPartition/DataLayer/IDataLayerEditorModule.h
@@ -6,6 +6,7 @@
 
 class AActor;
 class UDataLayerInstance;
+class UExternalDataLayerAsset;
 
 /**
  * The module holding all of the UI related pieces for DataLayer management
@@ -15,4 +16,5 @@ class IDataLayerEditorModule : public IModuleInterface
 public:
 	virtual ~IDataLayerEditorModule() {}
 	virtual bool AddActorToDataLayers(AActor* Actor, const TArray<UDataLayerInstance*>& DataLayers) = 0;
+	virtual void SetActorEditorContextCurrentExternalDataLayer(const UExternalDataLayerAsset* InExternalDataLayerAsset) = 0;
 };
\ No newline at end of file
diff --git a/Engine/Source/Runtime/Engine/Public/WorldPartition/DataLayer/WorldDataLayers.h b/Engine/Source/Runtime/Engine/Public/WorldPartition/DataLayer/WorldDataLayers.h
--- a/Engine/Source/Runtime/Engine/Public/WorldPartition/DataLayer/WorldDataLayers.h
+++ b/Engine/Source/Runtime/Engine/Public/WorldPartition/DataLayer/WorldDataLayers.h
@@ -8,6 +8,8 @@
 #include "WorldPartition/DataLayer/ActorDataLayer.h"
 #include "WorldPartition/DataLayer/DataLayerInstance.h"
 #include "WorldPartition/DataLayer/DataLayer.h"
+#include "WorldPartition/DataLayer/DataLayerInstanceProviderInterface.h"
+#include "WorldPartition/DataLayer/ExternalDataLayerAsset.h"
 #include "Net/Core/PushModel/PushModel.h"
 #include "Engine/World.h"
 #include "Engine/Level.h"
@@ -18,6 +20,7 @@ class UDEPRECATED_DataLayer;
 class UDataLayerInstance;
 class UDataLayerInstanceWithAsset;
 class UDataLayerAsset;
+class UExternalDataLayerInstance;
 
 USTRUCT()
 struct FActorPlacementDataLayers
@@ -27,12 +30,20 @@ struct FActorPlacementDataLayers
 	UPROPERTY()
 	TArray<FName> DataLayerInstanceNames;
 
+	UPROPERTY()
+	FName ExternalDataLayerName;
+
+	UPROPERTY()
+	FName CurrentColorizedDataLayerInstanceName;
+
 	UPROPERTY()
 	int32 ContextID = INT32_MAX;
 
 	void Reset()
 	{
 		DataLayerInstanceNames.Reset();
+		ExternalDataLayerName = NAME_None;
+		CurrentColorizedDataLayerInstanceName = NAME_None;
 	}
 };
 
@@ -40,12 +51,10 @@ struct FActorPlacementDataLayers
  * Actor containing data layers instances within a world.
  */
 UCLASS(hidecategories = (Actor, HLOD, Cooking, Transform, Advanced, Display, Events, Object, Physics, Attachment, Info, Input, Blueprint, Layers, Tags, Replication), notplaceable, MinimalAPI)
-class AWorldDataLayers : public AInfo
+class AWorldDataLayers : public AInfo, public IDataLayerInstanceProvider
 {
 	GENERATED_UCLASS_BODY()
 
-	friend class UDataLayerToAssetCommandlet;
-
 public:
 	ENGINE_API virtual void PostLoad() override;
 	ENGINE_API virtual void RewindForReplay() override;
@@ -55,11 +64,13 @@ public:
 #if WITH_EDITOR
 	ENGINE_API virtual void PreEditUndo() override;
 	ENGINE_API virtual void PostEditUndo() override;
-	virtual bool ShouldLevelKeepRefIfExternal() const override { return true; }
+	virtual bool ShouldLevelKeepRefIfExternal() const override;
+	virtual bool IsEditorOnly() const override;
 	virtual bool ShouldImport(FStringView ActorPropString, bool IsMovingLevel) override { return false; }
 	virtual bool IsLockLocation() const { return true; }
 	virtual bool IsUserManaged() const override { return false; }
-	virtual bool IsDataLayerTypeSupported(TSubclassOf<UDataLayerInstance> DataLayerType) const { return false; }
+	virtual bool ActorTypeSupportsDataLayer() const override { return false; }
+	virtual bool ActorTypeSupportsExternalDataLayer() const override { return false; }
 	ENGINE_API virtual TUniquePtr<class FWorldPartitionActorDesc> CreateClassActorDesc() const override;
 
 	static ENGINE_API AWorldDataLayers* Create(UWorld* World, FName InWorldDataLayerName = NAME_None);
@@ -76,10 +87,11 @@ public:
 	ENGINE_API void SetAllowRuntimeDataLayerEditing(bool bInAllowRuntimeDataLayerEditing);
 	bool GetAllowRuntimeDataLayerEditing() const { return bAllowRuntimeDataLayerEditing; }
 
+	ENGINE_API bool IsActorEditorContextCurrentColorized(const UDataLayerInstance* InDataLayerInstance) const;
 	ENGINE_API bool IsInActorEditorContext(const UDataLayerInstance* InDataLayerInstance) const;
 	ENGINE_API bool AddToActorEditorContext(UDataLayerInstance* InDataLayerInstance);
 	ENGINE_API bool RemoveFromActorEditorContext(UDataLayerInstance* InDataLayerInstance);
-	ENGINE_API void PushActorEditorContext(int32 InContextID);
+	ENGINE_API void PushActorEditorContext(int32 InContextID, bool bDuplicateContext);
 	ENGINE_API void PopActorEditorContext(int32 InContextID);
 	ENGINE_API TArray<UDataLayerInstance*> GetActorEditorContextDataLayers() const;
 
@@ -93,9 +105,8 @@ public:
 	template<class T>
 	void OverwriteDataLayerRuntimeStates(const TArray<T>* InActiveDataLayers, const TArray<T>* InLoadedDataLayers );
 
-	ENGINE_API bool IsReadOnly() const;
+	ENGINE_API bool IsReadOnly(FText* OutReason = nullptr) const;
 	ENGINE_API bool IsSubWorldDataLayers() const;
-	ENGINE_API bool IsRuntimeRelevant() const;
 
 	ENGINE_API bool SupportsExternalPackageDataLayerInstances() const;
 	ENGINE_API bool IsUsingExternalPackageDataLayerInstances() const { return bUseExternalPackageDataLayerInstances; }
@@ -120,6 +131,8 @@ public:
 	ENGINE_API void ForEachDataLayerInstance(TFunctionRef<bool(UDataLayerInstance*)> Func) const;
 
 	ENGINE_API TArray<const UDataLayerInstance*> GetDataLayerInstances(const TArray<FName>& InDataLayerInstanceNames) const;
+	ENGINE_API const UExternalDataLayerInstance* GetExternalDataLayerInstance(const UExternalDataLayerAsset* InExternalDataLayerAsset) const;
+	ENGINE_API bool IsExternalDataLayerWorldDataLayers() const;
 
 	// DataLayer Runtime State
 	ENGINE_API void SetDataLayerRuntimeState(const UDataLayerInstance* InDataLayerInstance, EDataLayerRuntimeState InState, bool bIsRecursive = false);
@@ -191,22 +204,38 @@ protected:
 	ENGINE_API void OnRep_EffectiveLoadedDataLayerNames();
 
 private:
+	// External Data Layers
+	bool AddExternalDataLayerInstance(UExternalDataLayerInstance* ExternalDataLayerInstance);
+	bool RemoveExternalDataLayerInstance(UExternalDataLayerInstance* ExternalDataLayerInstance);
+	ENGINE_API UExternalDataLayerInstance* GetExternalDataLayerInstance(const UExternalDataLayerAsset* InExternalDataLayerAsset);
+
 	ENGINE_API void OnDataLayerManagerInitialized();
 	ENGINE_API void OnDataLayerManagerDeinitialized();
 	ENGINE_API void ResolveEffectiveRuntimeState(const UDataLayerInstance* InDataLayer, bool bInNotifyChange = true);
 	ENGINE_API void DumpDataLayerRecursively(const UDataLayerInstance* DataLayer, FString Prefix, FOutputDevice& OutputDevice) const;
-	ENGINE_API TSet<TObjectPtr<UDataLayerInstance>>& GetDataLayerInstances();
-	ENGINE_API const TSet<TObjectPtr<UDataLayerInstance>>& GetDataLayerInstances() const;
+
+	//~ Begin IDataLayerInstanceProvider interface
+	ENGINE_API virtual TSet<TObjectPtr<UDataLayerInstance>>& GetDataLayerInstances() override;
+	virtual const TSet<TObjectPtr<UDataLayerInstance>>& GetDataLayerInstances() const override { return const_cast<AWorldDataLayers*>(this)->GetDataLayerInstances(); }
+	virtual const UExternalDataLayerInstance* GetRootExternalDataLayerInstance() const override { return RootExternalDataLayerInstance; }
+	//~ End IDataLayerInstanceProvider interface
+
+#if !WITH_EDITOR
+	void UpdateAccelerationTable(const UDataLayerInstance* DataLayerInstance, bool bIsAdding);
+#endif
 
 #if WITH_EDITOR
 	virtual bool ActorTypeIsMainWorldOnly() const override { return true; }
 
 	ENGINE_API void AddDataLayerInstance(UDataLayerInstance* InDataLayerInstance);
-	ENGINE_API void ConvertDataLayerToInstancces();
+	ENGINE_API void ConvertDataLayerToInstances();
 	ENGINE_API void UpdateContainsDeprecatedDataLayers();
 	ENGINE_API void ResolveActorDescContainers();
+	ENGINE_API void RemoveEditorDataLayers();
+	void UpdateCurrentColorizedDataLayerInstance();
 
 	friend class UDataLayerInstanceWithAsset;
+	friend class UDataLayerToAssetCommandlet;
 
 	// Used to compare state pre/post undo
 	TSet<TObjectPtr<UDataLayerInstance>> CachedDataLayerInstances;
@@ -231,6 +260,9 @@ private:
 	TArray<FActorPlacementDataLayers> CurrentDataLayersStack;
 #endif
 
+	UPROPERTY()
+	TObjectPtr<UExternalDataLayerInstance> RootExternalDataLayerInstance;
+
 	UPROPERTY()
 	TSet<TObjectPtr<UDataLayerInstance>> DataLayerInstances;
 
@@ -242,6 +274,9 @@ private:
 	UPROPERTY(Transient)
 	TArray<TObjectPtr<UDataLayerInstance>> LoadedExternalPackageDataLayerInstances;
 
+	UPROPERTY(Transient)
+	TArray<TObjectPtr<UDataLayerInstance>> TransientDataLayerInstances;
+
 	static_assert(DATALAYER_TO_INSTANCE_RUNTIME_CONVERSION_ENABLED, "DeprecatedDataLayerNameToDataLayerInstance Property is deprecated and needs to be deleted.");
 	UPROPERTY()
 	TMap<FName, TWeakObjectPtr<UDataLayerInstance>> DeprecatedDataLayerNameToDataLayerInstance;
@@ -287,7 +322,10 @@ private:
 	static_assert(DATALAYER_TO_INSTANCE_RUNTIME_CONVERSION_ENABLED, "bHasDeprecatedDataLayers is deprecated and needs to be deleted.");
 	bool bHasDeprecatedDataLayers;
 
+	friend class UWorldPartition;
 	friend class UDataLayerManager;
+	friend class UExternalDataLayerManager;
+	friend class UDataLayerEditorSubsystem;
 
 public:
 	DECLARE_DELEGATE_RetVal_ThreeParams(bool, FDataLayersFilterDelegate, FName /*DataLayerName*/, EDataLayerRuntimeState /*CurrentState*/, EDataLayerRuntimeState /*TargetState*/);
@@ -330,7 +368,8 @@ DataLayerInstanceType* AWorldDataLayers::CreateDataLayer(CreationsArgs... InCrea
 		GloballyUniqueObjectPath += OuterObject->GetPathName();
 		GloballyUniqueObjectPath += TEXT(".");
 		GloballyUniqueObjectPath += ShortName;
-		ExternalPackage = FExternalPackageHelper::CreateExternalPackage(OuterObject, *GloballyUniqueObjectPath);
+
+		ExternalPackage = FExternalPackageHelper::CreateExternalPackage(OuterObject, *GloballyUniqueObjectPath, FExternalPackageHelper::GetDefaultExternalPackageFlags(), GetRootExternalDataLayerAsset());
 		check(ExternalPackage);
 		NewObjectName = *ShortName;
 	}
diff --git a/Engine/Source/Runtime/Engine/Public/WorldPartition/DataLayer/WorldDataLayersActorDesc.h b/Engine/Source/Runtime/Engine/Public/WorldPartition/DataLayer/WorldDataLayersActorDesc.h
--- a/Engine/Source/Runtime/Engine/Public/WorldPartition/DataLayer/WorldDataLayersActorDesc.h
+++ b/Engine/Source/Runtime/Engine/Public/WorldPartition/DataLayer/WorldDataLayersActorDesc.h
@@ -76,6 +76,8 @@ protected:
 	virtual uint32 GetSizeOf() const override { return sizeof(FWorldDataLayersActorDesc); }
 	ENGINE_API virtual void Serialize(FArchive& Ar) override;
 	virtual bool IsResaveNeeded() const override { return !IsValid(); }
+	virtual bool IsRuntimeRelevant(const FWorldPartitionActorDescInstance* InActorDescInstance) const override;
+	virtual void OnUnloadingInstance(const FWorldPartitionActorDescInstance* InActorDescInstance) const override;
 	//~ End FWorldPartitionActorDesc Interface.
 
 private:
@@ -86,6 +88,7 @@ private:
 	TArray<FDataLayerInstanceDesc> DataLayerInstances;
 	mutable TOptional<TArray<FDataLayerInstanceDesc>> ExternalPackageDataLayerInstances;
 	bool bIsValid;
+	bool bIsExternalDataLayerWorldDataLayers;
 	bool bUseExternalPackageDataLayerInstances;
 };
 #endif
diff --git a/Engine/Source/Runtime/Engine/Public/WorldPartition/ErrorHandling/WorldPartitionStreamingGenerationErrorHandler.h b/Engine/Source/Runtime/Engine/Public/WorldPartition/ErrorHandling/WorldPartitionStreamingGenerationErrorHandler.h
--- a/Engine/Source/Runtime/Engine/Public/WorldPartition/ErrorHandling/WorldPartitionStreamingGenerationErrorHandler.h
+++ b/Engine/Source/Runtime/Engine/Public/WorldPartition/ErrorHandling/WorldPartitionStreamingGenerationErrorHandler.h
@@ -5,6 +5,7 @@
 #include "CoreMinimal.h"
 #include "WorldPartition/WorldPartitionActorDescInstanceViewInterface.h"
 
+class UDataLayerAsset;
 class UDataLayerInstance;
 class UDataLayerInstanceWithAsset;
 
@@ -32,9 +33,18 @@ public:
 	virtual void OnInvalidReferenceGridPlacement(const IWorldPartitionActorDescInstanceView& ActorDescView, const IWorldPartitionActorDescInstanceView& ReferenceActorDescView) = 0;
 
 	/**
-	 * Called when an actor references an actor using a different set of data layers.
+	 * Used to identify an actor data layer error
 	 */
-	virtual void OnInvalidReferenceDataLayers(const IWorldPartitionActorDescInstanceView& ActorDescView, const IWorldPartitionActorDescInstanceView& ReferenceActorDescView) = 0;
+	enum class EDataLayerInvalidReason
+	{
+		ReferencedActorDifferentRuntimeDataLayers,
+		ReferencedActorDifferentExternalDataLayer
+	};
+
+	/**
+	 * Called when there's an error with the data layers used by an actor and one or more referenced actor.
+	 */
+	virtual void OnInvalidReferenceDataLayers(const IWorldPartitionActorDescInstanceView& ActorDescView, const IWorldPartitionActorDescInstanceView& ReferenceActorDescView, EDataLayerInvalidReason Reason) = 0;
 
 	/**
 	 * Called when an actor references an actor using a different RuntimeGrid.
@@ -66,6 +76,11 @@ public:
 	 */
 	virtual void OnDataLayerAssetConflict(const UDataLayerInstanceWithAsset* DataLayerInstance, const UDataLayerInstanceWithAsset* ConflictingDataLayerInstance) = 0;
 
+	/**
+	 * Called when the data layer asset is not compatible with its data layer asset.
+	 */
+	virtual void OnInvalidDataLayerAssetType(const UDataLayerInstanceWithAsset* DataLayerInstance, const UDataLayerAsset* DataLayerAsset) = 0;
+
 	/**
 	 * Called when an actor needs to be resaved.
 	 */
@@ -73,7 +88,7 @@ public:
 
 
 	/**
-	 * Called when a level instance actor has errrors
+	 * Used to identify a level instance actor error
 	 */
 	enum class ELevelInstanceInvalidReason
 	{
@@ -102,6 +117,9 @@ public:
 	
 	PRAGMA_DISABLE_DEPRECATION_WARNINGS
 
+	UE_DEPRECATED(5.4, "Use OnInvalidReferenceDataLayers with EDataLayerInvalidReason instead.")
+	virtual void OnInvalidReferenceDataLayers(const IWorldPartitionActorDescInstanceView& ActorDescView, const IWorldPartitionActorDescInstanceView& ReferenceActorDescView) final {}
+
 	UE_DEPRECATED(5.2, "OnInvalidReference is deprecated, use the version which takes an optional actor descriptor view.")
 	virtual void OnInvalidReference(const FWorldPartitionActorDescView& ActorDescView, const FGuid& ReferenceGuid) {}
 
diff --git a/Engine/Source/Runtime/Engine/Public/WorldPartition/ErrorHandling/WorldPartitionStreamingGenerationLogErrorHandler.h b/Engine/Source/Runtime/Engine/Public/WorldPartition/ErrorHandling/WorldPartitionStreamingGenerationLogErrorHandler.h
--- a/Engine/Source/Runtime/Engine/Public/WorldPartition/ErrorHandling/WorldPartitionStreamingGenerationLogErrorHandler.h
+++ b/Engine/Source/Runtime/Engine/Public/WorldPartition/ErrorHandling/WorldPartitionStreamingGenerationLogErrorHandler.h
@@ -10,11 +10,12 @@ public:
 	ENGINE_API virtual void OnInvalidRuntimeGrid(const IWorldPartitionActorDescInstanceView& ActorDescView, FName GridName) override;
 	ENGINE_API virtual void OnInvalidReference(const IWorldPartitionActorDescInstanceView& ActorDescView, const FGuid& ReferenceGuid, IWorldPartitionActorDescInstanceView* ReferenceActorDescView) override;
 	ENGINE_API virtual void OnInvalidReferenceGridPlacement(const IWorldPartitionActorDescInstanceView& ActorDescView, const IWorldPartitionActorDescInstanceView& ReferenceActorDescView) override;
-	ENGINE_API virtual void OnInvalidReferenceDataLayers(const IWorldPartitionActorDescInstanceView& ActorDescView, const IWorldPartitionActorDescInstanceView& ReferenceActorDescView) override;
+	ENGINE_API virtual void OnInvalidReferenceDataLayers(const IWorldPartitionActorDescInstanceView& ActorDescView, const IWorldPartitionActorDescInstanceView& ReferenceActorDescView, EDataLayerInvalidReason Reason) override;
 	ENGINE_API virtual void OnInvalidReferenceLevelScriptStreamed(const IWorldPartitionActorDescInstanceView& ActorDescView) override;
 	ENGINE_API virtual void OnInvalidReferenceLevelScriptDataLayers(const IWorldPartitionActorDescInstanceView& ActorDescView) override;
 	ENGINE_API virtual void OnInvalidReferenceRuntimeGrid(const IWorldPartitionActorDescInstanceView& ActorDescView, const IWorldPartitionActorDescInstanceView& ReferenceActorDescView) override;
 	ENGINE_API virtual void OnInvalidReferenceDataLayerAsset(const UDataLayerInstanceWithAsset* DataLayerInstance) override;
+	ENGINE_API virtual void OnInvalidDataLayerAssetType(const UDataLayerInstanceWithAsset* DataLayerInstance, const UDataLayerAsset* DataLayerAsset) override;
 	ENGINE_API virtual void OnDataLayerHierarchyTypeMismatch(const UDataLayerInstance* DataLayerInstance, const UDataLayerInstance* Parent) override;
 	ENGINE_API virtual void OnDataLayerAssetConflict(const UDataLayerInstanceWithAsset* DataLayerInstance, const UDataLayerInstanceWithAsset* ConflictingDataLayerInstance) override;
 	ENGINE_API virtual void OnActorNeedsResave(const IWorldPartitionActorDescInstanceView& ActorDescView) override;
diff --git a/Engine/Source/Runtime/Engine/Public/WorldPartition/ErrorHandling/WorldPartitionStreamingGenerationNullErrorHandler.h b/Engine/Source/Runtime/Engine/Public/WorldPartition/ErrorHandling/WorldPartitionStreamingGenerationNullErrorHandler.h
--- a/Engine/Source/Runtime/Engine/Public/WorldPartition/ErrorHandling/WorldPartitionStreamingGenerationNullErrorHandler.h
+++ b/Engine/Source/Runtime/Engine/Public/WorldPartition/ErrorHandling/WorldPartitionStreamingGenerationNullErrorHandler.h
@@ -10,11 +10,12 @@ public:
 	virtual void OnInvalidRuntimeGrid(const IWorldPartitionActorDescInstanceView& ActorDescView, FName GridName) override {}
 	virtual void OnInvalidReference(const IWorldPartitionActorDescInstanceView& ActorDescView, const FGuid& ReferenceGuid, IWorldPartitionActorDescInstanceView* ReferenceActorDescView) override {}
 	virtual void OnInvalidReferenceGridPlacement(const IWorldPartitionActorDescInstanceView& ActorDescView, const IWorldPartitionActorDescInstanceView& ReferenceActorDescView) override {}
-	virtual void OnInvalidReferenceDataLayers(const IWorldPartitionActorDescInstanceView& ActorDescView, const IWorldPartitionActorDescInstanceView& ReferenceActorDescView) override {}
+	virtual void OnInvalidReferenceDataLayers(const IWorldPartitionActorDescInstanceView& ActorDescView, const IWorldPartitionActorDescInstanceView& ReferenceActorDescView, EDataLayerInvalidReason Reason) override {}
 	virtual void OnInvalidReferenceLevelScriptStreamed(const IWorldPartitionActorDescInstanceView& ActorDescView) override {}
 	virtual void OnInvalidReferenceLevelScriptDataLayers(const IWorldPartitionActorDescInstanceView& ActorDescView) override {}
 	virtual void OnInvalidReferenceRuntimeGrid(const IWorldPartitionActorDescInstanceView& ActorDescView, const IWorldPartitionActorDescInstanceView& ReferenceActorDescView) override {}
 	virtual void OnInvalidReferenceDataLayerAsset(const UDataLayerInstanceWithAsset* DataLayerInstance) override {}
+	virtual void OnInvalidDataLayerAssetType(const UDataLayerInstanceWithAsset* DataLayerInstance, const UDataLayerAsset* DataLayerAsset) override {}
 	virtual void OnDataLayerHierarchyTypeMismatch(const UDataLayerInstance* DataLayerInstance, const UDataLayerInstance* Parent) override {};
 	virtual void OnDataLayerAssetConflict(const UDataLayerInstanceWithAsset* DataLayerInstance, const UDataLayerInstanceWithAsset* ConflictingDataLayerInstance) override {};
 	virtual void OnActorNeedsResave(const IWorldPartitionActorDescInstanceView& ActorDescView) override {};
diff --git a/Engine/Source/Runtime/Engine/Public/WorldPartition/ErrorHandling/WorldPartitionStreamingGenerationProxyErrorHandler.h b/Engine/Source/Runtime/Engine/Public/WorldPartition/ErrorHandling/WorldPartitionStreamingGenerationProxyErrorHandler.h
--- a/Engine/Source/Runtime/Engine/Public/WorldPartition/ErrorHandling/WorldPartitionStreamingGenerationProxyErrorHandler.h
+++ b/Engine/Source/Runtime/Engine/Public/WorldPartition/ErrorHandling/WorldPartitionStreamingGenerationProxyErrorHandler.h
@@ -14,11 +14,12 @@ public:
 	ENGINE_API virtual void OnInvalidRuntimeGrid(const IWorldPartitionActorDescInstanceView& ActorDescView, FName GridName) override;
 	ENGINE_API virtual void OnInvalidReference(const IWorldPartitionActorDescInstanceView& ActorDescView, const FGuid& ReferenceGuid, IWorldPartitionActorDescInstanceView* ReferenceActorDescView) override;
 	ENGINE_API virtual void OnInvalidReferenceGridPlacement(const IWorldPartitionActorDescInstanceView& ActorDescView, const IWorldPartitionActorDescInstanceView& ReferenceActorDescView) override;
-	ENGINE_API virtual void OnInvalidReferenceDataLayers(const IWorldPartitionActorDescInstanceView& ActorDescView, const IWorldPartitionActorDescInstanceView& ReferenceActorDescView) override;
+	ENGINE_API virtual void OnInvalidReferenceDataLayers(const IWorldPartitionActorDescInstanceView& ActorDescView, const IWorldPartitionActorDescInstanceView& ReferenceActorDescView, EDataLayerInvalidReason Reason) override;
 	ENGINE_API virtual void OnInvalidReferenceLevelScriptStreamed(const IWorldPartitionActorDescInstanceView& ActorDescView) override;
 	ENGINE_API virtual void OnInvalidReferenceLevelScriptDataLayers(const IWorldPartitionActorDescInstanceView& ActorDescView) override;
 	ENGINE_API virtual void OnInvalidReferenceRuntimeGrid(const IWorldPartitionActorDescInstanceView& ActorDescView, const IWorldPartitionActorDescInstanceView& ReferenceActorDescView) override;
 	ENGINE_API virtual void OnInvalidReferenceDataLayerAsset(const UDataLayerInstanceWithAsset* DataLayerInstance) override;
+	ENGINE_API virtual void OnInvalidDataLayerAssetType(const UDataLayerInstanceWithAsset* DataLayerInstance, const UDataLayerAsset* DataLayerAsset) override;
 	ENGINE_API virtual void OnDataLayerHierarchyTypeMismatch(const UDataLayerInstance* DataLayerInstance, const UDataLayerInstance* Parent) override;
 	ENGINE_API virtual void OnDataLayerAssetConflict(const UDataLayerInstanceWithAsset* DataLayerInstance, const UDataLayerInstanceWithAsset* ConflictingDataLayerInstance) override;
 	ENGINE_API virtual void OnActorNeedsResave(const IWorldPartitionActorDescInstanceView& ActorDescView) override;
diff --git a/Engine/Source/Runtime/Engine/Public/WorldPartition/HLOD/IWorldPartitionHLODUtilities.h b/Engine/Source/Runtime/Engine/Public/WorldPartition/HLOD/IWorldPartitionHLODUtilities.h
--- a/Engine/Source/Runtime/Engine/Public/WorldPartition/HLOD/IWorldPartitionHLODUtilities.h
+++ b/Engine/Source/Runtime/Engine/Public/WorldPartition/HLOD/IWorldPartitionHLODUtilities.h
@@ -10,6 +10,9 @@
 
 #if WITH_EDITOR
 
+#include "WorldPartition/DataLayer/ExternalDataLayerAsset.h"
+#include "WorldPartition/DataLayer/ExternalDataLayerInstance.h"
+
 class AActor;
 class UWorldPartition;
 class UHLODLayer;
@@ -35,6 +38,16 @@ struct FHLODCreationParams
 	FGuid ContentBundleGuid;
 	TArray<const UDataLayerInstance*> DataLayerInstances;
 
+	const UExternalDataLayerAsset* GetExternalDataLayerAsset() const
+	{
+		auto IsAnExternalDataLayerPred = [](const UDataLayerInstance* DataLayerInstance) { return DataLayerInstance->IsA<UExternalDataLayerInstance>(); };
+		if (const UDataLayerInstance* const* ExternalDataLayerInstance = DataLayerInstances.FindByPredicate(IsAnExternalDataLayerPred))
+		{
+			return CastChecked<UExternalDataLayerInstance>(*ExternalDataLayerInstance)->GetExternalDataLayerAsset();
+		}
+		return nullptr;
+	}
+
 	double MinVisibleDistance;
 };
 
diff --git a/Engine/Source/Runtime/Engine/Public/WorldPartition/RuntimeHashSet/WorldPartitionRuntimeHashSet.h b/Engine/Source/Runtime/Engine/Public/WorldPartition/RuntimeHashSet/WorldPartitionRuntimeHashSet.h
--- a/Engine/Source/Runtime/Engine/Public/WorldPartition/RuntimeHashSet/WorldPartitionRuntimeHashSet.h
+++ b/Engine/Source/Runtime/Engine/Public/WorldPartition/RuntimeHashSet/WorldPartitionRuntimeHashSet.h
@@ -137,7 +137,6 @@ public:
 	ENGINE_API virtual bool SupportsHLODs() const override;
 	ENGINE_API virtual bool SetupHLODActors(const IStreamingGenerationContext* StreamingGenerationContext, const UWorldPartition::FSetupHLODActorsParams& Params) const override;
 	ENGINE_API virtual bool GenerateStreaming(class UWorldPartitionStreamingPolicy* StreamingPolicy, const IStreamingGenerationContext* StreamingGenerationContext, TArray<FString>* OutPackagesToGenerate) override;
-	ENGINE_API virtual void FlushStreaming() override;
 	ENGINE_API virtual bool IsValidGrid(FName GridName) const;
 	ENGINE_API virtual bool IsValidHLODLayer(FName GridName, const FSoftObjectPath& HLODLayerPath) const;
 	ENGINE_API virtual void DumpStateLog(FHierarchicalLogArchive& Ar) const override;
@@ -151,7 +150,7 @@ public:
 
 	// External streaming object interface
 #if WITH_EDITOR
-	ENGINE_API virtual URuntimeHashExternalStreamingObjectBase* StoreToExternalStreamingObject(UObject* StreamingObjectOuter, FName StreamingObjectName) override;
+	ENGINE_API virtual TSubclassOf<URuntimeHashExternalStreamingObjectBase> GetExternalStreamingObjectClass() const override { return URuntimeHashSetExternalStreamingObject::StaticClass(); }
 #endif
 	ENGINE_API virtual bool InjectExternalStreamingObject(URuntimeHashExternalStreamingObjectBase* ExternalStreamingObject) override;
 	ENGINE_API virtual bool RemoveExternalStreamingObject(URuntimeHashExternalStreamingObjectBase* ExternalStreamingObject) override;
@@ -161,6 +160,13 @@ public:
 	ENGINE_API virtual void ForEachStreamingCellsQuery(const FWorldPartitionStreamingQuerySource& QuerySource, TFunctionRef<bool(const UWorldPartitionRuntimeCell*)> Func, FWorldPartitionQueryCache* QueryCache) const override;
 	ENGINE_API virtual void ForEachStreamingCellsSources(const TArray<FWorldPartitionStreamingSource>& Sources, TFunctionRef<bool(const UWorldPartitionRuntimeCell*, EStreamingSourceTargetState)> Func) const override;
 
+protected:
+#if WITH_EDITOR
+	ENGINE_API virtual bool HasStreamingContent() const override;
+	ENGINE_API virtual void StoreStreamingContentToExternalStreamingObject(URuntimeHashExternalStreamingObjectBase* OutExternalStreamingObject) override;
+	ENGINE_API virtual void FlushStreamingContent() override;
+#endif
+
 private:
 #if WITH_EDITOR
 	/** Generate the runtime partitions streaming descs. */
diff --git a/Engine/Source/Runtime/Engine/Public/WorldPartition/WorldPartition.h b/Engine/Source/Runtime/Engine/Public/WorldPartition/WorldPartition.h
--- a/Engine/Source/Runtime/Engine/Public/WorldPartition/WorldPartition.h
+++ b/Engine/Source/Runtime/Engine/Public/WorldPartition/WorldPartition.h
@@ -49,6 +49,9 @@ class FAutoConsoleVariableRef;
 class FWorldPartitionDraw2DContext;
 class FContentBundleEditor;
 class IStreamingGenerationContext;
+class UExternalDataLayerManager;
+class IWorldPartitionCookPackageObject;
+
 struct IWorldPartitionStreamingSourceProvider;
 
 enum class EWorldPartitionRuntimeCellState : uint8;
@@ -167,7 +170,8 @@ private:
 	void OnActorDescInstanceUpdating(FWorldPartitionActorDescInstance* ActorDescInstance);
 	void OnActorDescInstanceUpdated(FWorldPartitionActorDescInstance* ActorDescInstance);
 
-	ENGINE_API void InitializeActorDescContainerEditorStreaming(UActorDescContainerInstance* InActorDescContainer, bool bInHashActorDescs);
+	bool ShouldHashUnhashActorDescInstances() const;
+	ENGINE_API void InitializeActorDescContainerEditorStreaming(UActorDescContainerInstance* InActorDescContainer);
 #endif
 
 public:
@@ -210,20 +214,29 @@ public:
 			: ErrorHandler(nullptr)
 		{}
 		
+		FGenerateStreamingParams& SetContainerInstanceCollection(const FActorDescContainerInstanceCollection& InContainerInstanceCollection, const FStreamingGenerationContainerInstanceCollection::ECollectionType& InCollectionType)
+		{
+			check(ContainerInstanceCollection.IsEmpty());
+			ContainerInstanceCollection.SetCollectionType(InCollectionType);
+			ContainerInstanceCollection.Append(InContainerInstanceCollection);
+			return *this;
+		}
+		FGenerateStreamingParams& SetErrorHandler(IStreamingGenerationErrorHandler* InErrorHandler) { ErrorHandler = InErrorHandler; return *this; }
+		FGenerateStreamingParams& SetOutputLogPath(const FString& InOutputLogPath) { OutputLogPath = InOutputLogPath; return *this; }
+
+		UE_DEPRECATED(5.4, "Use constructor receiving a ContainerInstanceCollection instead")
+		FGenerateStreamingParams& SetActorDescContainer(const UActorDescContainer* InActorDescContainer) { return *this; }
+
 		UE_DEPRECATED(5.4, "Use ContainerInstanceCollection instead")
 		FStreamingGenerationActorDescCollection ActorDescCollection;
 		
-		FStreamingGenerationContainerInstanceCollection ContainerInstanceCollection;
+	private:
 
+		FStreamingGenerationContainerInstanceCollection ContainerInstanceCollection;
 		TOptional<const FString> OutputLogPath;
 		IStreamingGenerationErrorHandler* ErrorHandler;
 
-		FGenerateStreamingParams& SetActorDescContainerInstance(const UActorDescContainerInstance* InContainerInstance) { ContainerInstanceCollection.AddContainer(InContainerInstance); return *this;	}
-		FGenerateStreamingParams& SetErrorHandler(IStreamingGenerationErrorHandler* InErrorHandler) { ErrorHandler = InErrorHandler; return *this; }
-		FGenerateStreamingParams& SetOutputLogPath(const FString& InOutputLogPath) { OutputLogPath = InOutputLogPath; return *this; }
-
-		UE_DEPRECATED(5.4, "Use SetActorDescContainerInstance instead")
-		FGenerateStreamingParams& SetActorDescContainer(const UActorDescContainer* InActorDescContainer) { return *this; }
+		friend class UWorldPartition;
 	};
 	PRAGMA_ENABLE_DEPRECATION_WARNINGS
 
@@ -232,10 +245,18 @@ public:
 		FGenerateStreamingContext()
 		{}
 
+		// Level Packages to generate
 		 TArray<FString>* PackagesToGenerate = nullptr;
+		
+		// Generated External Streaming Objects
+		TArray<URuntimeHashExternalStreamingObjectBase*>* GeneratedExternalStreamingObjects = nullptr;
+
 		 TOptional<FString> OutputLogFilename;
 
-		FGenerateStreamingContext& SetPackagesToGenerate(TArray<FString>* InPackagesToGenerate) { PackagesToGenerate = InPackagesToGenerate; return *this; }
+		UE_DEPRECATED(5.4, "SetPackagesToGenerate is deprecated, use SetLevelPackagesToGenerate")
+		FGenerateStreamingContext& SetPackagesToGenerate(TArray<FString>* InPackagesToGenerate) { return SetLevelPackagesToGenerate(InPackagesToGenerate); }
+		FGenerateStreamingContext& SetLevelPackagesToGenerate(TArray<FString>* InLevelPackagesToGenerate) { PackagesToGenerate = InLevelPackagesToGenerate; return *this; }
+		FGenerateStreamingContext& SetGeneratedExternalStreamingObjects(TArray<URuntimeHashExternalStreamingObjectBase*>* InGeneratedExternalStreamingObjects) { GeneratedExternalStreamingObjects = InGeneratedExternalStreamingObjects; return *this; }
 	};
 
 	ENGINE_API bool GenerateStreaming(const FGenerateStreamingParams& InParams, FGenerateStreamingContext& InContext);
@@ -275,7 +296,7 @@ public:
 	ENGINE_API virtual bool PrepareGeneratorPackageForCook(IWorldPartitionCookPackageContext& CookContext, TArray<UPackage*>& OutModifiedPackages) override;
 	ENGINE_API virtual bool PopulateGeneratorPackageForCook(IWorldPartitionCookPackageContext& CookContext, const TArray<FWorldPartitionCookPackage*>& InPackagesToCook, TArray<UPackage*>& OutModifiedPackages) override;
 	ENGINE_API virtual bool PopulateGeneratedPackageForCook(IWorldPartitionCookPackageContext& CookContext, const FWorldPartitionCookPackage& InPackagesToCool, TArray<UPackage*>& OutModifiedPackages) override;
-	ENGINE_API virtual UWorldPartitionRuntimeCell* GetCellForPackage(const FWorldPartitionCookPackage& PackageToCook) const override;
+	ENGINE_API virtual UWorldPartitionRuntimeCell* GetCellForPackage(const FWorldPartitionCookPackage& InPackageToCook) const override;
 	//~ End IWorldPartitionCookPackageGenerator Interface 
 	// End Cooking
 
@@ -324,8 +345,7 @@ public:
 		FCheckForErrorsParams& SetErrorHandler(IStreamingGenerationErrorHandler* InErrorHandler) { ErrorHandler = InErrorHandler; return *this; }
 		FCheckForErrorsParams& SetActorDescContainerInstanceCollection(const FActorDescContainerInstanceCollection* InActorDescContainerInstanceCollection) { ActorDescContainerInstanceCollection = InActorDescContainerInstanceCollection; return *this; }
 		FCheckForErrorsParams& SetEnableStreaming(bool bInEnableStreaming) { bEnableStreaming = bInEnableStreaming; return *this; }
-		FCheckForErrorsParams& SetActorGuidsToContainerMap(const TMap<FGuid, const UActorDescContainerInstance*>& InActorGuidsToContainerInstanceMap) { ActorGuidsToContainerInstanceMap = InActorGuidsToContainerInstanceMap; return *this; }
-
+		FCheckForErrorsParams& SetActorGuidsToContainerInstanceMap(const TMap<FGuid, const UActorDescContainerInstance*>& InActorGuidsToContainerInstanceMap) { ActorGuidsToContainerInstanceMap = InActorGuidsToContainerInstanceMap; return *this; }
 
 		UE_DEPRECATED(5.4, "Use ActorDescContainerInstanceCollection instead")
 		const FActorDescContainerCollection* ActorDescContainerCollection = nullptr;
@@ -442,6 +462,7 @@ public:
 	ENGINE_API bool IsStreamingCompleted(EWorldPartitionRuntimeCellState QueryState, const TArray<FWorldPartitionStreamingQuerySource>& QuerySources, bool bExactState) const;
 	ENGINE_API bool GetIntersectingCells(const TArray<FWorldPartitionStreamingQuerySource>& InSources, TArray<const IWorldPartitionCell*>& OutCells) const;
 
+	ENGINE_API bool IsExternalStreamingObjectInjected(URuntimeHashExternalStreamingObjectBase* InExternalStreamingObject) const;
 	ENGINE_API bool InjectExternalStreamingObject(URuntimeHashExternalStreamingObjectBase* ExternalStreamingObject);
 	ENGINE_API bool RemoveExternalStreamingObject(URuntimeHashExternalStreamingObjectBase* ExternalStreamingObject);
 
@@ -463,6 +484,7 @@ public:
 
 	ENGINE_API UDataLayerManager* GetDataLayerManager() const;
 	ENGINE_API UDataLayerManager* GetResolvingDataLayerManager() const;
+	ENGINE_API UExternalDataLayerManager* GetExternalDataLayerManager() const;
 
 	inline EWorldPartitionDataLayersLogicOperator GetDataLayersLogicOperator() const { return DataLayersLogicOperator; }
 
@@ -575,7 +597,7 @@ private:
 
 	TUniquePtr<FWorldPartitionExternalDirtyActorsTracker> ExternalDirtyActorsTracker;
 
-	TSet<FString> GeneratedStreamingPackageNames;
+	TSet<FString> GeneratedLevelStreamingPackageNames;
 
 	UPROPERTY(Transient)
 	TObjectPtr<UActorDescContainerInstance> ActorDescContainerInstance;
@@ -600,6 +622,9 @@ private:
 	UPROPERTY(Transient)
 	TObjectPtr<UDataLayerManager> DataLayerManager;
 
+	UPROPERTY(Transient)
+	TObjectPtr<UExternalDataLayerManager> ExternalDataLayerManager;
+
 	UPROPERTY(Transient)
 	mutable TObjectPtr<UWorldPartitionStreamingPolicy> StreamingPolicy;
 
@@ -635,6 +660,8 @@ private:
 	void HashActorDescInstance(FWorldPartitionActorDescInstance* ActorDescInstance);
 	void UnhashActorDescInstance(FWorldPartitionActorDescInstance* ActorDescInstance);
 	void OnContentBundleRemovedContent(const FContentBundleEditor* ContentBundle);
+	IWorldPartitionCookPackageObject* GetCookPackageObject(const FWorldPartitionCookPackage& PackageToCook) const;
+	bool HasStreamingContent() const;
 
 public:
 	// Editor loader adapters management
@@ -673,4 +700,5 @@ private:
 
 	friend class AWorldPartitionReplay;
 	friend class UWorldPartitionSubsystem;
+	friend class UExternalDataLayerManager;
 };
\ No newline at end of file
diff --git a/Engine/Source/Runtime/Engine/Public/WorldPartition/WorldPartitionActorDesc.h b/Engine/Source/Runtime/Engine/Public/WorldPartition/WorldPartitionActorDesc.h
--- a/Engine/Source/Runtime/Engine/Public/WorldPartition/WorldPartitionActorDesc.h
+++ b/Engine/Source/Runtime/Engine/Public/WorldPartition/WorldPartitionActorDesc.h
@@ -131,7 +131,9 @@ public:
 
 	inline bool GetActorIsHLODRelevant() const { return bActorIsHLODRelevant; }
 	inline FSoftObjectPath GetHLODLayer() const { return HLODLayer; }
-	inline const TArray<FName>& GetDataLayers() const { return DataLayers; }
+	ENGINE_API TArray<FName> GetDataLayers(bool bIncludeExternalDataLayer = true) const;
+	ENGINE_API FName GetExternalDataLayer() const;
+	inline const FSoftObjectPath& GetExternalDataLayerAsset() const { return ExternalDataLayerAsset; }
 
 	inline const TArray<FName>& GetTags() const { return Tags; }
 
@@ -376,6 +378,7 @@ protected:
 	bool							bIsBoundsValid;
 	FSoftObjectPath					HLODLayer;
 	TArray<FName>					DataLayers;
+	FSoftObjectPath					ExternalDataLayerAsset;
 	TArray<FGuid>					References;
 	TArray<FGuid>					EditorOnlyReferences; // References that aren't necessarily editor only but referenced through an editor only property.
 	TArray<FName>					Tags;
diff --git a/Engine/Source/Runtime/Engine/Public/WorldPartition/WorldPartitionActorDescInstance.h b/Engine/Source/Runtime/Engine/Public/WorldPartition/WorldPartitionActorDescInstance.h
--- a/Engine/Source/Runtime/Engine/Public/WorldPartition/WorldPartitionActorDescInstance.h
+++ b/Engine/Source/Runtime/Engine/Public/WorldPartition/WorldPartitionActorDescInstance.h
@@ -6,6 +6,7 @@
 
 #if WITH_EDITOR
 #include "WorldPartition/WorldPartitionActorDesc.h"
+#include "WorldPartition/DataLayer/DataLayerInstanceNames.h"
 #include "UObject/WeakObjectPtr.h"
 
 class AActor;
@@ -46,7 +47,7 @@ public:
 	ENGINE_API virtual bool IsEditorRelevant() const;
 
 	virtual bool IsUsingDataLayerAsset() const { return GetActorDesc()->IsUsingDataLayerAsset(); }
-	virtual const TArray<FName>& GetDataLayers() const { return GetActorDesc()->GetDataLayers(); }
+	virtual TArray<FName> GetDataLayers() const { return GetActorDesc()->GetDataLayers(); }
 
 	virtual bool GetActorIsHLODRelevant() const { return GetActorDesc()->GetActorIsHLODRelevant(); }
 	virtual FSoftObjectPath GetHLODLayer() const { return GetActorDesc()->GetHLODLayer(); }
@@ -72,6 +73,7 @@ public:
 	virtual const FGuid& GetParentActor() const { return GetActorDesc()->GetParentActor(); }
 		
 	virtual FGuid GetContentBundleGuid() const { return GetActorDesc()->GetContentBundleGuid(); }
+	virtual const FSoftObjectPath& GetExternalDataLayerAsset() const { return GetActorDesc()->GetExternalDataLayerAsset(); }
 
 	virtual bool IsChildContainerInstance() const { return ChildContainerInstance || GetActorDesc()->IsChildContainerInstance(); }
 	virtual FName GetChildContainerPackage() const { return GetActorDesc()->GetChildContainerPackage(); }
@@ -86,7 +88,7 @@ public:
 	ENGINE_API virtual const FWorldPartitionActorDesc* GetActorDesc() const { check(ActorDesc); return ActorDesc; }
 	
 	virtual bool HasResolvedDataLayerInstanceNames() const { return ResolvedDataLayerInstanceNames.IsSet(); }
-	ENGINE_API virtual const TArray<FName>& GetDataLayerInstanceNames() const;
+	ENGINE_API virtual const FDataLayerInstanceNames& GetDataLayerInstanceNames() const;
 
 	ENGINE_API virtual bool IsLoaded(bool bEvenIfPendingKill = false) const;
 	ENGINE_API virtual AActor* GetActor(bool bEvenIfPendingKill = true, bool bEvenIfUnreachable = false) const;
@@ -151,7 +153,7 @@ protected:
 		return HardRefCount;
 	}
 
-	inline void SetDataLayerInstanceNames(const TArray<FName>& InDataLayerInstanceNames) { ResolvedDataLayerInstanceNames = InDataLayerInstanceNames; }
+	inline void SetDataLayerInstanceNames(const FDataLayerInstanceNames& InDataLayerInstanceNames) { ResolvedDataLayerInstanceNames = InDataLayerInstanceNames; }
 
 	virtual ENGINE_API void RegisterChildContainerInstance();
 	virtual ENGINE_API void UnregisterChildContainerInstance();
@@ -162,7 +164,7 @@ protected:
 
 	mutable uint32								SoftRefCount;
 	mutable uint32								HardRefCount;
-	TOptional<TArray<FName>>					ResolvedDataLayerInstanceNames;
+	TOptional<FDataLayerInstanceNames>			ResolvedDataLayerInstanceNames;
 	bool										bIsForcedNonSpatiallyLoaded;
 	mutable FText*								UnloadedReason;
 
diff --git a/Engine/Source/Runtime/Engine/Public/WorldPartition/WorldPartitionActorDescInstanceView.h b/Engine/Source/Runtime/Engine/Public/WorldPartition/WorldPartitionActorDescInstanceView.h
--- a/Engine/Source/Runtime/Engine/Public/WorldPartition/WorldPartitionActorDescInstanceView.h
+++ b/Engine/Source/Runtime/Engine/Public/WorldPartition/WorldPartitionActorDescInstanceView.h
@@ -31,7 +31,7 @@ public:
 	virtual bool IsEditorRelevant() const override { return ActorDescInstance->IsEditorRelevant(); }
 
 	virtual bool IsUsingDataLayerAsset() const override { return ActorDescInstance->IsUsingDataLayerAsset(); }
-	virtual const TArray<FName>& GetDataLayers() const override { return ActorDescInstance->GetDataLayers(); }
+	virtual TArray<FName> GetDataLayers() const override { return ActorDescInstance->GetDataLayers(); }
 
 	virtual bool GetActorIsHLODRelevant() const override { return ActorDescInstance->GetActorIsHLODRelevant(); }
 	virtual FSoftObjectPath GetHLODLayer() const override { return ActorDescInstance->GetHLODLayer(); }
@@ -57,6 +57,7 @@ public:
 	virtual const FGuid& GetParentActor() const override { return ActorDescInstance->GetParentActor(); }
 
 	virtual FGuid GetContentBundleGuid() const override { return ActorDescInstance->GetContentBundleGuid(); }
+	virtual const FSoftObjectPath& GetExternalDataLayerAsset() const override { return ActorDescInstance->GetExternalDataLayerAsset(); }
 
 	virtual bool IsChildContainerInstance() const override { return ActorDescInstance->IsChildContainerInstance(); }
 	virtual FName GetChildContainerPackage() const override { return ActorDescInstance->GetChildContainerPackage(); }
@@ -74,7 +75,7 @@ public:
 	virtual const FWorldPartitionActorDesc* GetActorDesc() const override { return ActorDescInstance->GetActorDesc(); }
 
 	virtual bool HasResolvedDataLayerInstanceNames() const override { return ActorDescInstance->HasResolvedDataLayerInstanceNames(); }
-	const TArray<FName>& GetDataLayerInstanceNames() const override { return ActorDescInstance->GetDataLayerInstanceNames(); }
+	const FDataLayerInstanceNames& GetDataLayerInstanceNames() const override { return ActorDescInstance->GetDataLayerInstanceNames(); }
 
 	virtual AActor* GetActor(bool bEvenIfPendingKill = true, bool bEvenIfUnreachable = false) const override { return ActorDescInstance->GetActor(bEvenIfPendingKill, bEvenIfUnreachable); }
 	virtual bool IsLoaded(bool bEvenIfPendingKill = false) const override { return ActorDescInstance->IsLoaded(bEvenIfPendingKill); }
diff --git a/Engine/Source/Runtime/Engine/Public/WorldPartition/WorldPartitionActorDescInstanceViewInterface.h b/Engine/Source/Runtime/Engine/Public/WorldPartition/WorldPartitionActorDescInstanceViewInterface.h
--- a/Engine/Source/Runtime/Engine/Public/WorldPartition/WorldPartitionActorDescInstanceViewInterface.h
+++ b/Engine/Source/Runtime/Engine/Public/WorldPartition/WorldPartitionActorDescInstanceViewInterface.h
@@ -4,6 +4,7 @@
 #if WITH_EDITOR
 #include "CoreMinimal.h"
 #include "WorldPartition/WorldPartitionActorDesc.h"
+#include "WorldPartition/DataLayer/DataLayerInstanceNames.h"
 
 class AActor;
 class IStreamingGenerationErrorHandler;
@@ -33,7 +34,7 @@ public:
 	virtual bool IsEditorRelevant() const = 0;
 
 	virtual bool IsUsingDataLayerAsset() const = 0;
-	virtual const TArray<FName>& GetDataLayers() const = 0;
+	virtual TArray<FName> GetDataLayers() const = 0;
 
 	virtual bool GetActorIsHLODRelevant() const = 0;
 	virtual FSoftObjectPath GetHLODLayer() const = 0;
@@ -59,6 +60,7 @@ public:
 	virtual const FGuid& GetParentActor() const = 0;
 	
 	virtual FGuid GetContentBundleGuid() const = 0;
+	virtual const FSoftObjectPath& GetExternalDataLayerAsset() const = 0;
 	
 	virtual bool IsChildContainerInstance() const = 0;
 	virtual FName GetChildContainerPackage() const = 0;
@@ -76,7 +78,7 @@ public:
 	virtual FName GetActorLabelOrName() const { return GetActorLabel().IsNone() ? GetActorName() : GetActorLabel(); }
 
 	virtual bool HasResolvedDataLayerInstanceNames() const = 0;
-	virtual const TArray<FName>& GetDataLayerInstanceNames() const = 0;
+	virtual const FDataLayerInstanceNames& GetDataLayerInstanceNames() const = 0;
 
 	virtual AActor* GetActor(bool bEvenIfPendingKill = true, bool bEvenIfUnreachable = false) const = 0;
 	virtual bool IsLoaded(bool bEvenIfPendingKill = false) const = 0;
diff --git a/Engine/Source/Runtime/Engine/Public/WorldPartition/WorldPartitionLevelStreamingPolicy.h b/Engine/Source/Runtime/Engine/Public/WorldPartition/WorldPartitionLevelStreamingPolicy.h
--- a/Engine/Source/Runtime/Engine/Public/WorldPartition/WorldPartitionLevelStreamingPolicy.h
+++ b/Engine/Source/Runtime/Engine/Public/WorldPartition/WorldPartitionLevelStreamingPolicy.h
@@ -35,7 +35,7 @@ public:
 	virtual void RemapSoftObjectPath(FSoftObjectPath& ObjectPath) const override;
 	static FString GetCellPackagePath(const FName& InCellName, const UWorld* InWorld);
 
-	virtual bool StoreToExternalStreamingObject(URuntimeHashExternalStreamingObjectBase& OutExternalStreamingObject) override;
+	virtual bool StoreStreamingContentToExternalStreamingObject(URuntimeHashExternalStreamingObjectBase& OutExternalStreamingObject) override;
 	virtual bool ConvertContainerPathToEditorPath(const FActorContainerID& InContainerID, const FSoftObjectPath& InPath, FSoftObjectPath& OutPath) const;
 #endif
 
diff --git a/Engine/Source/Runtime/Engine/Public/WorldPartition/WorldPartitionMiniMap.h b/Engine/Source/Runtime/Engine/Public/WorldPartition/WorldPartitionMiniMap.h
--- a/Engine/Source/Runtime/Engine/Public/WorldPartition/WorldPartitionMiniMap.h
+++ b/Engine/Source/Runtime/Engine/Public/WorldPartition/WorldPartitionMiniMap.h
@@ -18,7 +18,8 @@ class AWorldPartitionMiniMap : public AInfo
 
 private:
 #if WITH_EDITOR
-	virtual bool IsDataLayerTypeSupported(TSubclassOf<UDataLayerInstance> DataLayerType) const final { return false; }
+	virtual bool ActorTypeSupportsDataLayer() const final { return false; }
+	virtual bool ActorTypeSupportsExternalDataLayer() const final { return false; }
 	virtual bool ActorTypeIsMainWorldOnly() const override  { return true; }
 	virtual bool ShouldImport(FStringView ActorPropString, bool IsMovingLevel) override { return false; }
 #endif
diff --git a/Engine/Source/Runtime/Engine/Public/WorldPartition/WorldPartitionMiniMapVolume.h b/Engine/Source/Runtime/Engine/Public/WorldPartition/WorldPartitionMiniMapVolume.h
--- a/Engine/Source/Runtime/Engine/Public/WorldPartition/WorldPartitionMiniMapVolume.h
+++ b/Engine/Source/Runtime/Engine/Public/WorldPartition/WorldPartitionMiniMapVolume.h
@@ -16,7 +16,8 @@ public:
 	virtual bool IsEditorOnly() const final { return true; }
 
 #if WITH_EDITOR
-	virtual bool IsDataLayerTypeSupported(TSubclassOf<UDataLayerInstance> DataLayerType) const final { return false; }
+	virtual bool ActorTypeSupportsDataLayer() const final { return false; }
+	virtual bool ActorTypeSupportsExternalDataLayer() const final { return false; }
 	virtual bool CanChangeIsSpatiallyLoadedFlag() const { return false; }
 #endif
 };
diff --git a/Engine/Source/Runtime/Engine/Public/WorldPartition/WorldPartitionRuntimeCell.h b/Engine/Source/Runtime/Engine/Public/WorldPartition/WorldPartitionRuntimeCell.h
--- a/Engine/Source/Runtime/Engine/Public/WorldPartition/WorldPartitionRuntimeCell.h
+++ b/Engine/Source/Runtime/Engine/Public/WorldPartition/WorldPartitionRuntimeCell.h
@@ -8,6 +8,8 @@
 #include "WorldPartition/WorldPartitionActorContainerID.h"
 #include "WorldPartition/WorldPartitionRuntimeCellData.h"
 #include "WorldPartition/WorldPartitionRuntimeCellInterface.h"
+#include "WorldPartition/Cook/WorldPartitionCookPackageInterface.h"
+#include "WorldPartition/DataLayer/DataLayerInstanceNames.h"
 #include "ProfilingDebugging/ProfilingHelpers.h"
 #include "Misc/HierarchicalLogArchive.h"
 #include "Algo/AnyOf.h"
@@ -17,8 +19,10 @@
 class UActorContainer;
 class UDataLayerAsset;
 class UDataLayerInstance;
+class UExternalDataLayerInstance;
 class UWorldPartition;
 class UDataLayerManager;
+class UExternalDataLayerAsset;
 class FStreamingGenerationActorDescView;
 struct FHierarchicalLogArchive;
 
@@ -169,7 +173,7 @@ static_assert(EWorldPartitionRuntimeCellState::Unloaded < EWorldPartitionRuntime
  * Represents a PIE/Game streaming cell which points to external actor/data chunk packages
  */
 UCLASS(Abstract, MinimalAPI)
-class UWorldPartitionRuntimeCell : public UObject, public IWorldPartitionCell
+class UWorldPartitionRuntimeCell : public UObject, public IWorldPartitionCell, public IWorldPartitionCookPackageObject
 {
 	GENERATED_UCLASS_BODY()
 
@@ -210,13 +214,15 @@ class UWorldPartitionRuntimeCell : public UObject, public IWorldPartitionCell
 
 	//~Begin IWorldPartitionCell Interface
 	ENGINE_API virtual TArray<const UDataLayerInstance*> GetDataLayerInstances() const override;
+	ENGINE_API virtual const UExternalDataLayerInstance* GetExternalDataLayerInstance() const override;
 	ENGINE_API virtual bool ContainsDataLayer(const UDataLayerAsset* DataLayerAsset) const override;
 	ENGINE_API virtual bool ContainsDataLayer(const UDataLayerInstance* DataLayerInstance) const override;
 	ENGINE_API virtual bool HasContentBundle() const override;
-	virtual const TArray<FName>& GetDataLayers() const override  { return DataLayers; }
+	virtual const TArray<FName>& GetDataLayers() const override { return DataLayers.GetRawArray(); }
+	virtual FName GetExternalDataLayer() const override { return DataLayers.GetExternalDataLayer(); }
 	virtual bool HasAnyDataLayer(const TSet<FName>& InDataLayers) const override
 	{
-		return Algo::AnyOf(DataLayers, [&InDataLayers](const FName& DataLayer) { return InDataLayers.Contains(DataLayer); });
+		return Algo::AnyOf(GetDataLayers(), [&InDataLayers](const FName& DataLayer) { return InDataLayers.Contains(DataLayer); });
 	}
 	ENGINE_API virtual FName GetLevelPackageName() const override;
 	ENGINE_API virtual FString GetDebugName() const override;
@@ -224,8 +230,6 @@ class UWorldPartitionRuntimeCell : public UObject, public IWorldPartitionCell
 	ENGINE_API virtual UWorld* GetOuterWorld() const override;
 	//~End IWorldPartitionCell Interface
 
-	inline bool HasDataLayers() const { return !DataLayers.IsEmpty(); }
-
 	ENGINE_API UDataLayerManager* GetDataLayerManager() const;
 	ENGINE_API EDataLayerRuntimeState GetCellEffectiveWantedState() const;
 
@@ -254,10 +258,20 @@ class UWorldPartitionRuntimeCell : public UObject, public IWorldPartitionCell
 	ENGINE_API virtual int32 GetActorCount() const PURE_VIRTUAL(UWorldPartitionRuntimeCell::GetActorCount, return 0;);
 
 	// Cook methods
-	virtual bool PrepareCellForCook(UPackage* InPackage) { return false; }
-	ENGINE_API virtual bool PopulateGeneratorPackageForCook(TArray<UPackage*>& OutModifiedPackages) PURE_VIRTUAL(UWorldPartitionRuntimeCell::PopulateGeneratorPackageForCook, return false;);
-	ENGINE_API virtual bool PopulateGeneratedPackageForCook(UPackage* InPackage, TArray<UPackage*>& OutModifiedPackages) PURE_VIRTUAL(UWorldPartitionRuntimeCell::PopulateGeneratedPackageForCook, return false;);
-	ENGINE_API virtual FString GetPackageNameToCreate() const PURE_VIRTUAL(UWorldPartitionRuntimeCell::GetPackageNameToCreate, return FString(""););
+	virtual bool PrepareCellForCook(UPackage* InPackage) { return OnPopulateGeneratorPackageForCook(InPackage); }
+	UE_DEPRECATED(5.4, "PopulateGeneratorPackageForCook is deprecated, it was replaced by OnPrepareGeneratorPackageForCook")
+	ENGINE_API virtual bool PopulateGeneratorPackageForCook(TArray<UPackage*>& OutModifiedPackages) { return OnPrepareGeneratorPackageForCook(OutModifiedPackages); }
+	UE_DEPRECATED(5.4, "PopulateGeneratedPackageForCook is deprecated, it was replaced by OnPopulateGeneratedPackageForCook")
+	ENGINE_API virtual bool PopulateGeneratedPackageForCook(UPackage* InPackage, TArray<UPackage*>& OutModifiedPackages) { return OnPopulateGeneratedPackageForCook(InPackage, OutModifiedPackages); }
+
+	//~Begin IWorldPartitionCookPackageObject
+	ENGINE_API virtual bool IsLevelPackage() const override { return true; }
+	ENGINE_API virtual const UExternalDataLayerAsset* GetExternalDataLayerAsset() const override { return ExternalDataLayerAsset; }
+	ENGINE_API virtual FString GetPackageNameToCreate() const { return FString(); }
+	ENGINE_API virtual bool OnPrepareGeneratorPackageForCook(TArray<UPackage*>& OutModifiedPackages) override { return true; }
+	ENGINE_API virtual bool OnPopulateGeneratorPackageForCook(UPackage* InPackage) override { return true; }
+	ENGINE_API virtual bool OnPopulateGeneratedPackageForCook(UPackage* InPackage, TArray<UPackage*>& OutModifiedPackages) override { return true; }
+	//~End IWorldPartitionCookPackageObject
 
 	ENGINE_API virtual void DumpStateLog(FHierarchicalLogArchive& Ar) const;
 #endif
@@ -297,7 +311,7 @@ protected:
 
 private:
 	UPROPERTY()
-	TArray<FName> DataLayers;
+	FDataLayerInstanceNames DataLayers;
 
 	UPROPERTY()
 	bool bClientOnlyVisible;
@@ -331,6 +345,11 @@ protected:
 	FName LevelPackageName;
 #endif
 
+#if WITH_EDITORONLY_DATA
+	UPROPERTY()
+	TObjectPtr<const UExternalDataLayerAsset> ExternalDataLayerAsset;
+#endif
+
 	mutable EDataLayerRuntimeState EffectiveWantedState;
 	mutable int32 EffectiveWantedStateEpoch;
 
diff --git a/Engine/Source/Runtime/Engine/Public/WorldPartition/WorldPartitionRuntimeCellInterface.h b/Engine/Source/Runtime/Engine/Public/WorldPartition/WorldPartitionRuntimeCellInterface.h
--- a/Engine/Source/Runtime/Engine/Public/WorldPartition/WorldPartitionRuntimeCellInterface.h
+++ b/Engine/Source/Runtime/Engine/Public/WorldPartition/WorldPartitionRuntimeCellInterface.h
@@ -5,8 +5,9 @@
 #include "UObject/Interface.h"
 #include "WorldPartitionRuntimeCellInterface.generated.h"
 
-class UDataLayerInstance;
 class UDataLayerAsset;
+class UDataLayerInstance;
+class UExternalDataLayerInstance;
 
 UINTERFACE(MinimalAPI)
 class UWorldPartitionCell : public UInterface
@@ -21,6 +22,8 @@ class IWorldPartitionCell
 public:
 	/** Returns the cell content associated data layer instances. */
 	virtual TArray<const UDataLayerInstance*> GetDataLayerInstances() const = 0;
+	/** Returns the cell associated external data layer instance. */
+	virtual const UExternalDataLayerInstance* GetExternalDataLayerInstance() const = 0;
 	/** Returns whether the cell data layers referenced the provided data layer asset or not. */
 	virtual bool ContainsDataLayer(const UDataLayerAsset* DataLayerAsset) const = 0;
 	/** Returns whether the cell data layers referenced the provided data layer instance or not. */
@@ -31,6 +34,8 @@ public:
 	virtual bool HasContentBundle() const = 0;
 	/** Returns the cell content associated data layers. */
 	virtual const TArray<FName>& GetDataLayers() const = 0;
+	/** Returns the cell content associated external data layer. */
+	virtual FName GetExternalDataLayer() const = 0;
 	/** Returns whether the cell data layers referenced any of the provided data layer or not. */
 	virtual bool HasAnyDataLayer(const TSet<FName>& InDataLayers) const = 0;
 	/** Returns the cell bounds computed using cell content. */
diff --git a/Engine/Source/Runtime/Engine/Public/WorldPartition/WorldPartitionRuntimeHash.h b/Engine/Source/Runtime/Engine/Public/WorldPartition/WorldPartitionRuntimeHash.h
--- a/Engine/Source/Runtime/Engine/Public/WorldPartition/WorldPartitionRuntimeHash.h
+++ b/Engine/Source/Runtime/Engine/Public/WorldPartition/WorldPartitionRuntimeHash.h
@@ -13,6 +13,7 @@
 #include "WorldPartition/WorldPartitionStreamingGeneration.h"
 #include "WorldPartition/WorldPartitionStreamingGenerationContext.h"
 #include "WorldPartition/WorldPartitionRuntimeContainerResolving.h"
+#include "WorldPartition/DataLayer/DataLayerInstanceProviderInterface.h"
 #if WITH_EDITOR
 #include "CookPackageSplitter.h"
 #include "Misc/HierarchicalLogArchive.h"
@@ -21,6 +22,8 @@
 
 struct FHierarchicalLogArchive;
 class FWorldPartitionDraw2DContext;
+class UExternalDataLayerAsset;
+class UExternalDataLayerInstance;
 
 extern ENGINE_API float GBlockOnSlowStreamingRatio;
 extern ENGINE_API float GBlockOnSlowStreamingWarningFactor;
@@ -34,7 +37,7 @@ enum class EWorldPartitionStreamingPerformance : uint8
 };
 
 UCLASS(Abstract, MinimalAPI)
-class URuntimeHashExternalStreamingObjectBase : public UObject
+class URuntimeHashExternalStreamingObjectBase : public UObject, public IWorldPartitionCookPackageObject, public IDataLayerInstanceProvider
 {
 	GENERATED_BODY()
 
@@ -55,10 +58,32 @@ public:
 	
 	ENGINE_API void OnStreamingObjectLoaded(UWorld* InjectedWorld);
 
+	// ~Being IDataLayerInstanceProvider
+	ENGINE_API virtual TSet<TObjectPtr<UDataLayerInstance>>& GetDataLayerInstances() override;
+	virtual const TSet<TObjectPtr<UDataLayerInstance>>& GetDataLayerInstances() const override { return const_cast<URuntimeHashExternalStreamingObjectBase*>(this)->GetDataLayerInstances(); }
+	virtual const UExternalDataLayerInstance* GetRootExternalDataLayerInstance() const override { return RootExternalDataLayerInstance; }
+	// ~End IDataLayerInstanceProvider
+	UExternalDataLayerInstance* GetRootExternalDataLayerInstance() { return const_cast<UExternalDataLayerInstance*>(RootExternalDataLayerInstance.Get()); }
+
 #if WITH_EDITOR
+	UE_DEPRECATED(5.4, "PopulateGeneratorPackageForCook is depreacted and was replaced by OnPopulateGeneratorPackageForCook")
 	ENGINE_API void PopulateGeneratorPackageForCook();
+
+	//~Begin IWorldPartitionCookPackageObject interface
+	ENGINE_API virtual bool IsLevelPackage() const override { return false; }
+	ENGINE_API virtual const UExternalDataLayerAsset* GetExternalDataLayerAsset() const override { return ExternalDataLayerAsset; }
+	ENGINE_API virtual FString GetPackageNameToCreate() const override;
+	ENGINE_API virtual bool OnPrepareGeneratorPackageForCook(TArray<UPackage*>& OutModifiedPackages) override { return true; }
+	ENGINE_API virtual bool OnPopulateGeneratorPackageForCook(UPackage* InPackage) override;
+	ENGINE_API virtual bool OnPopulateGeneratedPackageForCook(UPackage* InPackage, TArray<UPackage*>& OutModifiedPackages) override;
+	//~End IWorldPartitionCookPackageObject interface
+
+protected:
+	virtual void DumpStateLog(FHierarchicalLogArchive& Ar);
+	UWorldPartitionRuntimeCell* GetCellForCookPackage(const FString& InCookPackageName) const;
 #endif
 
+public:
 	UPROPERTY();
 	TMap<FName, FName> SubObjectsToCellRemapping;
 
@@ -75,11 +100,26 @@ protected:
 	UPROPERTY();
 	TMap<FName, FName> CellToLevelStreamingPackage;
 
+	UPROPERTY()
+	TSet<TObjectPtr<UDataLayerInstance>> DataLayerInstances;
+
+	UPROPERTY()
+	TObjectPtr<const UExternalDataLayerInstance> RootExternalDataLayerInstance;
+
+#if WITH_EDITORONLY_DATA
+	UPROPERTY(Transient)
+	TMap<FString, UWorldPartitionRuntimeCell*> PackagesToGenerateForCook;
+
+	UPROPERTY(Transient)
+	TObjectPtr<const UExternalDataLayerAsset> ExternalDataLayerAsset;
+#endif
+
 #if DO_CHECK
 	TWeakObjectPtr<UWorldPartition> TargetInjectedWorldPartition;
 #endif
 
 	friend class UWorldPartition;
+	friend class UExternalDataLayerManager;
 };
 
 struct FWorldPartitionQueryCache
@@ -104,27 +144,38 @@ class UWorldPartitionRuntimeHash : public UObject
 	virtual bool SupportsHLODs() const { return false; }
 	ENGINE_API virtual TArray<UWorldPartitionRuntimeCell*> GetAlwaysLoadedCells() const;
 	ENGINE_API virtual bool GenerateStreaming(class UWorldPartitionStreamingPolicy* StreamingPolicy, const IStreamingGenerationContext* StreamingGenerationContext, TArray<FString>* OutPackagesToGenerate);
-	ENGINE_API virtual void FlushStreaming();
 	virtual bool SetupHLODActors(const IStreamingGenerationContext* StreamingGenerationContext, const UWorldPartition::FSetupHLODActorsParams& Params) const { return false; }
 	virtual bool IsValidGrid(FName GridName) const { return false; }
 	virtual bool IsValidHLODLayer(FName GridName, const FSoftObjectPath& HLODLayerPath) const { return false; }
 	virtual void DrawPreview() const {}
 
-	virtual URuntimeHashExternalStreamingObjectBase* StoreToExternalStreamingObject(UObject* StreamingObjectOuter, FName StreamingObjectName) { return nullptr; }
-
+	ENGINE_API virtual bool HasStreamingContent() const { return false; }
+	ENGINE_API URuntimeHashExternalStreamingObjectBase* StoreStreamingContentToExternalStreamingObject(FName InStreamingObjectName);
+	ENGINE_API virtual void FlushStreamingContent();
+	ENGINE_API virtual TSubclassOf<URuntimeHashExternalStreamingObjectBase> GetExternalStreamingObjectClass() const PURE_VIRTUAL(UWorldPartitionRuntimeHash::GetExternalStreamingObjectClass, return nullptr;);
 	ENGINE_API virtual void DumpStateLog(FHierarchicalLogArchive& Ar) const;
 
-	// non-virtual
-	ENGINE_API bool PrepareGeneratorPackageForCook(TArray<UPackage*>& OutModifiedPackages);
-	ENGINE_API bool PopulateGeneratorPackageForCook(const TArray<FWorldPartitionCookPackage*>& PackagesToCook, TArray<UPackage*>& OutModifiedPackages);
-	ENGINE_API bool PopulateGeneratedPackageForCook(const FWorldPartitionCookPackage& PackagesToCook, TArray<UPackage*>& OutModifiedPackages);
-	ENGINE_API UWorldPartitionRuntimeCell* GetCellForPackage(const FWorldPartitionCookPackage& PackageToCook) const;
+	//~Begin Deprecation
+	UE_DEPRECATED(5.4, "Use StoreStreamingContentToExternalStreamingObject instead.")
+	virtual URuntimeHashExternalStreamingObjectBase* StoreToExternalStreamingObject(UObject* StreamingObjectOuter, FName StreamingObjectName) { return nullptr; }
+	UE_DEPRECATED(5.4, "Use FlushStreamingContent instead.")
+	virtual void FlushStreaming() { FlushStreamingContent(); }
+	UE_DEPRECATED(5.4, "PrepareGeneratorPackageForCook is deprecated.")
+	ENGINE_API bool PrepareGeneratorPackageForCook(TArray<UPackage*>& OutModifiedPackages) { return false; }
+	UE_DEPRECATED(5.4, "PopulateGeneratorPackageForCook is deprecated.")
+	ENGINE_API bool PopulateGeneratorPackageForCook(const TArray<FWorldPartitionCookPackage*>& PackagesToCook, TArray<UPackage*>& OutModifiedPackages) { return false; }
+	UE_DEPRECATED(5.4, "PopulateGeneratedPackageForCook is deprecated.")
+	ENGINE_API bool PopulateGeneratedPackageForCook(const FWorldPartitionCookPackage& PackagesToCook, TArray<UPackage*>& OutModifiedPackages) { return false; }
+	UE_DEPRECATED(5.4, "GetCellForPackage is deprecated.")
+	ENGINE_API UWorldPartitionRuntimeCell* GetCellForPackage(const FWorldPartitionCookPackage& PackageToCook) const { return nullptr; }
+	//~End Deprecation
 
 	// PIE/Game methods
 	ENGINE_API void OnBeginPlay();
 	ENGINE_API void OnEndPlay();
 
 protected:
+	ENGINE_API virtual void StoreStreamingContentToExternalStreamingObject(URuntimeHashExternalStreamingObjectBase* OutExternalStreamingObject);
 	ENGINE_API bool PopulateCellActorInstances(const TArray<const IStreamingGenerationContext::FActorSetInstance*>& ActorSetInstances, bool bIsMainWorldPartition, bool bIsCellAlwaysLoaded, TArray<IStreamingGenerationContext::FActorInstance>& OutCellActorInstances);
 	ENGINE_API void PopulateRuntimeCell(UWorldPartitionRuntimeCell* RuntimeCell, const TArray<IStreamingGenerationContext::FActorInstance>& ActorInstances, TArray<FString>* OutPackagesToGenerate);
 #endif
@@ -152,6 +203,7 @@ public:
 	ENGINE_API bool IsCellRelevantFor(bool bClientOnlyVisible) const;
 	ENGINE_API EWorldPartitionStreamingPerformance GetStreamingPerformance(const TSet<const UWorldPartitionRuntimeCell*>& CellToActivate) const;
 
+	ENGINE_API virtual bool IsExternalStreamingObjectInjected(URuntimeHashExternalStreamingObjectBase* InExternalStreamingObject) const;
 	ENGINE_API virtual bool InjectExternalStreamingObject(URuntimeHashExternalStreamingObjectBase* ExternalStreamingObject);
 	ENGINE_API virtual bool RemoveExternalStreamingObject(URuntimeHashExternalStreamingObjectBase* ExternalStreamingObject);
 
@@ -199,6 +251,9 @@ protected:
 
 #if WITH_EDITOR
 private:
+	UWorldPartitionRuntimeCell* GetCellForCookPackage(const FString& InCookPackageName) const;
+	friend class UWorldPartition;
+
 	using FRuntimeHashConvertFunc = TFunction<UWorldPartitionRuntimeHash*(const UWorldPartitionRuntimeHash*)>;
 	static TMap<TPair<const UClass*, const UClass*>, FRuntimeHashConvertFunc> WorldPartitionRuntimeHashConverters;
 
diff --git a/Engine/Source/Runtime/Engine/Public/WorldPartition/WorldPartitionRuntimeLevelStreamingCell.h b/Engine/Source/Runtime/Engine/Public/WorldPartition/WorldPartitionRuntimeLevelStreamingCell.h
--- a/Engine/Source/Runtime/Engine/Public/WorldPartition/WorldPartitionRuntimeLevelStreamingCell.h
+++ b/Engine/Source/Runtime/Engine/Public/WorldPartition/WorldPartitionRuntimeLevelStreamingCell.h
@@ -17,7 +17,7 @@ class UWorldPartitionRuntimeLevelStreamingCell : public UWorldPartitionRuntimeCe
 {
 	GENERATED_UCLASS_BODY()
 
-	//~Begin UWorldPartitionRuntimeCell Interface
+	//~Begin UWorldPartitionRuntimeCell interface
 	ENGINE_API virtual void Load() const override;
 	ENGINE_API virtual void Unload() const override;
 	ENGINE_API virtual bool CanUnload() const override;
@@ -28,11 +28,11 @@ class UWorldPartitionRuntimeLevelStreamingCell : public UWorldPartitionRuntimeCe
 	ENGINE_API virtual FLinearColor GetDebugColor(EWorldPartitionRuntimeCellVisualizeMode VisualizeMode) const override;
 	ENGINE_API virtual void SetIsAlwaysLoaded(bool bInIsAlwaysLoaded) override;
 	ENGINE_API virtual EStreamingStatus GetStreamingStatus() const override;
-	//~End UWorldPartitionRuntimeCell Interface
+	//~End UWorldPartitionRuntimeCell interface
 
-	//~Begin IWorldPartitionCell Interface
+	//~Begin IWorldPartitionCell interface
 	ENGINE_API FName GetLevelPackageName() const override;
-	//~End IWorldPartitionCell Interface
+	//~End IWorldPartitionCell interface
 
 	ENGINE_API virtual void SetStreamingPriority(int32 InStreamingPriority) const override;
 	ENGINE_API class UWorldPartitionLevelStreamingDynamic* GetLevelStreaming() const;
@@ -46,17 +46,20 @@ class UWorldPartitionRuntimeLevelStreamingCell : public UWorldPartitionRuntimeCe
 	
 
 #if WITH_EDITOR
-	//~Begin UWorldPartitionRuntimeCell Interface
+	//~Begin UWorldPartitionRuntimeCell interface
 	ENGINE_API virtual void AddActorToCell(const FStreamingGenerationActorDescView& ActorDescView) override;
 	ENGINE_API virtual void Fixup() override;
 	ENGINE_API virtual int32 GetActorCount() const override;
 	ENGINE_API virtual void DumpStateLog(FHierarchicalLogArchive& Ar) const override;
-	// Cook methods
 	ENGINE_API virtual bool PrepareCellForCook(UPackage* InPackage) override;
-	ENGINE_API virtual bool PopulateGeneratorPackageForCook(TArray<UPackage*>& OutModifiedPackages) override;
-	ENGINE_API virtual bool PopulateGeneratedPackageForCook(UPackage* InPackage, TArray<UPackage*>& OutModifiedPackages) override;
+	//~End UWorldPartitionRuntimeCell interface
+
+	//~Begin IWorldPartitionCookPackageObject interface
 	ENGINE_API virtual FString GetPackageNameToCreate() const override;
-	//~End UWorldPartitionRuntimeCell Interface
+	ENGINE_API virtual bool OnPrepareGeneratorPackageForCook(TArray<UPackage*>& OutModifiedPackages) override;
+	ENGINE_API virtual bool OnPopulateGeneratorPackageForCook(UPackage* InPackage) override;
+	ENGINE_API virtual bool OnPopulateGeneratedPackageForCook(UPackage* InPackage, TArray<UPackage*>& OutModifiedPackages) override;
+	//~End IWorldPartitionCookPackageObject interface
 
 	//~Begin IWorldPartitionCell Interface
 	ENGINE_API virtual TSet<FName> GetActorPackageNames() const override;
diff --git a/Engine/Source/Runtime/Engine/Public/WorldPartition/WorldPartitionRuntimeSpatialHash.h b/Engine/Source/Runtime/Engine/Public/WorldPartition/WorldPartitionRuntimeSpatialHash.h
--- a/Engine/Source/Runtime/Engine/Public/WorldPartition/WorldPartitionRuntimeSpatialHash.h
+++ b/Engine/Source/Runtime/Engine/Public/WorldPartition/WorldPartitionRuntimeSpatialHash.h
@@ -190,6 +190,10 @@ struct FSpatialHashStreamingGrid
 	ENGINE_API const FSquare2DGridHelper& GetGridHelper() const;
 	ENGINE_API float GetLoadingRange() const;
 
+#if WITH_EDITOR
+	void DumpStateLog(FHierarchicalLogArchive& Ar) const;
+#endif
+
 private:
 	ENGINE_API void ForEachRuntimeCell(const FGridCellCoord& Coords, TFunctionRef<void(const UWorldPartitionRuntimeCell*)> Func) const;
 	ENGINE_API void ForEachLayerCell(const FGridCellCoord& Coords, TFunctionRef<void(const FSpatialHashStreamingGridLayerCell*)> Func) const;
@@ -285,6 +289,10 @@ class URuntimeSpatialHashExternalStreamingObject : public URuntimeHashExternalSt
 {
 	GENERATED_BODY()
 
+#if WITH_EDITOR
+	virtual void DumpStateLog(FHierarchicalLogArchive& Ar) override;
+#endif
+
 public:
 	UPROPERTY();
 	TArray<FSpatialHashStreamingGrid> StreamingGrids;
@@ -313,13 +321,12 @@ public:
 
 	ENGINE_API virtual void SetDefaultValues() override;
 	virtual bool SupportsHLODs() const override { return true; }
-	ENGINE_API virtual void FlushStreaming() override;
 	ENGINE_API virtual bool SetupHLODActors(const IStreamingGenerationContext* StreamingGenerationContext, const UWorldPartition::FSetupHLODActorsParams& Params) const override;
 	ENGINE_API virtual bool IsValidGrid(FName GridName) const override;
 	ENGINE_API virtual bool IsValidHLODLayer(FName GridName, const FSoftObjectPath& HLODLayerPath) const override { return true; }
 	ENGINE_API virtual void DrawPreview() const override;
 
-	ENGINE_API virtual URuntimeHashExternalStreamingObjectBase* StoreToExternalStreamingObject(UObject* StreamingObjectOuter, FName StreamingObjectName) override;
+	ENGINE_API virtual TSubclassOf<URuntimeHashExternalStreamingObjectBase> GetExternalStreamingObjectClass() const override { return URuntimeSpatialHashExternalStreamingObject::StaticClass(); }
 
 	static ENGINE_API FString GetCellNameString(UWorld* InOuterWorld, FName InGridName, const FGridCellCoord& InCellGlobalCoord, const FDataLayersID& InDataLayerID, const FGuid& InContentBundleID, FString* OutInstanceSuffix = nullptr);
 	static ENGINE_API FGuid GetCellGuid(FName InGridName, int32 InCellSize, const FGridCellCoord& InCellGlobalCoord, const FDataLayersID& InDataLayerID, const FGuid& InContentBundleID);
@@ -354,6 +361,9 @@ protected:
 	ENGINE_API virtual EWorldPartitionStreamingPerformance GetStreamingPerformanceForCell(const UWorldPartitionRuntimeCell* Cell) const override;
 
 #if WITH_EDITOR
+	ENGINE_API virtual bool HasStreamingContent() const override;
+	ENGINE_API virtual void StoreStreamingContentToExternalStreamingObject(URuntimeHashExternalStreamingObjectBase* OutExternalStreamingObject) override;
+	ENGINE_API virtual void FlushStreamingContent() override;
 	ENGINE_API virtual bool GenerateStreaming(class UWorldPartitionStreamingPolicy* StreamingPolicy, const IStreamingGenerationContext* StreamingGenerationContext, TArray<FString>* OutPackagesToGenerate = nullptr) override;
 	ENGINE_API virtual void DumpStateLog(FHierarchicalLogArchive& Ar) const override;
 #endif
diff --git a/Engine/Source/Runtime/Engine/Public/WorldPartition/WorldPartitionStreamingGeneration.h b/Engine/Source/Runtime/Engine/Public/WorldPartition/WorldPartitionStreamingGeneration.h
--- a/Engine/Source/Runtime/Engine/Public/WorldPartition/WorldPartitionStreamingGeneration.h
+++ b/Engine/Source/Runtime/Engine/Public/WorldPartition/WorldPartitionStreamingGeneration.h
@@ -35,6 +35,14 @@ protected:
 	{}
 
 public:
+	struct FInvalidReference
+	{
+		FName ActorPackage;
+		FSoftObjectPath ActorSoftPath;
+		FTopLevelAssetPath BaseClass;
+		FTopLevelAssetPath NativeClass;
+	};
+
 	FStreamingGenerationActorDescView(const FStreamingGenerationActorDescViewMap& InActorDescViewMap, const FWorldPartitionActorDescInstance* InActorDescInstance, bool bInUnsaved = false)
 		: FStreamingGenerationActorDescView(InActorDescInstance)
 	{
@@ -49,22 +57,25 @@ public:
 	ENGINE_API virtual bool GetIsSpatiallyLoaded() const override;
 	ENGINE_API virtual FSoftObjectPath GetHLODLayer() const override;
 	ENGINE_API virtual const TArray<FGuid>& GetReferences() const override;
-	ENGINE_API const TArray<FName>& GetDataLayerInstanceNames() const override;
+	ENGINE_API const FDataLayerInstanceNames& GetDataLayerInstanceNames() const override;
+	ENGINE_API virtual bool IsEditorOnlyReference(const FGuid& ReferenceGuid) const override;
 	//~ End FWorldPartitionActorDescInstanceView interface
 
 	bool IsUnsaved() const { return bIsUnsaved; }
 	const FStreamingGenerationActorDescViewMap& GetActorDescViewMap() const { return *ActorDescViewMap; }
 
 	ENGINE_API const TArray<FGuid>& GetEditorReferences() const;
-	ENGINE_API const TArray<FName>& GetRuntimeDataLayerInstanceNames() const;
+	ENGINE_API const FDataLayerInstanceNames& GetRuntimeDataLayerInstanceNames() const;
 
+	ENGINE_API void AddForcedInvalidReference(const FStreamingGenerationActorDescView* ReferenceView);
+	ENGINE_API bool IsInvalidReference(const FGuid& InGuid, FInvalidReference* OutInvalidReference = nullptr) const;
 	ENGINE_API void SetForcedNonSpatiallyLoaded();
 	ENGINE_API void SetForcedNoRuntimeGrid();
 	ENGINE_API void SetForcedNoDataLayers();
-	ENGINE_API void SetRuntimeDataLayerInstanceNames(const TArray<FName>& InRuntimeDataLayerInstanceNames);
+	ENGINE_API void SetRuntimeDataLayerInstanceNames(const FDataLayerInstanceNames& InRuntimeDataLayerInstanceNames);
 	ENGINE_API void SetRuntimeReferences(const TArray<FGuid>& InRuntimeReferences);
 	ENGINE_API void SetEditorReferences(const TArray<FGuid>& InEditorReferences);
-	ENGINE_API void SetDataLayerInstanceNames(const TArray<FName>& InDataLayerInstanceNames);
+	ENGINE_API void SetDataLayerInstanceNames(const FDataLayerInstanceNames& InDataLayerInstanceNames);
 	ENGINE_API void SetParentView(const FStreamingGenerationActorDescView* InParentView);
 
 	ENGINE_API void SetForcedNoHLODLayer();
@@ -89,8 +100,15 @@ private:
 	bool bIsForcedNoDataLayers;
 	bool bIsForceNoHLODLayer;
 	bool bIsUnsaved;
-	TOptional<TArray<FName>> ResolvedDataLayerInstanceNames;
-	TOptional<TArray<FName>> RuntimeDataLayerInstanceNames;
+
+	// Stores invalid actor references
+	TMap<FGuid, FInvalidReference> ForcedInvalidReference;
+	
+	// Used by GetDataLayerInstanceNames() to avoid returning a copy
+	mutable FDataLayerInstanceNames LastReturnedDataLayerInstanceNames; 
+
+	TOptional<FDataLayerInstanceNames> ResolvedDataLayerInstanceNames;
+	TOptional<FDataLayerInstanceNames> RuntimeDataLayerInstanceNames;
 	TOptional<TArray<FGuid>> RuntimeReferences;
 	TOptional<FSoftObjectPath> RuntimedHLODLayer;
 	TArray<FGuid> EditorReferences;
@@ -205,27 +223,40 @@ protected:
 class FStreamingGenerationContainerInstanceCollection : public TActorDescContainerInstanceCollection<TObjectPtr<const UActorDescContainerInstance>>
 {
 public:
-	FStreamingGenerationContainerInstanceCollection() = default;
-	ENGINE_API FStreamingGenerationContainerInstanceCollection(std::initializer_list<TObjectPtr<const UActorDescContainerInstance>> ActorDescContainerInstanceArray);
-	ENGINE_API FStreamingGenerationContainerInstanceCollection(const TArray<const UActorDescContainerInstance*>& ActorDescContainerInstances);
-
-	ENGINE_API UWorld* GetWorld() const;
+	enum class ECollectionType
+	{
+		Invalid,
+		BaseAsContentBundle, // Collection contains only one content bundle container
+		BaseAndEDLs,		 // Collection contains the base container and optionally EDL containers
+		BaseAndAny			 // Collection contains the base container and any type of container (base,EDL,CB)
+	};
 
-	// @todo_ow : Remove once conversion to ExternalDataLayer is complete. 
-	// It is present to handle content bundles streaming generation via the same code path. 
-	ENGINE_API FGuid GetContentBundleGuid() const;
+	FStreamingGenerationContainerInstanceCollection() : CollectionType(ECollectionType::Invalid) {}
+	ENGINE_API FStreamingGenerationContainerInstanceCollection(std::initializer_list<TObjectPtr<const UActorDescContainerInstance>> ActorDescContainerInstanceArray, const ECollectionType& InCollectionType);
+	ENGINE_API FStreamingGenerationContainerInstanceCollection(const TArray<const UActorDescContainerInstance*>& ActorDescContainerInstances, const ECollectionType& InCollectionType);
 
-	ENGINE_API const UActorDescContainerInstance* GetMainContainer() const;
-	ENGINE_API FName GetMainContainerPackageName() const;
-	ENGINE_API TArrayView<const UActorDescContainerInstance* const> GetExternalDataLayerContainers();
+	ENGINE_API UWorld* GetWorld() const;
+	ENGINE_API const UActorDescContainerInstance* GetBaseContainerInstance() const;
+	ENGINE_API FName GetBaseContainerInstancePackageName() const;
+	ENGINE_API FGuid GetContentBundleGuid() const; // @todo_ow : Remove once conversion to ExternalDataLayer is complete. 
 
+private:
+	void SetCollectionType(const ECollectionType& InCollectionType) { check(CollectionType == ECollectionType::Invalid); CollectionType = InCollectionType; }
+	ECollectionType GetCollectionType() const { return CollectionType; }
+	void InitializeCollection();
+	FName GetBaseContainerPackageName() const;
+	TArrayView<const UActorDescContainerInstance* const> GetExternalDataLayerContainerInstances() const;
+	TArrayView<const UActorDescContainerInstance* const> GetContentBundleContainerInstances() const;
 	ENGINE_API virtual void OnCollectionChanged() override;
 
-private:
-	ENGINE_API void SortCollection();
+	ECollectionType CollectionType;
+	int32 ExternalDataLayerStartIdx = INDEX_NONE;
+	int32 ContentBundleStartIdx = INDEX_NONE;
+	static constexpr int BaseContainerIdx = 0;
 
-	static constexpr int MainContainerIdx = 0;
-	static constexpr int ExternalDataLayerContainerStartIdx = MainContainerIdx + 1;
+	friend class UWorldPartition;
+	friend class FWorldPartitionStreamingGenerator;
+	friend struct FGenerateStreamingParams;
 };
 
 class UE_DEPRECATED(5.4, "Use FStreamingGenerationContainerInstanceCollection instead") FStreamingGenerationActorDescCollection : public TActorDescContainerCollection<TObjectPtr<const UActorDescContainer>>
diff --git a/Engine/Source/Runtime/Engine/Public/WorldPartition/WorldPartitionStreamingGenerationContext.h b/Engine/Source/Runtime/Engine/Public/WorldPartition/WorldPartitionStreamingGenerationContext.h
--- a/Engine/Source/Runtime/Engine/Public/WorldPartition/WorldPartitionStreamingGenerationContext.h
+++ b/Engine/Source/Runtime/Engine/Public/WorldPartition/WorldPartitionStreamingGenerationContext.h
@@ -2,9 +2,14 @@
 #pragma once
 
 #include "CoreMinimal.h"
+#include "Misc/Optional.h"
 #include "OverrideVoidReturnInvoker.h"
 #include "WorldPartition/WorldPartitionActorContainerID.h"
 
+class UExternalDataLayerAsset;
+class UActorDescContainer;
+class UDataLayerInstance;
+class FGenerateStreamingActorDescCollection;
 class FStreamingGenerationActorDescViewMap;
 class FStreamingGenerationActorDescView;
 class FWorldPartitionStreamingGenerator;
@@ -64,7 +69,7 @@ public:
 	};
 
 	/**
-	 * An actor set instance is an actual intance of an actor set in the world.
+	 * An actor set instance is an actual instance of an actor set in the world.
 	 */
 	struct FActorSetInstance
 	{
@@ -104,6 +109,8 @@ public:
 				}
 			}
 		}
+
+		const UExternalDataLayerAsset* GetExternalDataLayerAsset() const;
 	};
 
 	/**
@@ -126,14 +133,14 @@ public:
 	};
 
 	virtual FBox GetWorldBounds() const = 0;
-
-	virtual const FActorSetContainerInstance* GetMainWorldContainerInstance() const = 0;
+	// Returns the ActorSetContainerInstance that contains the BaseActorDescContainerInstance of this context
+	virtual const FActorSetContainerInstance* GetActorSetContainerForContextBaseContainerInstance() const = 0;
 	virtual void ForEachActorSetInstance(TFunctionRef<void(const FActorSetInstance&)> Func) const = 0;
 	virtual void ForEachActorSetContainerInstance(TFunctionRef<void(const FActorSetContainerInstance&)> Func) const = 0;
 
 	PRAGMA_DISABLE_DEPRECATION_WARNINGS
 
-	UE_DEPRECATED(5.4, "Implement GetMainWorldContainerInstance instead")
+	UE_DEPRECATED(5.4, "Implement GetActorSetContainerForContextBaseContainerInstance instead")
 	virtual const FActorSetContainer* GetMainWorldContainer() const { return nullptr; };
 
 	UE_DEPRECATED(5.4, "Implement ForEachActorSetContainerInstance instead")
@@ -154,14 +161,27 @@ public:
 		return SourceContext->GetWorldBounds();
 	}
 
-	virtual const FActorSetContainerInstance* GetMainWorldContainerInstance() const override
+	virtual const FActorSetContainerInstance* GetActorSetContainerForContextBaseContainerInstance() const override
 	{
-		return SourceContext->GetMainWorldContainerInstance();
+		return SourceContext->GetActorSetContainerForContextBaseContainerInstance();
 	}
 
 	virtual void ForEachActorSetInstance(TFunctionRef<void(const FActorSetInstance&)> Func) const override
 	{
-		SourceContext->ForEachActorSetInstance(Func);
+		if (ActorSetInstanceFilterFunc.IsSet())
+		{
+			SourceContext->ForEachActorSetInstance([this, Func](const FActorSetInstance& ActorSetInstance)
+			{
+				if ((*ActorSetInstanceFilterFunc)(ActorSetInstance))
+				{
+					Func(ActorSetInstance);
+				}
+			});
+		}
+		else
+		{
+			SourceContext->ForEachActorSetInstance(Func);
+		}
 	}
 
 	virtual void ForEachActorSetContainerInstance(TFunctionRef<void(const FActorSetContainerInstance&)> Func) const override
@@ -169,7 +189,11 @@ public:
 		SourceContext->ForEachActorSetContainerInstance(Func);
 	}
 
+	void SetActorSetInstanceFilter(const TFunction<bool(const IStreamingGenerationContext::FActorSetInstance&)>& InActorSetInstanceFilterFunc) { ActorSetInstanceFilterFunc = InActorSetInstanceFilterFunc; }
+
 protected:
+	TOptional<TFunction<bool(const IStreamingGenerationContext::FActorSetInstance&)>> ActorSetInstanceFilterFunc;
+
 	const IStreamingGenerationContext* SourceContext;
 };
 #endif
\ No newline at end of file
diff --git a/Engine/Source/Runtime/Engine/Public/WorldPartition/WorldPartitionStreamingPolicy.h b/Engine/Source/Runtime/Engine/Public/WorldPartition/WorldPartitionStreamingPolicy.h
--- a/Engine/Source/Runtime/Engine/Public/WorldPartition/WorldPartitionStreamingPolicy.h
+++ b/Engine/Source/Runtime/Engine/Public/WorldPartition/WorldPartitionStreamingPolicy.h
@@ -72,7 +72,9 @@ public:
 	virtual void SetContainerResolver(const FWorldPartitionRuntimeContainerResolver& InContainerResolver) {}
 	virtual void RemapSoftObjectPath(FSoftObjectPath& ObjectPath) const {}
 
-	virtual bool StoreToExternalStreamingObject(URuntimeHashExternalStreamingObjectBase& OutExternalStreamingObject) { return true; }
+	UE_DEPRECATED(5.4, "Use StoreStreamingContentToExternalStreamingObject instead.")
+	virtual bool StoreToExternalStreamingObject(URuntimeHashExternalStreamingObjectBase& OutExternalStreamingObject) { return StoreStreamingContentToExternalStreamingObject(OutExternalStreamingObject); }
+	virtual bool StoreStreamingContentToExternalStreamingObject(URuntimeHashExternalStreamingObjectBase& OutExternalStreamingObject) { return true; }
 	virtual bool ConvertContainerPathToEditorPath(const FActorContainerID& InContainerID, const FSoftObjectPath& InPath, FSoftObjectPath& OutPath) const { return false; }
 #endif
 
