Fix storm brief one frame flickering when player is jumping while sliding.

Issue was caused by the sliding camera modifier removing itself from the player camera manager modifier list while camera manager was going through the list to apply each camera modifier.

Fix the issue by calling a forEach function with a lambda, that does a local copy of the modifier list of the player camera manager.

#tests console

diff --git a/Engine/Source/Runtime/Engine/Classes/Camera/PlayerCameraManager.h b/Engine/Source/Runtime/Engine/Classes/Camera/PlayerCameraManager.h
--- a/Engine/Source/Runtime/Engine/Classes/Camera/PlayerCameraManager.h
+++ b/Engine/Source/Runtime/Engine/Classes/Camera/PlayerCameraManager.h
@@ -283,6 +283,16 @@ protected:
 	TArray<TObjectPtr<UCameraModifier>> ModifierList;
 
 public:
+	/**
+	* Runs the given function on a local copy of the ModifierList of UCameraModifier, with the possibility of early exit
+	* Most easily used with a lambda as follows:
+	* ForEachCameraModifier([](UCameraModifier* Proxy) -> bool
+	* {
+	*     return continueLoop ? true : false;
+	* });
+	*/
+	ENGINE_API void ForEachCameraModifier(TFunctionRef<bool(UCameraModifier*)> Fn);
+
 	/** List of modifiers to create by default for this camera */
 	UPROPERTY(EditDefaultsOnly, BlueprintReadOnly, Category = CameraModifier)
 	TArray< TSubclassOf<UCameraModifier> > DefaultModifiers;
diff --git a/Engine/Source/Runtime/Engine/Private/PlayerCameraManager.cpp b/Engine/Source/Runtime/Engine/Private/PlayerCameraManager.cpp
--- a/Engine/Source/Runtime/Engine/Private/PlayerCameraManager.cpp
+++ b/Engine/Source/Runtime/Engine/Private/PlayerCameraManager.cpp
@@ -276,20 +276,21 @@ void APlayerCameraManager::ApplyCameraModifiers(float DeltaTime, FMinimalViewInf
 	ClearCachedPPBlends();
 
 	// Loop through each camera modifier
-	for (int32 ModifierIdx = 0; ModifierIdx < ModifierList.Num(); ++ModifierIdx)
+	ForEachCameraModifier([DeltaTime, &InOutPOV](UCameraModifier* CameraModifier)
 	{
+		bool bContinue = true;
+
 		// Apply camera modification and output into DesiredCameraOffset/DesiredCameraRotation
-		if ((ModifierList[ModifierIdx] != NULL) && !ModifierList[ModifierIdx]->IsDisabled())
+		if ((CameraModifier != NULL) && !CameraModifier->IsDisabled())
 		{
 			// If ModifyCamera returns true, exit loop
 			// Allows high priority things to dictate if they are
 			// the last modifier to be applied
-			if (ModifierList[ModifierIdx]->ModifyCamera(DeltaTime, InOutPOV))
-			{
-				break;
-			}
+			bContinue = !CameraModifier->ModifyCamera(DeltaTime, InOutPOV);
 		}
-	}
+
+		return bContinue;
+	});
 }
 
 void APlayerCameraManager::AddCachedPPBlend(struct FPostProcessSettings& PPSettings, float BlendWeight, EViewTargetBlendOrder BlendOrder)
@@ -1384,6 +1385,21 @@ void APlayerCameraManager::SetManualCameraFade(float InFadeAmount, FLinearColor
 	FadeTimeRemaining = 0.0f;
 }
 
+void APlayerCameraManager::ForEachCameraModifier(TFunctionRef<bool(UCameraModifier*)> Fn)
+{
+	// Local copy the modifiers array in case it get when calling the lambda on each modifiers
+	TArray<TObjectPtr<UCameraModifier>> LocalModifierList = ModifierList;
+
+	// Loop through each camera modifier
+	for (int32 ModifierIdx = 0; ModifierIdx < LocalModifierList.Num(); ++ModifierIdx)
+	{
+		if (!Fn(LocalModifierList[ModifierIdx]))
+		{
+			return;
+		}
+	}
+}
+
 void APlayerCameraManager::SetCameraCachePOV(const FMinimalViewInfo& InPOV)
 {
 	CameraCachePrivate.POV = InPOV;
