Fixed warning in ScanPathsSynchronous caused by world partition validation caused by CL 31011213

diff --git a/Engine/Source/Runtime/Engine/Private/WorldPartition/WorldPartitionStreamingGeneration.cpp b/Engine/Source/Runtime/Engine/Private/WorldPartition/WorldPartitionStreamingGeneration.cpp
--- a/Engine/Source/Runtime/Engine/Private/WorldPartition/WorldPartitionStreamingGeneration.cpp
+++ b/Engine/Source/Runtime/Engine/Private/WorldPartition/WorldPartitionStreamingGeneration.cpp
@@ -852,7 +852,7 @@ class FWorldPartitionStreamingGenerator
 				ResolveContainerDescriptor(ContainerCollectionInstanceDescriptor);
 
 				// Validate container, fixing anything illegal, etc.
-				ValidateContainerDescriptor(ContainerCollectionInstanceDescriptor, ContainerCollectionInstanceDescriptor.ID.IsMainContainer());
+				ValidateContainerDescriptor(ContainerCollectionInstanceDescriptor);
 
 				// Update container, computing cluster, bounds, etc.
 				UpdateContainerDescriptor(ContainerCollectionInstanceDescriptor);
@@ -1021,9 +1021,10 @@ class FWorldPartitionStreamingGenerator
 	 * Perform various validations on the container descriptor, and adjust it based on different requirements. This needs to happen before updating
 	 * containers bounds because some actor descriptor views might change grid placement, etc.
 	 */
-	void ValidateContainerDescriptor(FContainerCollectionInstanceDescriptor& ContainerCollectionInstanceDescriptor, bool bIsMainContainer)
+	void ValidateContainerDescriptor(FContainerCollectionInstanceDescriptor& ContainerCollectionInstanceDescriptor)
 	{
-		if (bIsMainContainer)
+		const bool bIsMainContainerNonContentBundle = ContainerCollectionInstanceDescriptor.ID.IsMainContainer() && !ContainerCollectionInstanceDescriptor.ContentBundleID.IsValid();
+		if (bIsMainContainerNonContentBundle)
 		{
 			TArray<FGuid> LevelScriptReferences;
 			if (WorldPartitionContext)
