Clean up ALLOW_LQ_LIGHTMAPS for platform extensions

#rnx

diff --git a/Engine/Source/Runtime/Engine/Public/SceneManagement.h b/Engine/Source/Runtime/Engine/Public/SceneManagement.h
--- a/Engine/Source/Runtime/Engine/Public/SceneManagement.h
+++ b/Engine/Source/Runtime/Engine/Public/SceneManagement.h
@@ -371,7 +371,7 @@ static const int32 LQ_LIGHTMAP_COEF_INDEX = 2;
 /** Compile out low quality lightmaps to save memory */
 // @todo-mobile: Need to fix this!
 #ifndef ALLOW_LQ_LIGHTMAPS
-#define ALLOW_LQ_LIGHTMAPS (PLATFORM_DESKTOP || PLATFORM_IOS || PLATFORM_ANDROID || PLATFORM_SWITCH || PLATFORM_HOLOLENS)
+#define ALLOW_LQ_LIGHTMAPS (PLATFORM_DESKTOP || PLATFORM_IOS || PLATFORM_ANDROID)
 #endif
 
 /** Compile out high quality lightmaps to save memory */
