Remove old deprecated function from 5.0, UInputDelegateBinding::BindToInputComponent in favor of the version that takes an object to bind to.

#rnx

diff --git a/Engine/Source/Runtime/Engine/Classes/Engine/InputDelegateBinding.h b/Engine/Source/Runtime/Engine/Classes/Engine/InputDelegateBinding.h
--- a/Engine/Source/Runtime/Engine/Classes/Engine/InputDelegateBinding.h
+++ b/Engine/Source/Runtime/Engine/Classes/Engine/InputDelegateBinding.h
@@ -35,9 +35,6 @@ UCLASS(abstract, MinimalAPI)
 class UInputDelegateBinding : public UDynamicBlueprintBinding
 {
 	GENERATED_UCLASS_BODY()
-
-	UE_DEPRECATED(5.0, "This version of BindToInputComponent is deprecated, please use the version that takes in an ObjectToBindTo instead.")
-	virtual void BindToInputComponent(UInputComponent* InputComponent) const { };
 	
 	/**
 	 * Override this function to bind a delegate to the given input component.
