- Added support to iterate on FDataLayerInstanceNames (for backward compatibility)
- So facilitate that, editor only option that forcibly empties FDataLayerInstanceNames non-external data layers is now an irreversible operation (which is fine with the current use cases)

diff --git a/Engine/Source/Runtime/Engine/Private/WorldPartition/DataLayer/DataLayerUtils.cpp b/Engine/Source/Runtime/Engine/Private/WorldPartition/DataLayer/DataLayerUtils.cpp
--- a/Engine/Source/Runtime/Engine/Private/WorldPartition/DataLayer/DataLayerUtils.cpp
+++ b/Engine/Source/Runtime/Engine/Private/WorldPartition/DataLayer/DataLayerUtils.cpp
@@ -59,8 +59,8 @@ FDataLayerInstanceNames FDataLayerUtils::ResolveDataLayerInstanceNames(const UDa
 			ResolvedDLINames.Reserve(ActorDescNonExternalDataLayers.Num() + 1);
 			const UDataLayerInstance* ExternalDataLayerInstance = InDataLayerManager->GetDataLayerInstanceFromAssetName(ActorDescExternalDataLayer);
 			FName ExternalDataLayer = ExternalDataLayerInstance ? ExternalDataLayerInstance->GetDataLayerFName() : NAME_None;
-			bool bIsFirstDataLayerIsExternal = !ExternalDataLayer.IsNone();
-			if (bIsFirstDataLayerIsExternal)
+			bool bIsFirstDataLayerExternal = !ExternalDataLayer.IsNone();
+			if (bIsFirstDataLayerExternal)
 			{
 				ResolvedDLINames.Add(ExternalDataLayer);
 			}
@@ -71,7 +71,7 @@ FDataLayerInstanceNames FDataLayerUtils::ResolveDataLayerInstanceNames(const UDa
 					ResolvedDLINames.Add(DataLayerInstance->GetDataLayerFName());
 				}
 			}
-			return FDataLayerInstanceNames(ResolvedDLINames, bIsFirstDataLayerIsExternal);
+			return FDataLayerInstanceNames(ResolvedDLINames, bIsFirstDataLayerExternal);
 		}
 		// Fallback on FWorldDataLayersActorDesc
 		else if (!InWorldDataLayersActorDescs.IsEmpty() && AreWorldDataLayersActorDescsSane(InWorldDataLayersActorDescs))
@@ -80,8 +80,8 @@ FDataLayerInstanceNames FDataLayerUtils::ResolveDataLayerInstanceNames(const UDa
 			ResolvedDLINames.Reserve(ActorDescNonExternalDataLayers.Num() + 1);
 			const FDataLayerInstanceDesc* ExternalDataLayerInstanceDesc = GetDataLayerInstanceDescFromAssetPath(InWorldDataLayersActorDescs, ActorDescExternalDataLayer);
 			FName ExternalDataLayer = ExternalDataLayerInstanceDesc ? ExternalDataLayerInstanceDesc->GetName() : NAME_None;
-			bool bIsFirstDataLayerIsExternal = !ExternalDataLayer.IsNone();
-			if (bIsFirstDataLayerIsExternal)
+			bool bIsFirstDataLayerExternal = !ExternalDataLayer.IsNone();
+			if (bIsFirstDataLayerExternal)
 			{
 				ResolvedDLINames.Add(ExternalDataLayer);
 			}
@@ -92,7 +92,7 @@ FDataLayerInstanceNames FDataLayerUtils::ResolveDataLayerInstanceNames(const UDa
 					ResolvedDLINames.Add(DataLayerInstanceDesc->GetName());
 				}
 			}
-			return FDataLayerInstanceNames(ResolvedDLINames, bIsFirstDataLayerIsExternal);
+			return FDataLayerInstanceNames(ResolvedDLINames, bIsFirstDataLayerExternal);
 		}
 	}
 	return FDataLayerInstanceNames(ActorDescNonExternalDataLayers, ActorDescExternalDataLayer);
@@ -110,8 +110,8 @@ bool FDataLayerUtils::ResolveRuntimeDataLayerInstanceNames(const UDataLayerManag
 		TArray<FName> ResolvedRuntimeDLINames;
 		ResolvedRuntimeDLINames.Reserve(NonEDLInstanceNames.Num() + 1);
 		const UDataLayerInstance* ExternalDataLayerInstance = InDataLayerManager->GetDataLayerInstanceFromName(EDLInstanceName);
-		const bool bIsFirstDataLayerIsExternal = (ExternalDataLayerInstance && ExternalDataLayerInstance->IsRuntime());
-		if (bIsFirstDataLayerIsExternal)
+		const bool bIsFirstDataLayerExternal = (ExternalDataLayerInstance && ExternalDataLayerInstance->IsRuntime());
+		if (bIsFirstDataLayerExternal)
 		{
 			ResolvedRuntimeDLINames.Add(EDLInstanceName);
 		}
@@ -123,7 +123,7 @@ bool FDataLayerUtils::ResolveRuntimeDataLayerInstanceNames(const UDataLayerManag
 				ResolvedRuntimeDLINames.Add(DataLayerInstanceName);
 			}
 		}
-		OutRuntimeDataLayerInstanceNames = FDataLayerInstanceNames(ResolvedRuntimeDLINames, bIsFirstDataLayerIsExternal);
+		OutRuntimeDataLayerInstanceNames = FDataLayerInstanceNames(ResolvedRuntimeDLINames, bIsFirstDataLayerExternal);
 		return true;
 	}
 
@@ -136,8 +136,8 @@ bool FDataLayerUtils::ResolveRuntimeDataLayerInstanceNames(const UDataLayerManag
 		TArray<FName> ResolvedRuntimeDLINames;
 		ResolvedRuntimeDLINames.Reserve(NonEDLInstanceNames.Num() + 1);
 		const FDataLayerInstanceDesc* ExternalDataLayerInstanceDesc = GetDataLayerInstanceDescFromInstanceName(WorldDataLayersActorDescs, EDLInstanceName);
-		const bool bIsFirstDataLayerIsExternal = (ExternalDataLayerInstanceDesc && (ExternalDataLayerInstanceDesc->GetDataLayerType() == EDataLayerType::Runtime));
-		if (bIsFirstDataLayerIsExternal)
+		const bool bIsFirstDataLayerExternal = (ExternalDataLayerInstanceDesc && (ExternalDataLayerInstanceDesc->GetDataLayerType() == EDataLayerType::Runtime));
+		if (bIsFirstDataLayerExternal)
 		{
 			ResolvedRuntimeDLINames.Add(EDLInstanceName);
 		}
@@ -150,7 +150,7 @@ bool FDataLayerUtils::ResolveRuntimeDataLayerInstanceNames(const UDataLayerManag
 				ResolvedRuntimeDLINames.Add(DataLayerInstanceName);
 			}
 		}
-		OutRuntimeDataLayerInstanceNames = FDataLayerInstanceNames(ResolvedRuntimeDLINames, bIsFirstDataLayerIsExternal);
+		OutRuntimeDataLayerInstanceNames = FDataLayerInstanceNames(ResolvedRuntimeDLINames, bIsFirstDataLayerExternal);
 		return true;
 	}
 
diff --git a/Engine/Source/Runtime/Engine/Private/WorldPartition/WorldPartitionRuntimeCell.cpp b/Engine/Source/Runtime/Engine/Private/WorldPartition/WorldPartitionRuntimeCell.cpp
--- a/Engine/Source/Runtime/Engine/Private/WorldPartition/WorldPartitionRuntimeCell.cpp
+++ b/Engine/Source/Runtime/Engine/Private/WorldPartition/WorldPartitionRuntimeCell.cpp
@@ -112,14 +112,14 @@ void UWorldPartitionRuntimeCell::SetDataLayers(const TArray<const UDataLayerInst
 	});
 
 	TArray<FName> SortedDataLayerInstanceNames;
-	bool bIsFirstDataLayerIsExternal = false;
+	bool bIsFirstDataLayerExternal = false;
 	Algo::Transform(SortedDataLayerInstances, SortedDataLayerInstanceNames, [](const UDataLayerInstance* DataLayerInstance) { return DataLayerInstance->GetDataLayerFName(); });
 	if (const UExternalDataLayerInstance* ExternalDataLayerInstance = Cast<UExternalDataLayerInstance>(SortedDataLayerInstances[0]))
 	{
-		bIsFirstDataLayerIsExternal = true;
+		bIsFirstDataLayerExternal = true;
 		ExternalDataLayerAsset = ExternalDataLayerInstance->GetExternalDataLayerAsset();
 	}
-	DataLayers = FDataLayerInstanceNames(SortedDataLayerInstanceNames, bIsFirstDataLayerIsExternal);
+	DataLayers = FDataLayerInstanceNames(SortedDataLayerInstanceNames, bIsFirstDataLayerExternal);
 }
 
 void UWorldPartitionRuntimeCell::DumpStateLog(FHierarchicalLogArchive& Ar) const
diff --git a/Engine/Source/Runtime/Engine/Private/WorldPartition/WorldPartitionStreamingGeneration.cpp b/Engine/Source/Runtime/Engine/Private/WorldPartition/WorldPartitionStreamingGeneration.cpp
--- a/Engine/Source/Runtime/Engine/Private/WorldPartition/WorldPartitionStreamingGeneration.cpp
+++ b/Engine/Source/Runtime/Engine/Private/WorldPartition/WorldPartitionStreamingGeneration.cpp
@@ -262,8 +262,12 @@ void FStreamingGenerationActorDescView::SetParentView(const FStreamingGeneration
 void FStreamingGenerationActorDescView::SetDataLayerInstanceNames(const FDataLayerInstanceNames& InDataLayerInstanceNames)
 {
 	check(!Super::HasResolvedDataLayerInstanceNames());
+	check(!InDataLayerInstanceNames.IsForcedEmptyNonExternalDataLayers())
 	ResolvedDataLayerInstanceNames = InDataLayerInstanceNames;
-	ResolvedDataLayerInstanceNames->SetIsForcedEmptyNonExternalDataLayers(bIsForcedNoDataLayers);
+	if (bIsForcedNoDataLayers)
+	{
+		ResolvedDataLayerInstanceNames->ForceEmptyNonExternalDataLayers();
+	}
 }
 
 bool FStreamingGenerationActorDescView::IsInvalidReference(const FGuid& InGuid, FInvalidReference* OutInvalidReference) const
@@ -321,20 +325,24 @@ void FStreamingGenerationActorDescView::SetForcedNoDataLayers()
 
 		if (ResolvedDataLayerInstanceNames.IsSet())
 		{
-			ResolvedDataLayerInstanceNames->SetIsForcedEmptyNonExternalDataLayers(bIsForcedNoDataLayers);
+			ResolvedDataLayerInstanceNames->ForceEmptyNonExternalDataLayers();
 		}
 
 		if (RuntimeDataLayerInstanceNames.IsSet())
 		{
-			RuntimeDataLayerInstanceNames->SetIsForcedEmptyNonExternalDataLayers(bIsForcedNoDataLayers);
+			RuntimeDataLayerInstanceNames->ForceEmptyNonExternalDataLayers();
 		}
 	}
 }
 
 void FStreamingGenerationActorDescView::SetRuntimeDataLayerInstanceNames(const FDataLayerInstanceNames& InRuntimeDataLayerInstanceNames)
 {
+	check(!InRuntimeDataLayerInstanceNames.IsForcedEmptyNonExternalDataLayers())
 	RuntimeDataLayerInstanceNames = InRuntimeDataLayerInstanceNames;
-	RuntimeDataLayerInstanceNames->SetIsForcedEmptyNonExternalDataLayers(bIsForcedNoDataLayers);
+	if (bIsForcedNoDataLayers)
+	{
+		RuntimeDataLayerInstanceNames->ForceEmptyNonExternalDataLayers();
+	}
 }
 
 void FStreamingGenerationActorDescView::SetRuntimeReferences(const TArray<FGuid>& InRuntimeReferences)
diff --git a/Engine/Source/Runtime/Engine/Public/WorldPartition/DataLayer/DataLayerInstanceNames.h b/Engine/Source/Runtime/Engine/Public/WorldPartition/DataLayer/DataLayerInstanceNames.h
--- a/Engine/Source/Runtime/Engine/Public/WorldPartition/DataLayer/DataLayerInstanceNames.h
+++ b/Engine/Source/Runtime/Engine/Public/WorldPartition/DataLayer/DataLayerInstanceNames.h
@@ -11,15 +11,15 @@ struct FDataLayerInstanceNames
 	GENERATED_USTRUCT_BODY()
 
 	FDataLayerInstanceNames()
-	: bIsFirstDataLayerIsExternal(false)
+	: bIsFirstDataLayerExternal(false)
 	{
 #if WITH_EDITOR
 		bIsForcedEmptyNonExternalDataLayers = false;
 #endif
 	}
 
-	FDataLayerInstanceNames(const TArray<FName>& InDataLayers, bool bInIsFirstDataLayerIsExternal)
-	: bIsFirstDataLayerIsExternal(bInIsFirstDataLayerIsExternal)
+	FDataLayerInstanceNames(const TArray<FName>& InDataLayers, bool bInIsFirstDataLayerExternal)
+	: bIsFirstDataLayerExternal(bInIsFirstDataLayerExternal)
 	, DataLayers(InDataLayers)
 	{
 #if WITH_EDITOR
@@ -29,8 +29,8 @@ struct FDataLayerInstanceNames
 
 	FDataLayerInstanceNames(const TArray<FName>& InNonExternalDataLayers, FName InExternalDataLayer)
 	{
-		bIsFirstDataLayerIsExternal = !InExternalDataLayer.IsNone();
-		if (bIsFirstDataLayerIsExternal)
+		bIsFirstDataLayerExternal = !InExternalDataLayer.IsNone();
+		if (bIsFirstDataLayerExternal)
 		{
 			DataLayers.Add(InExternalDataLayer);
 		}
@@ -39,68 +39,58 @@ struct FDataLayerInstanceNames
 
 	const FName GetExternalDataLayer() const
 	{
-		check(!bIsFirstDataLayerIsExternal || !DataLayers.IsEmpty());
-		return bIsFirstDataLayerIsExternal ? DataLayers[0] : NAME_None;
+		check(!bIsFirstDataLayerExternal || !DataLayers.IsEmpty());
+		return bIsFirstDataLayerExternal ? DataLayers[0] : NAME_None;
 	}
 
+	/**
+	* DO NOT USE DIRECTLY
+	* STL-like iterators to enable range-based for loop support.
+	*/
+	TArray<FName>::RangedForIteratorType     		 begin() { return DataLayers.begin(); }
+	TArray<FName>::RangedForConstIteratorType		 begin() const { return DataLayers.begin(); }
+	TArray<FName>::RangedForIteratorType      		 end() { return DataLayers.end(); }
+	TArray<FName>::RangedForConstIteratorType		 end() const { return DataLayers.end(); }
+	TArray<FName>::RangedForReverseIteratorType      rbegin() { return DataLayers.rbegin(); }
+	TArray<FName>::RangedForConstReverseIteratorType rbegin() const { return DataLayers.rbegin(); }
+	TArray<FName>::RangedForReverseIteratorType      rend() { return DataLayers.rend(); }
+	TArray<FName>::RangedForConstReverseIteratorType rend() const { return DataLayers.rend(); }
+
 	TArrayView<const FName> GetNonExternalDataLayers() const
 	{
 		static TArray<FName> EmptyArray;
-
-#if WITH_EDITOR
-		if (bIsForcedEmptyNonExternalDataLayers)
-		{
-			return EmptyArray;
-		}
-#endif
-		check(!bIsFirstDataLayerIsExternal || !DataLayers.IsEmpty());
-		const int32 Offset = bIsFirstDataLayerIsExternal ? 1 : 0;
+		check(!bIsFirstDataLayerExternal || !DataLayers.IsEmpty());
+		const int32 Offset = bIsFirstDataLayerExternal ? 1 : 0;
 		const int32 NonExternalDataLayersCount = DataLayers.Num() - Offset;
 		return NonExternalDataLayersCount > 0 ? MakeArrayView(&DataLayers[Offset], NonExternalDataLayersCount) : EmptyArray;
 	}
 
+	int32 Num() const { return DataLayers.Num(); }
+	bool IsEmpty() const { return DataLayers.IsEmpty(); }
+	bool Contains(FName InDataLayer) const { return DataLayers.Contains(InDataLayer); }
+	bool HasExternalDataLayer() const { return bIsFirstDataLayerExternal; }
+	const TArray<FName>& ToArray() const { return DataLayers; }
+
+private:
 #if WITH_EDITOR
-	TArray<FName> ToArray() const
+	bool IsForcedEmptyNonExternalDataLayers() const { return bIsForcedEmptyNonExternalDataLayers; }
+	void ForceEmptyNonExternalDataLayers()
 	{
-		if (bIsForcedEmptyNonExternalDataLayers)
+		bIsForcedEmptyNonExternalDataLayers = true;
+		if (bIsFirstDataLayerExternal)
 		{
-			static TArray<FName> EmptyArray;
-			return HasExternalDataLayer() ? TArray<FName>({ GetExternalDataLayer() }) : EmptyArray;
+			DataLayers = { DataLayers[0] };
 		}
-		return DataLayers;
-	}
-
-	int32 Num() const
-	{
-		if (bIsForcedEmptyNonExternalDataLayers)
+		else
 		{
-			return HasExternalDataLayer() ? 1 : 0;
+			DataLayers.Empty();
 		}
-		return DataLayers.Num();
 	}
-
-	bool IsEmpty() const
-	{
-		return Num() == 0;
-	}
-#endif
-
-	bool HasExternalDataLayer() const { return bIsFirstDataLayerIsExternal; }
-
-	const TArray<FName>& GetRawArray() const
-	{
-		return DataLayers;
-	}
-
-private:
-#if WITH_EDITOR
-	bool IsForcedEmptyNonExternalDataLayers() const { return bIsForcedEmptyNonExternalDataLayers; }
-	void SetIsForcedEmptyNonExternalDataLayers(bool bInNewValue) { bIsForcedEmptyNonExternalDataLayers = bInNewValue; }
 	bool bIsForcedEmptyNonExternalDataLayers;
 #endif
 
 	UPROPERTY()
-	bool bIsFirstDataLayerIsExternal;
+	bool bIsFirstDataLayerExternal;
 
 	UPROPERTY()
 	TArray<FName> DataLayers;
diff --git a/Engine/Source/Runtime/Engine/Public/WorldPartition/WorldPartitionRuntimeCell.h b/Engine/Source/Runtime/Engine/Public/WorldPartition/WorldPartitionRuntimeCell.h
--- a/Engine/Source/Runtime/Engine/Public/WorldPartition/WorldPartitionRuntimeCell.h
+++ b/Engine/Source/Runtime/Engine/Public/WorldPartition/WorldPartitionRuntimeCell.h
@@ -218,7 +218,7 @@ class UWorldPartitionRuntimeCell : public UObject, public IWorldPartitionCell, p
 	ENGINE_API virtual bool ContainsDataLayer(const UDataLayerAsset* DataLayerAsset) const override;
 	ENGINE_API virtual bool ContainsDataLayer(const UDataLayerInstance* DataLayerInstance) const override;
 	ENGINE_API virtual bool HasContentBundle() const override;
-	virtual const TArray<FName>& GetDataLayers() const override { return DataLayers.GetRawArray(); }
+	virtual const TArray<FName>& GetDataLayers() const override { return DataLayers.ToArray(); }
 	virtual FName GetExternalDataLayer() const override { return DataLayers.GetExternalDataLayer(); }
 	virtual bool HasAnyDataLayer(const TSet<FName>& InDataLayers) const override
 	{
