NetDriver
* Added exclusive csv cpu scope to FlushNet call to track time spent inside NetworkHandlerComponent (compression, encryption, socket send, etc.)

diff --git a/Engine/Source/Runtime/Engine/Private/Net/Experimental/Iris/DataStreamChannel.cpp b/Engine/Source/Runtime/Engine/Private/Net/Experimental/Iris/DataStreamChannel.cpp
--- a/Engine/Source/Runtime/Engine/Private/Net/Experimental/Iris/DataStreamChannel.cpp
+++ b/Engine/Source/Runtime/Engine/Private/Net/Experimental/Iris/DataStreamChannel.cpp
@@ -197,7 +197,7 @@ void UDataStreamChannel::PostTickDispatch()
 	}
 
 #if UE_NET_IRIS_CSV_STATS
-	CSV_SCOPED_TIMING_STAT_EXCLUSIVE(UDataStreamChannel_PostTickDispatch);
+	CSV_SCOPED_TIMING_STAT_EXCLUSIVE(UDataStreamChannel_PostTickDispatch_Write);
 #endif
 
 	IRIS_PROFILER_SCOPE(UDataStreamChannel_PostTickDispatch);
@@ -238,7 +238,7 @@ void UDataStreamChannel::Tick()
 	}
 
 #if UE_NET_IRIS_CSV_STATS
-	CSV_SCOPED_TIMING_STAT_EXCLUSIVE(UDataStreamChannel_Tick);
+	CSV_SCOPED_TIMING_STAT_EXCLUSIVE(UDataStreamChannel_Tick_Write);
 #endif
 
 	IRIS_PROFILER_SCOPE(UDataStreamChannel_Tick);
diff --git a/Engine/Source/Runtime/Engine/Private/NetConnection.cpp b/Engine/Source/Runtime/Engine/Private/NetConnection.cpp
--- a/Engine/Source/Runtime/Engine/Private/NetConnection.cpp
+++ b/Engine/Source/Runtime/Engine/Private/NetConnection.cpp
@@ -2091,6 +2091,8 @@ uint32 GNetOutBytes = 0;
 
 void UNetConnection::FlushNet(bool bIgnoreSimulation)
 {
+	CSV_SCOPED_TIMING_STAT_EXCLUSIVE(FlushNet);
+
 	check(Driver);
 
 	// Update info.
