Fixed uninitialized variable causing invalid checks while running content validation commandlet.

diff --git a/Engine/Source/Runtime/Engine/Public/WorldPartition/DataLayer/DataLayerInstanceNames.h b/Engine/Source/Runtime/Engine/Public/WorldPartition/DataLayer/DataLayerInstanceNames.h
--- a/Engine/Source/Runtime/Engine/Public/WorldPartition/DataLayer/DataLayerInstanceNames.h
+++ b/Engine/Source/Runtime/Engine/Public/WorldPartition/DataLayer/DataLayerInstanceNames.h
@@ -29,6 +29,9 @@ struct FDataLayerInstanceNames
 
 	FDataLayerInstanceNames(const TArray<FName>& InNonExternalDataLayers, FName InExternalDataLayer)
 	{
+#if WITH_EDITOR
+		bIsForcedEmptyNonExternalDataLayers = false;
+#endif
 		bIsFirstDataLayerExternal = !InExternalDataLayer.IsNone();
 		if (bIsFirstDataLayerExternal)
 		{
