Fix for freshly built virtual textures not showing encoder metadata.

diff --git a/Engine/Source/Runtime/Engine/Private/TextureDerivedDataTask.cpp b/Engine/Source/Runtime/Engine/Private/TextureDerivedDataTask.cpp
--- a/Engine/Source/Runtime/Engine/Private/TextureDerivedDataTask.cpp
+++ b/Engine/Source/Runtime/Engine/Private/TextureDerivedDataTask.cpp
@@ -766,6 +766,7 @@ static void DDC1_BuildTexture(
 		DerivedData->SizeY = DerivedData->VTData->Height;
 		DerivedData->PixelFormat = DerivedData->VTData->LayerTypes[0];
 		DerivedData->SetNumSlices(1);
+		DerivedData->ResultMetadata = InBuildResultMetadata;
 
 		// Verify our predicted count matches.
 		check(PredictedInfo.NumMips == DerivedData->VTData->GetNumMips());
