[NetPhysics] Implementing NetworkPhysicsSettingsComponent as a way to set individual settings for networked physics per actor.

Currently implements settings for physics replication through Predictive Interpolation and physics Resimulation.
Settings are distributed to networked physics systems on both GT and PT.

Also introduces a new internal callback OnInitialize_Internal() in Chaos::TSimCallbackObject that is only called once, on the physics-thread after the first initial input data has been marshalled.
Used in both PhysicsReplication.h and NetworkPhysicsSettingsComponent.h

#virtualized

diff --git a/Engine/Source/Runtime/Engine/Private/PhysicsEngine/NetworkPhysicsComponent.cpp b/Engine/Source/Runtime/Engine/Private/PhysicsEngine/NetworkPhysicsComponent.cpp
--- a/Engine/Source/Runtime/Engine/Private/PhysicsEngine/NetworkPhysicsComponent.cpp
+++ b/Engine/Source/Runtime/Engine/Private/PhysicsEngine/NetworkPhysicsComponent.cpp
@@ -1,7 +1,7 @@
 // Copyright Epic Games, Inc. All Rights Reserved.
 
 #include "Physics/NetworkPhysicsComponent.h"
-
+#include "Physics/NetworkPhysicsSettingsComponent.h"
 #include "Components/PrimitiveComponent.h"
 #include "EngineLogs.h"
 #include "EngineUtils.h"
@@ -387,6 +387,24 @@ UNetworkPhysicsComponent::UNetworkPhysicsComponent() : Super()
 
 void UNetworkPhysicsComponent::InitPhysics()
 {
+	AActor* Owner = GetOwner();
+	if (Owner)
+	{
+		if (UNetworkPhysicsSettingsComponent* PhysicsSettings = Owner->GetComponentByClass<UNetworkPhysicsSettingsComponent>())
+		{
+			InputRedundancy = PhysicsSettings->ResimulationSettings.bOverrideRedundantInputs ? PhysicsSettings->ResimulationSettings.RedundantInputs : InputRedundancy;
+			StateRedundancy = PhysicsSettings->ResimulationSettings.bOverrideRedundantStates ? PhysicsSettings->ResimulationSettings.RedundantStates : StateRedundancy;
+		}
+
+		if (APawn* Pawn = Cast<APawn>(Owner))
+		{
+			FRepMovement& RepMovement = Pawn->GetReplicatedMovement_Mutable();
+			RepMovement.LocationQuantizationLevel = EVectorQuantization::RoundTwoDecimals;
+			RepMovement.RotationQuantizationLevel = ERotatorQuantization::ShortComponents;
+			RepMovement.VelocityQuantizationLevel = EVectorQuantization::RoundTwoDecimals;
+		}
+	}
+
 	bAutoActivate = true;
 	bWantsInitializeComponent = true;
 	SetIsReplicatedByDefault(true);
@@ -394,14 +412,6 @@ void UNetworkPhysicsComponent::InitPhysics()
 
 	StateOffsets.SetNumZeroed(StateRedundancy + 1);
 	InputOffsets.SetNumZeroed(InputRedundancy + 1);
-
-	if (APawn* Pawn = Cast<APawn>(GetOwner()))
-	{
-		FRepMovement& RepMovement = Pawn->GetReplicatedMovement_Mutable();
-		RepMovement.LocationQuantizationLevel = EVectorQuantization::RoundTwoDecimals;
-		RepMovement.RotationQuantizationLevel = ERotatorQuantization::ShortComponents;
-		RepMovement.VelocityQuantizationLevel = EVectorQuantization::RoundTwoDecimals;
-	}
 }
 
 void UNetworkPhysicsComponent::BeginPlay()
diff --git a/Engine/Source/Runtime/Engine/Private/PhysicsEngine/NetworkPhysicsSettingsComponent.cpp b/Engine/Source/Runtime/Engine/Private/PhysicsEngine/NetworkPhysicsSettingsComponent.cpp
new file mode 100644
--- /dev/null
+++ b/Engine/Source/Runtime/Engine/Private/PhysicsEngine/NetworkPhysicsSettingsComponent.cpp
@@ -0,0 +1,123 @@
+// Copyright Epic Games, Inc. All Rights Reserved.
+
+/*=============================================================================
+	NetworkPhysicsSettingsComponent.cpp
+	Handles data distribution of networked physics settings to systems that need it, on both Game-Thread and Physics-Thread.
+=============================================================================*/
+
+#include "Physics/NetworkPhysicsSettingsComponent.h"
+#include "Engine/Engine.h"
+#include "Chaos/Declares.h"
+#include "Physics/Experimental/PhysScene_Chaos.h"
+#include "PBDRigidsSolver.h"
+#include "PhysicsReplication.h"
+#include "Components/PrimitiveComponent.h"
+
+namespace PhysicsReplicationCVars
+{
+	namespace ResimulationCVars
+	{
+		int32 SimProxyRepMode = -1;
+		static FAutoConsoleVariableRef CVarSimProxyRepMode(TEXT("np2.Resim.SimProxyRepMode"), SimProxyRepMode, TEXT("All actors with a NetworkPhysicsSettingsComponent and that are running resimulation and is ROLE_SimulatedProxy will change their physics replication mode. -1 = Disabled, 0 = Default, 1 = PredictiveInterpolation, 2 = Resimulation"));
+	}
+}
+
+
+UNetworkPhysicsSettingsComponent::UNetworkPhysicsSettingsComponent()
+{
+	bWantsInitializeComponent = true;
+	bAutoActivate = true;
+}
+
+void UNetworkPhysicsSettingsComponent::InitializeComponent()
+{
+	Super::InitializeComponent();
+
+	using namespace Chaos;
+	NetworkPhysicsSettingsAsync = nullptr;
+	if (UWorld* World = GetWorld())
+	{
+		if (FPhysScene* PhysScene = World->GetPhysicsScene())
+		{
+			if (Chaos::FPhysicsSolver* Solver = PhysScene->GetSolver())
+			{
+				NetworkPhysicsSettingsAsync = Solver->CreateAndRegisterSimCallbackObject_External<FNetworkPhysicsSettingsComponentAsync>();
+				
+				// Marshal settings data from GT to PT
+				if (NetworkPhysicsSettingsAsync)
+				{
+					if (AActor* Owner = GetOwner())
+					{
+						if (UPrimitiveComponent* RootPrimComp = Cast<UPrimitiveComponent>(Owner->GetRootComponent()))
+						{
+							if (Chaos::FConstPhysicsObjectHandle PhysicsObject = RootPrimComp->GetPhysicsObjectByName(NAME_None))
+							{
+								FNetworkPhysicsSettingsAsyncInput* AsyncInput = NetworkPhysicsSettingsAsync->GetProducerInputData_External();
+								AsyncInput->PhysicsObject = PhysicsObject;
+								AsyncInput->Settings.GeneralSettings = GeneralSettings;
+								AsyncInput->Settings.ResimulationSettings = ResimulationSettings;
+								AsyncInput->Settings.PredictiveInterpolationSettings = PredictiveInterpolationSettings;
+							}
+						}
+					}
+				}
+			}
+		}
+	}
+}
+
+void UNetworkPhysicsSettingsComponent::UninitializeComponent()
+{
+	Super::UninitializeComponent();
+
+	using namespace Chaos;
+	if (UWorld* World = GetWorld())
+	{
+		if (FPhysScene* PhysScene = World->GetPhysicsScene())
+		{
+			if (Chaos::FPhysicsSolver* Solver = PhysScene->GetSolver())
+			{
+				Solver->UnregisterAndFreeSimCallbackObject_External(NetworkPhysicsSettingsAsync);
+			}
+		}
+	}
+	NetworkPhysicsSettingsAsync = nullptr;
+}
+
+void UNetworkPhysicsSettingsComponent::BeginPlay()
+{
+	Super::BeginPlay();
+
+	// Apply overrides on actor
+	if (AActor* Owner = GetOwner())
+	{
+		if ((GeneralSettings.bOverrideSimProxyRepMode || PhysicsReplicationCVars::ResimulationCVars::SimProxyRepMode >= 0)
+			&& Owner->GetLocalRole() == ENetRole::ROLE_SimulatedProxy)
+		{
+			EPhysicsReplicationMode RepMode = GeneralSettings.bOverrideSimProxyRepMode ? GeneralSettings.SimProxyRepMode : static_cast<EPhysicsReplicationMode>(PhysicsReplicationCVars::ResimulationCVars::SimProxyRepMode);
+			Owner->SetPhysicsReplicationMode(RepMode);
+		}
+	}
+}
+
+
+#pragma region // FNetworkPhysicsSettingsComponentAsync
+
+void FNetworkPhysicsSettingsComponentAsync::OnPostInitialize_Internal()
+{
+	// Receive settings data on PT from GT
+	if (const FNetworkPhysicsSettingsAsyncInput* AsyncInput = GetConsumerInput_Internal())
+	{
+		Settings = AsyncInput->Settings;
+		
+		if (Chaos::FPBDRigidsSolver* RigidsSolver = static_cast<Chaos::FPBDRigidsSolver*>(GetSolver()))
+		{
+			if (FPhysicsReplicationAsync* PhysRep = RigidsSolver->GetPhysicsReplication())
+			{
+				PhysRep->RegisterSettings(AsyncInput->PhysicsObject, Settings);
+			}
+		}
+	}
+}
+
+#pragma endregion // FNetworkPhysicsSettingsComponentAsync
diff --git a/Engine/Source/Runtime/Engine/Private/PhysicsEngine/PhysicsReplication.cpp b/Engine/Source/Runtime/Engine/Private/PhysicsEngine/PhysicsReplication.cpp
--- a/Engine/Source/Runtime/Engine/Private/PhysicsEngine/PhysicsReplication.cpp
+++ b/Engine/Source/Runtime/Engine/Private/PhysicsEngine/PhysicsReplication.cpp
@@ -1,7 +1,7 @@
 // Copyright Epic Games, Inc. All Rights Reserved.
 
 /*=============================================================================
-	PhysicsReplication.cpp: Code for updating body instance physics state based on replication
+	PhysicsReplication.cpp: Code for keeping replicated physics objects in sync with the server based on replicated server state data.
 =============================================================================*/
 
 #include "PhysicsReplication.h"
@@ -11,12 +11,12 @@
 #include "Physics/Experimental/PhysScene_Chaos.h"
 #include "PhysicsEngine/PhysicsSettings.h"
 #include "Chaos/PhysicsObjectInternalInterface.h"
-#include "Chaos/PhysicsObject.h"
 #include "GameFramework/PlayerController.h"
 #include "GameFramework/PlayerState.h"
 #include "PhysicsProxy/SingleParticlePhysicsProxy.h"
 #include "PBDRigidsSolver.h"
 #include "Chaos/DebugDrawQueue.h"
+#include "Chaos/Particles.h"
 
 namespace CharacterMovementCVars
 {
@@ -77,133 +77,133 @@ namespace CharacterMovementCVars
 
 namespace PhysicsReplicationCVars
 {
-	static int32 SkipSkeletalRepOptimization = 1;
+	int32 SkipSkeletalRepOptimization = 1;
 	static FAutoConsoleVariableRef CVarSkipSkeletalRepOptimization(TEXT("p.SkipSkeletalRepOptimization"), SkipSkeletalRepOptimization, TEXT("If true, we don't move the skeletal mesh component during replication. This is ok because the skeletal mesh already polls physx after its results"));
 #if !UE_BUILD_SHIPPING
 	int32 LogPhysicsReplicationHardSnaps = 0;
 	static FAutoConsoleVariableRef CVarLogPhysicsReplicationHardSnaps(TEXT("p.LogPhysicsReplicationHardSnaps"), LogPhysicsReplicationHardSnaps, TEXT(""));
 #endif
 
-	static int32 EnableDefaultReplication = 0;
+	int32 EnableDefaultReplication = 0;
 	static FAutoConsoleVariableRef CVarEnableDefaultReplication(TEXT("np2.EnableDefaultReplication"), EnableDefaultReplication, TEXT("Enable default replication in the networked physics prediction flow."));
 
 	namespace ResimulationCVars
 	{
-		static bool bRuntimeCorrectionEnabled = true;
+		bool bRuntimeCorrectionEnabled = true;
 		static FAutoConsoleVariableRef CVarResimRuntimeCorrectionEnabled(TEXT("np2.Resim.RuntimeCorrectionEnabled"), bRuntimeCorrectionEnabled, TEXT("Apply runtime corrections while error is smalle enough not to trigger a resim."));
 
-		static bool bDisableReplicationOnInteraction = true;
+		bool bDisableReplicationOnInteraction = true;
 		static FAutoConsoleVariableRef CVarResimDisableReplicationOnInteraction(TEXT("np2.Resim.DisableReplicationOnInteraction"), bDisableReplicationOnInteraction, TEXT("If a resim object interacts with another object not running resimulation, deactivate that objects replication until interaction stops."));
 
-		static float PosStabilityMultiplier = 0.5f;
+		float PosStabilityMultiplier = 0.5f;
 		static FAutoConsoleVariableRef CVarResimPosStabilityMultiplier(TEXT("np2.Resim.PosStabilityMultiplier"), PosStabilityMultiplier, TEXT("Recommended range between 0.0-1.0. Lower value means more stable positional corrections."));
 
-		static float RotStabilityMultiplier = 1.0f;
+		float RotStabilityMultiplier = 1.0f;
 		static FAutoConsoleVariableRef CVarResimRotStabilityMultiplier(TEXT("np2.Resim.RotStabilityMultiplier"), RotStabilityMultiplier, TEXT("Recommended range between 0.0-1.0. Lower value means more stable rotational corrections."));
 	}
 
 	namespace PredictiveInterpolationCVars
 	{
-		static float PosCorrectionTimeBase = 0.0f;
+		float PosCorrectionTimeBase = 0.0f;
 		static FAutoConsoleVariableRef CVarPosCorrectionTimeBase(TEXT("np2.PredictiveInterpolation.PosCorrectionTimeBase"), PosCorrectionTimeBase, TEXT("Base time to correct positional offset over. RTT * PosCorrectionTimeMultiplier are added on top of this."));
 
-		static float PosCorrectionTimeMin = 0.1f;
+		float PosCorrectionTimeMin = 0.1f;
 		static FAutoConsoleVariableRef CVarPosCorrectionTimeMin(TEXT("np2.PredictiveInterpolation.PosCorrectionTimeMin"), PosCorrectionTimeMin, TEXT("Min time time to correct positional offset over. DeltaSeconds is added on top of this."));
 
-		static float PosCorrectionTimeMultiplier = 1.0f;
+		float PosCorrectionTimeMultiplier = 1.0f;
 		static FAutoConsoleVariableRef CVarPosCorrectionTimeMultiplier(TEXT("np2.PredictiveInterpolation.PosCorrectionTimeMultiplier"), PosCorrectionTimeMultiplier, TEXT("Multiplier to adjust how much of RTT (network Round Trip Time) to add to positional offset correction."));
 
-		static float RotCorrectionTimeBase = 0.0f;
+		float RotCorrectionTimeBase = 0.0f;
 		static FAutoConsoleVariableRef CVarRotCorrectionTimeBase(TEXT("np2.PredictiveInterpolation.RotCorrectionTimeBase"), RotCorrectionTimeBase, TEXT("Base time to correct positional offset over. RTT * PosCorrectionTimeMultiplier are added on top of this."));
 
-		static float RotCorrectionTimeMin = 0.1f;
+		float RotCorrectionTimeMin = 0.1f;
 		static FAutoConsoleVariableRef CVarRotCorrectionTimeMin(TEXT("np2.PredictiveInterpolation.RotCorrectionTimeMin"), RotCorrectionTimeMin, TEXT("Min time time to correct rotational offset over. DeltaSeconds is added on top of this."));
 
-		static float RotCorrectionTimeMultiplier = 1.0f;
+		float RotCorrectionTimeMultiplier = 1.0f;
 		static FAutoConsoleVariableRef CVarRotCorrectionTimeMultiplier(TEXT("np2.PredictiveInterpolation.RotCorrectionTimeMultiplier"), RotCorrectionTimeMultiplier, TEXT("Multiplier to adjust how much of RTT (network Round Trip Time) to add to positional offset correction."));
 
-		static float PosInterpolationTimeMultiplier = 1.1f;
+		float PosInterpolationTimeMultiplier = 1.1f;
 		static FAutoConsoleVariableRef CVarInterpolationTimeMultiplier(TEXT("np2.PredictiveInterpolation.InterpolationTimeMultiplier"), PosInterpolationTimeMultiplier, TEXT("Multiplier to adjust the replication interpolation time which is based on the sendrate of replication data from the server."));
 		
-		static float RotInterpolationTimeMultiplier = 1.25f;
+		float RotInterpolationTimeMultiplier = 1.25f;
 		static FAutoConsoleVariableRef CVarRotInterpolationTimeMultiplier(TEXT("np2.PredictiveInterpolation.RotInterpolationTimeMultiplier"), RotInterpolationTimeMultiplier, TEXT("Multiplier to adjust the rotational replication interpolation time which is based on the sendrate of replication data from the server."));
 		
-		static float AverageReceiveIntervalSmoothing = 3.0f;
+		float AverageReceiveIntervalSmoothing = 3.0f;
 		static FAutoConsoleVariableRef CVarAverageReceiveIntervalSmoothing(TEXT("np2.PredictiveInterpolation.AverageReceiveIntervalSmoothing"), AverageReceiveIntervalSmoothing, TEXT("Recommended range: 1.0 - 5.0. Higher value makes the average receive interval adjust itself slower, reducing spikes in InterpolationTime."));
 
-		static float ExtrapolationTimeMultiplier = 3.0f;
+		float ExtrapolationTimeMultiplier = 3.0f;
 		static FAutoConsoleVariableRef CVarExtrapolationTimeMultiplier(TEXT("np2.PredictiveInterpolation.ExtrapolationTimeMultiplier"), ExtrapolationTimeMultiplier, TEXT("Multiplier to adjust the time to extrapolate the target forward over, the time is based on current send-rate."));
 
-		static float ExtrapolationMinTime = 0.75f;
+		float ExtrapolationMinTime = 0.75f;
 		static FAutoConsoleVariableRef CVarExtrapolationMinTime(TEXT("np2.PredictiveInterpolation.ExtrapolationMinTime"), ExtrapolationMinTime, TEXT("Clamps minimum extrapolation time. Value in seconds. Disable minimum clamp by setting to 0."));
 
-		static float MinExpectedDistanceCovered = 0.5f;
+		float MinExpectedDistanceCovered = 0.5f;
 		static FAutoConsoleVariableRef CVarMinExpectedDistanceCovered(TEXT("np2.PredictiveInterpolation.MinExpectedDistanceCovered"), MinExpectedDistanceCovered, TEXT("Value between 0-1, in percentage where 0.25 = 25%. How much of the expected distance based on replication velocity should the object have covered in a simulation tick to Not be considered stuck."));
 
-		static float ErrorAccumulationDecreaseMultiplier = 0.5f;
+		float ErrorAccumulationDecreaseMultiplier = 0.5f;
 		static FAutoConsoleVariableRef CVarErrorAccumulationDecreaseMultiplier(TEXT("np2.PredictiveInterpolation.ErrorAccumulationDecreaseMultiplier"), ErrorAccumulationDecreaseMultiplier, TEXT("Multiplier to adjust how fast we decrease accumulated error time when we no longer accumulate error."));
 
-		static float ErrorAccumulationSeconds = 3.0f;
+		float ErrorAccumulationSeconds = 3.0f;
 		static FAutoConsoleVariableRef CVarErrorAccumulationSeconds(TEXT("np2.PredictiveInterpolation.ErrorAccumulationSeconds"), ErrorAccumulationSeconds, TEXT("Perform a reposition if replication have not been able to cover the min expected distance towards the target for this amount of time."));
 		
-		static bool bDisableErrorVelocityLimits = false;
+		bool bDisableErrorVelocityLimits = false;
 		static FAutoConsoleVariableRef CVarDisableErrorVelocityLimits(TEXT("np2.PredictiveInterpolation.DisableErrorVelocityLimits"), bDisableErrorVelocityLimits, TEXT("Disable the velocity limit and allow error accumulation at any velocity."));
 		
-		static float ErrorAccLinVelMaxLimit = 50.0f;
+		float ErrorAccLinVelMaxLimit = 50.0f;
 		static FAutoConsoleVariableRef CVarErrorAccLinVelMaxLimit(TEXT("np2.PredictiveInterpolation.ErrorAccLinVelMaxLimit"), ErrorAccLinVelMaxLimit, TEXT("If target velocity is below this limit we check for desync to trigger softsnap and accumulate time to build up to a hardsnap."));
 		
-		static float ErrorAccAngVelMaxLimit = 1.5f;
+		float ErrorAccAngVelMaxLimit = 1.5f;
 		static FAutoConsoleVariableRef CVarErrorAccAngVelMaxLimit(TEXT("np2.PredictiveInterpolation.ErrorAccAngVelMaxLimit"), ErrorAccAngVelMaxLimit, TEXT("If target angular velocity (in degrees) is below this limit we check for desync to trigger softsnap and accumulate time to build up to a hardsnap."));
 		
-		static float SoftSnapPosStrength = 0.5f;
+		float SoftSnapPosStrength = 0.5f;
 		static FAutoConsoleVariableRef CVarSoftSnapPosStrength(TEXT("np2.PredictiveInterpolation.SoftSnapPosStrength"), SoftSnapPosStrength, TEXT("Value in percent between 0.0 - 1.0 representing how much to softsnap each tick of the remaining distance."));
 		
-		static float SoftSnapRotStrength = 0.5f;
+		float SoftSnapRotStrength = 0.5f;
 		static FAutoConsoleVariableRef CVarSoftSnapRotStrength(TEXT("np2.PredictiveInterpolation.SoftSnapRotStrength"), SoftSnapRotStrength, TEXT("Value in percent between 0.0 - 1.0 representing how much to softsnap each tick of the remaining distance."));
 
-		static bool bSoftSnapToSource = false;
+		bool bSoftSnapToSource = false;
 		static FAutoConsoleVariableRef CVarSoftSnapToSource(TEXT("np2.PredictiveInterpolation.SoftSnapToSource"), bSoftSnapToSource, TEXT("If true, soft snap will be performed towards the source state of the current target instead of the predicted state of the current target."));
 
-		static float EarlyOutDistanceSqr = 1.0f;
+		float EarlyOutDistanceSqr = 1.0f;
 		static FAutoConsoleVariableRef CVarEarlyOutDistanceSqr(TEXT("np2.PredictiveInterpolation.EarlyOutDistanceSqr"), EarlyOutDistanceSqr, TEXT("Squared value. If object is within this distance from the source target, early out from replication and apply sleep if replicated."));
 		
-		static float EarlyOutAngle = 1.5f;
+		float EarlyOutAngle = 1.5f;
 		static FAutoConsoleVariableRef CVarEarlyOutAngle(TEXT("np2.PredictiveInterpolation.EarlyOutAngle"), EarlyOutAngle, TEXT("If object is within this rotational angle (in degrees) from the source target, early out from replication and apply sleep if replicated."));
 		
-		static bool bEarlyOutWithVelocity = true;
+		bool bEarlyOutWithVelocity = true;
 		static FAutoConsoleVariableRef CVarEarlyOutWithVelocity(TEXT("np2.PredictiveInterpolation.EarlyOutWithVelocity"), bEarlyOutWithVelocity, TEXT("If true, allow replication logic to early out if current velocities are driving replication well enough. If false, only early out if target velocity is zero."));
 
-		static bool bSkipVelocityRepOnPosEarlyOut = true;
+		bool bSkipVelocityRepOnPosEarlyOut = true;
 		static FAutoConsoleVariableRef CVarSkipVelocityRepOnPosEarlyOut(TEXT("np2.PredictiveInterpolation.SkipVelocityRepOnPosEarlyOut"), bSkipVelocityRepOnPosEarlyOut, TEXT("If true, don't run linear velocity replication if position can early out but angular can't early out."));
 
-		static bool bPostResimWaitForUpdate = false;
+		bool bPostResimWaitForUpdate = false;
 		static FAutoConsoleVariableRef CVarPostResimWaitForUpdate(TEXT("np2.PredictiveInterpolation.PostResimWaitForUpdate"), bPostResimWaitForUpdate, TEXT("After a resimulation, wait for replicated states that correspond to post-resim state before processing replication again."));
 		
-		static bool bVelocityBased = true;
+		bool bVelocityBased = true;
 		static FAutoConsoleVariableRef CVarVelocityBased(TEXT("np2.PredictiveInterpolation.VelocityBased"), bVelocityBased, TEXT("When true, predictive interpolation replication mode will only apply linear velocity and angular velocity"));
 		
-		static bool bDisableSoftSnap = false;
+		bool bDisableSoftSnap = false;
 		static FAutoConsoleVariableRef CVarDisableSoftSnap(TEXT("np2.PredictiveInterpolation.DisableSoftSnap"), bDisableSoftSnap, TEXT("When true, predictive interpolation will not use softsnap to correct the replication with when velocity fails. Hardsnap will still eventually kick in if replication can't reach the target."));
 	
-		static bool bAlwaysHardSnap = false;
+		bool bAlwaysHardSnap = false;
 		static FAutoConsoleVariableRef CVarAlwaysHardSnap(TEXT("np2.PredictiveInterpolation.AlwaysHardSnap"), bAlwaysHardSnap, TEXT("When true, predictive interpolation replication mode will always hard snap. Used as a backup measure"));
 
-		static bool bSkipReplication = false;
+		bool bSkipReplication = false;
 		static FAutoConsoleVariableRef CVarSkipReplication(TEXT("np2.PredictiveInterpolation.SkipReplication"), bSkipReplication, TEXT("When true, predictive interpolation is not applied anymore letting the object simulate freely instead"));
 
-		static bool bDontClearTarget = false;
+		bool bDontClearTarget = false;
 		static FAutoConsoleVariableRef CVarDontClearTarget(TEXT("np2.PredictiveInterpolation.DontClearTarget"), bDontClearTarget, TEXT("When true, predictive interpolation will not lose track of the last replicated state after coming to rest."));
 		
-		static bool bDrawDebugTargets = false;
+		bool bDrawDebugTargets = false;
 		static FAutoConsoleVariableRef CVarDrawDebugTargets(TEXT("np2.PredictiveInterpolation.DrawDebugTargets"), bDrawDebugTargets, TEXT("Draw target states, color coded by which ServerFrame they originate from, replicated targets are large and extrapolated targets are small. There is a Z offset to the draw calls."));
 		
-		static bool bDrawDebugVectors = false;
+		bool bDrawDebugVectors = false;
 		static FAutoConsoleVariableRef CVarDrawDebugVectors(TEXT("np2.PredictiveInterpolation.DrawDebugVectors"), bDrawDebugVectors, TEXT("Draw replication vectors, target velocity, replicated velocity, velocity change between replication calls etc."));
 		
-		static float SleepSecondsClearTarget = 15.0f;
+		float SleepSecondsClearTarget = 15.0f;
 		static FAutoConsoleVariableRef CVarSleepSecondsClearTarget(TEXT("np2.PredictiveInterpolation.SleepSecondsClearTarget"), SleepSecondsClearTarget, TEXT("Wait for the object to sleep for this many seconds before clearing the replication target, to ensure nothing wakes up the object just after it goes to sleep on the client."));
 		
-		static int32 TargetTickAlignmentClampMultiplier = 10;
+		int32 TargetTickAlignmentClampMultiplier = 10;
 		static FAutoConsoleVariableRef CVarTargetTickAlignmentClampMultiplier(TEXT("np2.PredictiveInterpolation.TargetTickAlignmentClampMultiplier"), TargetTickAlignmentClampMultiplier, TEXT("Multiplier to adjust clamping of target alignment via TickCount. Multiplier is performed on AverageReceiveInterval."));
 	}
 
@@ -745,10 +745,36 @@ void FPhysicsReplication::PrepareAsyncData_External(const FRigidBodyErrorCorrect
 	AsyncInput->ErrorCorrection.AngularVelocityCoefficient = AngularVelocityCoefficient;
 }
 
-#pragma region AsyncPhysicsReplication
+#pragma region FPhysicsReplicationAsync
 void FPhysicsReplicationAsync::OnPhysicsObjectUnregistered_Internal(Chaos::FConstPhysicsObjectHandle PhysicsObject)
 {
 	ObjectToTarget.Remove(PhysicsObject);
+	ObjectToSettings.Remove(PhysicsObject);
+}
+
+void FPhysicsReplicationAsync::RegisterSettings(Chaos::FConstPhysicsObjectHandle PhysicsObject, FNetworkPhysicsSettingsAsync InSettings)
+{
+	if (PhysicsObject != nullptr)
+	{
+		FNetworkPhysicsSettingsAsync& Settings = ObjectToSettings.FindOrAdd(PhysicsObject);
+		Settings = InSettings;
+	}
+}
+
+void FPhysicsReplicationAsync::FetchObjectSettings(Chaos::FConstPhysicsObjectHandle PhysicsObject)
+{
+	FNetworkPhysicsSettingsAsync* CustomSettings = ObjectToSettings.Find(PhysicsObject);
+	SettingsCurrent = (CustomSettings != nullptr) ? *CustomSettings : SettingsDefault;
+}
+
+void FPhysicsReplicationAsync::OnPostInitialize_Internal()
+{
+	Chaos::FPBDRigidsSolver* RigidsSolver = static_cast<Chaos::FPBDRigidsSolver*>(GetSolver());
+	
+	if (ensure(RigidsSolver))
+	{
+		RigidsSolver->SetPhysicsReplication(this);
+	}
 }
 
 void FPhysicsReplicationAsync::OnPreSimulate_Internal()
@@ -769,7 +795,7 @@ void FPhysicsReplicationAsync::OnPreSimulate_Internal()
 		if (bRewindDataExist && RewindData->IsResim())
 		{
 			// TODO, Handle the transition from post-resim to interpolation better (disabled by default, resim vs replication interaction is handled via FPhysicsReplicationAsync::CacheResimInteractions)
-			if (PhysicsReplicationCVars::PredictiveInterpolationCVars::bPostResimWaitForUpdate && RewindData->IsFinalResim())
+			if (SettingsCurrent.PredictiveInterpolationSettings.GetPostResimWaitForUpdate() && RewindData->IsFinalResim())
 			{
 				for (auto Itr = ObjectToTarget.CreateIterator(); Itr; ++Itr)
 				{
@@ -1060,6 +1086,9 @@ void FPhysicsReplicationAsync::ApplyTargetStatesAsync(const float DeltaSeconds,
 
 			if (FPBDRigidParticleHandle* RigidHandle = Handle->CastToRigidParticle())
 			{
+				// Cache custom settings for this object if there are any
+				FetchObjectSettings(POHandle);
+
 				const EPhysicsReplicationMode RepMode = Target.IsWaiting() ? Target.RepModeOverride : Target.RepMode;
 				switch (RepMode)
 				{
@@ -1536,7 +1565,7 @@ bool FPhysicsReplicationAsync::PredictiveInterpolation(Chaos::FPBDRigidParticleH
 		Target.AccumulatedErrorSeconds = 0;
 	}
 
-	if (PhysicsReplicationCVars::PredictiveInterpolationCVars::bDisableSoftSnap && PhysicsReplicationCVars::PredictiveInterpolationCVars::bVelocityBased)
+	if (SettingsCurrent.PredictiveInterpolationSettings.GetDisableSoftSnap() && PhysicsReplicationCVars::PredictiveInterpolationCVars::bVelocityBased)
 	{
 		bSoftSnap = false;
 	}
@@ -1576,16 +1605,16 @@ bool FPhysicsReplicationAsync::PredictiveInterpolation(Chaos::FPBDRigidParticleH
 	{
 		// Calculate interpolation time based on current average receive rate
 		const float AverageReceiveIntervalSeconds = Target.AverageReceiveInterval * DeltaSeconds;
-		const float InterpolationTime = AverageReceiveIntervalSeconds * PhysicsReplicationCVars::PredictiveInterpolationCVars::PosInterpolationTimeMultiplier;
+		const float InterpolationTime = AverageReceiveIntervalSeconds * SettingsCurrent.PredictiveInterpolationSettings.GetPosInterpolationTimeMultiplier();
 
 		// Calculate position correction time based on current Round Trip Time
 		const float RTT = LatencyOneWay * 2.f;
-		const float PosCorrectionTime = FMath::Max(PhysicsReplicationCVars::PredictiveInterpolationCVars::PosCorrectionTimeBase + AverageReceiveIntervalSeconds + RTT * PhysicsReplicationCVars::PredictiveInterpolationCVars::PosCorrectionTimeMultiplier,
-			DeltaSeconds + PhysicsReplicationCVars::PredictiveInterpolationCVars::PosCorrectionTimeMin);
-		const float RotCorrectionTime = FMath::Max(PhysicsReplicationCVars::PredictiveInterpolationCVars::RotCorrectionTimeBase + AverageReceiveIntervalSeconds + RTT * PhysicsReplicationCVars::PredictiveInterpolationCVars::RotCorrectionTimeMultiplier,
-			DeltaSeconds + PhysicsReplicationCVars::PredictiveInterpolationCVars::RotCorrectionTimeMin);
+		const float PosCorrectionTime = FMath::Max(SettingsCurrent.PredictiveInterpolationSettings.GetPosCorrectionTimeBase() + AverageReceiveIntervalSeconds + RTT * SettingsCurrent.PredictiveInterpolationSettings.GetPosCorrectionTimeMultiplier(),
+			DeltaSeconds + SettingsCurrent.PredictiveInterpolationSettings.GetPosCorrectionTimeMin());
+		const float RotCorrectionTime = FMath::Max(SettingsCurrent.PredictiveInterpolationSettings.GetRotCorrectionTimeBase() + AverageReceiveIntervalSeconds + RTT * SettingsCurrent.PredictiveInterpolationSettings.GetRotCorrectionTimeMultiplier(),
+			DeltaSeconds + SettingsCurrent.PredictiveInterpolationSettings.GetRotCorrectionTimeMin());
 
-		if ((bXCanEarlyOut && PhysicsReplicationCVars::PredictiveInterpolationCVars::bSkipVelocityRepOnPosEarlyOut) == false)
+		if ((bXCanEarlyOut && SettingsCurrent.PredictiveInterpolationSettings.GetSkipVelocityRepOnPosEarlyOut()) == false)
 		{	// --- Velocity Replication ---
 			
 			// Get PosDiff
@@ -1649,7 +1678,7 @@ bool FPhysicsReplicationAsync::PredictiveInterpolation(Chaos::FPBDRigidParticleH
 			float WAngle;
 			FVector WAxis;
 			TargetRotDelta.ToAxisAndAngle(WAxis, WAngle);
-			const FVector TargetRotDeltaBlend = FVector(WAxis * (WAngle / (DeltaSeconds * PhysicsReplicationCVars::PredictiveInterpolationCVars::RotInterpolationTimeMultiplier)));
+			const FVector TargetRotDeltaBlend = FVector(WAxis * (WAngle / (DeltaSeconds * SettingsCurrent.PredictiveInterpolationSettings.GetRotInterpolationTimeMultiplier())));
 			const FVector RepAngVel = FMath::DegreesToRadians(TargetAngVel) + TargetRotDeltaBlend;
 
 			Handle->SetW(RepAngVel);
@@ -1661,12 +1690,12 @@ bool FPhysicsReplicationAsync::PredictiveInterpolation(Chaos::FPBDRigidParticleH
 		if (bSoftSnap)
 		{
 			const FVector SoftSnapPos = FMath::Lerp(FVector(CurrentState.Position),
-			PhysicsReplicationCVars::PredictiveInterpolationCVars::bSoftSnapToSource ? Target.PrevPosTarget : Target.TargetState.Position,
-			FMath::Clamp(PhysicsReplicationCVars::PredictiveInterpolationCVars::SoftSnapPosStrength, 0.0f, 1.0f));
+				SettingsCurrent.PredictiveInterpolationSettings.GetSoftSnapToSource() ? Target.PrevPosTarget : Target.TargetState.Position,
+			FMath::Clamp(SettingsCurrent.PredictiveInterpolationSettings.GetSoftSnapPosStrength(), 0.0f, 1.0f));
 		
 			const FQuat SoftSnapRot = FQuat::Slerp(CurrentState.Quaternion,
-				PhysicsReplicationCVars::PredictiveInterpolationCVars::bSoftSnapToSource ? Target.PrevRotTarget : Target.TargetState.Quaternion,
-			FMath::Clamp(PhysicsReplicationCVars::PredictiveInterpolationCVars::SoftSnapRotStrength, 0.0f, 1.0f));
+				SettingsCurrent.PredictiveInterpolationSettings.GetSoftSnapToSource() ? Target.PrevRotTarget : Target.TargetState.Quaternion,
+			FMath::Clamp(SettingsCurrent.PredictiveInterpolationSettings.GetSoftSnapRotStrength(), 0.0f, 1.0f));
 		
 			Handle->SetX(SoftSnapPos);
 			Handle->SetP(SoftSnapPos);
@@ -1725,10 +1754,8 @@ bool FPhysicsReplicationAsync::ResimulationReplication(Chaos::FPBDRigidParticleH
 	bool bClearTarget = true;
 
 	static constexpr Chaos::FFrameAndPhase::EParticleHistoryPhase RewindPhase = Chaos::FFrameAndPhase::EParticleHistoryPhase::PostPushData;
-	FAsyncPhysicsTimestamp TimeStamp;
-	TimeStamp.LocalFrame = RewindData->CurrentFrame();
 
-	const float ResimErrorThreshold = Chaos::FPhysicsSolverBase::ResimulationErrorThreshold();
+	const float ResimErrorThreshold = SettingsCurrent.ResimulationSettings.GetResimulationErrorThreshold(Chaos::FPhysicsSolverBase::ResimulationErrorThreshold());
 	const Chaos::FGeometryParticleState PastState = RewindData->GetPastStateAtFrame(*Handle, LocalFrame, RewindPhase);
 
 	const FVector ErrorOffset = (Target.TargetState.Position - PastState.X());
@@ -1768,18 +1795,18 @@ bool FPhysicsReplicationAsync::ResimulationReplication(Chaos::FPBDRigidParticleH
 		ResimFrame = (ResimFrame == INDEX_NONE) ? LocalFrame : FMath::Min(ResimFrame, LocalFrame);
 		RewindData->SetResimFrame(ResimFrame);
 	}
-	else if (PhysicsReplicationCVars::ResimulationCVars::bRuntimeCorrectionEnabled)
+	else if (SettingsCurrent.ResimulationSettings.GetRuntimeCorrectionEnabled())
 	{
 		const int32 NumPredictedFrames = RigidsSolver->GetCurrentFrame() - LocalFrame - Target.TickCount;
 
 		if (Target.TickCount <= NumPredictedFrames && NumPredictedFrames > 0)
 		{
 			// Calculate correction to position
-			const float CorrectionAmountX = PhysicsReplicationCVars::ResimulationCVars::PosStabilityMultiplier / NumPredictedFrames; // Same result as (ErrorOffset / NumPredictedFrames) * PosStabilityMultiplier
+			const float CorrectionAmountX = SettingsCurrent.ResimulationSettings.GetPosStabilityMultiplier() / NumPredictedFrames; // Same result as (ErrorOffset / NumPredictedFrames) * PosStabilityMultiplier
 			const FVector CorrectedX = Handle->X() + (ErrorOffset * CorrectionAmountX);
 
 			// Calculate correction to rotation
-			const float CorrectionAmountR = (1.f / NumPredictedFrames) * PhysicsReplicationCVars::ResimulationCVars::RotStabilityMultiplier;
+			const float CorrectionAmountR = (1.f / NumPredictedFrames) * SettingsCurrent.ResimulationSettings.GetRotStabilityMultiplier();
 			const FQuat InvCurrentQuat = PastState.R().Inverse();
 			const FQuat DeltaQuat = Target.TargetState.Quaternion * InvCurrentQuat;
 			const FQuat TargetCorrectionR = Handle->R() * DeltaQuat;
@@ -1808,7 +1835,7 @@ FName FPhysicsReplicationAsync::GetFNameForStatId() const
 	const static FLazyName StaticName("FPhysicsReplicationAsyncCallback");
 	return StaticName;
 }
-#pragma endregion // AsyncPhysicsReplication
+#pragma endregion // FPhysicsReplicationAsync
 
 
 
diff --git a/Engine/Source/Runtime/Engine/Public/Physics/NetworkPhysicsSettingsComponent.h b/Engine/Source/Runtime/Engine/Public/Physics/NetworkPhysicsSettingsComponent.h
new file mode 100644
--- /dev/null
+++ b/Engine/Source/Runtime/Engine/Public/Physics/NetworkPhysicsSettingsComponent.h
@@ -0,0 +1,303 @@
+// Copyright Epic Games, Inc. All Rights Reserved.
+
+/*=============================================================================
+	NetworkPhysicsSettingsComponent.h
+	Manage networked physics settings per actor through ActorComponent and the subsequent physics-thread data flow for the settings.
+=============================================================================*/
+
+#pragma once
+#include "CoreMinimal.h"
+#include "Components/ActorComponent.h"
+#include "Chaos/SimCallbackObject.h"
+
+#include "NetworkPhysicsSettingsComponent.generated.h"
+
+
+class FNetworkPhysicsSettingsComponentAsync;
+
+//  Alias
+using FPredictiveInterpolationSettings	= FNetworkPhysicsSettingsPredictiveInterpolation;
+using FResimulationSettings				= FNetworkPhysicsSettingsResimulation;
+/*
+using FRewindSettings					= FNetworkPhysicsSettingsRewindData;
+using FRenderInterpolationSettings		= FNetworkPhysicsSettingsRenderInterpolation;
+*/
+
+namespace PhysicsReplicationCVars
+{
+	namespace PredictiveInterpolationCVars
+	{
+		extern float PosCorrectionTimeBase;
+		extern float PosCorrectionTimeMin;
+		extern float PosCorrectionTimeMultiplier;
+		extern float RotCorrectionTimeBase;
+		extern float RotCorrectionTimeMin;
+		extern float RotCorrectionTimeMultiplier;
+		extern float PosInterpolationTimeMultiplier;
+		extern float RotInterpolationTimeMultiplier;
+		extern float SoftSnapPosStrength;
+		extern float SoftSnapRotStrength;
+		extern bool bSoftSnapToSource;
+		extern bool bSkipVelocityRepOnPosEarlyOut;
+		extern bool bPostResimWaitForUpdate;
+		extern bool bDisableSoftSnap;
+	}
+
+	namespace ResimulationCVars
+	{
+		extern bool bRuntimeCorrectionEnabled;
+		extern float PosStabilityMultiplier;
+		extern float RotStabilityMultiplier;
+	}
+}
+
+
+USTRUCT()
+struct FNetworkPhysicsSettings
+{
+	GENERATED_BODY()
+
+	UPROPERTY(config, EditDefaultsOnly, Category = "Overrides", Meta = (InlineEditConditionToggle))
+	uint32 bOverrideSimProxyRepMode : 1;
+	// Override the EPhysicsReplicationMode if Actor is ENetRole::ROLE_SimulatedProxy and has EPhysicsReplicationMode::Resimulation set
+	UPROPERTY(config, EditDefaultsOnly, Category = "Overrides", Meta = (EditCondition = "bOverrideSimProxyRepMode"))
+	EPhysicsReplicationMode SimProxyRepMode = EPhysicsReplicationMode::PredictiveInterpolation;
+};
+
+
+USTRUCT()
+struct FNetworkPhysicsSettingsPredictiveInterpolation
+{
+	GENERATED_BODY()
+
+	// np2.PredictiveInterpolation.PosCorrectionTimeBase
+	UPROPERTY(config, EditDefaultsOnly, Category = "Overrides", Meta = (InlineEditConditionToggle))
+	uint32 bOverridePosCorrectionTimeBase : 1;
+	UPROPERTY(config, EditDefaultsOnly, Category = "Overrides", Meta = (EditCondition = "bOverridePosCorrectionTimeBase"))
+	float PosCorrectionTimeBase = PhysicsReplicationCVars::PredictiveInterpolationCVars::PosCorrectionTimeBase;
+	float GetPosCorrectionTimeBase() { return bOverridePosCorrectionTimeBase ? PosCorrectionTimeBase : PhysicsReplicationCVars::PredictiveInterpolationCVars::PosCorrectionTimeBase; }
+
+	// np2.PredictiveInterpolation.PosCorrectionTimeMin
+	UPROPERTY(config, EditDefaultsOnly, Category = "Overrides", Meta = (InlineEditConditionToggle))
+	uint32 bOverridePosCorrectionTimeMin : 1;
+	UPROPERTY(config, EditDefaultsOnly, Category = "Overrides", Meta = (EditCondition = "bOverridePosCorrectionTimeMin"))
+	float PosCorrectionTimeMin = PhysicsReplicationCVars::PredictiveInterpolationCVars::PosCorrectionTimeMin;
+	float GetPosCorrectionTimeMin() { return bOverridePosCorrectionTimeMin ? PosCorrectionTimeMin : PhysicsReplicationCVars::PredictiveInterpolationCVars::PosCorrectionTimeMin; }
+
+	// np2.PredictiveInterpolation.PosCorrectionTimeMultiplier
+	UPROPERTY(config, EditDefaultsOnly, Category = "Overrides", Meta = (InlineEditConditionToggle))
+	uint32 bOverridePosCorrectionTimeMultiplier : 1;
+	UPROPERTY(config, EditDefaultsOnly, Category = "Overrides", Meta = (EditCondition = "bOverridePosCorrectionTimeMultiplier"))
+	float PosCorrectionTimeMultiplier = PhysicsReplicationCVars::PredictiveInterpolationCVars::PosCorrectionTimeMultiplier;
+	float GetPosCorrectionTimeMultiplier() { return bOverridePosCorrectionTimeMultiplier ? PosCorrectionTimeMultiplier : PhysicsReplicationCVars::PredictiveInterpolationCVars::PosCorrectionTimeMultiplier; }
+
+	// np2.PredictiveInterpolation.RotCorrectionTimeBase
+	UPROPERTY(config, EditDefaultsOnly, Category = "Overrides", Meta = (InlineEditConditionToggle))
+	uint32 bOverrideRotCorrectionTimeBase : 1;
+	UPROPERTY(config, EditDefaultsOnly, Category = "Overrides", Meta = (EditCondition = "bOverrideRotCorrectionTimeBase"))
+	float RotCorrectionTimeBase = PhysicsReplicationCVars::PredictiveInterpolationCVars::RotCorrectionTimeBase;
+	float GetRotCorrectionTimeBase() { return bOverrideRotCorrectionTimeBase ? RotCorrectionTimeBase : PhysicsReplicationCVars::PredictiveInterpolationCVars::RotCorrectionTimeBase; }
+
+	// np2.PredictiveInterpolation.RotCorrectionTimeMin
+	UPROPERTY(config, EditDefaultsOnly, Category = "Overrides", Meta = (InlineEditConditionToggle))
+	uint32 bOverrideRotCorrectionTimeMin : 1;
+	UPROPERTY(config, EditDefaultsOnly, Category = "Overrides", Meta = (EditCondition = "bOverrideRotCorrectionTimeMin"))
+	float RotCorrectionTimeMin = PhysicsReplicationCVars::PredictiveInterpolationCVars::RotCorrectionTimeMin;
+	float GetRotCorrectionTimeMin() { return bOverrideRotCorrectionTimeMin ? RotCorrectionTimeMin : PhysicsReplicationCVars::PredictiveInterpolationCVars::RotCorrectionTimeMin; }
+
+	// np2.PredictiveInterpolation.RotCorrectionTimeMultiplier
+	UPROPERTY(config, EditDefaultsOnly, Category = "Overrides", Meta = (InlineEditConditionToggle))
+	uint32 bOverrideRotCorrectionTimeMultiplier : 1;
+	UPROPERTY(config, EditDefaultsOnly, Category = "Overrides", Meta = (EditCondition = "bOverrideRotCorrectionTimeMultiplier"))
+	float RotCorrectionTimeMultiplier = PhysicsReplicationCVars::PredictiveInterpolationCVars::RotCorrectionTimeMultiplier;
+	float GetRotCorrectionTimeMultiplier() { return bOverrideRotCorrectionTimeMultiplier ? RotCorrectionTimeMultiplier : PhysicsReplicationCVars::PredictiveInterpolationCVars::RotCorrectionTimeMultiplier; }
+
+	// np2.PredictiveInterpolation.InterpolationTimeMultiplier
+	UPROPERTY(config, EditDefaultsOnly, Category = "Overrides", Meta = (InlineEditConditionToggle))
+	uint32 bOverridePosInterpolationTimeMultiplier : 1;
+	UPROPERTY(config, EditDefaultsOnly, Category = "Overrides", Meta = (EditCondition = "bOverridePosInterpolationTimeMultiplier"))
+	float PosInterpolationTimeMultiplier = PhysicsReplicationCVars::PredictiveInterpolationCVars::PosInterpolationTimeMultiplier;
+	float GetPosInterpolationTimeMultiplier() { return bOverridePosInterpolationTimeMultiplier ? PosInterpolationTimeMultiplier : PhysicsReplicationCVars::PredictiveInterpolationCVars::PosInterpolationTimeMultiplier; }
+
+	// np2.PredictiveInterpolation.RotInterpolationTimeMultiplier
+	UPROPERTY(config, EditDefaultsOnly, Category = "Overrides", Meta = (InlineEditConditionToggle))
+	uint32 bOverrideRotInterpolationTimeMultiplier : 1;
+	UPROPERTY(config, EditDefaultsOnly, Category = "Overrides", Meta = (EditCondition = "bOverrideRotInterpolationTimeMultiplier"))
+	float RotInterpolationTimeMultiplier = PhysicsReplicationCVars::PredictiveInterpolationCVars::RotInterpolationTimeMultiplier;
+	float GetRotInterpolationTimeMultiplier() { return bOverrideRotInterpolationTimeMultiplier ? RotInterpolationTimeMultiplier : PhysicsReplicationCVars::PredictiveInterpolationCVars::RotInterpolationTimeMultiplier; }
+
+	// np2.PredictiveInterpolation.SoftSnapPosStrength
+	UPROPERTY(config, EditDefaultsOnly, Category = "Overrides", Meta = (InlineEditConditionToggle))
+	uint32 bOverrideSoftSnapPosStrength : 1;
+	UPROPERTY(config, EditDefaultsOnly, Category = "Overrides", Meta = (EditCondition = "bOverrideSoftSnapPosStrength"))
+	float SoftSnapPosStrength = PhysicsReplicationCVars::PredictiveInterpolationCVars::SoftSnapPosStrength;
+	float GetSoftSnapPosStrength() { return bOverrideSoftSnapPosStrength ? SoftSnapPosStrength : PhysicsReplicationCVars::PredictiveInterpolationCVars::SoftSnapPosStrength; }
+
+	// np2.PredictiveInterpolation.SoftSnapRotStrength
+	UPROPERTY(config, EditDefaultsOnly, Category = "Overrides", Meta = (InlineEditConditionToggle))
+	uint32 bOverrideSoftSnapRotStrength : 1;
+	UPROPERTY(config, EditDefaultsOnly, Category = "Overrides", Meta = (EditCondition = "bOverrideSoftSnapRotStrength"))
+	float SoftSnapRotStrength = PhysicsReplicationCVars::PredictiveInterpolationCVars::SoftSnapRotStrength;
+	float GetSoftSnapRotStrength() { return bOverrideSoftSnapRotStrength ? SoftSnapRotStrength : PhysicsReplicationCVars::PredictiveInterpolationCVars::SoftSnapRotStrength; }
+
+	// np2.PredictiveInterpolation.SoftSnapToSource
+	UPROPERTY(config, EditDefaultsOnly, Category = "Overrides", Meta = (InlineEditConditionToggle))
+	uint32 bOverrideSoftSnapToSource : 1;
+	UPROPERTY(config, EditDefaultsOnly, Category = "Overrides", Meta = (EditCondition = "bOverrideSoftSnapToSource"))
+	bool bSoftSnapToSource = PhysicsReplicationCVars::PredictiveInterpolationCVars::bSoftSnapToSource;
+	bool GetSoftSnapToSource() { return bOverrideSoftSnapToSource ? bSoftSnapToSource : PhysicsReplicationCVars::PredictiveInterpolationCVars::bSoftSnapToSource; }
+
+	// np2.PredictiveInterpolation.SkipVelocityRepOnPosEarlyOut
+	UPROPERTY(config, EditDefaultsOnly, Category = "Overrides", Meta = (InlineEditConditionToggle))
+	uint32 bOverrideSkipVelocityRepOnPosEarlyOut : 1;
+	UPROPERTY(config, EditDefaultsOnly, Category = "Overrides", Meta = (EditCondition = "bOverrideSkipVelocityRepOnPosEarlyOut"))
+	bool bSkipVelocityRepOnPosEarlyOut = PhysicsReplicationCVars::PredictiveInterpolationCVars::bSkipVelocityRepOnPosEarlyOut;
+	bool GetSkipVelocityRepOnPosEarlyOut() { return bOverrideSkipVelocityRepOnPosEarlyOut ? bSkipVelocityRepOnPosEarlyOut : PhysicsReplicationCVars::PredictiveInterpolationCVars::bSkipVelocityRepOnPosEarlyOut; }
+
+	// np2.PredictiveInterpolation.PostResimWaitForUpdate
+	UPROPERTY(config, EditDefaultsOnly, Category = "Overrides", Meta = (InlineEditConditionToggle))
+	uint32 bOverridePostResimWaitForUpdate : 1;
+	UPROPERTY(config, EditDefaultsOnly, Category = "Overrides", Meta = (EditCondition = "bOverridePostResimWaitForUpdate"))
+	bool bPostResimWaitForUpdate = PhysicsReplicationCVars::PredictiveInterpolationCVars::bPostResimWaitForUpdate;
+	bool GetPostResimWaitForUpdate() { return bOverridePostResimWaitForUpdate ? bPostResimWaitForUpdate : PhysicsReplicationCVars::PredictiveInterpolationCVars::bPostResimWaitForUpdate; }
+
+	// np2.PredictiveInterpolation.DisableSoftSnap
+	UPROPERTY(config, EditDefaultsOnly, Category = "Overrides", Meta = (InlineEditConditionToggle))
+	uint32 bOverrideDisableSoftSnap : 1;
+	UPROPERTY(config, EditDefaultsOnly, Category = "Overrides", Meta = (EditCondition = "bOverrideDisableSoftSnap"))
+	bool bDisableSoftSnap = PhysicsReplicationCVars::PredictiveInterpolationCVars::bDisableSoftSnap;
+	bool GetDisableSoftSnap() { return bOverrideDisableSoftSnap ? bDisableSoftSnap : PhysicsReplicationCVars::PredictiveInterpolationCVars::bDisableSoftSnap; }
+};
+
+USTRUCT()
+struct FNetworkPhysicsSettingsResimulation
+{
+	GENERATED_BODY()
+
+	// Override how many inputs to synchronize each sync to cover packet loss.
+	UPROPERTY(config, EditDefaultsOnly, Category = "Overrides", Meta = (InlineEditConditionToggle))
+	uint32 bOverrideRedundantInputs : 1;
+	UPROPERTY(config, EditDefaultsOnly, Category = "Overrides", Meta = (EditCondition = "bOverrideRedundantInputs"))
+	uint8 RedundantInputs = 4;
+
+	// Override how many states to synchronize each sync to cover packet loss.
+	UPROPERTY(config, EditDefaultsOnly, Category = "Overrides", Meta = (InlineEditConditionToggle))
+	uint32 bOverrideRedundantStates : 1;
+	UPROPERTY(config, EditDefaultsOnly, Category = "Overrides", Meta = (EditCondition = "bOverrideRedundantStates"))
+	uint8 RedundantStates = 1;
+
+	// np2.Resim.RuntimeCorrectionEnabled
+	UPROPERTY(config, EditDefaultsOnly, Category = "Overrides", Meta = (InlineEditConditionToggle))
+	uint32 bOverrideRuntimeCorrectionEnabled : 1;
+	UPROPERTY(config, EditDefaultsOnly, Category = "Overrides", Meta = (EditCondition = "bOverrideRuntimeCorrectionEnabled"))
+	bool bRuntimeCorrectionEnabled = PhysicsReplicationCVars::ResimulationCVars::bRuntimeCorrectionEnabled;
+	bool GetRuntimeCorrectionEnabled() { return bOverrideRuntimeCorrectionEnabled ? bRuntimeCorrectionEnabled : PhysicsReplicationCVars::ResimulationCVars::bRuntimeCorrectionEnabled; }
+
+	// np2.Resim.PosStabilityMultiplier
+	UPROPERTY(config, EditDefaultsOnly, Category = "Overrides", Meta = (InlineEditConditionToggle))
+	uint32 bOverridePosStabilityMultiplier : 1;
+	UPROPERTY(config, EditDefaultsOnly, Category = "Overrides", Meta = (EditCondition = "bOverridePosStabilityMultiplier"))
+	float PosStabilityMultiplier = PhysicsReplicationCVars::ResimulationCVars::PosStabilityMultiplier;
+	float GetPosStabilityMultiplier() { return bOverridePosStabilityMultiplier ? PosStabilityMultiplier : PhysicsReplicationCVars::ResimulationCVars::PosStabilityMultiplier; }
+
+	// np2.Resim.RotStabilityMultiplier
+	UPROPERTY(config, EditDefaultsOnly, Category = "Overrides", Meta = (InlineEditConditionToggle))
+	uint32 bOverrideRotStabilityMultiplier : 1;
+	UPROPERTY(config, EditDefaultsOnly, Category = "Overrides", Meta = (EditCondition = "bOverrideRotStabilityMultiplier"))
+	float RotStabilityMultiplier = PhysicsReplicationCVars::ResimulationCVars::RotStabilityMultiplier;
+	float GetRotStabilityMultiplier() { return bOverrideRotStabilityMultiplier ? RotStabilityMultiplier : PhysicsReplicationCVars::ResimulationCVars::RotStabilityMultiplier; }
+
+	// Project Settings -> Physics -> Replication -> Physics Prediction -> Resimulation Error Threshold
+	UPROPERTY(config, EditDefaultsOnly, Category = "Overrides", Meta = (InlineEditConditionToggle))
+	uint32 bOverrideResimulationErrorThreshold : 1;
+	UPROPERTY(config, EditDefaultsOnly, Category = "Overrides", Meta = (EditCondition = "bOverrideResimulationErrorThreshold"))
+	uint32 ResimulationErrorThreshold = 10;
+	uint32 GetResimulationErrorThreshold(uint32 DefaultValue) { return bOverrideRotStabilityMultiplier ? ResimulationErrorThreshold : DefaultValue; }
+};
+
+/*
+USTRUCT()
+struct FNetworkPhysicsSettingsRewindData
+{
+	GENERATED_BODY()
+};
+
+USTRUCT()
+struct FNetworkPhysicsSettingsRenderInterpolation
+{
+	GENERATED_BODY()
+};
+*/
+
+UCLASS(BlueprintType, meta = (BlueprintSpawnableComponent))
+class UNetworkPhysicsSettingsComponent : public UActorComponent
+{
+	GENERATED_BODY()
+
+public:
+	UNetworkPhysicsSettingsComponent();
+
+	virtual void InitializeComponent() override;
+	virtual void UninitializeComponent() override;
+
+	virtual void BeginPlay() override;
+
+
+public:
+	UPROPERTY(EditDefaultsOnly, Category = "Networked Physics Settings")
+	FNetworkPhysicsSettings GeneralSettings;
+	
+	UPROPERTY(EditDefaultsOnly, Category = "Networked Physics Settings")
+	FNetworkPhysicsSettingsPredictiveInterpolation PredictiveInterpolationSettings;
+	
+	UPROPERTY(EditDefaultsOnly, Category = "Networked Physics Settings")
+	FNetworkPhysicsSettingsResimulation ResimulationSettings;
+	
+	/*
+	UPROPERTY(EditDefaultsOnly, Category = "Networked Physics Settings")
+	FNetworkPhysicsSettingsRewindData RewindSettings;
+	
+	UPROPERTY(EditDefaultsOnly, Category = "Networked Physics Settings")
+	FNetworkPhysicsSettingsRenderInterpolation RenderInterpolationSettings;
+	*/
+
+private:
+	FNetworkPhysicsSettingsComponentAsync* NetworkPhysicsSettingsAsync;
+};
+
+
+
+
+
+
+#pragma region // FNetworkPhysicsSettingsComponentAsync
+
+struct FNetworkPhysicsSettingsAsync
+{
+	FNetworkPhysicsSettings GeneralSettings;
+	FNetworkPhysicsSettingsResimulation ResimulationSettings;
+	FNetworkPhysicsSettingsPredictiveInterpolation PredictiveInterpolationSettings;
+};
+
+struct FNetworkPhysicsSettingsAsyncInput : public Chaos::FSimCallbackInput
+{
+	Chaos::FConstPhysicsObjectHandle PhysicsObject;
+	FNetworkPhysicsSettingsAsync Settings;
+
+	void Reset()
+	{
+		Settings = FNetworkPhysicsSettingsAsync();
+	}
+};
+
+class FNetworkPhysicsSettingsComponentAsync : public Chaos::TSimCallbackObject<FNetworkPhysicsSettingsAsyncInput>
+{
+public:
+	virtual void OnPostInitialize_Internal() override;
+	virtual void OnPreSimulate_Internal() override { };
+
+public:
+	FNetworkPhysicsSettingsAsync Settings;
+};
+
+#pragma endregion // FNetworkPhysicsSettingsComponentAsync
diff --git a/Engine/Source/Runtime/Engine/Public/PhysicsReplication.h b/Engine/Source/Runtime/Engine/Public/PhysicsReplication.h
--- a/Engine/Source/Runtime/Engine/Public/PhysicsReplication.h
+++ b/Engine/Source/Runtime/Engine/Public/PhysicsReplication.h
@@ -12,10 +12,10 @@
 #include "PhysicsReplicationInterface.h"
 #include "Physics/PhysicsInterfaceDeclares.h"
 #include "PhysicsProxy/SingleParticlePhysicsProxyFwd.h"
-#include "Chaos/Particles.h"
 #include "Chaos/PhysicsObject.h"
 #include "Chaos/SimCallbackObject.h"
 #include "Physics/PhysicsInterfaceUtils.h"
+#include "Physics/NetworkPhysicsSettingsComponent.h"
 
 namespace CharacterMovementCVars
 {
@@ -42,13 +42,38 @@ namespace CharacterMovementCVars
 namespace PhysicsReplicationCVars
 {
 	extern ENGINE_API int32 LogPhysicsReplicationHardSnaps;
+	
+	namespace PredictiveInterpolationCVars
+	{
+		extern float PosCorrectionTimeBase;
+		extern float PosCorrectionTimeMin;
+		extern float PosCorrectionTimeMultiplier;
+		extern float RotCorrectionTimeBase;
+		extern float RotCorrectionTimeMin;
+		extern float RotCorrectionTimeMultiplier;
+		extern float PosInterpolationTimeMultiplier;
+		extern float RotInterpolationTimeMultiplier;
+		extern float SoftSnapPosStrength;
+		extern float SoftSnapRotStrength;
+		extern bool bSoftSnapToSource;
+		extern bool bSkipVelocityRepOnPosEarlyOut;
+		extern bool bPostResimWaitForUpdate;
+		extern bool bDisableSoftSnap;
+	}
+
+	namespace ResimulationCVars
+	{
+		extern bool bRuntimeCorrectionEnabled;
+		extern float PosStabilityMultiplier;
+		extern float RotStabilityMultiplier;
+	}
 }
 #endif
 
 class FPhysScene_PhysX;
 
 
-// -------- Async Flow ------->
+#pragma region FPhysicsReplicationAsync
 
 struct FPhysicsRepErrorCorrectionData
 {
@@ -132,7 +157,7 @@ struct FReplicatedPhysicsTargetAsync
 	int32 ReceiveFrame;
 
 	/** Local physics frames between received targets */
-	uint8 ReceiveInterval;
+	int32 ReceiveInterval;
 	float AverageReceiveInterval;
 
 	/** The replication mode this PhysicsObject should use */
@@ -201,6 +226,7 @@ class FPhysicsReplicationAsync : public Chaos::TSimCallbackObject<
 	Chaos::ESimCallbackOptions::Presimulate | Chaos::ESimCallbackOptions::PhysicsObjectUnregister>
 {
 	virtual FName GetFNameForStatId() const override;
+	virtual void OnPostInitialize_Internal() override;
 	virtual void OnPreSimulate_Internal() override;
 	virtual void OnPhysicsObjectUnregistered_Internal(Chaos::FConstPhysicsObjectHandle PhysicsObject) override;
 
@@ -212,16 +238,24 @@ class FPhysicsReplicationAsync : public Chaos::TSimCallbackObject<
 	virtual bool PredictiveInterpolation(Chaos::FPBDRigidParticleHandle* Handle, FReplicatedPhysicsTargetAsync& Target, const float DeltaSeconds);
 	virtual bool ResimulationReplication(Chaos::FPBDRigidParticleHandle* Handle, FReplicatedPhysicsTargetAsync& Target, const float DeltaSeconds);
 
+public:
+	virtual void RegisterSettings(Chaos::FConstPhysicsObjectHandle PhysicsObject, FNetworkPhysicsSettingsAsync InSettings);
+
 private:
 	float LatencyOneWay;
 	FRigidBodyErrorCorrection ErrorCorrectionDefault;
+	FNetworkPhysicsSettingsAsync SettingsCurrent;
+	FNetworkPhysicsSettingsAsync SettingsDefault;
 	TMap<Chaos::FConstPhysicsObjectHandle, FReplicatedPhysicsTargetAsync> ObjectToTarget;
+	TMap<Chaos::FConstPhysicsObjectHandle, FNetworkPhysicsSettingsAsync> ObjectToSettings;
 	TArray<int32> ParticlesInResimIslands;
 
 private:
 	void UpdateAsyncTarget(const FPhysicsRepAsyncInputData& Input, Chaos::FPBDRigidsSolver* RigidsSolver);
 	void UpdateRewindDataTarget(const FPhysicsRepAsyncInputData& Input);
 	void CacheResimInteractions();
+	// Sets SettingsCurrent to either the objects custom settings or to the default settings
+	void FetchObjectSettings(Chaos::FConstPhysicsObjectHandle PhysicsObject); 
 	static void ExtrapolateTarget(FReplicatedPhysicsTargetAsync& Target, const int32 ExtrapolateFrames, const float DeltaSeconds);
 
 public:
@@ -231,7 +265,7 @@ public:
 	}
 };
 
-// <-------- Async Flow --------
+#pragma endregion // FPhysicsReplicationAsync
 
 
 
diff --git a/Engine/Source/Runtime/Experimental/Chaos/Private/PBDRigidsSolver.cpp b/Engine/Source/Runtime/Experimental/Chaos/Private/PBDRigidsSolver.cpp
--- a/Engine/Source/Runtime/Experimental/Chaos/Private/PBDRigidsSolver.cpp
+++ b/Engine/Source/Runtime/Experimental/Chaos/Private/PBDRigidsSolver.cpp
@@ -1856,8 +1856,13 @@ namespace Chaos
 			delete SimCallbackObject;
 		}
 		PushData.SimCommands.Reset();
+		
+		for (ISimCallbackObject* SimCallbackObject : PushData.SimCallbackObjectsToAdd)
+		{
+			SimCallbackObject->PostInitialize_Internal();
+		}
 
-		if(MRewindCallback && MRewindData && !IsShuttingDown())
+		if (MRewindCallback && MRewindData && !IsShuttingDown())
 		{
 			MRewindCallback->ProcessInputs_Internal(MRewindData->CurrentFrame(), PushData.SimCallbackInputs);
 		}
diff --git a/Engine/Source/Runtime/Experimental/Chaos/Public/Chaos/Framework/PhysicsSolverBase.h b/Engine/Source/Runtime/Experimental/Chaos/Public/Chaos/Framework/PhysicsSolverBase.h
--- a/Engine/Source/Runtime/Experimental/Chaos/Public/Chaos/Framework/PhysicsSolverBase.h
+++ b/Engine/Source/Runtime/Experimental/Chaos/Public/Chaos/Framework/PhysicsSolverBase.h
@@ -19,6 +19,7 @@
 #endif
 
 class FChaosSolversModule;
+class FPhysicsReplicationAsync;
 
 DECLARE_MULTICAST_DELEGATE_OneParam(FSolverPreAdvance, Chaos::FReal);
 DECLARE_MULTICAST_DELEGATE_OneParam(FSolverPreBuffer, Chaos::FReal);
@@ -411,6 +412,16 @@ namespace Chaos
 			return MRewindCallback.IsValid() && MRewindData.IsValid();
 		}
 
+		void SetPhysicsReplication(FPhysicsReplicationAsync* InPhysicsReplication)
+		{
+			PhysicsReplication = InPhysicsReplication;
+		}
+
+		FPhysicsReplicationAsync* GetPhysicsReplication()
+		{
+			return PhysicsReplication;
+		}
+
 		//Used as helper for GT to go from unique idx back to gt particle
 		//If GT deletes a particle, this function will return null (that's a good thing when consuming async outputs as GT may have already deleted the particle we care about)
 		//Note: if the physics solver has been advanced after the particle was freed on GT, the index may have been freed and reused.
@@ -761,6 +772,7 @@ namespace Chaos
 
 		TUniquePtr<FRewindData> MRewindData;
 		TUniquePtr<IRewindCallback> MRewindCallback;
+		FPhysicsReplicationAsync* PhysicsReplication;
 
 		bool bUseCollisionResimCache;
 
diff --git a/Engine/Source/Runtime/Experimental/Chaos/Public/Chaos/SimCallbackObject.h b/Engine/Source/Runtime/Experimental/Chaos/Public/Chaos/SimCallbackObject.h
--- a/Engine/Source/Runtime/Experimental/Chaos/Public/Chaos/SimCallbackObject.h
+++ b/Engine/Source/Runtime/Experimental/Chaos/Public/Chaos/SimCallbackObject.h
@@ -79,6 +79,11 @@ public:
 
 	virtual bool IsFAsyncObjectManagerCallback() const { return false;}
 
+	void PostInitialize_Internal()
+	{
+		OnPostInitialize_Internal();
+	}
+
 	void PreSimulate_Internal()
 	{
 		OnPreSimulate_Internal();
@@ -247,11 +252,16 @@ private:
 	 */
 	virtual FSimCallbackInput* AllocateInputData_External() = 0;
 
+	/**
+	* Called once when callback object is registered and after input data have been marshalled
+	*/
+	virtual void OnPostInitialize_Internal() { }
+
 	/**
 	* Called before simulation step (NOTE: not once per sub-step when sub-stepping is enabled)
 	*/
 	virtual void OnPreSimulate_Internal() = 0;
-	
+
 	/**
 	* Called once per simulation sub-step, before Integrate. Can be used to modify particle positions, velocities etc.
 	*
