remove unused platform-specific code path for legacy particle system

#rnx

diff --git a/Engine/Source/Runtime/Engine/Private/Particles/ParticleSystemRender.cpp b/Engine/Source/Runtime/Engine/Private/Particles/ParticleSystemRender.cpp
--- a/Engine/Source/Runtime/Engine/Private/Particles/ParticleSystemRender.cpp
+++ b/Engine/Source/Runtime/Engine/Private/Particles/ParticleSystemRender.cpp
@@ -952,13 +952,7 @@ void FDynamicSpriteEmitterData::GetDynamicMeshElementsEmitter(const FParticleSys
 
 				// Set the sprite uniform buffer for this view.
 				SpriteVertexFactory->SetSpriteUniformBuffer(CollectorResources.UniformBuffer);
-#if PLATFORM_SWITCH
-				// use the full vertex for non-instancing case, and the "4 float" instance data for the instanced case
-				uint32 InstanceBufferStride = ((sizeof(float) * 4) * NumVerticesPerParticle);
-#else
-				uint32 InstanceBufferStride = VertexSize;
-#endif
-				SpriteVertexFactory->SetInstanceBuffer(Allocation.VertexBuffer, Allocation.VertexOffset, InstanceBufferStride);
+				SpriteVertexFactory->SetInstanceBuffer(Allocation.VertexBuffer, Allocation.VertexOffset, VertexSize);
 				SpriteVertexFactory->SetDynamicParameterBuffer(DynamicParameterAllocation.VertexBuffer, DynamicParameterAllocation.VertexOffset, GetDynamicParameterVertexStride());
 
 				if (SourceData->RequiredModule->bCutoutTexureIsValid)
