Fixed Build (Engine Localization - Text conflict)

diff --git a/Engine/Source/Runtime/Engine/Private/WorldPartition/DataLayer/ExternalDataLayerManager.cpp b/Engine/Source/Runtime/Engine/Private/WorldPartition/DataLayer/ExternalDataLayerManager.cpp
--- a/Engine/Source/Runtime/Engine/Private/WorldPartition/DataLayer/ExternalDataLayerManager.cpp
+++ b/Engine/Source/Runtime/Engine/Private/WorldPartition/DataLayer/ExternalDataLayerManager.cpp
@@ -128,11 +128,11 @@ bool UExternalDataLayerManager::CanInjectExternalDataLayerAsset(const UExternalD
 		{
 			if (GetTypedOuter<UWorldPartition>()->GetWorld() != GetTypedOuter<UWorld>())
 			{
-				*OutReason = LOCTEXT("CantInjectExternalDataLayerAsset", "External Data Layer instances are not yet supported for sub-levels.");
+				*OutReason = LOCTEXT("CantInjectExternalDataLayerAssetSubLevelsNotSupported", "External Data Layers are not yet supported for sub-levels.");
 			}
 			else
 			{
-				*OutReason = LOCTEXT("CantInjectExternalDataLayerAsset", "External Data Layers are not supported.");
+				*OutReason = LOCTEXT("CantInjectExternalDataLayerAssetNotSupported", "External Data Layers are not supported.");
 			}
 		}
 		return false;
