Fix a couple regressions with custom character gravity:
- Floor sweep test failing to find floor with custom gravity. Fixed sweep shape orientation to respect custom gravity.
- Capsule rotation, when configured to orient to movement, would flip axis when applying delta rotation. Changed to evaluate in gravity relative space.
#rnx

diff --git a/Engine/Source/Runtime/Engine/Private/Components/CharacterMovementComponent.cpp b/Engine/Source/Runtime/Engine/Private/Components/CharacterMovementComponent.cpp
--- a/Engine/Source/Runtime/Engine/Private/Components/CharacterMovementComponent.cpp
+++ b/Engine/Source/Runtime/Engine/Private/Components/CharacterMovementComponent.cpp
@@ -6357,24 +6357,52 @@ void UCharacterMovementComponent::PhysicsRotation(float DeltaTime)
 			}
 		}
 
-		// PITCH
-		if (!FMath::IsNearlyEqual(CurrentRotation.Pitch, DesiredRotation.Pitch, AngleTolerance))
+		if (HasCustomGravity())
 		{
-			DesiredRotation.Pitch = FMath::FixedTurn(CurrentRotation.Pitch, DesiredRotation.Pitch, DeltaRot.Pitch);
-		}
+			FRotator GravityRelativeCurrentRotation = (GravityToWorldTransform * CurrentRotation.Quaternion()).Rotator();
+			FRotator GravityRelativeDesiredRotation = (GravityToWorldTransform * DesiredRotation.Quaternion()).Rotator();
 
-		// YAW
-		if (!FMath::IsNearlyEqual(CurrentRotation.Yaw, DesiredRotation.Yaw, AngleTolerance))
-		{
-			DesiredRotation.Yaw = FMath::FixedTurn(CurrentRotation.Yaw, DesiredRotation.Yaw, DeltaRot.Yaw);
-		}
+			// PITCH
+			if (!FMath::IsNearlyEqual(GravityRelativeCurrentRotation.Pitch, GravityRelativeDesiredRotation.Pitch, AngleTolerance))
+			{
+				GravityRelativeDesiredRotation.Pitch = FMath::FixedTurn(GravityRelativeCurrentRotation.Pitch, GravityRelativeDesiredRotation.Pitch, DeltaRot.Pitch);
+			}
 
-		// ROLL
-		if (!FMath::IsNearlyEqual(CurrentRotation.Roll, DesiredRotation.Roll, AngleTolerance))
-		{
-			DesiredRotation.Roll = FMath::FixedTurn(CurrentRotation.Roll, DesiredRotation.Roll, DeltaRot.Roll);
+			// YAW
+			if (!FMath::IsNearlyEqual(GravityRelativeCurrentRotation.Yaw, GravityRelativeDesiredRotation.Yaw, AngleTolerance))
+			{
+				GravityRelativeDesiredRotation.Yaw = FMath::FixedTurn(GravityRelativeCurrentRotation.Yaw, GravityRelativeDesiredRotation.Yaw, DeltaRot.Yaw);
+			}
+
+			// ROLL
+			if (!FMath::IsNearlyEqual(GravityRelativeCurrentRotation.Roll, GravityRelativeDesiredRotation.Roll, AngleTolerance))
+			{
+				GravityRelativeDesiredRotation.Roll = FMath::FixedTurn(GravityRelativeCurrentRotation.Roll, GravityRelativeDesiredRotation.Roll, DeltaRot.Roll);
+			}
+
+			DesiredRotation = (WorldToGravityTransform * GravityRelativeDesiredRotation.Quaternion()).Rotator();
 		}
+		else
+		{
+			// PITCH
+			if (!FMath::IsNearlyEqual(CurrentRotation.Pitch, DesiredRotation.Pitch, AngleTolerance))
+			{
+				DesiredRotation.Pitch = FMath::FixedTurn(CurrentRotation.Pitch, DesiredRotation.Pitch, DeltaRot.Pitch);
+			}
 
+			// YAW
+			if (!FMath::IsNearlyEqual(CurrentRotation.Yaw, DesiredRotation.Yaw, AngleTolerance))
+			{
+				DesiredRotation.Yaw = FMath::FixedTurn(CurrentRotation.Yaw, DesiredRotation.Yaw, DeltaRot.Yaw);
+			}
+
+			// ROLL
+			if (!FMath::IsNearlyEqual(CurrentRotation.Roll, DesiredRotation.Roll, AngleTolerance))
+			{
+				DesiredRotation.Roll = FMath::FixedTurn(CurrentRotation.Roll, DesiredRotation.Roll, DeltaRot.Roll);
+			}
+		}
+		
 		// Set the new rotation.
 		DesiredRotation.DiagnosticCheckNaN(TEXT("CharacterMovementComponent::PhysicsRotation(): DesiredRotation"));
 		MoveUpdatedComponent( FVector::ZeroVector, DesiredRotation, /*bSweep*/ false );
@@ -6949,7 +6977,7 @@ bool UCharacterMovementComponent::FloorSweepTest(
 
 	if (!bUseFlatBaseForFloorChecks)
 	{
-		bBlockingHit = GetWorld()->SweepSingleByChannel(OutHit, Start, End, FQuat::Identity, TraceChannel, CollisionShape, Params, ResponseParam);
+		bBlockingHit = GetWorld()->SweepSingleByChannel(OutHit, Start, End, GetWorldToGravityTransform(), TraceChannel, CollisionShape, Params, ResponseParam);
 	}
 	else
 	{
@@ -6965,7 +6993,7 @@ bool UCharacterMovementComponent::FloorSweepTest(
 		{
 			// Test again with the same box, not rotated.
 			OutHit.Reset(1.f, false);
-			bBlockingHit = GetWorld()->SweepSingleByChannel(OutHit, Start, End, FQuat::Identity, TraceChannel, BoxShape, Params, ResponseParam);
+			bBlockingHit = GetWorld()->SweepSingleByChannel(OutHit, Start, End, GetWorldToGravityTransform(), TraceChannel, BoxShape, Params, ResponseParam);
 		}
 	}
 
