World Partition: added a project setting to control wheter to enable world partition and streaming on new maps created in the content browser,

diff --git a/Engine/Source/Editor/UnrealEd/Private/Factories/WorldFactory.cpp b/Engine/Source/Editor/UnrealEd/Private/Factories/WorldFactory.cpp
--- a/Engine/Source/Editor/UnrealEd/Private/Factories/WorldFactory.cpp
+++ b/Engine/Source/Editor/UnrealEd/Private/Factories/WorldFactory.cpp
@@ -5,6 +5,7 @@
 #include "Editor.h"
 #include "UObject/Package.h"
 #include "ThumbnailRendering/WorldThumbnailInfo.h"
+#include "WorldPartition/WorldPartitionSettings.h"
 #include "EditorClassUtils.h"
 #include "Modules/ModuleManager.h"
 
@@ -17,9 +18,8 @@ UWorldFactory::UWorldFactory(const FObjectInitializer& ObjectInitializer)
 	SupportedClass = UWorld::StaticClass();
 	WorldType = EWorldType::Inactive;
 	bInformEngineOfWorld = false;
-	bCreateWorldPartition = false;
-	// default to true to preserve previous behavior
-	bEnableWorldPartitionStreaming = true;
+	bCreateWorldPartition = UWorldPartitionSettings::Get()->GetNewMapsEnableWorldPartition();
+	bEnableWorldPartitionStreaming = UWorldPartitionSettings::Get()->GetNewMapsEnableWorldPartitionStreaming();
 	FeatureLevel = ERHIFeatureLevel::Num;
 }
 
diff --git a/Engine/Source/Runtime/Engine/Public/WorldPartition/WorldPartitionSettings.h b/Engine/Source/Runtime/Engine/Public/WorldPartition/WorldPartitionSettings.h
--- a/Engine/Source/Runtime/Engine/Public/WorldPartition/WorldPartitionSettings.h
+++ b/Engine/Source/Runtime/Engine/Public/WorldPartition/WorldPartitionSettings.h
@@ -16,6 +16,8 @@ public:
 	static UWorldPartitionSettings* Get() { return CastChecked<UWorldPartitionSettings>(UWorldPartitionSettings::StaticClass()->GetDefaultObject()); }
 
 	EWorldPartitionDataLayersLogicOperator GetNewMapsDataLayersLogicOperator() const { return NewMapsDataLayersLogicOperator; }
+	bool GetNewMapsEnableWorldPartition() const { return bNewMapsEnableWorldPartition; }
+	bool GetNewMapsEnableWorldPartitionStreaming() const { return bNewMapsEnableWorldPartitionStreaming; }
 	
 	TSubclassOf<UWorldPartitionEditorHash> GetEditorHashDefaultClass() const { return EditorHashDefaultClass; }
 	TSubclassOf<UWorldPartitionRuntimeHash> GetRuntimeHashDefaultClass() const { return RuntimeHashDefaultClass; }
@@ -25,6 +27,14 @@ protected:
 	UPROPERTY(EditAnywhere, Config, Category = WorldPartition)
 	EWorldPartitionDataLayersLogicOperator NewMapsDataLayersLogicOperator = EWorldPartitionDataLayersLogicOperator::Or;
 
+	/** Set the default to whether enable world partition for new maps created in the content broswer */
+	UPROPERTY(EditAnywhere, Config, Category = WorldPartition)
+	bool bNewMapsEnableWorldPartition = false;
+
+	/** Set the default to whether enable world partition streaming for new maps created in the content broswer */
+	UPROPERTY(EditAnywhere, Config, Category = WorldPartition)
+	bool bNewMapsEnableWorldPartitionStreaming = true;
+
 	/** Set the default editor hash class to use for the editor */
 	UPROPERTY(Config)
 	TSubclassOf<UWorldPartitionEditorHash> EditorHashDefaultClass;
