Prevent null dereference when fixing up material references.
#rnx

diff --git a/Engine/Source/Runtime/Engine/Private/Materials/MaterialExpressions.cpp b/Engine/Source/Runtime/Engine/Private/Materials/MaterialExpressions.cpp
--- a/Engine/Source/Runtime/Engine/Private/Materials/MaterialExpressions.cpp
+++ b/Engine/Source/Runtime/Engine/Private/Materials/MaterialExpressions.cpp
@@ -15255,7 +15255,7 @@ static int32 FindOutputIndexByName(const FName& Name, const TArray<FFunctionExpr
 	for (int32 OutputIndex = 0; OutputIndex < Outputs.Num(); OutputIndex++)
 	{
 		const FFunctionExpressionOutput& CurrentOutput = Outputs[OutputIndex];
-		if (CurrentOutput.ExpressionOutput->OutputName == Name)
+		if (CurrentOutput.ExpressionOutput && CurrentOutput.ExpressionOutput->OutputName == Name)
 		{
 			return OutputIndex;
 		}
@@ -17206,7 +17206,10 @@ static void FixupReferencingInputs(
 			{
 				if (bMatchByName)
 				{
-					CurrentInput->OutputIndex = FindOutputIndexByName(OriginalOutputs[CurrentInput->OutputIndex].ExpressionOutput->OutputName, NewOutputs);
+					if (OriginalOutputs[CurrentInput->OutputIndex].ExpressionOutput)
+					{
+						CurrentInput->OutputIndex = FindOutputIndexByName(OriginalOutputs[CurrentInput->OutputIndex].ExpressionOutput->OutputName, NewOutputs);
+					}
 				}
 				else
 				{
