Add an API to get committed memory usage for local shader compile workers

To better track the total memory usage for shader compilation during cooking, this API gives an estimate of the total committed memory used by local SCWs. This is only an estimate, as you generally cannot sum up the committed memory across processes and assume that this is the total committed memory of the group of processes. But it is the best estimate available.

diff --git a/Engine/Source/Runtime/Engine/Private/ShaderCompiler/ShaderCompiler.cpp b/Engine/Source/Runtime/Engine/Private/ShaderCompiler/ShaderCompiler.cpp
--- a/Engine/Source/Runtime/Engine/Private/ShaderCompiler/ShaderCompiler.cpp
+++ b/Engine/Source/Runtime/Engine/Private/ShaderCompiler/ShaderCompiler.cpp
@@ -4113,6 +4113,28 @@ void FShaderCompileThreadRunnable::PrintWorkerMemoryUsage()
 	}
 }
 
+FShaderCompileMemoryUsage FShaderCompileThreadRunnable::GetExternalWorkerMemoryUsage()
+{
+	FScopeLock WorkerScopeLock(&WorkerInfosLock);
+	FShaderCompileMemoryUsage MemoryUsage{};
+	for (const TUniquePtr<FShaderCompileWorkerInfo>& WorkerInfo : WorkerInfos)
+	{
+		FProcHandle ProcHandle = WorkerInfo->WorkerProcess;
+		if (!ProcHandle.IsValid())
+		{
+			continue;
+		}
+		FPlatformProcessMemoryStats MemoryStats;
+		if (FPlatformProcess::TryGetMemoryUsage(ProcHandle, MemoryStats))
+		{
+			// Virtual memory is committed memory on Windows.
+			MemoryUsage.VirtualMemory += MemoryStats.UsedVirtual;
+			MemoryUsage.PhysicalMemory += MemoryStats.UsedPhysical;
+		}
+	}
+	return MemoryUsage;
+}
+
 void FShaderCompileUtilities::ExecuteShaderCompileJob(FShaderCommonCompileJob& Job)
 {
 	TRACE_CPUPROFILER_EVENT_SCOPE(FShaderCompileUtilities::ExecuteShaderCompileJob);
@@ -6943,6 +6965,18 @@ void FShaderCompilingManager::GetLocalStats(FShaderCompilerStats& OutStats) cons
 	}
 }
 
+FShaderCompileMemoryUsage FShaderCompilingManager::GetExternalMemoryUsage()
+{
+	FShaderCompileMemoryUsage TotalMemoryUsage{};
+	for (const TUniquePtr<FShaderCompileThreadRunnableBase>& ThreadPtr : Threads)
+	{
+		FShaderCompileMemoryUsage MemoryUsage = ThreadPtr->GetExternalWorkerMemoryUsage();
+		TotalMemoryUsage.VirtualMemory += MemoryUsage.VirtualMemory;
+		TotalMemoryUsage.PhysicalMemory += MemoryUsage.PhysicalMemory;
+	}
+	return TotalMemoryUsage;
+}
+
 static bool GatherUniqueErrors(const TArray<FShaderCommonCompileJobPtr>& CompleteJobs, FShaderErrorInfo& OutShaderErrorInfo)
 {
 	// Gather unique errors
diff --git a/Engine/Source/Runtime/Engine/Public/ShaderCompiler.h b/Engine/Source/Runtime/Engine/Public/ShaderCompiler.h
--- a/Engine/Source/Runtime/Engine/Public/ShaderCompiler.h
+++ b/Engine/Source/Runtime/Engine/Public/ShaderCompiler.h
@@ -113,6 +113,18 @@ public:
 };
 #endif // WITH_EDITOR
 
+struct FShaderCompileMemoryUsage
+{
+	/**
+	* The amount of virtual memory used (committed on Windows)
+	*/
+	uint64 VirtualMemory;
+	/**
+	* The amount of physical memory used.
+	*/
+	uint64 PhysicalMemory;
+};
+
 class FShaderCompileThreadRunnableBase : public FRunnable
 {
 	friend class FShaderCompilingManager;
@@ -129,6 +141,8 @@ protected:
 	TAtomic<bool> bForceFinish;
 
 	virtual void PrintWorkerMemoryUsage() {}
+	/** Returns the amount of memory (in bytes) used by external processes related to this, if any. */
+	virtual FShaderCompileMemoryUsage GetExternalWorkerMemoryUsage() { return FShaderCompileMemoryUsage{}; }
 
 	// Returns a name for this thread instance. Defaults to "ShaderCompilingThread".
 	virtual const TCHAR* GetThreadName() const
@@ -194,6 +208,7 @@ public:
 
 protected:
 	virtual void PrintWorkerMemoryUsage() override;
+	virtual FShaderCompileMemoryUsage GetExternalWorkerMemoryUsage() override;
 
 private:
 
@@ -978,6 +993,11 @@ public:
 	  */
 	ENGINE_API void GetLocalStats(FShaderCompilerStats & OutStats) const;
 
+	/**
+	 * Returns the current memory usage of external local compilation processes in bytes.
+	 */
+	ENGINE_API FShaderCompileMemoryUsage GetExternalMemoryUsage();
+
 	/** 
 	 * Processes completed asynchronous shader maps, and assigns them to relevant materials.
 	 * @param TimeSlice - When more than 0, ProcessAsyncResults will be bandwidth throttled by the given timeslice, to limit hitching.
