World Partition Actor Descriptors:
- Make loading functions private to make sure all code requesting actors loading state changes in the editor goes through the loading contexts.
- Actor descriptors can now load async when using the deferred context, in preparation of editor async loading through a new async context.

diff --git a/Engine/Source/Runtime/Engine/Private/WorldPartition/WorldPartitionActorDescInstance.cpp b/Engine/Source/Runtime/Engine/Private/WorldPartition/WorldPartitionActorDescInstance.cpp
--- a/Engine/Source/Runtime/Engine/Private/WorldPartition/WorldPartitionActorDescInstance.cpp
+++ b/Engine/Source/Runtime/Engine/Private/WorldPartition/WorldPartitionActorDescInstance.cpp
@@ -16,11 +16,12 @@ FWorldPartitionActorDescInstance::FWorldPartitionActorDescInstance()
 	, SoftRefCount(0)
 	, HardRefCount(0)
 	, bIsForcedNonSpatiallyLoaded(false)
+	, bIsRegisteringOrUnregistering(false)
 	, UnloadedReason(nullptr)
+	, AsyncLoadID(INDEX_NONE)
 	, ActorDesc(nullptr)
 	, ChildContainerInstance(nullptr)
-{
-}
+{}
 
 FWorldPartitionActorDescInstance::FWorldPartitionActorDescInstance(UActorDescContainerInstance* InContainerInstance, FWorldPartitionActorDesc* InActorDesc)
 	: FWorldPartitionActorDescInstance()
@@ -49,6 +50,8 @@ bool FWorldPartitionActorDescInstance::IsLoaded(bool bEvenIfPendingKill) const
 
 AActor* FWorldPartitionActorDescInstance::GetActor(bool bEvenIfPendingKill, bool bEvenIfUnreachable) const
 {
+	FlushAsyncLoad();
+
 	if (ActorPtr.IsExplicitlyNull() || ActorPtr.IsStale())
 	{
 		ActorPtr = FindObject<AActor>(nullptr, *GetActorSoftPath().ToString());
@@ -87,7 +90,7 @@ bool FWorldPartitionActorDescInstance::IsRuntimeRelevant() const
 	return GetActorDesc()->IsRuntimeRelevant(this);
 }
 
-AActor* FWorldPartitionActorDescInstance::Load()
+bool FWorldPartitionActorDescInstance::StartAsyncLoad()
 {
 	static FText FailedToLoad(LOCTEXT("FailedToLoadReason", "Failed to load"));
 	UnloadedReason = nullptr;
@@ -95,43 +98,59 @@ AActor* FWorldPartitionActorDescInstance::Load()
 	if (ActorPtr.IsExplicitlyNull() || ActorPtr.IsStale())
 	{
 		// First, try to find the existing actor which could have been loaded by another actor (through standard serialization)
-		ActorPtr = FindObject<AActor>(nullptr, *GetActorSoftPath().ToString());
+		ActorPtr = FindObject<AActor>(nullptr, * GetActorSoftPath().ToString());
 	}
 
 	// Then, if the actor isn't loaded, load it
 	if (ActorPtr.IsExplicitlyNull())
 	{
-		const FLinkerInstancingContext* InstancingContext = GetContainerInstance()->GetInstancingContext();
-		FSoftObjectPath LocalActorPath = GetActorSoftPath();
-
-		UPackage* Package = nullptr;
+		const FName ActorPackage = GetActorPackage();
+		const FLinkerInstancingContext* InstancingContext = GetContainerInstance()->GetInstancingContext();		
+		const FName PackageName = InstancingContext ? InstancingContext->RemapPackage(ActorPackage) : ActorPackage;
+		const FPackagePath PackagePath = FPackagePath::FromPackageNameChecked(ActorPackage);
 
-		if (InstancingContext)
+		AsyncLoadID = LoadPackageAsync(PackagePath, PackageName, FLoadPackageAsyncDelegate::CreateLambda([this, ActorPackage](const FName& PackageName, UPackage* Package, EAsyncLoadingResult::Type Result)
 		{
-			FName RemappedPackageName = InstancingContext->RemapPackage(GetActorPackage());
-			check(RemappedPackageName != LocalActorPath.GetLongPackageFName());
+			check(AsyncLoadID != INDEX_NONE);
+			AsyncLoadID = INDEX_NONE;
 
-			Package = CreatePackage(*RemappedPackageName.ToString());
-		}
+			if ((Result != EAsyncLoadingResult::Succeeded) || !Package)
+			{
+				UE_LOG(LogWorldPartition, Warning, TEXT("Can't load actor guid `%s` ('%s') from package '%s'"), *GetGuid().ToString(), *GetActorName().ToString(), *ActorPackage.ToString());
+				UnloadedReason = &FailedToLoad;
+				return;
+			}
 
-		Package = LoadPackage(Package, *GetActorPackage().ToString(), LOAD_None, nullptr, InstancingContext);
+			ActorPtr = FindObject<AActor>(nullptr, * GetActorSoftPath().ToString());
 
-		if (Package)
-		{
-			ActorPtr = FindObject<AActor>(nullptr, *LocalActorPath.ToString());
 			if (!ActorPtr.IsValid())
 			{
-				UE_LOG(LogWorldPartition, Warning, TEXT("Can't load actor guid `%s` ('%s') from package '%s'"), *GetGuid().ToString(), *GetActorName().ToString(), *GetActorPackage().ToString());
+				UE_LOG(LogWorldPartition, Warning, TEXT("Can't find actor guid `%s` ('%s') in package '%s'"), *GetGuid().ToString(), *GetActorName().ToString(), *ActorPackage.ToString());
 				UnloadedReason = &FailedToLoad;
+				return;
 			}
-		}
+
+			check(ActorPtr->GetPackage() == Package);
+		})
+		, PKG_None, INDEX_NONE, 0, InstancingContext);
 	}
 
-	return ActorPtr.Get();
+	return (AsyncLoadID != INDEX_NONE) || ActorPtr.IsValid(false);
 }
 
-void FWorldPartitionActorDescInstance::Unload()
+void FWorldPartitionActorDescInstance::FlushAsyncLoad() const
 {
+	if (AsyncLoadID != INDEX_NONE)
+	{
+		FlushAsyncLoading(AsyncLoadID);
+		AsyncLoadID = INDEX_NONE;
+	}
+}
+
+void FWorldPartitionActorDescInstance::MarkUnload()
+{
+	FlushAsyncLoad();
+
 	// Notify Desc as it can have some custom code to run on the actor depending on type
 	GetActorDesc()->OnUnloadingInstance(this);
 
@@ -145,18 +164,18 @@ void FWorldPartitionActorDescInstance::Unload()
 		//
 		// FWorldPartitionCookPackageSplitter:
 		//		should mark each FWorldPartitionActorDesc as moved, and the splitter should take responsbility for calling ClearFlags on every object in 
-		//		the package when it does the move
+		//		the package when it does the move.
 
 		if (Actor->IsPackageExternal())
 		{
 			ForEachObjectWithPackage(Actor->GetPackage(), [](UObject* Object)
+			{
+				if (Object->HasAnyFlags(RF_Public | RF_Standalone))
 				{
-					if (Object->HasAnyFlags(RF_Public | RF_Standalone))
-					{
-						CastChecked<UMetaData>(Object)->ClearFlags(RF_Public | RF_Standalone);
-					}
-					return true;
-				}, false);
+					CastChecked<UMetaData>(Object)->ClearFlags(RF_Public | RF_Standalone);
+				}
+				return true;
+			}, false);
 		}
 
 		ActorPtr = nullptr;
diff --git a/Engine/Source/Runtime/Engine/Private/WorldPartition/WorldPartitionHandle.cpp b/Engine/Source/Runtime/Engine/Private/WorldPartition/WorldPartitionHandle.cpp
--- a/Engine/Source/Runtime/Engine/Private/WorldPartition/WorldPartitionHandle.cpp
+++ b/Engine/Source/Runtime/Engine/Private/WorldPartition/WorldPartitionHandle.cpp
@@ -57,18 +57,16 @@ FWorldPartitionLoadingContext::IContext* FWorldPartitionLoadingContext::ActiveCo
 
 void FWorldPartitionLoadingContext::LoadAndRegisterActor(FWorldPartitionActorDescInstance* InActorDescInstance)
 {
-#if DO_CHECK
-	FWorldPartitionActorDescInstance::FRegisteringUnregisteringGuard Guard(InActorDescInstance);
-#endif
+	check(!InActorDescInstance->bIsRegisteringOrUnregistering);
+	InActorDescInstance->bIsRegisteringOrUnregistering = true;
 
 	ActiveContext->RegisterActor(InActorDescInstance);
 }
 
 void FWorldPartitionLoadingContext::UnloadAndUnregisterActor(FWorldPartitionActorDescInstance* InActorDescInstance)
 {
-#if DO_CHECK
-	FWorldPartitionActorDescInstance::FRegisteringUnregisteringGuard Guard(InActorDescInstance);
-#endif
+	check(!InActorDescInstance->bIsRegisteringOrUnregistering);
+	InActorDescInstance->bIsRegisteringOrUnregistering = true;
 
 	ActiveContext->UnregisterActor(InActorDescInstance);
 }
@@ -93,9 +91,10 @@ FWorldPartitionLoadingContext::IContext::~IContext()
 */
 void FWorldPartitionLoadingContext::FImmediate::RegisterActor(FWorldPartitionActorDescInstance* InActorDescInstance)
 {
+	// Set GIsEditorLoadingPackage to avoid dirtying the Actor package if Modify() is called during the load sequence
 	TGuardValue<bool> IsEditorLoadingPackageGuard(GIsEditorLoadingPackage, true);
 
-	if (AActor* Actor = InActorDescInstance->Load())
+	if (InActorDescInstance->StartAsyncLoad())
 	{
 		UActorDescContainerInstance* Container = InActorDescInstance->GetContainerInstance();
 		check(Container);
@@ -103,8 +102,14 @@ void FWorldPartitionLoadingContext::FImmediate::RegisterActor(FWorldPartitionAct
 		const FTransform& ContainerTransform = Container->GetTransform();
 		const FTransform* ContainerTransformPtr = ContainerTransform.Equals(FTransform::Identity) ? nullptr : &ContainerTransform;
 
-		Actor->GetLevel()->AddLoadedActor(Actor, ContainerTransformPtr);
+		if (AActor* Actor = InActorDescInstance->GetActor())
+		{
+			Actor->GetLevel()->AddLoadedActor(Actor, ContainerTransformPtr);
+		}
 	}
+
+	check(InActorDescInstance->bIsRegisteringOrUnregistering);
+	InActorDescInstance->bIsRegisteringOrUnregistering = false;
 }
 
 void FWorldPartitionLoadingContext::FImmediate::UnregisterActor(FWorldPartitionActorDescInstance* InActorDescInstance)
@@ -123,8 +128,11 @@ void FWorldPartitionLoadingContext::FImmediate::UnregisterActor(FWorldPartitionA
 
 		Actor->GetLevel()->RemoveLoadedActor(Actor, ContainerTransformPtr);
 
-		InActorDescInstance->Unload();
+		InActorDescInstance->MarkUnload();
 	}
+
+	check(InActorDescInstance->bIsRegisteringOrUnregistering);
+	InActorDescInstance->bIsRegisteringOrUnregistering = false;
 }
 
 /**
@@ -148,7 +156,8 @@ FWorldPartitionLoadingContext::FDeferred::~FDeferred()
 				ULevel* Level = nullptr;
 				for (FWorldPartitionActorDescInstance* ActorDescInstance : SourceList)
 				{
-					// When cleaning up worlds, actors are already marked as garbage at this point, so no need to remove them from the world
+					// For async loads, the actor loading might have failed for several reasons, this will be handled here.
+					// Also, when cleaning up worlds, actors are already marked as garbage at this point, so no need to remove them from the world.
 					if (AActor* Actor = ActorDescInstance->GetActor(); IsValid(Actor))
 					{
 						ActorList.Add(Actor);
@@ -164,6 +173,7 @@ FWorldPartitionLoadingContext::FDeferred::~FDeferred()
 
 			if (ContainerOp.Registrations.Num())
 			{
+				// Set GIsEditorLoadingPackage to avoid dirtying the Actor package if Modify() is called during the load sequence
 				TGuardValue<bool> IsEditorLoadingPackageGuard(GIsEditorLoadingPackage, true);
 
 				TArray<AActor*> ActorList;
@@ -171,6 +181,12 @@ FWorldPartitionLoadingContext::FDeferred::~FDeferred()
 				{
 					Level->AddLoadedActors(ActorList, ContainerTransformPtr);
 				}
+
+				for (FWorldPartitionActorDescInstance* ActorDescInstance : ContainerOp.Registrations)
+				{
+					check(ActorDescInstance->bIsRegisteringOrUnregistering);
+					ActorDescInstance->bIsRegisteringOrUnregistering = false;
+				}
 			}
 
 			if (ContainerOp.Unregistrations.Num())
@@ -186,7 +202,9 @@ FWorldPartitionLoadingContext::FDeferred::~FDeferred()
 
 				for (FWorldPartitionActorDescInstance* ActorDescInstance : ContainerOp.Unregistrations)
 				{
-					ActorDescInstance->Unload();
+					ActorDescInstance->MarkUnload();
+					check(ActorDescInstance->bIsRegisteringOrUnregistering);
+					ActorDescInstance->bIsRegisteringOrUnregistering = false;
 				}
 			}
 		}
@@ -194,7 +212,6 @@ FWorldPartitionLoadingContext::FDeferred::~FDeferred()
 		// Continue with potentially new registrations/unregistrations that may have happenned during previous cycle
 		LocalContainerInstanceOps = MoveTemp(ContainerInstanceOps);
 	}
-	
 }
 
 void FWorldPartitionLoadingContext::FDeferred::RegisterActor(FWorldPartitionActorDescInstance* InActorDescInstance)
@@ -202,7 +219,7 @@ void FWorldPartitionLoadingContext::FDeferred::RegisterActor(FWorldPartitionActo
 	TGuardValue<bool> IsEditorLoadingPackageGuard(GIsEditorLoadingPackage, true);
 
 	check(InActorDescInstance);
-	if (InActorDescInstance->Load())
+	if (InActorDescInstance->StartAsyncLoad())
 	{
 		UActorDescContainerInstance* Container = InActorDescInstance->GetContainerInstance();
 		check(Container);
diff --git a/Engine/Source/Runtime/Engine/Public/WorldPartition/WorldPartitionActorDescInstance.h b/Engine/Source/Runtime/Engine/Public/WorldPartition/WorldPartitionActorDescInstance.h
--- a/Engine/Source/Runtime/Engine/Public/WorldPartition/WorldPartitionActorDescInstance.h
+++ b/Engine/Source/Runtime/Engine/Public/WorldPartition/WorldPartitionActorDescInstance.h
@@ -26,6 +26,7 @@ class FWorldPartitionActorDescInstance
 	friend struct FWorldPartitionActorDescUnitTestAcccessor;
 	friend class UDataLayerManager;
 	friend class UWorldPartition;
+	friend class FWorldPartitionLoadingContext;
 	friend class IWorldPartitionActorLoaderInterface;
 
 protected:
@@ -34,58 +35,58 @@ public:
 	ENGINE_API FWorldPartitionActorDescInstance(UActorDescContainerInstance* InContainerInstance, FWorldPartitionActorDesc* InActorDesc);
 	virtual ~FWorldPartitionActorDescInstance() {}
 
-	virtual const FGuid& GetGuid() const { return GetActorDesc()->GetGuid(); }
-	virtual FTopLevelAssetPath GetBaseClass() const { return GetActorDesc()->GetBaseClass(); }
-	virtual FTopLevelAssetPath GetNativeClass() const { return GetActorDesc()->GetNativeClass(); }
-	virtual UClass* GetActorNativeClass() const { return GetActorDesc()->GetActorNativeClass(); }
+	virtual const FGuid& GetGuid() const { return ActorDesc->GetGuid(); }
+	virtual FTopLevelAssetPath GetBaseClass() const { return ActorDesc->GetBaseClass(); }
+	virtual FTopLevelAssetPath GetNativeClass() const { return ActorDesc->GetNativeClass(); }
+	virtual UClass* GetActorNativeClass() const { return ActorDesc->GetActorNativeClass(); }
 	
-	virtual FName GetRuntimeGrid() const { return GetActorDesc()->GetRuntimeGrid(); }
-	virtual bool GetIsSpatiallyLoaded() const { return !GetForceNonSpatiallyLoaded() && GetActorDesc()->GetIsSpatiallyLoaded(); }
-	virtual bool GetActorIsEditorOnly() const { return GetActorDesc()->GetActorIsEditorOnly(); }
-	virtual bool GetActorIsRuntimeOnly() const { return GetActorDesc()->GetActorIsRuntimeOnly(); }
+	virtual FName GetRuntimeGrid() const { return ActorDesc->GetRuntimeGrid(); }
+	virtual bool GetIsSpatiallyLoaded() const { return !GetForceNonSpatiallyLoaded() && ActorDesc->GetIsSpatiallyLoaded(); }
+	virtual bool GetActorIsEditorOnly() const { return ActorDesc->GetActorIsEditorOnly(); }
+	virtual bool GetActorIsRuntimeOnly() const { return ActorDesc->GetActorIsRuntimeOnly(); }
 	ENGINE_API virtual bool IsRuntimeRelevant() const;
 	ENGINE_API virtual bool IsEditorRelevant() const;
 
-	virtual bool IsUsingDataLayerAsset() const { return GetActorDesc()->IsUsingDataLayerAsset(); }
-	virtual TArray<FName> GetDataLayers() const { return GetActorDesc()->GetDataLayers(); }
+	virtual bool IsUsingDataLayerAsset() const { return ActorDesc->IsUsingDataLayerAsset(); }
+	virtual TArray<FName> GetDataLayers() const { return ActorDesc->GetDataLayers(); }
 
-	virtual bool GetActorIsHLODRelevant() const { return GetActorDesc()->GetActorIsHLODRelevant(); }
-	virtual FSoftObjectPath GetHLODLayer() const { return GetActorDesc()->GetHLODLayer(); }
+	virtual bool GetActorIsHLODRelevant() const { return ActorDesc->GetActorIsHLODRelevant(); }
+	virtual FSoftObjectPath GetHLODLayer() const { return ActorDesc->GetHLODLayer(); }
 
-	virtual const TArray<FName>& GetTags() const { return GetActorDesc()->GetTags(); }
-	virtual FName GetActorPackage() const { return GetActorDesc()->GetActorPackage(); }
+	virtual const TArray<FName>& GetTags() const { return ActorDesc->GetTags(); }
+	virtual FName GetActorPackage() const { return ActorDesc->GetActorPackage(); }
 	ENGINE_API virtual FSoftObjectPath GetActorSoftPath() const;
-	virtual FName GetActorLabel() const { return GetActorDesc()->GetActorLabel(); }
+	virtual FName GetActorLabel() const { return ActorDesc->GetActorLabel(); }
 	ENGINE_API virtual FName GetActorName() const;
-	virtual FName GetFolderPath() const { return GetActorDesc()->GetFolderPath(); }
-	virtual const FGuid& GetFolderGuid() const { return GetActorDesc()->GetFolderGuid(); }
+	virtual FName GetFolderPath() const { return ActorDesc->GetFolderPath(); }
+	virtual const FGuid& GetFolderGuid() const { return ActorDesc->GetFolderGuid(); }
 
-	virtual FBox GetEditorBounds() const { return GetActorDesc()->GetEditorBounds(); }
-	virtual FBox GetRuntimeBounds() const { return GetActorDesc()->GetRuntimeBounds(); }
+	virtual FBox GetEditorBounds() const { return ActorDesc->GetEditorBounds(); }
+	virtual FBox GetRuntimeBounds() const { return ActorDesc->GetRuntimeBounds(); }
 
-	virtual bool GetProperty(FName PropertyName, FName* PropertyValue) const { return GetActorDesc()->GetProperty(PropertyName, PropertyValue); }
-	virtual bool HasProperty(FName PropertyName) const { return GetActorDesc()->HasProperty(PropertyName); }
+	virtual bool GetProperty(FName PropertyName, FName* PropertyValue) const { return ActorDesc->GetProperty(PropertyName, PropertyValue); }
+	virtual bool HasProperty(FName PropertyName) const { return ActorDesc->HasProperty(PropertyName); }
 
-	virtual const TArray<FGuid>& GetReferences() const { return GetActorDesc()->GetReferences(); }
-	virtual const TArray<FGuid>& GetEditorOnlyReferences() const { return GetActorDesc()->GetEditorOnlyReferences(); }
-	virtual bool IsEditorOnlyReference(const FGuid& ReferenceGuid) const { return GetActorDesc()->IsEditorOnlyReference(ReferenceGuid); }
+	virtual const TArray<FGuid>& GetReferences() const { return ActorDesc->GetReferences(); }
+	virtual const TArray<FGuid>& GetEditorOnlyReferences() const { return ActorDesc->GetEditorOnlyReferences(); }
+	virtual bool IsEditorOnlyReference(const FGuid& ReferenceGuid) const { return ActorDesc->IsEditorOnlyReference(ReferenceGuid); }
 	
-	virtual const FGuid& GetParentActor() const { return GetActorDesc()->GetParentActor(); }
+	virtual const FGuid& GetParentActor() const { return ActorDesc->GetParentActor(); }
 		
-	virtual FGuid GetContentBundleGuid() const { return GetActorDesc()->GetContentBundleGuid(); }
-	virtual const FSoftObjectPath& GetExternalDataLayerAsset() const { return GetActorDesc()->GetExternalDataLayerAsset(); }
-
-	virtual bool IsChildContainerInstance() const { return ChildContainerInstance || GetActorDesc()->IsChildContainerInstance(); }
-	virtual FName GetChildContainerPackage() const { return GetActorDesc()->GetChildContainerPackage(); }
-	virtual EWorldPartitionActorFilterType GetChildContainerFilterType() const { return GetActorDesc()->GetChildContainerFilterType(); }
-	virtual const FWorldPartitionActorFilter* GetChildContainerFilter() const { return GetActorDesc()->GetChildContainerFilter(); }
-	virtual bool GetChildContainerInstance(FWorldPartitionActorDesc::FContainerInstance& OutContainerInstance) const { return GetActorDesc()->GetChildContainerInstance(this, OutContainerInstance); }
+	virtual FGuid GetContentBundleGuid() const { return ActorDesc->GetContentBundleGuid(); }
+	virtual const FSoftObjectPath& GetExternalDataLayerAsset() const { return ActorDesc->GetExternalDataLayerAsset(); }
+
+	virtual bool IsChildContainerInstance() const { return ChildContainerInstance || ActorDesc->IsChildContainerInstance(); }
+	virtual FName GetChildContainerPackage() const { return ActorDesc->GetChildContainerPackage(); }
+	virtual EWorldPartitionActorFilterType GetChildContainerFilterType() const { return ActorDesc->GetChildContainerFilterType(); }
+	virtual const FWorldPartitionActorFilter* GetChildContainerFilter() const { return ActorDesc->GetChildContainerFilter(); }
+	virtual bool GetChildContainerInstance(FWorldPartitionActorDesc::FContainerInstance& OutContainerInstance) const { return ActorDesc->GetChildContainerInstance(this, OutContainerInstance); }
 			
-	virtual bool IsMainWorldOnly() const { return GetActorDesc()->IsMainWorldOnly(); }
-	virtual bool IsListedInSceneOutliner() const { return GetActorDesc()->IsListedInSceneOutliner(); }
-	virtual const FGuid& GetSceneOutlinerParent() const { return GetActorDesc()->GetSceneOutlinerParent(); }
+	virtual bool IsMainWorldOnly() const { return ActorDesc->IsMainWorldOnly(); }
+	virtual bool IsListedInSceneOutliner() const { return ActorDesc->IsListedInSceneOutliner(); }
+	virtual const FGuid& GetSceneOutlinerParent() const { return ActorDesc->GetSceneOutlinerParent(); }
 			
-	ENGINE_API virtual const FWorldPartitionActorDesc* GetActorDesc() const { check(ActorDesc); return ActorDesc; }
+	inline const FWorldPartitionActorDesc* GetActorDesc() const { return ActorDesc; }
 	
 	virtual bool HasResolvedDataLayerInstanceNames() const { return ResolvedDataLayerInstanceNames.IsSet(); }
 	ENGINE_API virtual const FDataLayerInstanceNames& GetDataLayerInstanceNames() const;
@@ -109,15 +110,16 @@ public:
 
 	inline void SetUnloadedReason(FText* InUnloadedReason) { UnloadedReason = InUnloadedReason; }
 
-	FName GetDisplayClassName() const { return GetActorDesc()->GetDisplayClassName(); }
+	FName GetDisplayClassName() const { return ActorDesc->GetDisplayClassName(); }
 	
 	ENGINE_API const FText& GetUnloadedReason() const;
 		
-	ENGINE_API AActor* Load();
-	ENGINE_API void Unload();
-
 protected:
-	UWorldPartition* GetLoadedChildWorldPartition() const { return GetActorDesc()->GetLoadedChildWorldPartition(this); }
+	ENGINE_API bool StartAsyncLoad();
+	ENGINE_API void FlushAsyncLoad() const;
+	ENGINE_API void MarkUnload();
+
+	UWorldPartition* GetLoadedChildWorldPartition() const { return ActorDesc->GetLoadedChildWorldPartition(this); }
 	ENGINE_API void UpdateActorDesc(FWorldPartitionActorDesc* InActorDesc);
 	void Invalidate();
 		
@@ -166,7 +168,9 @@ protected:
 	mutable uint32								HardRefCount;
 	TOptional<FDataLayerInstanceNames>			ResolvedDataLayerInstanceNames;
 	bool										bIsForcedNonSpatiallyLoaded;
+	bool										bIsRegisteringOrUnregistering;
 	mutable FText*								UnloadedReason;
+	mutable int32								AsyncLoadID;
 
 	// Instancing Path if set
 	TOptional<FSoftObjectPath>					ActorPath;
@@ -175,33 +179,5 @@ protected:
 	FWorldPartitionActorDesc*					ActorDesc;
 
 	TObjectPtr<UActorDescContainerInstance>		ChildContainerInstance;
-
-#if DO_CHECK
-private:
-	struct FRegisteringUnregisteringGuard
-	{
-		FRegisteringUnregisteringGuard(FWorldPartitionActorDescInstance* InActorDescInstance)
-			:ActorDescInstance(InActorDescInstance)
-		{
-			check(!ActorDescInstance->bIsRegisteringOrUnregistering);
-			ActorDescInstance->bIsRegisteringOrUnregistering = true;
-		}
-
-		~FRegisteringUnregisteringGuard()
-		{
-			check(ActorDescInstance->bIsRegisteringOrUnregistering);
-			ActorDescInstance->bIsRegisteringOrUnregistering = false;
-		}
-
-		FRegisteringUnregisteringGuard(const FRegisteringUnregisteringGuard&) = delete;
-		FRegisteringUnregisteringGuard& operator=(const FRegisteringUnregisteringGuard&) = delete;
-
-	private:
-		FWorldPartitionActorDescInstance* ActorDescInstance;
-	};
-	bool bIsRegisteringOrUnregistering = false;
-#endif // DO_CHECK
-
 #endif // WITH_EDITOR
 };
-
diff --git a/Engine/Source/Runtime/Engine/Public/WorldPartition/WorldPartitionHandle.h b/Engine/Source/Runtime/Engine/Public/WorldPartition/WorldPartitionHandle.h
--- a/Engine/Source/Runtime/Engine/Public/WorldPartition/WorldPartitionHandle.h
+++ b/Engine/Source/Runtime/Engine/Public/WorldPartition/WorldPartitionHandle.h
@@ -24,6 +24,8 @@ class TActorDescContainerInstanceCollection;
 
 class FWorldPartitionLoadingContext
 {
+	friend struct FWorldPartitionReferenceImpl;
+
 public:
 	/**
 	 * Base class for loading contexts
@@ -97,10 +99,10 @@ public:
 	UE_DEPRECATED(5.4, "Use FWorldPartitionActorDescInstance version instead")
 	static ENGINE_API void UnloadAndUnregisterActor(FWorldPartitionActorDesc* ActorDesc) {}
 
-	static ENGINE_API void LoadAndRegisterActor(FWorldPartitionActorDescInstance* ActorDesc);
-	static ENGINE_API void UnloadAndUnregisterActor(FWorldPartitionActorDescInstance* ActorDesc);
-
 private:
+	static void LoadAndRegisterActor(FWorldPartitionActorDescInstance* ActorDescInstance);
+	static void UnloadAndUnregisterActor(FWorldPartitionActorDescInstance* ActorDescInstance);
+
 	static FImmediate DefaultContext;
 	static IContext* ActiveContext;	
 };
