World Partition - HLOD: Turn off all collision channels on HLOD actors, except for ECC_Visibility & ECC_Camera (only in editor)

diff --git a/Engine/Source/Runtime/Engine/Private/WorldPartition/HLOD/HLODActor.cpp b/Engine/Source/Runtime/Engine/Private/WorldPartition/HLOD/HLODActor.cpp
--- a/Engine/Source/Runtime/Engine/Private/WorldPartition/HLOD/HLODActor.cpp
+++ b/Engine/Source/Runtime/Engine/Private/WorldPartition/HLOD/HLODActor.cpp
@@ -316,6 +316,7 @@ void AWorldPartitionHLOD::PreRegisterAllComponents()
 			ForEachComponent<UPrimitiveComponent>(false, [bShouldEnableCollision](UPrimitiveComponent* PrimitiveComponent)
 			{
 				PrimitiveComponent->SetCollisionEnabled(bShouldEnableCollision ? ECollisionEnabled::QueryOnly : ECollisionEnabled::NoCollision);
+				PrimitiveComponent->SetCollisionResponseToAllChannels(ECollisionResponse::ECR_Ignore);
 				PrimitiveComponent->SetCollisionResponseToChannel(ECC_Visibility, bShouldEnableCollision ? ECR_Block : ECR_Ignore);
 				PrimitiveComponent->SetCollisionResponseToChannel(ECC_Camera, bShouldEnableCollision ? ECR_Block : ECR_Ignore);
 			});
