Add an API on UAnimInstance to override the value of a curve

#changelist validated
#virtualized

diff --git a/Engine/Source/Runtime/Engine/Classes/Animation/AnimInstance.h b/Engine/Source/Runtime/Engine/Classes/Animation/AnimInstance.h
--- a/Engine/Source/Runtime/Engine/Classes/Animation/AnimInstance.h
+++ b/Engine/Source/Runtime/Engine/Classes/Animation/AnimInstance.h
@@ -1155,6 +1155,9 @@ public:
 	/** Returns value of named curved in OutValue, returns whether the curve was actually found or not. */
 	ENGINE_API bool GetCurveValue(FName CurveName, float& OutValue) const;
 
+	/** Overrides the value of a named curve. Will be reset next evaluation */
+	ENGINE_API void OverrideCurveValue(FName CurveName, float Value);
+
 	/** Returns the name of a currently active state in a state machine. */
 	UFUNCTION(BlueprintPure, Category="Animation|State Machines", meta=(BlueprintInternalUseOnly = "true", AnimGetter = "true", BlueprintThreadSafe))
 	ENGINE_API FName GetCurrentStateName(int32 MachineIndex);
diff --git a/Engine/Source/Runtime/Engine/Private/Animation/AnimInstance.cpp b/Engine/Source/Runtime/Engine/Private/Animation/AnimInstance.cpp
--- a/Engine/Source/Runtime/Engine/Private/Animation/AnimInstance.cpp
+++ b/Engine/Source/Runtime/Engine/Private/Animation/AnimInstance.cpp
@@ -1759,6 +1759,14 @@ void UAnimInstance::GetAllCurveNames(TArray<FName>& OutNames) const
 	GetActiveCurveNames(EAnimCurveType::AttributeCurve, OutNames);
 }
 
+void UAnimInstance::OverrideCurveValue(FName CurveName, float Value)
+{
+	FAnimInstanceProxy& Proxy = GetProxyOnAnyThread<FAnimInstanceProxy>();
+
+	TMap<FName, float>& AnimationCurves = Proxy.GetAnimationCurves(EAnimCurveType::AttributeCurve);
+	AnimationCurves.FindOrAdd(CurveName) = Value;
+}
+
 void UAnimInstance::SetRootMotionMode(TEnumAsByte<ERootMotionMode::Type> Value)
 {
 	RootMotionMode = Value;
