make geomtools2d stop converting to float in the middle of its triangulation point-on-line predicate, and use a smaller tolerance

diff --git a/Engine/Source/Runtime/Engine/Private/GeomTools.cpp b/Engine/Source/Runtime/Engine/Private/GeomTools.cpp
--- a/Engine/Source/Runtime/Engine/Private/GeomTools.cpp
+++ b/Engine/Source/Runtime/Engine/Private/GeomTools.cpp
@@ -984,9 +984,9 @@ static bool IsPointOnLineSegment(const FVector2D& TestPoint, const FVector2D& A,
 {
 	FVector2D BA = B - A;
 	FVector2D PA = TestPoint - A;
-	float SizeSquaredBA = FVector2D::DotProduct(BA, BA);
-	float AreaCompareThreshold = 0.01f * SizeSquaredBA;
-	float ParallelogramArea = BA.X * PA.Y - BA.Y * PA.X;
+	double SizeSquaredBA = FVector2D::DotProduct(BA, BA);
+	double AreaCompareThreshold = UE_DOUBLE_KINDA_SMALL_NUMBER * SizeSquaredBA;
+	double ParallelogramArea = BA.X * PA.Y - BA.Y * PA.X;
 
 	return  TestPoint.X >= FMath::Min(A.X, B.X) && TestPoint.X <= FMath::Max(A.X, B.X) && // X within AB.X, including ON A or B
 		TestPoint.Y >= FMath::Min(A.Y, B.Y) && TestPoint.Y <= FMath::Max(A.Y, B.Y) && // Y within AB.Y, including ON A or B
