Exposing Engine function FSendingRepState::BuildConditionMapFromRepFlags as inline UE::Net::BuildConditionMapFromRepFlags for use in subsequent GAS Gameplay Debugger changes.

Also const qualifying APlayerController on method UActorChannel::CanSubObjectReplicateToClient.

diff --git a/Engine/Source/Runtime/Engine/Classes/Engine/ActorChannel.h b/Engine/Source/Runtime/Engine/Classes/Engine/ActorChannel.h
--- a/Engine/Source/Runtime/Engine/Classes/Engine/ActorChannel.h
+++ b/Engine/Source/Runtime/Engine/Classes/Engine/ActorChannel.h
@@ -372,7 +372,7 @@ public:
 	ENGINE_API virtual void AddedToChannelPool() override;
 
 	static ENGINE_API bool CanSubObjectReplicateToClient(
-		APlayerController* PlayerController,
+		const APlayerController* PlayerController,
 		ELifetimeCondition NetCondition,
 		FObjectKey SubObjectKey,
 		const TStaticBitArray<COND_Max>& ConditionMap,
diff --git a/Engine/Source/Runtime/Engine/Private/DataChannel.cpp b/Engine/Source/Runtime/Engine/Private/DataChannel.cpp
--- a/Engine/Source/Runtime/Engine/Private/DataChannel.cpp
+++ b/Engine/Source/Runtime/Engine/Private/DataChannel.cpp
@@ -3919,7 +3919,7 @@ bool UActorChannel::UpdateDeletedSubObjects(FOutBunch& Bunch)
 }
 
 bool UActorChannel::CanSubObjectReplicateToClient(
-	APlayerController* PlayerController,
+	const APlayerController* PlayerController,
 	ELifetimeCondition NetCondition,
 	FObjectKey SubObjectKey,
 	const TStaticBitArray<COND_Max>& ConditionMap,
@@ -3970,7 +3970,7 @@ bool UActorChannel::ReplicateRegisteredSubObjects(FOutBunch& Bunch, FReplication
 	}
 #endif
 
-	const TStaticBitArray<COND_Max> ConditionMap = FSendingRepState::BuildConditionMapFromRepFlags(RepFlags);
+	const TStaticBitArray<COND_Max> ConditionMap = UE::Net::BuildConditionMapFromRepFlags(RepFlags);
 
 	bool bWroteSomethingImportant = false;
 
@@ -4136,7 +4136,7 @@ bool UActorChannel::ReplicateSubobject(UActorComponent* ReplicatedComponent, FOu
 
 		bool bWroteSomethingImportant = false;
 
-		const TStaticBitArray<COND_Max> ConditionMap = FSendingRepState::BuildConditionMapFromRepFlags(RepFlags);
+		const TStaticBitArray<COND_Max> ConditionMap = UE::Net::BuildConditionMapFromRepFlags(RepFlags);
 
 		checkf(Actor->IsUsingRegisteredSubObjectList() == false, TEXT("This code should only be hit when an Actor that does NOT support the SubObjectList is replicating an ActorComponent that does: %s replicating %s."), 
 			*Actor->GetName(), *ReplicatedComponent->GetName());
diff --git a/Engine/Source/Runtime/Engine/Private/DataReplication.cpp b/Engine/Source/Runtime/Engine/Private/DataReplication.cpp
--- a/Engine/Source/Runtime/Engine/Private/DataReplication.cpp
+++ b/Engine/Source/Runtime/Engine/Private/DataReplication.cpp
@@ -1588,7 +1588,7 @@ void FObjectReplicator::ReplicateCustomDeltaProperties( FNetBitWriter & Bunch, F
 	};
 
 	// Initialize a map of which conditions are valid
-	const TStaticBitArray<COND_Max> ConditionMap = FSendingRepState::BuildConditionMapFromRepFlags(RepFlags);
+	const TStaticBitArray<COND_Max> ConditionMap = UE::Net::BuildConditionMapFromRepFlags(RepFlags);
 
 	// Make sure net field export group is registered
 	FNetFieldExportGroup* NetFieldExportGroup = OwningChannel->GetOrCreateNetFieldExportGroupForClassNetCache( Object );
diff --git a/Engine/Source/Runtime/Engine/Private/NetConnection.cpp b/Engine/Source/Runtime/Engine/Private/NetConnection.cpp
--- a/Engine/Source/Runtime/Engine/Private/NetConnection.cpp
+++ b/Engine/Source/Runtime/Engine/Private/NetConnection.cpp
@@ -5003,7 +5003,7 @@ void UNetConnection::FlushDormancy(AActor* Actor)
 			RepFlags.bRepPhysics = Actor->GetReplicatedMovement().bRepPhysics;
 			RepFlags.bReplay = bReplay;
 
-			ConditionMap = FSendingRepState::BuildConditionMapFromRepFlags(RepFlags);
+			ConditionMap = UE::Net::BuildConditionMapFromRepFlags(RepFlags);
 			
 			const UWorld* const World = Actor->GetWorld();
 			const UNetworkSubsystem* const NetworkSubsystem = World ? World->GetSubsystem<UNetworkSubsystem>() : nullptr;
diff --git a/Engine/Source/Runtime/Engine/Private/RepLayout.cpp b/Engine/Source/Runtime/Engine/Private/RepLayout.cpp
--- a/Engine/Source/Runtime/Engine/Private/RepLayout.cpp
+++ b/Engine/Source/Runtime/Engine/Private/RepLayout.cpp
@@ -7104,42 +7104,6 @@ void FRepLayout::BuildHandleToCmdIndexTable_r(
 	}
 }
 
-TStaticBitArray<COND_Max> FSendingRepState::BuildConditionMapFromRepFlags(const FReplicationFlags RepFlags)
-{
-	TStaticBitArray<COND_Max> ConditionMap;
-
-	// Setup condition map
-	const bool bIsInitial = RepFlags.bNetInitial ? true : false;
-	const bool bIsOwner = RepFlags.bNetOwner ? true : false;
-	const bool bIsSimulated = RepFlags.bNetSimulated ? true : false;
-	const bool bIsPhysics = RepFlags.bRepPhysics ? true : false;
-	const bool bIsReplay = RepFlags.bReplay ? true : false;
-
-	ConditionMap[COND_None] = true;
-	ConditionMap[COND_InitialOnly] = bIsInitial;
-
-	ConditionMap[COND_OwnerOnly] = bIsOwner;
-	ConditionMap[COND_SkipOwner] = !bIsOwner;
-
-	ConditionMap[COND_SimulatedOnly] = bIsSimulated;
-	ConditionMap[COND_SimulatedOnlyNoReplay] = bIsSimulated && !bIsReplay;
-	ConditionMap[COND_AutonomousOnly] = !bIsSimulated;
-
-	ConditionMap[COND_SimulatedOrPhysics] = bIsSimulated || bIsPhysics;
-	ConditionMap[COND_SimulatedOrPhysicsNoReplay] = (bIsSimulated || bIsPhysics) && !bIsReplay;
-
-	ConditionMap[COND_InitialOrOwner] = bIsInitial || bIsOwner;
-	ConditionMap[COND_ReplayOrOwner] = bIsReplay || bIsOwner;
-	ConditionMap[COND_ReplayOnly] = bIsReplay;
-	ConditionMap[COND_SkipReplay] = !bIsReplay;
-
-	ConditionMap[COND_Custom] = true;
-	ConditionMap[COND_Dynamic] = true;
-	ConditionMap[COND_Never] = false;
-
-	return ConditionMap;
-}
-
 bool FSendingRepState::HasAnyPendingRetirements() const
 {
 	for (const FPropertyRetirement& PropRet : Retirement)
@@ -7159,7 +7123,7 @@ void FRepLayout::RebuildConditionalProperties(FSendingRepState* RepState, const
 	
 	RepState->RepFlags = RepFlags;
 
-	TStaticBitArray<COND_Max> ConditionMap = FSendingRepState::BuildConditionMapFromRepFlags(RepFlags);
+	TStaticBitArray<COND_Max> ConditionMap = UE::Net::BuildConditionMapFromRepFlags(RepFlags);
 	if (EnumHasAnyFlags(Flags, ERepLayoutFlags::HasDynamicConditionProperties) && RepState->RepChangedPropertyTracker.IsValid())
 	{
 		const FRepChangedPropertyTracker* RepChangedPropertyTracker = RepState->RepChangedPropertyTracker.Get();
diff --git a/Engine/Source/Runtime/Engine/Public/Net/RepLayout.h b/Engine/Source/Runtime/Engine/Public/Net/RepLayout.h
--- a/Engine/Source/Runtime/Engine/Public/Net/RepLayout.h
+++ b/Engine/Source/Runtime/Engine/Public/Net/RepLayout.h
@@ -124,6 +124,50 @@ typedef TConstRepDataBuffer<ERepDataBufferType::ShadowBuffer> FConstRepShadowDat
 class FRepLayout;
 class FRepLayoutCmd;
 
+namespace UE::Net
+{
+	/**
+	 * Builds a new ConditionMap given the input RepFlags.
+	 * This can be used to determine whether or not a given property should be
+	 * considered enabled / disabled based on ELifetimeCondition.
+	 */
+	inline TStaticBitArray<COND_Max> BuildConditionMapFromRepFlags(const FReplicationFlags& RepFlags)
+	{
+		TStaticBitArray<COND_Max> ConditionMap;
+
+		// Setup condition map
+		const bool bIsInitial = RepFlags.bNetInitial ? true : false;
+		const bool bIsOwner = RepFlags.bNetOwner ? true : false;
+		const bool bIsSimulated = RepFlags.bNetSimulated ? true : false;
+		const bool bIsPhysics = RepFlags.bRepPhysics ? true : false;
+		const bool bIsReplay = RepFlags.bReplay ? true : false;
+
+		ConditionMap[COND_None] = true;
+		ConditionMap[COND_InitialOnly] = bIsInitial;
+
+		ConditionMap[COND_OwnerOnly] = bIsOwner;
+		ConditionMap[COND_SkipOwner] = !bIsOwner;
+
+		ConditionMap[COND_SimulatedOnly] = bIsSimulated;
+		ConditionMap[COND_SimulatedOnlyNoReplay] = bIsSimulated && !bIsReplay;
+		ConditionMap[COND_AutonomousOnly] = !bIsSimulated;
+
+		ConditionMap[COND_SimulatedOrPhysics] = bIsSimulated || bIsPhysics;
+		ConditionMap[COND_SimulatedOrPhysicsNoReplay] = (bIsSimulated || bIsPhysics) && !bIsReplay;
+
+		ConditionMap[COND_InitialOrOwner] = bIsInitial || bIsOwner;
+		ConditionMap[COND_ReplayOrOwner] = bIsReplay || bIsOwner;
+		ConditionMap[COND_ReplayOnly] = bIsReplay;
+		ConditionMap[COND_SkipReplay] = !bIsReplay;
+
+		ConditionMap[COND_Custom] = true;
+		ConditionMap[COND_Dynamic] = true;
+		ConditionMap[COND_Never] = false;
+
+		return ConditionMap;
+	}
+}
+
 struct FRepSharedPropertyKey
 {
 private:
@@ -516,15 +560,9 @@ private:
 public:
 
 	void CountBytes(FArchive& Ar) const;
-	
-	/**
-	 * Builds a new ConditionMap given the input RepFlags.
-	 * This can be used to determine whether or not a given property should be
-	 * considered enabled / disabled based on ELifetimeCondition.
-	 *
-	 * TODO: This doesn't have to be part of FRepState.
-	 */
-	static TStaticBitArray<COND_Max> BuildConditionMapFromRepFlags(const FReplicationFlags InFlags);
+
+	UE_DEPRECATED(5.4, "Use UE::Net::BuildConditionMapFromRepFlags instead")
+	static inline TStaticBitArray<COND_Max> BuildConditionMapFromRepFlags(const FReplicationFlags InFlags) { return UE::Net::BuildConditionMapFromRepFlags(InFlags); }
 
 	bool HasAnyPendingRetirements() const;
 
