Audio Insights VirtualLoopStop messages now have DeviceId attached (was silently failing when getting the properly previously), so the VirtualLoop tab no longer displays virtual loops after they have stopped being active

diff --git a/Engine/Source/Runtime/Engine/Private/AudioDevice.cpp b/Engine/Source/Runtime/Engine/Private/AudioDevice.cpp
--- a/Engine/Source/Runtime/Engine/Private/AudioDevice.cpp
+++ b/Engine/Source/Runtime/Engine/Private/AudioDevice.cpp
@@ -225,6 +225,7 @@ static FAutoConsoleCommandWithWorldArgsAndOutputDevice GSetCurrentSpatialPluginC
 
 #if UE_AUDIO_PROFILERTRACE_ENABLED
 UE_TRACE_EVENT_BEGIN(Audio, VirtualLoopStop)
+	UE_TRACE_EVENT_FIELD(uint32, DeviceId)
 	UE_TRACE_EVENT_FIELD(uint64, Timestamp)
 	UE_TRACE_EVENT_FIELD(uint32, PlayOrder)
 UE_TRACE_EVENT_END()
@@ -5375,7 +5376,7 @@ bool FAudioDevice::RemoveVirtualLoop(FActiveSound& InActiveSound)
 		if (InActiveSound.Sound)
 		{
 			const FVector Location = InActiveSound.Transform.GetLocation();
-			UE_LOG(LogAudio, Verbose, TEXT("Removing virtual looping sound '%s' at location %s."), *InActiveSound.Sound->GetName(), *Location.ToCompactString());
+			UE_LOG(LogAudio, Verbose, TEXT("Removing virtual looping sound '%s' with play order %d at location %s."), *InActiveSound.Sound->GetName(), InActiveSound.GetPlayOrder(), *Location.ToCompactString());
 		}
 #endif // !(UE_BUILD_SHIPPING || UE_BUILD_TEST)
 
@@ -5384,6 +5385,7 @@ bool FAudioDevice::RemoveVirtualLoop(FActiveSound& InActiveSound)
 		if (bChannelEnabled)
 		{
 			UE_TRACE_LOG(Audio, VirtualLoopStop, AudioChannel)
+				<< VirtualLoopStop.DeviceId(static_cast<uint32>(DeviceID))
 				<< VirtualLoopStop.Timestamp(FPlatformTime::Cycles64())
 				<< VirtualLoopStop.PlayOrder(InActiveSound.GetPlayOrder());
 		}
