Fix for AnimSequence cook indeterminism.
Curve names are not sorted using lexical sort instead of FName sort
Added checks for two cases of indeterminism :
- An animation asset referencing another animation asset with different skeleton (i.e. preview base pose), which causes indeterminism as the referenced asset gets assigned the new skeleton
- Montages forcing root motion settings on AnimSequences and the AnimSequence not being saved after setting the flag (which causes the flag to be applied at Montage PostLoad, wich causes indeterminism)

diff --git a/Engine/Source/Runtime/Engine/Classes/Animation/AnimCompress.h b/Engine/Source/Runtime/Engine/Classes/Animation/AnimCompress.h
--- a/Engine/Source/Runtime/Engine/Classes/Animation/AnimCompress.h
+++ b/Engine/Source/Runtime/Engine/Classes/Animation/AnimCompress.h
@@ -272,7 +272,7 @@ namespace UE
 			// This is a version string that mimics the old versioning scheme. If you
 			// want to bump this version, generate a new guid using VS->Tools->Create GUID and
 			// return it here. Ex.
-			static FString AnimationCompressionVersionString = TEXT("55AC7D7DCC1F46C2A8B121A56FA37296");
+			static const FString AnimationCompressionVersionString = TEXT("6AB870F5E91D43739796EA2787A3FB07");
 			
 			struct FAnimationCompressionMemorySummaryScope
 			{
diff --git a/Engine/Source/Runtime/Engine/Classes/Animation/AnimSequence.h b/Engine/Source/Runtime/Engine/Classes/Animation/AnimSequence.h
--- a/Engine/Source/Runtime/Engine/Classes/Animation/AnimSequence.h
+++ b/Engine/Source/Runtime/Engine/Classes/Animation/AnimSequence.h
@@ -701,6 +701,7 @@ public:
 
 	// to support anim sequence base to all montages
 	ENGINE_API virtual void EnableRootMotionSettingFromMontage(bool bInEnableRootMotion, const ERootMotionRootLock::Type InRootMotionRootLock) override;
+	ENGINE_API virtual bool GetEnableRootMotionSettingFromMontage() const override;
 
 #if WITH_EDITOR
 	virtual class UAnimSequence* GetAdditiveBasePose() const override 
diff --git a/Engine/Source/Runtime/Engine/Classes/Animation/AnimSequenceBase.h b/Engine/Source/Runtime/Engine/Classes/Animation/AnimSequenceBase.h
--- a/Engine/Source/Runtime/Engine/Classes/Animation/AnimSequenceBase.h
+++ b/Engine/Source/Runtime/Engine/Classes/Animation/AnimSequenceBase.h
@@ -260,6 +260,7 @@ public:
 
 	// to support anim sequence base to montage
 	virtual void EnableRootMotionSettingFromMontage(bool bInEnableRootMotion, const ERootMotionRootLock::Type InRootMotionRootLock) {};
+	virtual bool GetEnableRootMotionSettingFromMontage() const { return false; }
 
 #if WITH_EDITOR
 private:
diff --git a/Engine/Source/Runtime/Engine/Private/Animation/AnimCompressionTypes.cpp b/Engine/Source/Runtime/Engine/Private/Animation/AnimCompressionTypes.cpp
--- a/Engine/Source/Runtime/Engine/Private/Animation/AnimCompressionTypes.cpp
+++ b/Engine/Source/Runtime/Engine/Private/Animation/AnimCompressionTypes.cpp
@@ -853,7 +853,7 @@ void FCompressibleAnimData::FetchData(const ITargetPlatform* InPlatform)
 	// Sorting increases the chances of linearizing the decompression later
 	RawFloatCurves.Sort([](const FFloatCurve& InLHS, const FFloatCurve& InRHS)
 	{
-		return InLHS.GetName().FastLess(InRHS.GetName());
+		return InLHS.GetName().LexicalLess(InRHS.GetName());
 	});
 	
 	// Apply any key reduction if possible
diff --git a/Engine/Source/Runtime/Engine/Private/Animation/AnimMontage.cpp b/Engine/Source/Runtime/Engine/Private/Animation/AnimMontage.cpp
--- a/Engine/Source/Runtime/Engine/Private/Animation/AnimMontage.cpp
+++ b/Engine/Source/Runtime/Engine/Private/Animation/AnimMontage.cpp
@@ -476,6 +476,12 @@ void UAnimMontage::PostLoad()
 			{
 				if (UAnimSequenceBase* AnimReference = Segment.GetAnimReference())
 				{
+#if WITH_EDITOR
+					if (!AnimReference->GetEnableRootMotionSettingFromMontage())
+					{
+						UE_LOG(LogAnimation, Warning, TEXT("[Montage %s] has RootMotionEnabled, but [AnimationSequence %s] has not been saved after setting the flag. Please open the Montage and the AnimationSequence and save the AnimationSequence as this will generate non determistic cooks."), *GetFullName(), *AnimReference->GetFullName());
+					}
+#endif // WITH_EDITOR
 					AnimReference->EnableRootMotionSettingFromMontage(true, RootMotionRootLock);
 				}
 			}
diff --git a/Engine/Source/Runtime/Engine/Private/Animation/AnimSequence.cpp b/Engine/Source/Runtime/Engine/Private/Animation/AnimSequence.cpp
--- a/Engine/Source/Runtime/Engine/Private/Animation/AnimSequence.cpp
+++ b/Engine/Source/Runtime/Engine/Private/Animation/AnimSequence.cpp
@@ -4062,6 +4062,11 @@ void UAnimSequence::EnableRootMotionSettingFromMontage(bool bInEnableRootMotion,
 	}
 }
 
+bool UAnimSequence::GetEnableRootMotionSettingFromMontage() const
+{
+	return bRootMotionSettingsCopiedFromMontage;
+}
+
 bool UAnimSequence::IsCompressedDataValid() const
 {
 	// For bone compressed data, we don't check if we have a codec. It is valid to have no compressed data
diff --git a/Engine/Source/Runtime/Engine/Private/AnimationAsset.cpp b/Engine/Source/Runtime/Engine/Private/AnimationAsset.cpp
--- a/Engine/Source/Runtime/Engine/Private/AnimationAsset.cpp
+++ b/Engine/Source/Runtime/Engine/Private/AnimationAsset.cpp
@@ -489,6 +489,11 @@ bool UAnimationAsset::ReplaceSkeleton(USkeleton* NewSkeleton, bool bConvertSpace
 				}
 				AnimAsset->ConditionalPostLoad();
 
+				if (AnimAsset->GetSkeleton() != GetSkeleton())
+				{
+					UE_LOG(LogAnimation, Warning, TEXT("AnimationAsset referencing asset using different skeleton. This will generate undeterministic builds. Please Fix the Asset : AnimationAsset: [%s] - ReferencedAsset : [%s]"), *GetName(), *AnimAsset->GetName());
+				}
+
 				// This ensure that in subsequent behaviour the RawData GUID is never 'new-ed' but always calculated from the 
 				// raw animation data itself.
 				if (UAnimSequence* Sequence = Cast<UAnimSequence>(AnimAsset))
