Fix for bSkipAssigningGamepadToFirstPlayer in the GameViewportClient.

This was broken because we did not actually update the InputDeviceId on the input event, only the ControllerId. The fix is to create a new "fake" input device ID and map it to the next available local player if needed for any gamepad events.

This has the resulting behavior of te Keyboard and Mouse controlling LocalPlayer[0], and the gamepad controlling LocalPlayer[1], instead of both devices controlling LocalPlayer[0].

In case this somehow caues issues unforeseen issues with split screen games, it can be disabled with the cvar 'input.bRemapDeviceIdForOffsetPlayerGamepadIds'.

diff --git a/Engine/Source/Runtime/Engine/Private/GameViewportClient.cpp b/Engine/Source/Runtime/Engine/Private/GameViewportClient.cpp
--- a/Engine/Source/Runtime/Engine/Private/GameViewportClient.cpp
+++ b/Engine/Source/Runtime/Engine/Private/GameViewportClient.cpp
@@ -116,6 +116,15 @@ static TAutoConsoleVariable<float> CVarSecondaryScreenPercentage( // TODO: make
 	TEXT(" 1: override secondary screen percentage."),
 	ECVF_Default);
 
+
+static TAutoConsoleVariable<bool> CVarRemapDeviceIdForOffsetPlayerGamepadIds(
+	TEXT("input.bRemapDeviceIdForOffsetPlayerGamepadIds"),
+	true,
+	TEXT("If true, then when bOffsetPlayerGamepadIds is true we will create a new Input Device Id\n")
+	TEXT("as needed for the next local player. This fixes the behavior in split screen.\n")
+	TEXT("Note: This CVar will be removed in a future release, this is a temporary wrapper for bug fix behavior."),
+	ECVF_Default);
+
 #if CSV_PROFILER
 struct FCsvLocalPlayer
 {
@@ -606,6 +615,33 @@ void UGameViewportClient::RemapControllerInput(FInputKeyEventArgs& InOutEventArg
 
 	if (NumLocalPlayers > 1 && InOutEventArgs.Key.IsGamepadKey() && GetDefault<UGameMapsSettings>()->bOffsetPlayerGamepadIds)
 	{
+		// Temp cvar in case this change somehow breaks input for any split screen games.
+		if (CVarRemapDeviceIdForOffsetPlayerGamepadIds.GetValueOnAnyThread())
+		{
+			const TArray<ULocalPlayer*>& CurrentLocalPlayers = World->GetGameInstance()->GetLocalPlayers();
+			if (CurrentLocalPlayers.IsValidIndex(InOutEventArgs.ControllerId) && CurrentLocalPlayers.IsValidIndex(InOutEventArgs.ControllerId + 1))
+			{
+				const FPlatformUserId DesiredPlatformUser = CurrentLocalPlayers[InOutEventArgs.ControllerId + 1]->GetPlatformUserId();
+
+				IPlatformInputDeviceMapper& DeviceMapper = IPlatformInputDeviceMapper::Get();
+
+				// Check for if this FPlatformUserID already has a primary input device ID. If it does, we can use that
+				FInputDeviceId DesiredInputDeviceId = DeviceMapper.GetPrimaryInputDeviceForUser(DesiredPlatformUser);
+				if (!DesiredInputDeviceId.IsValid())
+				{
+					// Otherwise we need to create a new "Fake" input device ID...
+					DesiredInputDeviceId = DeviceMapper.AllocateNewInputDeviceId();
+
+					// ...  and map it to our desired platform user so that the PlayerController knows it that this is associated with the local player
+					DeviceMapper.Internal_MapInputDeviceToUser(DesiredInputDeviceId, DesiredPlatformUser, EInputDeviceConnectionState::Connected);
+				}
+
+				// Say that this input event is from the other local player's input device!
+				InOutEventArgs.InputDevice = DesiredInputDeviceId;
+			}
+		}
+
+		// We still want to increment the controller ID in case there is any legacy code listening for it
 		InOutEventArgs.ControllerId++;
 	}
 	else if (InOutEventArgs.Viewport->IsPlayInEditorViewport() && InOutEventArgs.Key.IsGamepadKey())
