Add a CVar for propagating local SCW OOMs

When we run in a docker container with a memory limit, every out-of-memory condition should be reported as a total failure. Right now, when a local SCW goes OOM we do not report it as a failure, meaning that our reporting for live failures is inaccurate and under-reports memory issues. For future work, we could consider more explicitly limiting the amount of memory SCWs can use and be more intelligent about handling this, but for the moment this is a worthwhile improvement to correctly assess the scope of the problem.

diff --git a/Engine/Source/Runtime/Engine/Private/ShaderCompiler/ShaderCompiler.cpp b/Engine/Source/Runtime/Engine/Private/ShaderCompiler/ShaderCompiler.cpp
--- a/Engine/Source/Runtime/Engine/Private/ShaderCompiler/ShaderCompiler.cpp
+++ b/Engine/Source/Runtime/Engine/Private/ShaderCompiler/ShaderCompiler.cpp
@@ -2298,6 +2298,13 @@ static TAutoConsoleVariable<int32> CVarShadersRemoveDeadCode(
 	TEXT("\t1: Remove unreferenced code before compilation (Default)\n"),
 	ECVF_ReadOnly);
 
+static TAutoConsoleVariable<bool> CVarShadersPropagateLocalWorkerOOMs(
+	TEXT("r.Shaders.PropagateLocalWorkerOOMs"),
+	false,
+	TEXT("When set, out-of-memory conditions in a local shader compile worker will be treated as regular out-of-memory conditions and propagated to the main process.\n")
+	TEXT("This is useful when running in environment with hard memory limits, where it does not matter which process in particular caused us to violate the memory limit."),
+	ECVF_Default);
+
 #if ENABLE_COOK_STATS
 namespace ShaderCompilerCookStats
 {
@@ -2454,6 +2461,10 @@ namespace ShaderCompileWorkerError
 		}
 		else
 		{
+			if (CVarShadersPropagateLocalWorkerOOMs.GetValueOnAnyThread())
+			{
+				FPlatformMemory::OnOutOfMemory(0, 64);
+			}
 			ModalErrorOrLog(TEXT("ShaderCompileWorker failed"), ErrorReport);
 			return false;
 		}
