Fix incorrect scale stripping during animation compression

diff --git a/Engine/Source/Runtime/Engine/Private/Animation/AnimCompress.cpp b/Engine/Source/Runtime/Engine/Private/Animation/AnimCompress.cpp
--- a/Engine/Source/Runtime/Engine/Private/Animation/AnimCompress.cpp
+++ b/Engine/Source/Runtime/Engine/Private/Animation/AnimCompress.cpp
@@ -6,6 +6,7 @@
 
 #include "Animation/AnimCompress.h"
 #include "Animation/AnimSequence.h"
+#include "Animation/Skeleton.h"
 #include "Misc/MessageDialog.h"
 #include "Animation/AnimSequenceDecompressionContext.h"
 #include "Misc/FeedbackContext.h"
@@ -261,6 +262,47 @@ void UAnimCompress::PadByteStream(TArray<uint8>& ByteStream, const int32 Alignme
 	}
 }
 
+namespace UE::Anim::Compression::Private
+{
+	static bool ShouldCompressScale(const FCompressibleAnimData& CompressibleAnimData, const TArray<FScaleTrack>& ScaleData)
+	{
+		if (CompressibleAnimData.bIsValidAdditive)
+		{
+			// We always compress scale with additive sequences
+			return true;
+		}
+
+		for (int32 TrackIndex = 0; TrackIndex < ScaleData.Num(); TrackIndex++)
+		{
+			const FScaleTrack& ScaleTrack = ScaleData[TrackIndex];
+			if (ScaleTrack.ScaleKeys.Num() != 1)
+			{
+				// If we have 0 keys or more than 1 key, we aren't a trivial track with a value equal to the bind pose
+				// We thus need to compress scale
+				return true;
+			}
+
+			const int32 BoneIndex = CompressibleAnimData.TrackToSkeletonMapTable[TrackIndex].BoneTreeIndex;
+			if (!CompressibleAnimData.RefSkeleton.IsValidIndex(BoneIndex))
+			{
+				// This track isn't mapped to a bone, ignore it
+				continue;
+			}
+
+			const FVector3f RefScale(CompressibleAnimData.RefSkeleton.GetRefBonePose()[BoneIndex].GetScale3D());
+			if (!RefScale.Equals(ScaleTrack.ScaleKeys[0]))
+			{
+				// This trivial track isn't equal to the bind pose scale, we need to retain it
+				return true;
+			}
+		}
+
+		// We don't need to retain any scale, it can safely be stripped and ignored
+		// At runtime, when we decompress we'll skip missing scale entries, leaving the
+		// bind pose value untouched in the output buffer
+		return false;
+	}
+}
 
 void UAnimCompress::BitwiseCompressAnimationTracks(
 	const FCompressibleAnimData& CompressibleAnimData,
@@ -309,7 +351,7 @@ void UAnimCompress::BitwiseCompressAnimationTracks(
 
 		check(TranslationData.Num() == RotationData.Num());
 		const int32 NumTracks = RotationData.Num();
-		const bool bHasScale = ScaleData.Num() > 0;
+		const bool bHasScale = UE::Anim::Compression::Private::ShouldCompressScale(CompressibleAnimData, ScaleData);
 
 		if (NumTracks == 0)
 		{
@@ -669,6 +711,8 @@ void UAnimCompress::DecompressPose(FAnimSequenceDecompressionContext& DecompCont
 
 	AnimData.RotationCodec->GetPoseRotations(OutAtoms, RotationPairs, DecompContext);
 
+	// Scale can be stripped when it is trivial and equal to the bind pose
+	// When this is the case, we assume that the output pose has already been populated with the bind pose
 	if (AnimData.CompressedScaleOffsets.IsValid())
 	{
 		AnimData.ScaleCodec->GetPoseScales(OutAtoms, ScalePairs, DecompContext);
@@ -694,6 +738,25 @@ void UAnimCompress::DecompressBone(FAnimSequenceDecompressionContext& DecompCont
 		// decompress the rotation component using the proper method
 		((AnimEncodingLegacyBase*)AnimData.ScaleCodec)->GetBoneAtomScale(OutAtom, DecompContext, TrackIndex);
 	}
+	else
+	{
+		// If scale is stripped, we must output the bind pose value since it's been stripped
+		if (DecompContext.IsAdditiveAnimation())
+		{
+			// Additive animations use the additive identity
+			OutAtom.SetScale3D(FVector::ZeroVector);
+		}
+		else
+		{
+			checkf(DecompContext.GetRefLocalPoses().Num() > 0, TEXT("Reference pose must be provided in the FAnimSequenceDecompressionContext constructor"));
+			checkf(DecompContext.GetTrackToSkeletonMap().Num() > 0, TEXT("TrackToSkeletonMap must be provided in the FAnimSequenceDecompressionContext constructor"));
+
+			const int32 BoneIndex = DecompContext.GetTrackToSkeletonMap()[TrackIndex].BoneTreeIndex;
+			const FTransform& BindTransform = DecompContext.GetRefLocalPoses()[BoneIndex];
+
+			OutAtom.SetScale3D(BindTransform.GetScale3D());
+		}
+	}
 }
 
 #if WITH_EDITOR
@@ -708,6 +771,25 @@ void UAnimCompress::PopulateDDCKey(const UE::Anim::Compression::FAnimDDCKeyArgs&
 	SCF = (uint8)ScaleCompressionFormat.GetValue();
 
 	Ar << TCF << RCF << SCF;
+
+	// Additive sequences use the additive identity as their bind pose, scale is never stripped
+	if (!KeyArgs.AnimSequence.IsValidAdditive())
+	{
+		// We have to include the bind pose in the DDC key.
+		// If a sequence is compressed with bind pose A, and we strip a few bones and later modify the bind pose,
+		// bind pose B might now contain values that would not be stripped in our sequence.
+		// To avoid data being stale, the DDC must reflect this.
+
+		const USkeleton* Skeleton = KeyArgs.AnimSequence.GetSkeleton();
+		const TArray<FTransform>& BindPose = Skeleton->GetRefLocalPoses();
+		for (const FTransform& BoneBindTransform : BindPose)
+		{
+			// We only strip scale
+
+			FVector Scale = BoneBindTransform.GetScale3D();
+			Ar << Scale;
+		}
+	}
 }
 
 void UAnimCompress::FilterTrivialPositionKeys(
diff --git a/Engine/Source/Runtime/Engine/Private/Animation/AnimSequenceHelpers.cpp b/Engine/Source/Runtime/Engine/Private/Animation/AnimSequenceHelpers.cpp
--- a/Engine/Source/Runtime/Engine/Private/Animation/AnimSequenceHelpers.cpp
+++ b/Engine/Source/Runtime/Engine/Private/Animation/AnimSequenceHelpers.cpp
@@ -440,36 +440,6 @@ bool Compression::CompressAnimationDataTracks(TArray<FRawAnimSequenceTrack>& Raw
 		{
 			bRemovedKeys |= CompressRawAnimSequenceTrack(RawAnimationData[TrackIndex], NumberOfKeys, ErrorName, MaxPosDiff, MaxAngleDiff);
 		}
-
-		bool bCompressScaleKeys = false;
-		// go through remove keys if not needed
-		for (int32 TrackIndex = 0; TrackIndex < RawAnimationData.Num(); TrackIndex++)
-		{
-			FRawAnimSequenceTrack const& RawData = RawAnimationData[TrackIndex];
-			if (RawData.ScaleKeys.Num() > 0)
-			{
-				// if scale key exists, see if we can just empty it
-				const bool bHaveMultipleKeys = RawData.ScaleKeys.Num() > 1;
-				const bool bHaveNonIdentityKey = !RawData.ScaleKeys[0].Equals(FVector3f::OneVector);
-				if (bHaveMultipleKeys|| bHaveNonIdentityKey)
-				{
-					bCompressScaleKeys = true;
-					break;
-				}
-			}
-		}
-
-		// if we don't have scale, we should delete all scale keys
-		// if you have one track that has scale, we still should support scale, so compress scale
-		if (!bCompressScaleKeys)
-		{
-			// then remove all scale keys
-			for (int32 TrackIndex = 0; TrackIndex < RawAnimationData.Num(); TrackIndex++)
-			{
-				FRawAnimSequenceTrack& RawData = RawAnimationData[TrackIndex];
-				RawData.ScaleKeys.Empty();
-			}
-		}
 	}
 #endif
 	return bRemovedKeys;
@@ -487,39 +457,6 @@ bool Compression::CompressAnimationDataTracks(const USkeleton* Skeleton, const T
 		{
 			bRemovedKeys |= CompressRawAnimSequenceTrack(RawAnimationData[TrackIndex], NumberOfKeys, ErrorName, MaxPosDiff, MaxAngleDiff);
 		}
-
-		bool bCompressScaleKeys = false;
-		// go through remove keys if not needed
-		for (int32 TrackIndex = 0; TrackIndex < RawAnimationData.Num(); TrackIndex++)
-		{
-			FRawAnimSequenceTrack const& RawData = RawAnimationData[TrackIndex];
-			if (RawData.ScaleKeys.Num() > 0)
-			{
-				// if scale key exists, see if we can just empty it
-				const bool bHaveMultipleKeys = RawData.ScaleKeys.Num() > 1;
-				const bool bHaveNonIdentityKey = !RawData.ScaleKeys[0].Equals(FVector3f::OneVector);
-				const FReferenceSkeleton& ReferenceSkeleton = Skeleton->GetReferenceSkeleton();
-				const int32 TrackBoneIndex = TrackToSkeleton[TrackIndex].BoneTreeIndex;				
-				const bool bHaveNonRefPoseKey = ReferenceSkeleton.IsValidIndex(TrackBoneIndex) && !RawData.ScaleKeys[0].Equals(FVector3f(ReferenceSkeleton.GetRefBonePose()[TrackBoneIndex].GetScale3D()));
-				if (bHaveMultipleKeys|| bHaveNonIdentityKey || bHaveNonRefPoseKey)
-				{
-					bCompressScaleKeys = true;
-					break;
-				}
-			}
-		}
-
-		// if we don't have scale, we should delete all scale keys
-		// if you have one track that has scale, we still should support scale, so compress scale
-		if (!bCompressScaleKeys)
-		{
-			// then remove all scale keys
-			for (int32 TrackIndex = 0; TrackIndex < RawAnimationData.Num(); TrackIndex++)
-			{
-				FRawAnimSequenceTrack& RawData = RawAnimationData[TrackIndex];
-				RawData.ScaleKeys.Empty();
-			}
-		}
 	}
 #endif
 	return bRemovedKeys;
