Physics: Support filtering hits based on non-actor source objects.

- Make it so that FBodyInstance can initialize its query data using an id from an arbitrary UObject, instead of just actors. The stored data (UObject ID) stays the same as it was before.
- Make FCollisionQueryParams take arbitrary ignore objects instead of just actors. This gets used in FCollisionQueryFilterCallback.

diff --git a/Engine/Source/Runtime/Engine/Classes/PhysicsEngine/BodyInstance.h b/Engine/Source/Runtime/Engine/Classes/PhysicsEngine/BodyInstance.h
--- a/Engine/Source/Runtime/Engine/Classes/PhysicsEngine/BodyInstance.h
+++ b/Engine/Source/Runtime/Engine/Classes/PhysicsEngine/BodyInstance.h
@@ -83,7 +83,13 @@ struct FInitBodySpawnParams
 
 struct FInitBodiesHelperBase
 {
-	ENGINE_API FInitBodiesHelperBase(TArray<FBodyInstance*>& InBodies, TArray<FTransform>& InTransforms, class UBodySetup* InBodySetup, class UPrimitiveComponent* InPrimitiveComp, FPhysScene* InRBScene, const FInitBodySpawnParams& InSpawnParams, FPhysicsAggregateHandle InAggregate);
+	ENGINE_API FInitBodiesHelperBase(TArray<FBodyInstance*>& InBodies, TArray<FTransform>& InTransforms, 
+		class UBodySetup* InBodySetup, class UPrimitiveComponent* InPrimitiveComp, FPhysScene* InRBScene, 
+		const FInitBodySpawnParams& InSpawnParams, FPhysicsAggregateHandle InAggregate);
+	
+	ENGINE_API FInitBodiesHelperBase(TArray<FBodyInstance*>& InBodies, TArray<FTransform>& InTransforms, 
+		class UBodySetup* InBodySetup, class UPrimitiveComponent* InComponent, UObject* InSourceObject, FPhysScene* InRBScene,
+		const FInitBodySpawnParams& InSpawnParams, FPhysicsAggregateHandle InAggregate);
 
 	FInitBodiesHelperBase(const FInitBodiesHelperBase& InHelper) = delete;
 	FInitBodiesHelperBase(FInitBodiesHelperBase&& InHelper) = delete;
@@ -97,6 +103,7 @@ struct FInitBodiesHelperBase
 	TArray<FTransform>& Transforms;
 	class UBodySetup* BodySetup;
 	class UPrimitiveComponent* PrimitiveComp;
+	UObject* SourceObject;
 	FPhysScene* PhysScene;
 	FPhysicsAggregateHandle Aggregate;
 
@@ -133,13 +140,19 @@ protected:
 template <bool bCompileStatic>
 struct FInitBodiesHelper : public FInitBodiesHelperBase
 {
-	FInitBodiesHelper(TArray<FBodyInstance*>& InBodies, TArray<FTransform>& InTransforms, class UBodySetup* InBodySetup, class UPrimitiveComponent* InPrimitiveComp, FPhysScene* InRBScene, const FInitBodySpawnParams& InSpawnParams, FPhysicsAggregateHandle InAggregate)
-	: FInitBodiesHelperBase(InBodies, InTransforms, InBodySetup, InPrimitiveComp, InRBScene, InSpawnParams, InAggregate)
+	FInitBodiesHelper(TArray<FBodyInstance*>& InBodies, TArray<FTransform>& InTransforms, class UBodySetup* InBodySetup,
+		class UPrimitiveComponent* InPrimitiveComp, FPhysScene* InRBScene,
+		const FInitBodySpawnParams& InSpawnParams, FPhysicsAggregateHandle InAggregate);
+
+	FInitBodiesHelper(TArray<FBodyInstance*>& InBodies, TArray<FTransform>& InTransforms, class UBodySetup* InBodySetup, 
+		class UPrimitiveComponent* InComponent, class UObject* InOwnerObject, FPhysScene* InRBScene,
+		const FInitBodySpawnParams& InSpawnParams, FPhysicsAggregateHandle InAggregate)
+		: FInitBodiesHelperBase(InBodies, InTransforms, InBodySetup, InComponent, InOwnerObject, InRBScene, InSpawnParams, InAggregate)
 	{
 		//Compute all the needed constants
 		bStatic = bCompileStatic || SpawnParams.bStaticPhysics;
 		SkelMeshComp = bCompileStatic ? nullptr : Cast<USkeletalMeshComponent>(PrimitiveComp);
-		if(SpawnParams.bPhysicsTypeDeterminesSimulation)
+		if (SpawnParams.bPhysicsTypeDeterminesSimulation)
 		{
 			this->UpdateSimulatingAndBlendWeight();
 		}
@@ -623,9 +636,12 @@ public:
 	TSharedPtr<TArray<ANSICHAR>> CharDebugName;
 #endif
 
-	/** PrimitiveComponent containing this body.   */
+	/** PrimitiveComponent containing this body, if relevant.   */
 	TWeakObjectPtr<class UPrimitiveComponent> OwnerComponent;
 
+	/** Source object for this body. If OwnerComponent is not null, this will typically be the owning actor. */
+	TWeakObjectPtr<UObject> SourceObject;
+
 	/** Constructor **/
 	ENGINE_API FBodyInstance();
 	ENGINE_API ~FBodyInstance();
@@ -655,6 +671,21 @@ public:
 	*/
 	ENGINE_API void InitBody(UBodySetup* Setup, const FTransform& Transform, UPrimitiveComponent* PrimComp, FPhysScene* InRBScene, const FInitBodySpawnParams& SpawnParams);
 
+	/**
+	 * Initialise a single rigid body (this FBodyInstance) for the given body setup
+	 *  @param Setup The setup to use to create the body
+	 *  @param Transform Transform of the body
+	 *  @param PrimComp The owning component
+	 *  @param SourceObject The source object. If null and PrimComp is not null, it will be set
+	 *   to the owning actor.
+	 *  @param InRBScene The physics scene to place the body into
+	 *  @param SpawnParams The parameters for determining certain spawn behavior
+	 *  @param InAggregate An aggregate to place the body into
+	 */
+	ENGINE_API void InitBody(UBodySetup* Setup, const FTransform& Transform, 
+		UPrimitiveComponent* PrimComp, UObject* SourceObject, 
+		FPhysScene* InRBScene, const FInitBodySpawnParams& SpawnParams);
+
 	/** Validate a body transform, outputting debug info
 	 *	@param Transform Transform to debug
 	 *	@param DebugName Name of the instance for logging
diff --git a/Engine/Source/Runtime/Engine/Private/Collision/Collision.cpp b/Engine/Source/Runtime/Engine/Private/Collision/Collision.cpp
--- a/Engine/Source/Runtime/Engine/Private/Collision/Collision.cpp
+++ b/Engine/Source/Runtime/Engine/Private/Collision/Collision.cpp
@@ -107,15 +107,12 @@ static FORCEINLINE_DEBUGGABLE bool CheckForCollision(const UPrimitiveComponent*
 
 void FCollisionQueryParams::AddIgnoredActor(const AActor* InIgnoreActor)
 {
-	if (InIgnoreActor)
-	{	
-		IgnoreActors.Add(InIgnoreActor->GetUniqueID());
-	}
+	AddIgnoredSourceObject(InIgnoreActor);
 }
 
 void FCollisionQueryParams::AddIgnoredActor(const uint32 InIgnoreActorID)
 {
-	IgnoreActors.Add(InIgnoreActorID);
+	AddIgnoredSourceObject(InIgnoreActorID);
 }
 
 void FCollisionQueryParams::AddIgnoredActors(const TArray<AActor*>& InIgnoreActors)
@@ -142,6 +139,43 @@ void FCollisionQueryParams::AddIgnoredActors(const TArray<TWeakObjectPtr<const A
 	}
 }
 
+void FCollisionQueryParams::AddIgnoredSourceObject(const UObject* InIgnoreObject)
+{
+	if (InIgnoreObject)
+	{
+		IgnoreSourceObjects.Add(InIgnoreObject->GetUniqueID());
+	}
+}
+
+void FCollisionQueryParams::AddIgnoredSourceObject(const uint32 InIgnoreActorID)
+{
+	IgnoreSourceObjects.Add(InIgnoreActorID);
+}
+
+void FCollisionQueryParams::AddIgnoredSourceObjects(const TArray<UObject*>& InIgnoreObjects)
+{
+	for (int32 Idx = 0; Idx < InIgnoreObjects.Num(); ++Idx)
+	{
+		AddIgnoredSourceObject(InIgnoreObjects[Idx]);
+	}
+}
+
+void FCollisionQueryParams::AddIgnoredSourceObjects(const TArray<const UObject*>& InIgnoreObjects)
+{
+	for (int32 Idx = 0; Idx < InIgnoreObjects.Num(); ++Idx)
+	{
+		AddIgnoredSourceObject(InIgnoreObjects[Idx]);
+	}
+}
+
+void FCollisionQueryParams::AddIgnoredSourceObjects(const TArray<TWeakObjectPtr<const UObject>>& InIgnoreObjects)
+{
+	for (int32 Idx = 0; Idx < InIgnoreObjects.Num(); ++Idx)
+	{
+		AddIgnoredSourceObject(InIgnoreObjects[Idx].Get());
+	}
+}
+
 FORCEINLINE_DEBUGGABLE void FCollisionQueryParams::Internal_AddIgnoredComponent(const UPrimitiveComponent* InIgnoreComponent)
 {
 	if (InIgnoreComponent && (!CheckForCollision(InIgnoreComponent) || IsQueryCollisionEnabled(InIgnoreComponent)))
diff --git a/Engine/Source/Runtime/Engine/Private/Collision/WorldCollision.cpp b/Engine/Source/Runtime/Engine/Private/Collision/WorldCollision.cpp
--- a/Engine/Source/Runtime/Engine/Private/Collision/WorldCollision.cpp
+++ b/Engine/Source/Runtime/Engine/Private/Collision/WorldCollision.cpp
@@ -464,11 +464,11 @@ bool UWorld::ComponentSweepMultiByChannel(TArray<struct FHitResult>& OutHits, cl
 	{
 		FComponentQueryParams ParamsCopy{ Params };
 		TSet<uint32> ActorsToExclude;
-		for (uint32 ActorID : ParamsCopy.GetIgnoredActors())
+		for (uint32 ActorID : ParamsCopy.GetIgnoredSourceObjects())
 		{
 			ActorsToExclude.Add(ActorID);
 		}		
-		ParamsCopy.ClearIgnoredActors(); // This will be populated a bit later
+		ParamsCopy.ClearIgnoredSourceObjects(); // This will be populated a bit later
 		// All actors pointed to by shapes should be ignored (This deals with welded Actors)
 		TArray<Chaos::FShapeInstanceProxy*> Shapes = Interface->GetAllThreadShapes({ &Object, 1 });
 		for (Chaos::FShapeInstanceProxy* Shape : Shapes)
diff --git a/Engine/Source/Runtime/Engine/Private/PhysicsEngine/BodyInstance.cpp b/Engine/Source/Runtime/Engine/Private/PhysicsEngine/BodyInstance.cpp
--- a/Engine/Source/Runtime/Engine/Private/PhysicsEngine/BodyInstance.cpp
+++ b/Engine/Source/Runtime/Engine/Private/PhysicsEngine/BodyInstance.cpp
@@ -1167,12 +1167,28 @@ void FInitBodiesHelperBase::UpdateSimulatingAndBlendWeight()
 	GetSimulatingAndBlendWeight(SkelMeshComp, BodySetup, InstanceBlendWeight, bInstanceSimulatePhysics);
 }
 
+template<bool bCompileStatic>
+FInitBodiesHelper<bCompileStatic>::FInitBodiesHelper(
+	TArray<FBodyInstance*>& InBodies, TArray<FTransform>& InTransforms, UBodySetup* InBodySetup, 
+	UPrimitiveComponent* InPrimitiveComp, FPhysScene* InRBScene, 
+	const FInitBodySpawnParams& InSpawnParams, FPhysicsAggregateHandle InAggregate)
+	: FInitBodiesHelper<bCompileStatic>(InBodies, InTransforms, InBodySetup, InPrimitiveComp,
+		InPrimitiveComp ? InPrimitiveComp->GetOwner() : nullptr, 
+		InRBScene, InSpawnParams, InAggregate)
+{
+}
+template struct FInitBodiesHelper<true>;
+template struct FInitBodiesHelper<false>;
 
-FInitBodiesHelperBase::FInitBodiesHelperBase(TArray<FBodyInstance*>& InBodies, TArray<FTransform>& InTransforms, class UBodySetup* InBodySetup, class UPrimitiveComponent* InPrimitiveComp, FPhysScene* InRBScene, const FInitBodySpawnParams& InSpawnParams, FPhysicsAggregateHandle InAggregate)
+
+FInitBodiesHelperBase::FInitBodiesHelperBase(TArray<FBodyInstance*>& InBodies, TArray<FTransform>& InTransforms, 
+	class UBodySetup* InBodySetup, UPrimitiveComponent* InComponent, UObject* InSourceObject, FPhysScene* InRBScene, 
+	const FInitBodySpawnParams& InSpawnParams, FPhysicsAggregateHandle InAggregate)
 	: Bodies(InBodies)
 	, Transforms(InTransforms)
 	, BodySetup(InBodySetup)
-	, PrimitiveComp(InPrimitiveComp)
+	, PrimitiveComp(InComponent)
+	, SourceObject(InSourceObject)
 	, PhysScene(InRBScene)
 	, Aggregate(InAggregate)
 #if USE_BODYINSTANCE_DEBUG_NAMES
@@ -1189,6 +1205,17 @@ FInitBodiesHelperBase::FInitBodiesHelperBase(TArray<FBodyInstance*>& InBodies, T
 #endif
 }
 
+FInitBodiesHelperBase::FInitBodiesHelperBase(TArray<FBodyInstance*>& InBodies, TArray<FTransform>& InTransforms,
+	class UBodySetup* InBodySetup, class UPrimitiveComponent* InPrimitiveComp, FPhysScene* InRBScene,
+	const FInitBodySpawnParams& InSpawnParams, FPhysicsAggregateHandle InAggregate)
+	// Forward to the other constructor
+	: FInitBodiesHelperBase(InBodies, InTransforms, InBodySetup,
+		InPrimitiveComp,
+		InPrimitiveComp ? InPrimitiveComp->GetOwner() : nullptr,
+		InRBScene, InSpawnParams, InAggregate)
+{
+}
+
 // Return to actor ref
 void FInitBodiesHelperBase::CreateActor_AssumesLocked(FBodyInstance* Instance, const FTransform& Transform) const
 {
@@ -1377,6 +1404,7 @@ bool FInitBodiesHelperBase::CreateShapesAndActors()
 		}
 
 		Instance->OwnerComponent = PrimitiveComp;
+		Instance->SourceObject = SourceObject;
 		Instance->BodySetup = BodySetup;
 		Instance->Scale3D = Transform.GetScale3D();
 #if USE_BODYINSTANCE_DEBUG_NAMES
@@ -1410,7 +1438,8 @@ bool FInitBodiesHelperBase::CreateShapesAndActors()
 		if (FPhysicsInterface::IsValid(Instance->GetPhysicsActorHandle()))
 		{
 			Instance->OwnerComponent = nullptr;
-			Instance->BodySetup      = nullptr;
+			Instance->SourceObject = nullptr;
+			Instance->BodySetup = nullptr;
 			Bodies.RemoveAt(BodyIdx);  // so we wont add it to the physx scene again later.
 			Transforms.RemoveAt(BodyIdx);
 			continue;
@@ -1444,6 +1473,7 @@ bool FInitBodiesHelperBase::CreateShapesAndActors()
 			FPhysicsInterface::ReleaseActor(Instance->ActorHandle, PhysScene);
 
 			Instance->OwnerComponent = nullptr;
+			Instance->SourceObject = nullptr;
 			Instance->BodySetup = nullptr;
 			Instance->ExternalCollisionProfileBodySetup = nullptr;
 
@@ -1589,7 +1619,12 @@ static TAutoConsoleVariable<int32> CVarAllowCreatePhysxBodies(
 	ECVF_ReadOnly);
 
 
-void FBodyInstance::InitBody(class UBodySetup* Setup, const FTransform& Transform, UPrimitiveComponent* PrimComp, FPhysScene* InRBScene, const FInitBodySpawnParams& SpawnParams)
+void FBodyInstance::InitBody(UBodySetup* Setup, const FTransform& Transform, UPrimitiveComponent* PrimComp, FPhysScene* InRBScene, const FInitBodySpawnParams& SpawnParams)
+{
+	InitBody(Setup, Transform, PrimComp, nullptr, InRBScene, SpawnParams);
+}
+
+void FBodyInstance::InitBody(class UBodySetup* Setup, const FTransform& Transform, UPrimitiveComponent* PrimComp, UObject* InSourceObject, FPhysScene* InRBScene, const FInitBodySpawnParams& SpawnParams)
 {
 	if (CVarAllowCreatePhysxBodies.GetValueOnGameThread() == 0)
 	{
@@ -1608,15 +1643,21 @@ void FBodyInstance::InitBody(class UBodySetup* Setup, const FTransform& Transfor
 	Bodies.Add(this);
 	Transforms.Add(Transform);
 
+	UObject* SourceObjectToUse = InSourceObject;
+	if (!SourceObjectToUse && PrimComp)
+	{
+		SourceObjectToUse = PrimComp->GetOwner();
+	}
+
 	bool bIsStatic = SpawnParams.bStaticPhysics;
 	if(bIsStatic)
 	{
-		FInitBodiesHelper<true> InitBodiesHelper(Bodies, Transforms, Setup, PrimComp, InRBScene, SpawnParams, SpawnParams.Aggregate);
+		FInitBodiesHelper<true> InitBodiesHelper(Bodies, Transforms, Setup, PrimComp, SourceObjectToUse, InRBScene, SpawnParams, SpawnParams.Aggregate);
 		InitBodiesHelper.InitBodies();
 	}
 	else
 	{
-		FInitBodiesHelper<false> InitBodiesHelper(Bodies, Transforms, Setup, PrimComp, InRBScene, SpawnParams, SpawnParams.Aggregate);
+		FInitBodiesHelper<false> InitBodiesHelper(Bodies, Transforms, Setup, PrimComp, SourceObjectToUse, InRBScene, SpawnParams, SpawnParams.Aggregate);
 		InitBodiesHelper.InitBodies();
 	}
 
@@ -4229,9 +4270,12 @@ void FBodyInstance::BuildBodyFilterData(FBodyCollisionFilterData& OutFilterData,
 		return;
 	}
 
-	// Figure out if we are static
 	UPrimitiveComponent* OwnerComponentInst = OwnerComponent.Get();
-	AActor* Owner = OwnerComponentInst ? OwnerComponentInst->GetOwner() : NULL;
+	AActor* ActorOwner = Cast<AActor>(SourceObject.Get());
+	if (!ActorOwner && OwnerComponentInst)
+	{
+		ActorOwner = OwnerComponentInst->GetOwner();
+	}
 	const bool bPhysicsStatic = !OwnerComponentInst || OwnerComponentInst->IsWorldGeometry();
 
 	// Grab collision setting from body instance.
@@ -4293,9 +4337,9 @@ void FBodyInstance::BuildBodyFilterData(FBodyCollisionFilterData& OutFilterData,
 
 #if WITH_EDITOR
 	// if no collision, but if world wants to enable trace collision for components, allow it
-	if((UseCollisionEnabled == ECollisionEnabled::NoCollision) && Owner && (Owner->IsA(AVolume::StaticClass()) == false))
+	if((UseCollisionEnabled == ECollisionEnabled::NoCollision) && ActorOwner && (ActorOwner->IsA(AVolume::StaticClass()) == false))
 	{
-		UWorld* World = Owner->GetWorld();
+		UWorld* World = ActorOwner->GetWorld();
 		UPrimitiveComponent* PrimComp = OwnerComponentInst;
 		if(World && World->bEnableTraceCollision &&
 		   (PrimComp->IsA(UStaticMeshComponent::StaticClass()) || PrimComp->IsA(USkeletalMeshComponent::StaticClass()) || PrimComp->IsA(UBrushComponent::StaticClass())))
@@ -4319,9 +4363,9 @@ void FBodyInstance::BuildBodyFilterData(FBodyCollisionFilterData& OutFilterData,
 		FCollisionFilterData SimFilterData;
 		FCollisionFilterData SimpleQueryData;
 
-		uint32 ActorID = Owner ? Owner->GetUniqueID() : 0;
+		uint32 SourceObjectID = SourceObject.IsValid() ? SourceObject->GetUniqueID() : 0;
 		uint32 CompID = (OwnerComponentInst != nullptr) ? OwnerComponentInst->GetUniqueID() : 0;
-		CreateShapeFilterData(UseChannel, MaskFilter, ActorID, UseResponse, CompID, InstanceBodyIndex, SimpleQueryData, SimFilterData, bRootCCD && !bPhysicsStatic, bUseNotifyRBCollision, bPhysicsStatic, bUseContactModification);
+		CreateShapeFilterData(UseChannel, MaskFilter, SourceObjectID, UseResponse, CompID, InstanceBodyIndex, SimpleQueryData, SimFilterData, bRootCCD && !bPhysicsStatic, bUseNotifyRBCollision, bPhysicsStatic, bUseContactModification);
 
 		FCollisionFilterData ComplexQueryData = SimpleQueryData;
 			
diff --git a/Engine/Source/Runtime/Engine/Private/PhysicsEngine/ClusterUnionComponent.cpp b/Engine/Source/Runtime/Engine/Private/PhysicsEngine/ClusterUnionComponent.cpp
--- a/Engine/Source/Runtime/Engine/Private/PhysicsEngine/ClusterUnionComponent.cpp
+++ b/Engine/Source/Runtime/Engine/Private/PhysicsEngine/ClusterUnionComponent.cpp
@@ -1673,8 +1673,8 @@ void UClusterUnionComponent::VisitAllCurrentActors(const TFunction<bool(AActor*)
 void UClusterUnionComponent::VisitAllCurrentChildComponentsForCollision(ECollisionChannel TraceChannel, const struct FCollisionQueryParams& Params, const struct FCollisionResponseParams& ResponseParams, const struct FCollisionObjectQueryParams& ObjectParams, const TFunction<bool(UPrimitiveComponent*)>& Lambda) const
 {
 	TSet<uint32> IgnoredActors;
-	IgnoredActors.Reserve(Params.GetIgnoredActors().Num());
-	for (uint32 Id : Params.GetIgnoredActors())
+	IgnoredActors.Reserve(Params.GetIgnoredSourceObjects().Num());
+	for (uint32 Id : Params.GetIgnoredSourceObjects())
 	{
 		IgnoredActors.Add(Id);
 	}
diff --git a/Engine/Source/Runtime/Engine/Private/PhysicsEngine/CollisionQueryFilterCallback.h b/Engine/Source/Runtime/Engine/Private/PhysicsEngine/CollisionQueryFilterCallback.h
--- a/Engine/Source/Runtime/Engine/Private/PhysicsEngine/CollisionQueryFilterCallback.h
+++ b/Engine/Source/Runtime/Engine/Private/PhysicsEngine/CollisionQueryFilterCallback.h
@@ -25,6 +25,8 @@ public:
 	/** List of ComponentIds for this query to ignore */
 	const FilterIgnoreComponentsArrayType& IgnoreComponents;
 
+	//~ TODO: It would be nice to rename this to IgnoreSourceObjects, because these might not be actors in
+	//~ non-actor workflows (requires deprecation).
 	/** List of ActorIds for this query to ignore */
 	const FilterIgnoreActorsArrayType& IgnoreActors;
 
@@ -39,7 +41,7 @@ public:
 
 	FCollisionQueryFilterCallback(const FCollisionQueryParams& InQueryParams, bool bInIsSweep)
 		: IgnoreComponents(InQueryParams.GetIgnoredComponents())
-		, IgnoreActors(InQueryParams.GetIgnoredActors())
+		, IgnoreActors(InQueryParams.GetIgnoredSourceObjects())
 #if DETECT_SQ_HITCHES
 		, bRecordHitches(false)
 #endif
diff --git a/Engine/Source/Runtime/Engine/Public/CollisionQueryParams.h b/Engine/Source/Runtime/Engine/Public/CollisionQueryParams.h
--- a/Engine/Source/Runtime/Engine/Public/CollisionQueryParams.h
+++ b/Engine/Source/Runtime/Engine/Public/CollisionQueryParams.h
@@ -80,7 +80,9 @@ struct FCollisionQueryParams
 	typedef TArray<uint32, TInlineAllocator<8>> IgnoreComponentsArrayType;
 
 	/** TArray typedef of actors to ignore. */
-	typedef TArray<uint32, TInlineAllocator<4>> IgnoreActorsArrayType;
+	typedef TArray<uint32, TInlineAllocator<4>> IgnoreSourceObjectsArrayType;
+
+	typedef IgnoreSourceObjectsArrayType IgnoreActorsArrayType;
 
 	/** Extra filtering done on the query. See declaration for filtering logic */
 	FMaskFilter IgnoreMask;
@@ -105,8 +107,8 @@ private:
 	/** Set of components to ignore during the trace */
 	mutable IgnoreComponentsArrayType IgnoreComponents;
 
-	/** Set of actors to ignore during the trace */
-	IgnoreActorsArrayType IgnoreActors;
+	/** Set of source objects to ignore during the trace. In actor workflows, these are actors. */
+	IgnoreSourceObjectsArrayType IgnoreSourceObjects;
 
 	ENGINE_API void Internal_AddIgnoredComponent(const UPrimitiveComponent* InIgnoreComponent);
 
@@ -115,10 +117,21 @@ public:
 	/** Returns set of unique components to ignore during the trace. Elements are guaranteed to be unique (they are made so internally if they are not already). */
 	ENGINE_API const IgnoreComponentsArrayType& GetIgnoredComponents() const;
 
-	/** Returns set of actors to ignore during the trace. Note that elements are NOT guaranteed to be unique. This is less important for actors since it's less likely that duplicates are added.*/
+	UE_DEPRECATED(5.5, "Use GetIgnoredSourceObjects instead.")
+	/** Returns set of source objects (including actors) to ignore during the trace. Note that elements are NOT guaranteed to be unique. This is less important for actors since it's less likely that duplicates are added.*/
 	const IgnoreActorsArrayType& GetIgnoredActors() const
 	{
-		return IgnoreActors;
+		return IgnoreSourceObjects;
+	}
+
+	/**
+	 * Returns the set of source objects (such as actors) to ignore during the trace. Note that elements are
+	 * NOT guaranteed to be unique. This is less important for source objects than components since it's
+	 * less likely that duplicates are added.
+	 */
+	const IgnoreActorsArrayType& GetIgnoredSourceObjects() const
+	{
+		return IgnoreSourceObjects;
 	}
 
 	/** Clears the set of components to ignore during the trace. */
@@ -128,10 +141,17 @@ public:
 		bComponentListUnique = true;
 	}
 
+	UE_DEPRECATED(5.5, "Use ClearIgnoredSourceObjects instead.")
 	/** Clears the set of actors to ignore during the trace. */
 	void ClearIgnoredActors()
 	{
-		IgnoreActors.Reset();
+		IgnoreSourceObjects.Reset();
+	}
+
+	/** Clears the set of source objects (such as actors) to ignore during the trace. */
+	void ClearIgnoredSourceObjects()
+	{
+		IgnoreSourceObjects.Reset();
 	}
 
 	/**
@@ -175,19 +195,28 @@ public:
 
 	// Utils
 
-	/** Add an actor for this trace to ignore */
+	/** Add an actor for this trace to ignore. Equivalent to calling AddIgnoredSourceObject. */
 	ENGINE_API void AddIgnoredActor(const AActor* InIgnoreActor);
 
-	/** Add an actor by ID for this trace to ignore */
+	/** Add an actor by ID for this trace to ignore. Equivalent to calling AddIgnoredSourceObject. */
 	ENGINE_API void AddIgnoredActor(const uint32 InIgnoreActorID);
 
-	/** Add a collection of actors for this trace to ignore */
+	/** Add a source object for this trace to ignore */
+	ENGINE_API void AddIgnoredSourceObject(const UObject* InIgnoreActor);
+
+	/** Add a source object for this trace to ignore */
+	ENGINE_API void AddIgnoredSourceObject(const uint32 InIgnoreActorID);
+
+	/** Add a collection of actors for this trace to ignore. Equivalent to calling AddIgnoredSourceObjects */
 	ENGINE_API void AddIgnoredActors(const TArray<AActor*>& InIgnoreActors);
 	ENGINE_API void AddIgnoredActors(const TArray<const AActor*>& InIgnoreActors);
-
-	/** Variant that uses an array of TWeakObjectPtrs */
 	ENGINE_API void AddIgnoredActors(const TArray<TWeakObjectPtr<const AActor> >& InIgnoreActors);
 
+	/** Add a collection of source objects for this trace to ignore */
+	ENGINE_API void AddIgnoredSourceObjects(const TArray<UObject*>& InIgnoreObjects);
+	ENGINE_API void AddIgnoredSourceObjects(const TArray<const UObject*>& InIgnoreObjects);
+	ENGINE_API void AddIgnoredSourceObjects(const TArray<TWeakObjectPtr<const UObject> >& InIgnoreObjects);
+
 	/** Add a component for this trace to ignore */
 	ENGINE_API void AddIgnoredComponent(const UPrimitiveComponent* InIgnoreComponent);
 
diff --git a/Engine/Source/Runtime/Engine/Public/Physics/PhysicsFiltering.h b/Engine/Source/Runtime/Engine/Public/Physics/PhysicsFiltering.h
--- a/Engine/Source/Runtime/Engine/Public/Physics/PhysicsFiltering.h
+++ b/Engine/Source/Runtime/Engine/Public/Physics/PhysicsFiltering.h
@@ -47,7 +47,7 @@ struct FPhysicsFilterBuilder
 		}
 	}
 
-	inline void GetQueryData(uint32 ActorID, uint32& OutWord0, uint32& OutWord1, uint32& OutWord2, uint32& OutWord3) const
+	inline void GetQueryData(uint32 SourceObjectID, uint32& OutWord0, uint32& OutWord1, uint32& OutWord2, uint32& OutWord3) const
 	{
 		/**
 		 * Format for QueryData : 
@@ -56,7 +56,7 @@ struct FPhysicsFilterBuilder
 		 *		word2 (touching channels)
 		 *		word3 (ExtraFilter (top NumExtraFilterBits) + MyChannel (next NumCollisionChannelBits) as ECollisionChannel + Flags (remaining NumFilterDataFlagBits)
 		 */
-		OutWord0 = ActorID;
+		OutWord0 = SourceObjectID;
 		OutWord1 = BlockingBits;
 		OutWord2 = TouchingBits;
 		OutWord3 = Word3;
@@ -94,7 +94,7 @@ private:
 inline void CreateShapeFilterData(
 	const uint8 MyChannel,
 	const FMaskFilter MaskFilter,
-	const int32 ActorID,
+	const int32 SourceObjectID,
 	const FCollisionResponseContainer& ResponseToChannels,
 	uint32 ComponentID,
 	uint16 BodyIndex,
@@ -113,7 +113,7 @@ inline void CreateShapeFilterData(
 
 	OutQueryData = FCollisionFilterData();
 	OutSimData = FCollisionFilterData();
-	Builder.GetQueryData(ActorID, OutQueryData.Word0, OutQueryData.Word1, OutQueryData.Word2, OutQueryData.Word3);
+	Builder.GetQueryData(SourceObjectID, OutQueryData.Word0, OutQueryData.Word1, OutQueryData.Word2, OutQueryData.Word3);
 	Builder.GetSimData(BodyIndex, ComponentID, OutSimData.Word0, OutSimData.Word1, OutSimData.Word2, OutSimData.Word3);
 }
 
diff --git a/Engine/Source/Runtime/Experimental/ChaosVisualDebugger/Public/DataWrappers/ChaosVDQueryDataWrappers.h b/Engine/Source/Runtime/Experimental/ChaosVisualDebugger/Public/DataWrappers/ChaosVDQueryDataWrappers.h
--- a/Engine/Source/Runtime/Experimental/ChaosVisualDebugger/Public/DataWrappers/ChaosVDQueryDataWrappers.h
+++ b/Engine/Source/Runtime/Experimental/ChaosVisualDebugger/Public/DataWrappers/ChaosVDQueryDataWrappers.h
@@ -118,7 +118,7 @@ struct FChaosVDCollisionQueryParams : public FChaosVDWrapperDataBase
 		IgnoreMask = Other.IgnoreMask;
 
 		IgnoredComponentsIDs = Other.GetIgnoredComponents();
-		IgnoredActorsIDs = Other.GetIgnoredActors();
+		IgnoredActorsIDs = Other.GetIgnoredSourceObjects();
 
 		bHasValidData = true;
 	}
