Updated USoundWave::HasStreamingChunks() to return false when the loading behavior is ForceInline (not a streaming sound).  This will fix a family of ensures where systems attemt to prime sounds without checking to see if they are forceinline (these prime calls will now early exit).
#rnx

diff --git a/Engine/Source/Runtime/Engine/Private/SoundWave.cpp b/Engine/Source/Runtime/Engine/Private/SoundWave.cpp
--- a/Engine/Source/Runtime/Engine/Private/SoundWave.cpp
+++ b/Engine/Source/Runtime/Engine/Private/SoundWave.cpp
@@ -1914,7 +1914,9 @@ void USoundWave::FlushAudioRenderingCommands() const
 bool USoundWave::HasStreamingChunks()
 {
 	check(SoundWaveDataPtr);
-	return SoundWaveDataPtr->RunningPlatformData.GetNumChunks() > 0;
+	const bool bIsForceInline = (GetLoadingBehavior() == ESoundWaveLoadingBehavior::ForceInline);
+	const bool bHasChunks = (SoundWaveDataPtr->RunningPlatformData.GetNumChunks() > 0);
+	return !bIsForceInline && bHasChunks;
 }
 
 void USoundWave::PostLoad()
