Fixed WorldPartitionRuntimeHashSet generated cell names for PIE instanced partitioned worlds (client/server had different names which closed the net connection, also world instance suffix was present twice)

diff --git a/Engine/Source/Runtime/Engine/Private/WorldPartition/RuntimeHashSet/WorldPartitionRuntimeHashSet.cpp b/Engine/Source/Runtime/Engine/Private/WorldPartition/RuntimeHashSet/WorldPartitionRuntimeHashSet.cpp
--- a/Engine/Source/Runtime/Engine/Private/WorldPartition/RuntimeHashSet/WorldPartitionRuntimeHashSet.cpp
+++ b/Engine/Source/Runtime/Engine/Private/WorldPartition/RuntimeHashSet/WorldPartitionRuntimeHashSet.cpp
@@ -678,25 +678,12 @@ UWorldPartitionRuntimeHashSet::FCellUniqueId UWorldPartitionRuntimeHashSet::GetC
 
 	// Build cell unique name
 	{
-		TStringBuilder<128> CellNameBuilder;
-
 		UWorld* OuterWorld = GetTypedOuter<UWorld>();
 		check(OuterWorld);
 
+		FString InstanceSuffix;
 		FString WorldName = FPackageName::GetShortName(OuterWorld->GetPackage());
 
-		CellNameBuilder.Appendf(TEXT("%s_%s"), *WorldName, *CellNameID.ToString());
-
-		if (DataLayersID.GetHash())
-		{
-			CellNameBuilder.Appendf(TEXT("_d%X"), DataLayersID.GetHash());
-		}
-
-		if (ContentBundleID.IsValid())
-		{
-			CellNameBuilder.Appendf(TEXT("_c%s"), *UContentBundleDescriptor::GetContentBundleCompactString(ContentBundleID));
-		}
-
 		if (!IsRunningCookCommandlet() && OuterWorld->IsGameWorld())
 		{
 			FString SourceWorldPath;
@@ -714,10 +701,28 @@ UWorldPartitionRuntimeHashSet::FCellUniqueId UWorldPartitionRuntimeHashSet::GetC
 
 				if (int32 Index = InstancePackageName.Find(SourcePackageName); Index != INDEX_NONE)
 				{
-					CellNameBuilder.Appendf(TEXT("_i%s"), *InstancePackageName.Mid(Index + SourcePackageName.Len()));
+					InstanceSuffix = InstancePackageName.Mid(Index + SourcePackageName.Len());
 				}
 			}
 		}
+
+		TStringBuilder<128> CellNameBuilder;
+		CellNameBuilder.Appendf(TEXT("%s_%s"), *WorldName, *CellNameID.ToString());
+
+		if (DataLayersID.GetHash())
+		{
+			CellNameBuilder.Appendf(TEXT("_d%X"), DataLayersID.GetHash());
+		}
+
+		if (ContentBundleID.IsValid())
+		{
+			CellNameBuilder.Appendf(TEXT("_c%s"), *UContentBundleDescriptor::GetContentBundleCompactString(ContentBundleID));
+		}
+
+		if (!InstanceSuffix.IsEmpty())
+		{
+			CellNameBuilder.Appendf(TEXT("_i%s"), *InstanceSuffix);
+		}
 	
 		CellUniqueId.Name = CellNameBuilder.ToString();
 	}
