Avoid stack overflow in AnimBP evaluation

diff --git a/Engine/Source/Runtime/Engine/Private/Animation/AnimNodeBase.cpp b/Engine/Source/Runtime/Engine/Private/Animation/AnimNodeBase.cpp
--- a/Engine/Source/Runtime/Engine/Private/Animation/AnimNodeBase.cpp
+++ b/Engine/Source/Runtime/Engine/Private/Animation/AnimNodeBase.cpp
@@ -257,7 +257,8 @@ FAnimNode_Base* FPoseLinkBase::GetLinkNode()
 	return LinkedNode;
 }
 
-const FExposedValueHandler& FAnimNode_Base::GetEvaluateGraphExposedInputs() const
+// Don't inline this function to keep the stack usage down
+FORCENOINLINE const FExposedValueHandler& FAnimNode_Base::GetEvaluateGraphExposedInputs() const
 {
 	if(NodeData)
 	{
diff --git a/Engine/Source/Runtime/Engine/Private/Animation/ExposedValueHandler.cpp b/Engine/Source/Runtime/Engine/Private/Animation/ExposedValueHandler.cpp
--- a/Engine/Source/Runtime/Engine/Private/Animation/ExposedValueHandler.cpp
+++ b/Engine/Source/Runtime/Engine/Private/Animation/ExposedValueHandler.cpp
@@ -38,7 +38,8 @@ void FExposedValueHandler::Initialize(UClass* InClass, const FPropertyAccessLibr
 	}
 }
 
-void FExposedValueHandler::Execute(const FAnimationBaseContext& Context) const
+// Don't inline this function to keep the stack usage down
+FORCENOINLINE void FExposedValueHandler::Execute(const FAnimationBaseContext& Context) const
 {
 	if (Handler)
 	{
