Skeletal Mesh: Recreate a material override map for older skelmeshes that did not fill it in when overriding materials on auto-generated LODs.

#rnx

diff --git a/Engine/Source/Runtime/Core/Private/UObject/FortniteMainBranchObjectVersion.cpp b/Engine/Source/Runtime/Core/Private/UObject/FortniteMainBranchObjectVersion.cpp
--- a/Engine/Source/Runtime/Core/Private/UObject/FortniteMainBranchObjectVersion.cpp
+++ b/Engine/Source/Runtime/Core/Private/UObject/FortniteMainBranchObjectVersion.cpp
@@ -12,7 +12,7 @@ TMap<FGuid, FGuid> FFortniteMainBranchObjectVersion::GetSystemGuids()
 	SystemGuids.Add(DevGuids.NANITE_DERIVEDDATA_VER, FGuid("D28C003F825F4C178F698938FB0E47F7"));
 	SystemGuids.Add(DevGuids.NIAGARASHADERMAP_DERIVEDDATA_VER, FGuid("C417A50271B3427D98785388266E3B66"));
 	SystemGuids.Add(DevGuids.Niagara_LatestScriptCompileVersion, FGuid("CE33E29C6268469E8E18F1D96A55DFDA"));
-	SystemGuids.Add(DevGuids.SkeletalMeshDerivedDataVersion, FGuid("69ABAF00AC9648E2A2EE57D4CC449CF3"));
+	SystemGuids.Add(DevGuids.SkeletalMeshDerivedDataVersion, FGuid("B9AEAA3EE7AC4FCFB3860E8F3E06DD75"));
 	SystemGuids.Add(DevGuids.STATICMESH_DERIVEDDATA_VER, FGuid("543843B103794E0E9BA4BE80FB602F79"));
 	SystemGuids.Add(DevGuids.MaterialTranslationDDCVersion, FGuid("BBD650804D314CE7A44ADF9BB4F3F5F9"));
 	return SystemGuids;
diff --git a/Engine/Source/Runtime/Engine/Private/SkeletalMesh.cpp b/Engine/Source/Runtime/Engine/Private/SkeletalMesh.cpp
--- a/Engine/Source/Runtime/Engine/Private/SkeletalMesh.cpp
+++ b/Engine/Source/Runtime/Engine/Private/SkeletalMesh.cpp
@@ -93,6 +93,7 @@
 #include "Animation/SkinWeightProfileManager.h"
 #include "BoneWeights.h"
 #include "SkeletalMeshAttributes.h"
+#include "Algo/AnyOf.h"
 #include "Logging/StructuredLog.h"
 
 #define LOCTEXT_NAMESPACE "SkeltalMesh"
@@ -3274,7 +3275,7 @@ void USkeletalMesh::BeginPostLoadInternal(FSkinnedAssetPostLoadContext& Context)
 				}
 			}
 		}
-
+		
 		// If we didn't get any meshes from the bulk data, then try to recover them from the LODModel listings.
 		for (int32 LODIndex = 0; LODIndex < GetImportedModel()->LODModels.Num(); ++LODIndex)
 		{
@@ -3287,31 +3288,69 @@ void USkeletalMesh::BeginPostLoadInternal(FSkinnedAssetPostLoadContext& Context)
 			// so that the mesh doesn't get reduced again if it gets regenerated.
 			FSkeletalMeshLODInfo* MeshLODInfo = GetLODInfo(LODIndex); 
 			const bool bReductionActive = IsReductionActive(LODIndex);
-			const bool bInlineReduction = (GetLODInfo(LODIndex)->ReductionSettings.BaseLOD == LODIndex);
-			if (bReductionActive && !bInlineReduction)
+			const bool bInlineReduction = (MeshLODInfo->ReductionSettings.BaseLOD == LODIndex);
+			if (!bReductionActive || bInlineReduction)
 			{
-				//Generated LOD (not inline) do not need imported data
-				continue;
-			}
+				const FSkeletalMeshLODModel& LODModel = GetImportedModel()->LODModels[LODIndex];
+				FMeshDescription MeshDescription;
+				LODModel.GetMeshDescription(this, LODIndex, MeshDescription);
+				CreateMeshDescription(LODIndex, MoveTemp(MeshDescription));
+				CommitMeshDescription(LODIndex);
+
+				// Reset the reduction settings so that we don't re-reduce the mesh and possibly lose morph targets
+				// in the process.
+				FSkeletalMeshOptimizationSettings& ReductionSettings = MeshLODInfo->ReductionSettings;
 			
-			const FSkeletalMeshLODModel& LODModel = GetImportedModel()->LODModels[LODIndex];
+				//Remove the reduction settings
+				ReductionSettings.NumOfTrianglesPercentage = 1.0f;
+				ReductionSettings.NumOfVertPercentage = 1.0f;
+				ReductionSettings.MaxNumOfTrianglesPercentage = MAX_uint32;
+				ReductionSettings.MaxNumOfVertsPercentage = MAX_uint32;
+				ReductionSettings.TerminationCriterion = SMTC_NumOfTriangles;
+				MeshLODInfo->bHasBeenSimplified = false;
+			}
+			else if (MeshLODInfo->LODMaterialMap.IsEmpty())
+			{
+				// Generated LODs (not inline) do not need imported data. We do need a material map though,
+				// because in many cases the map was not created when a section material got overridden, so reconstruct one if it isn't available.
+				const FSkeletalMeshLODModel& BaseLODModel = GetImportedModel()->LODModels[MeshLODInfo->ReductionSettings.BaseLOD];
+				const FSkeletalMeshLODModel& LODModel = GetImportedModel()->LODModels[LODIndex];
+				TArray<int32> MaterialMap;
+				MaterialMap.Init(INDEX_NONE, LODModel.Sections.Num());
 
-			FMeshDescription MeshDescription;
-			LODModel.GetMeshDescription(this, LODIndex, MeshDescription);
-			CreateMeshDescription(LODIndex, MoveTemp(MeshDescription));
-			CommitMeshDescription(LODIndex);
+				if (BaseLODModel.Sections.Num() == LODModel.Sections.Num())
+				{
+					for (int32 SectionIndex = 0; SectionIndex < LODModel.Sections.Num(); SectionIndex++)
+					{
+						const int32 MaterialIndex = LODModel.Sections[SectionIndex].MaterialIndex;
+						if (BaseLODModel.Sections[SectionIndex].MaterialIndex != MaterialIndex)
+						{
+							MaterialMap[SectionIndex] = MaterialIndex;
+						}
+					}
+				}
+				else
+				{
+					for (int32 SectionIndex = 0; SectionIndex < LODModel.Sections.Num(); SectionIndex++)
+					{
+						const int32 BaseSectionIndex = LODModel.Sections[SectionIndex].OriginalDataSectionIndex;
+						
+						if (BaseLODModel.Sections.IsValidIndex(BaseSectionIndex))
+						{
+							const int32 MaterialIndex = LODModel.Sections[SectionIndex].MaterialIndex;
 
-			// Reset the reduction settings so that we don't re-reduce the mesh and possibly lose morph targets
-			// in the process.
-			FSkeletalMeshOptimizationSettings& ReductionSettings = MeshLODInfo->ReductionSettings;
-			
-			//Remove the reduction settings
-			ReductionSettings.NumOfTrianglesPercentage = 1.0f;
-			ReductionSettings.NumOfVertPercentage = 1.0f;
-			ReductionSettings.MaxNumOfTrianglesPercentage = MAX_uint32;
-			ReductionSettings.MaxNumOfVertsPercentage = MAX_uint32;
-			ReductionSettings.TerminationCriterion = SMTC_NumOfTriangles;
-			MeshLODInfo->bHasBeenSimplified = false;
+							if (BaseLODModel.Sections[BaseSectionIndex].MaterialIndex != MaterialIndex)
+							{
+								MaterialMap[SectionIndex] = MaterialIndex;
+							}
+						}
+					}
+				}
+				if (Algo::AnyOf(MaterialMap, [](int32 Item) { return Item != INDEX_NONE; }))
+				{
+					MeshLODInfo->LODMaterialMap = MoveTemp(MaterialMap);
+				}
+			}
 		}
 
 		if (GetLinkerCustomVersion(FEditorObjectVersion::GUID) < FEditorObjectVersion::SkeletalMeshBuildRefactor)
