Fix setting flags atomically on a new object in closed code.

diff --git a/Engine/Source/Runtime/Engine/Private/SimpleConstructionScript.cpp b/Engine/Source/Runtime/Engine/Private/SimpleConstructionScript.cpp
--- a/Engine/Source/Runtime/Engine/Private/SimpleConstructionScript.cpp
+++ b/Engine/Source/Runtime/Engine/Private/SimpleConstructionScript.cpp
@@ -692,7 +692,13 @@ void USimpleConstructionScript::ExecuteScriptOnActor(AActor* Actor, const TInlin
 	else if(Actor->GetRootComponent() == nullptr) // Must have a root component at the end of SCS, so if we don't have one already (from base class), create a SceneComponent now
 	{
 		USceneComponent* SceneComp = NewObject<USceneComponent>(Actor);
-		SceneComp->SetFlags(RF_Transactional);
+
+		// The object is new, so its safe for us to atomically set the flag in the open.
+		UE_AUTORTFM_OPEN(
+			{
+				SceneComp->SetFlags(RF_Transactional);
+			});
+
 		SceneComp->CreationMethod = EComponentCreationMethod::SimpleConstructionScript;
 		if (RootRelativeRotationCache)
 		{   // Enforces using the same rotator as much as possible.
