Include task name in AssetCompilingManager log output

diff --git a/Engine/Source/Runtime/Engine/Private/AssetCompilingManager.cpp b/Engine/Source/Runtime/Engine/Private/AssetCompilingManager.cpp
--- a/Engine/Source/Runtime/Engine/Private/AssetCompilingManager.cpp
+++ b/Engine/Source/Runtime/Engine/Private/AssetCompilingManager.cpp
@@ -253,7 +253,8 @@ public:
 		
 		// Add next work memory requirement to see if it still fits in available memory
 		EQueuedWorkPriority Priority;
-		if (IQueuedWork* NextWork = QueuedWork.Peek(&Priority))
+		IQueuedWork* NextWork = QueuedWork.Peek(&Priority);
+		if (NextWork)
 		{
 			NewRequiredMemory += GetRequiredMemory(NextWork);
 		}
@@ -281,9 +282,19 @@ public:
 			{
 				// Never limit below a concurrency of 1 or else we'll starve the asset processing and never be scheduled again
 				// The idea for now is to let assets get built one by one when starving on memory, which should not increase OOM compared to the old synchronous behavior.
-				
-				UE_LOG(LogAsyncCompilation, Display, TEXT("AssetCompile memory estimate is greater than available, but we're running it anyway, likely to fail! ")
+				const TCHAR* DebugName = nullptr;
+				if (NextWork)
+				{
+					DebugName = NextWork->GetDebugName();
+				}
+				if (DebugName == nullptr)
+				{
+					DebugName = TEXT("No DebugName");
+				}
+
+				UE_LOG(LogAsyncCompilation, Display, TEXT("AssetCompile memory estimate is greater than available, but we're running it [%s] anyway, likely to fail! ")
 					TEXT("RequiredMemory = %.3f MB + %.3f MB MemoryLimit = %.3f MB "),
+					DebugName,
 					NewRequiredMemory/(1024*1024.f),
 					TotalEstimatedMemory/(1024*1024.f),
 					MemoryLimit/(1024*1024.f));
