LevelStreaming:
- virtual ULevelStreaming::OnCurrentStateChanged called in SetCurrentState so that child classes can implement without hooking to the global event
- FLoadLevelInstanceParams::LevelStreamingCreatedCallback: called after the LevelStreaming object is created, allowing calling code to do some initialization code

diff --git a/Engine/Source/Runtime/Engine/Classes/Engine/LevelStreaming.h b/Engine/Source/Runtime/Engine/Classes/Engine/LevelStreaming.h
--- a/Engine/Source/Runtime/Engine/Classes/Engine/LevelStreaming.h
+++ b/Engine/Source/Runtime/Engine/Classes/Engine/LevelStreaming.h
@@ -687,6 +687,9 @@ protected:
 	/** Called by SetLoadedLevel */
 	virtual void OnLevelLoadedChanged(ULevel* Level) {}
 
+	/** Called when the current state changes */
+	virtual void OnCurrentStateChanged(ELevelStreamingState OldState, ELevelStreamingState NewState) {}
+
 	/** Called by RequestLevel to detect existing streaming level with same world asset */
 	bool ValidateUniqueWorldAsset(UWorld* PersistentWorld);
 
diff --git a/Engine/Source/Runtime/Engine/Classes/Engine/LevelStreamingDynamic.h b/Engine/Source/Runtime/Engine/Classes/Engine/LevelStreamingDynamic.h
--- a/Engine/Source/Runtime/Engine/Classes/Engine/LevelStreamingDynamic.h
+++ b/Engine/Source/Runtime/Engine/Classes/Engine/LevelStreamingDynamic.h
@@ -63,6 +63,9 @@ class ULevelStreamingDynamic : public ULevelStreaming
 
 		/** Set EditorPath Owner */
 		UObject* EditorPathOwner = nullptr;
+
+		/** Callback when ULevelStreaming object is created */
+		TUniqueFunction<void(ULevelStreaming*)> LevelStreamingCreatedCallback;
 	};
 
 	/**  
diff --git a/Engine/Source/Runtime/Engine/Private/LevelStreaming.cpp b/Engine/Source/Runtime/Engine/Private/LevelStreaming.cpp
--- a/Engine/Source/Runtime/Engine/Private/LevelStreaming.cpp
+++ b/Engine/Source/Runtime/Engine/Private/LevelStreaming.cpp
@@ -536,11 +536,17 @@ void ULevelStreaming::OnLevelRemoved()
 
 void ULevelStreaming::SetCurrentState(ELevelStreamingState NewState)
 {
+	ELevelStreamingState OldState = CurrentState;
+	CurrentState = NewState;
+
+	if (OldState != NewState)
+	{
+		OnCurrentStateChanged(OldState, NewState);
+	}
+
 	// TODO: We should only fire the delegate when the current state has changed, but first AsyncLevelLoadComplete needs to be fixed to 
 	// only set the new state once the LoadedLevel is assigned. Clients currently rely on getting a repeated notification after the loaded
 	// level is available.
-	ELevelStreamingState OldState = CurrentState;
-	CurrentState = NewState;
 	FLevelStreamingDelegates::OnLevelStreamingStateChanged.Broadcast(GetWorld(), this, GetLoadedLevel(), OldState, NewState);
 }
 
@@ -2613,6 +2619,10 @@ ULevelStreamingDynamic* ULevelStreamingDynamic::LoadLevelInstance_Internal(const
     
 	// Setup streaming level object that will load specified map
 	ULevelStreamingDynamic* StreamingLevel = NewObject<ULevelStreamingDynamic>(Params.World, LevelStreamingClass, NAME_None, RF_Transient, NULL);
+	if (Params.LevelStreamingCreatedCallback)
+	{
+		Params.LevelStreamingCreatedCallback(StreamingLevel);
+	}
 
 	FSoftObjectPath WorldAssetPath(*WriteToString<512>(UnmodifiedLevelPackageName, TEXT("."), ShortPackageName));
     StreamingLevel->SetWorldAsset(TSoftObjectPtr<UWorld>(WorldAssetPath));
