Skinning:  Fix for previous frame bone data not persisting when recreating the render proxy.  The cause of this was the previous frame data being ignored due to the same RevisonNumber being passed to UpdateBoneData for both the previous and current frame, causing the second call to ignore the data.  To solve this, we need to propagate PreviousRevisonNumber to the call, and fix a couple other minor bugs.

#rnx

#changelist validated
#virtualized

diff --git a/Engine/Source/Runtime/Engine/Classes/Components/SkinnedMeshComponent.h b/Engine/Source/Runtime/Engine/Classes/Components/SkinnedMeshComponent.h
--- a/Engine/Source/Runtime/Engine/Classes/Components/SkinnedMeshComponent.h
+++ b/Engine/Source/Runtime/Engine/Classes/Components/SkinnedMeshComponent.h
@@ -290,6 +290,12 @@ public:
 		return (LeaderPoseComponentPtr ? LeaderPoseComponentPtr->CurrentBoneTransformRevisionNumber : CurrentBoneTransformRevisionNumber);
 	}
 
+	uint32 GetPreviousBoneTransformRevisionNumber() const
+	{
+		const USkinnedMeshComponent* LeaderPoseComponentPtr = LeaderPoseComponent.Get();
+		return (LeaderPoseComponentPtr ? LeaderPoseComponentPtr->PreviousBoneTransformRevisionNumber : PreviousBoneTransformRevisionNumber);
+	}
+
 	/* this update renderer with new revision number twice so to clear bone velocity for motion blur or temporal AA */
 	ENGINE_API void ClearMotionVector();
 	
diff --git a/Engine/Source/Runtime/Engine/Private/GPUSkinVertexFactory.cpp b/Engine/Source/Runtime/Engine/Private/GPUSkinVertexFactory.cpp
--- a/Engine/Source/Runtime/Engine/Private/GPUSkinVertexFactory.cpp
+++ b/Engine/Source/Runtime/Engine/Private/GPUSkinVertexFactory.cpp
@@ -218,7 +218,7 @@ TConsoleVariableData<int32>* FGPUBaseSkinVertexFactory::FShaderDataType::MaxBone
 uint32 FGPUBaseSkinVertexFactory::FShaderDataType::MaxGPUSkinBones = 0;
 
 void FGPUBaseSkinVertexFactory::FShaderDataType::UpdateBoneData(FRHICommandList& RHICmdList, const TArray<FMatrix44f>& ReferenceToLocalMatrices,
-	const TArray<FBoneIndexType>& BoneMap, uint32 RevisionNumber, bool bPrevious, ERHIFeatureLevel::Type InFeatureLevel, bool bUseSkinCache, bool bForceUpdateImmediately, const FName& AssetPathName)
+	const TArray<FBoneIndexType>& BoneMap, uint32 RevisionNumber, ERHIFeatureLevel::Type InFeatureLevel, bool bUseSkinCache, bool bForceUpdateImmediately, const FName& AssetPathName)
 {
 	// stat disabled by default due to low-value/high-frequency
 	//QUICK_SCOPE_CYCLE_COUNTER(STAT_FGPUBaseSkinVertexFactory_UpdateBoneData);
@@ -234,6 +234,7 @@ void FGPUBaseSkinVertexFactory::FShaderDataType::UpdateBoneData(FRHICommandList&
 		// make sure current revision is up-to-date
 		SetCurrentRevisionNumber(RevisionNumber);
 
+		const bool bPrevious = false;
 		CurrentBoneBuffer = &GetBoneBufferForWriting(bPrevious);
 
 		static FSharedPoolPolicyData PoolPolicy;
diff --git a/Engine/Source/Runtime/Engine/Private/SkeletalRenderGPUSkin.cpp b/Engine/Source/Runtime/Engine/Private/SkeletalRenderGPUSkin.cpp
--- a/Engine/Source/Runtime/Engine/Private/SkeletalRenderGPUSkin.cpp
+++ b/Engine/Source/Runtime/Engine/Private/SkeletalRenderGPUSkin.cpp
@@ -404,6 +404,7 @@ void FSkeletalMeshObjectGPUSkin::Update(
 
 	uint64 FrameNumberToPrepare = GFrameCounter;
 	uint32 RevisionNumber = 0;
+	uint32 PreviousRevisionNumber = 0;
 
 	FGPUSkinCache* GPUSkinCache = nullptr;
 	FSceneInterface* Scene = nullptr;
@@ -413,15 +414,16 @@ void FSkeletalMeshObjectGPUSkin::Update(
 		Scene = &(InMeshComponent->SceneProxy->GetScene());
 		GPUSkinCache = Scene->GetGPUSkinCache();
 		RevisionNumber = InMeshComponent->GetBoneTransformRevisionNumber();
+		PreviousRevisionNumber = InMeshComponent->GetPreviousBoneTransformRevisionNumber();
 	}
 
 	// queue a call to update this data
 	FSkeletalMeshObjectGPUSkin* MeshObject = this;
 	ENQUEUE_RENDER_COMMAND(SkelMeshObjectUpdateDataCommand)(UE::RenderCommandPipe::SkeletalMesh,
-		[MeshObject, FrameNumberToPrepare, RevisionNumber, NewDynamicData, GPUSkinCache, Scene](FRHICommandList& RHICmdList)
+		[MeshObject, FrameNumberToPrepare, RevisionNumber, PreviousRevisionNumber, NewDynamicData, GPUSkinCache, Scene](FRHICommandList& RHICmdList)
 		{
 			FScopeCycleCounter Context(MeshObject->GetStatId());
-			MeshObject->UpdateDynamicData_RenderThread(GPUSkinCache, RHICmdList, NewDynamicData, Scene, FrameNumberToPrepare, RevisionNumber);
+			MeshObject->UpdateDynamicData_RenderThread(GPUSkinCache, RHICmdList, NewDynamicData, Scene, FrameNumberToPrepare, RevisionNumber, PreviousRevisionNumber);
 		}
 	);
 }
@@ -479,7 +481,7 @@ FORCEINLINE bool IsDeferredSkeletalDynamicDataUpdateEnabled()
 	return CVarDeferSkeletalDynamicDataUpdateUntilGDME.GetValueOnRenderThread() > 0 && !IsParallelGatherDynamicMeshElementsEnabled();
 }
 
-void FSkeletalMeshObjectGPUSkin::UpdateDynamicData_RenderThread(FGPUSkinCache* GPUSkinCache, FRHICommandList& RHICmdList, FDynamicSkelMeshObjectDataGPUSkin* InDynamicData, FSceneInterface* Scene, uint64 FrameNumberToPrepare, uint32 RevisionNumber)
+void FSkeletalMeshObjectGPUSkin::UpdateDynamicData_RenderThread(FGPUSkinCache* GPUSkinCache, FRHICommandList& RHICmdList, FDynamicSkelMeshObjectDataGPUSkin* InDynamicData, FSceneInterface* Scene, uint64 FrameNumberToPrepare, uint32 RevisionNumber, uint32 PreviousRevisionNumber)
 {
 	TRACE_CPUPROFILER_EVENT_SCOPE(GPUSkin::UpdateDynamicData_RT);
 
@@ -518,13 +520,13 @@ void FSkeletalMeshObjectGPUSkin::UpdateDynamicData_RenderThread(FGPUSkinCache* G
 	}
 	else
 	{
-		ProcessUpdatedDynamicData(EGPUSkinCacheEntryMode::Raster, GPUSkinCache, RHICmdList, FrameNumberToPrepare, RevisionNumber, bMorphNeedsUpdate, DynamicData->LODIndex);
+		ProcessUpdatedDynamicData(EGPUSkinCacheEntryMode::Raster, GPUSkinCache, RHICmdList, FrameNumberToPrepare, RevisionNumber, PreviousRevisionNumber, bMorphNeedsUpdate, DynamicData->LODIndex);
 
 	#if RHI_RAYTRACING
 		if (ShouldUseSeparateSkinCacheEntryForRayTracing() && FGPUSkinCache::IsGPUSkinCacheRayTracingSupported() && GPUSkinCache && SkeletalMeshRenderData->bSupportRayTracing)
 		{
 			// Morph delta is updated in raster pass above, no need to update again for ray tracing
-			ProcessUpdatedDynamicData(EGPUSkinCacheEntryMode::RayTracing, GPUSkinCache, RHICmdList, FrameNumberToPrepare, RevisionNumber, /*bMorphNeedsUpdate=*/false, DynamicData->RayTracingLODIndex);
+			ProcessUpdatedDynamicData(EGPUSkinCacheEntryMode::RayTracing, GPUSkinCache, RHICmdList, FrameNumberToPrepare, RevisionNumber, PreviousRevisionNumber, /*bMorphNeedsUpdate=*/false, DynamicData->RayTracingLODIndex);
 		}
 		else
 		{
@@ -548,11 +550,11 @@ void FSkeletalMeshObjectGPUSkin::PreGDMECallback(FRHICommandList& RHICmdList, FG
 {
 	if (bNeedsUpdateDeferred)
 	{
-		ProcessUpdatedDynamicData(EGPUSkinCacheEntryMode::Raster, GPUSkinCache, RHICmdList, FrameNumber, LastBoneTransformRevisionNumber, bMorphNeedsUpdateDeferred, DynamicData->LODIndex);
+		ProcessUpdatedDynamicData(EGPUSkinCacheEntryMode::Raster, GPUSkinCache, RHICmdList, FrameNumber, LastBoneTransformRevisionNumber, 0, bMorphNeedsUpdateDeferred, DynamicData->LODIndex);
 	}
 }
 
-void FSkeletalMeshObjectGPUSkin::ProcessUpdatedDynamicData(EGPUSkinCacheEntryMode Mode, FGPUSkinCache* GPUSkinCache, FRHICommandList& RHICmdList, uint32 FrameNumberToPrepare, uint32 RevisionNumber, bool bMorphNeedsUpdate, int32 LODIndex)
+void FSkeletalMeshObjectGPUSkin::ProcessUpdatedDynamicData(EGPUSkinCacheEntryMode Mode, FGPUSkinCache* GPUSkinCache, FRHICommandList& RHICmdList, uint32 FrameNumberToPrepare, uint32 RevisionNumber, uint32 PreviousRevisionNumber, bool bMorphNeedsUpdate, int32 LODIndex)
 {
 	QUICK_SCOPE_CYCLE_COUNTER(STAT_FSkeletalMeshObjectGPUSkin_ProcessUpdatedDynamicData);
 	bNeedsUpdateDeferred = false;
@@ -680,16 +682,17 @@ void FSkeletalMeshObjectGPUSkin::ProcessUpdatedDynamicData(EGPUSkinCacheEntryMod
 
 			const FName OwnerName = GetAssetPathName(LODIndex);
 
-			// if we have previous reference to loca, we also update to previous frame
+			// If we have previous reference to local, we also update to previous frame.  We technically are storing the previous data to "current" first,
+			// then the second call to UpdateBoneData immediately below swaps that to become "previous".
 			if (DynamicData->PreviousReferenceToLocal.Num() > 0)
 			{
 				TArray<FMatrix44f>& PreviousReferenceToLocalMatrices = bShouldUseSeparateMatricesForRayTracing ? DynamicData->PreviousReferenceToLocalForRayTracing : DynamicData->PreviousReferenceToLocal;
-				ShaderData.UpdateBoneData(RHICmdList, PreviousReferenceToLocalMatrices, Section.BoneMap, RevisionNumber, true, FeatureLevel, bUseSkinCache, DynamicData->bForceUpdateDynamicDataImmediately, OwnerName);
+				ShaderData.UpdateBoneData(RHICmdList, PreviousReferenceToLocalMatrices, Section.BoneMap, PreviousRevisionNumber, FeatureLevel, bUseSkinCache, DynamicData->bForceUpdateDynamicDataImmediately, OwnerName);
 			}
 
 			// Create a uniform buffer from the bone transforms.
 			TArray<FMatrix44f>& ReferenceToLocalMatrices = bShouldUseSeparateMatricesForRayTracing ? DynamicData->ReferenceToLocalForRayTracing : DynamicData->ReferenceToLocal;
-			ShaderData.UpdateBoneData(RHICmdList, ReferenceToLocalMatrices, Section.BoneMap, RevisionNumber, false, FeatureLevel, bUseSkinCache, DynamicData->bForceUpdateDynamicDataImmediately, OwnerName);
+			ShaderData.UpdateBoneData(RHICmdList, ReferenceToLocalMatrices, Section.BoneMap, RevisionNumber, FeatureLevel, bUseSkinCache, DynamicData->bForceUpdateDynamicDataImmediately, OwnerName);
 			ShaderData.UpdatedFrameNumber = FrameNumberToPrepare;
 
 			// Update uniform buffer for APEX cloth simulation mesh positions and normals
diff --git a/Engine/Source/Runtime/Engine/Private/SkeletalRenderGPUSkin.h b/Engine/Source/Runtime/Engine/Private/SkeletalRenderGPUSkin.h
--- a/Engine/Source/Runtime/Engine/Private/SkeletalRenderGPUSkin.h
+++ b/Engine/Source/Runtime/Engine/Private/SkeletalRenderGPUSkin.h
@@ -352,7 +352,7 @@ public:
 	ENGINE_API virtual void InitResources(USkinnedMeshComponent* InMeshComponent) override;
 	ENGINE_API virtual void ReleaseResources() override;
 	ENGINE_API virtual void Update(int32 LODIndex,USkinnedMeshComponent* InMeshComponent,const FMorphTargetWeightMap& InActiveMorphTargets, const TArray<float>& InMorphTargetWeights, EPreviousBoneTransformUpdateMode PreviousBoneTransformUpdateMode, const FExternalMorphWeightData& InExternalMorphWeightData) override;
-	ENGINE_API void UpdateDynamicData_RenderThread(FGPUSkinCache* GPUSkinCache, FRHICommandList& RHICmdList, FDynamicSkelMeshObjectDataGPUSkin* InDynamicData, FSceneInterface* Scene, uint64 FrameNumberToPrepare, uint32 RevisionNumber);
+	ENGINE_API void UpdateDynamicData_RenderThread(FGPUSkinCache* GPUSkinCache, FRHICommandList& RHICmdList, FDynamicSkelMeshObjectDataGPUSkin* InDynamicData, FSceneInterface* Scene, uint64 FrameNumberToPrepare, uint32 RevisionNumber, uint32 PreviousRevisionNumber);
 	ENGINE_API virtual void PreGDMECallback(FRHICommandList& RHICmdList, FGPUSkinCache* GPUSkinCache, uint32 FrameNumber) override;
 	ENGINE_API virtual const FVertexFactory* GetSkinVertexFactory(const FSceneView* View, int32 LODIndex,int32 ChunkIdx, ESkinVertexFactoryMode VFMode = ESkinVertexFactoryMode::Default) const override;
 	ENGINE_API virtual const FSkinBatchVertexFactoryUserData* GetVertexFactoryUserData(const int32 LODIndex, int32 ChunkIdx, ESkinVertexFactoryMode VFMode) const override;
@@ -665,7 +665,7 @@ protected:
 	*/
 	ENGINE_API void ReleaseMorphResources();
 
-	ENGINE_API void ProcessUpdatedDynamicData(EGPUSkinCacheEntryMode Mode, FGPUSkinCache* GPUSkinCache, FRHICommandList& RHICmdList, uint32 FrameNumberToPrepare, uint32 RevisionNumber, bool bMorphNeedsUpdate, int32 LODIndex);
+	ENGINE_API void ProcessUpdatedDynamicData(EGPUSkinCacheEntryMode Mode, FGPUSkinCache* GPUSkinCache, FRHICommandList& RHICmdList, uint32 FrameNumberToPrepare, uint32 RevisionNumber, uint32 PreviousRevisionNumber, bool bMorphNeedsUpdate, int32 LODIndex);
 
 	ENGINE_API virtual void UpdateMorphVertexBuffer(FRHICommandList& RHICmdList, EGPUSkinCacheEntryMode Mode, FSkeletalMeshObjectLOD& LOD, const FSkeletalMeshLODRenderData& LODData, bool bGPUSkinCacheEnabled, FMorphVertexBuffer& MorphVertexBuffer);
 
diff --git a/Engine/Source/Runtime/Engine/Public/GPUSkinVertexFactory.h b/Engine/Source/Runtime/Engine/Public/GPUSkinVertexFactory.h
--- a/Engine/Source/Runtime/Engine/Public/GPUSkinVertexFactory.h
+++ b/Engine/Source/Runtime/Engine/Public/GPUSkinVertexFactory.h
@@ -202,7 +202,7 @@ public:
 
 		// @param FrameTime from GFrameTime
 		void UpdateBoneData(FRHICommandList& RHICmdList, const TArray<FMatrix44f>& ReferenceToLocalMatrices,
-			const TArray<FBoneIndexType>& BoneMap, uint32 RevisionNumber, bool bPrevious, ERHIFeatureLevel::Type FeatureLevel, 
+			const TArray<FBoneIndexType>& BoneMap, uint32 RevisionNumber, ERHIFeatureLevel::Type FeatureLevel, 
 			bool bUseSkinCache, bool bForceUpdateImmediately, const FName& AssetPathName);
 
 		void ReleaseBoneData()
