Fixes a couple of cases where mid-tick changes in movement mode could continue movement with the old mode's Phys**** function. This could occur in a customized project when root motion causes lift-off into a mode other than MOVE_Falling.

#changelist validated
#virtualized

diff --git a/Engine/Source/Runtime/Engine/Private/Components/CharacterMovementComponent.cpp b/Engine/Source/Runtime/Engine/Private/Components/CharacterMovementComponent.cpp
--- a/Engine/Source/Runtime/Engine/Private/Components/CharacterMovementComponent.cpp
+++ b/Engine/Source/Runtime/Engine/Private/Components/CharacterMovementComponent.cpp
@@ -5427,9 +5427,9 @@ void UCharacterMovementComponent::PhysWalking(float deltaTime, int32 Iterations)
 		ApplyRootMotionToVelocity(timeTick);
 		devCode(ensureMsgf(!Velocity.ContainsNaN(), TEXT("PhysWalking: Velocity contains NaN after Root Motion application (%s)\n%s"), *GetPathNameSafe(this), *Velocity.ToString()));
 
-		if( IsFalling() )
+		if (MovementMode != MOVE_Walking)
 		{
-			// Root motion could have put us into Falling.
+			// Root motion could have taken us out of walking mode
 			// No movement has taken place this movement tick so we pass on full time/past iteration count
 			StartNewPhysics(remainingTime+timeTick, Iterations-1);
 			return;
@@ -5450,9 +5450,15 @@ void UCharacterMovementComponent::PhysWalking(float deltaTime, int32 Iterations)
 			// try to move forward
 			MoveAlongFloor(MoveVelocity, timeTick, &StepDownResult);
 
-			if ( IsFalling() )
+			if (IsSwimming()) //just entered water
 			{
-				// pawn decided to jump up
+				StartSwimming(OldLocation, OldVelocity, timeTick, remainingTime, Iterations);
+				return;
+			}
+			else if (MovementMode != MOVE_Walking)
+			{
+				// pawn ended up in a different mode, probably due to the step-up-and-over flow
+				// let's refund the estimated unused time (if any) and keep moving in the new mode
 				const float DesiredDist = Delta.Size();
 				if (DesiredDist > UE_KINDA_SMALL_NUMBER)
 				{
@@ -5462,11 +5468,6 @@ void UCharacterMovementComponent::PhysWalking(float deltaTime, int32 Iterations)
 				StartNewPhysics(remainingTime,Iterations);
 				return;
 			}
-			else if ( IsSwimming() ) //just entered water
-			{
-				StartSwimming(OldLocation, OldVelocity, timeTick, remainingTime, Iterations);
-				return;
-			}
 		}
 
 		// Update floor.
@@ -5625,9 +5626,9 @@ void UCharacterMovementComponent::PhysNavWalking(float deltaTime, int32 Iteratio
 
 	ApplyRootMotionToVelocity(deltaTime);
 
-	if( IsFalling() )
+	if (MovementMode != MOVE_NavWalking)
 	{
-		// Root motion could have put us into Falling
+		// Root motion could have taken us out of walking mode
 		StartNewPhysics(deltaTime, Iterations);
 		return;
 	}
