Added stats for audio commands being run on the game thread.

#tests local win64 test client

#changelist validated
#virtualized

diff --git a/Engine/Source/Runtime/Engine/Private/AudioThread.cpp b/Engine/Source/Runtime/Engine/Private/AudioThread.cpp
--- a/Engine/Source/Runtime/Engine/Private/AudioThread.cpp
+++ b/Engine/Source/Runtime/Engine/Private/AudioThread.cpp
@@ -7,6 +7,7 @@
 #include "AudioThread.h"
 #include "Audio.h"
 #include "Async/Async.h"
+#include "ProfilingDebugging/CsvProfiler.h"
 #include "Tasks/Pipe.h"
 
 //
@@ -364,11 +365,23 @@ void FAudioThread::RunCommandOnGameThread(TUniqueFunction<void()> InFunction, co
 	if (IsAudioThreadRunning())
 	{
 		check(IsInAudioThread());
-		FFunctionGraphTask::CreateAndDispatchWhenReady(MoveTemp(InFunction), InStatId, nullptr, ENamedThreads::GameThread);
+		FFunctionGraphTask::CreateAndDispatchWhenReady(
+			[Function = MoveTemp(InFunction), InStatId]()
+			{
+				CSV_SCOPED_TIMING_STAT_EXCLUSIVE(Audio);
+				QUICK_SCOPE_CYCLE_COUNTER(STAT_AudioThread_RunCommandOnGameThread);
+				FScopeCycleCounter ScopeCycleCounter(InStatId);
+				Function();
+			},
+			TStatId(),
+			nullptr,
+			ENamedThreads::GameThread);
 	}
 	else
 	{
 		check(IsInGameThread());
+		CSV_SCOPED_TIMING_STAT_EXCLUSIVE(Audio);
+		QUICK_SCOPE_CYCLE_COUNTER(STAT_AudioThread_RunCommandOnGameThread);
 		FScopeCycleCounter ScopeCycleCounter(InStatId);
 		InFunction();
 	}
