Skeletal Mesh Import: Retain deformer asset on skeletal mesh re-import

#changelist validated
#virtualized

diff --git a/Engine/Source/Editor/UnrealEd/Private/ImportUtils/SkeletalMeshImportUtils.cpp b/Engine/Source/Editor/UnrealEd/Private/ImportUtils/SkeletalMeshImportUtils.cpp
--- a/Engine/Source/Editor/UnrealEd/Private/ImportUtils/SkeletalMeshImportUtils.cpp
+++ b/Engine/Source/Editor/UnrealEd/Private/ImportUtils/SkeletalMeshImportUtils.cpp
@@ -404,6 +404,7 @@ TSharedPtr<FExistingSkelMeshData> SkeletalMeshImportUtils::SaveExistingSkelMeshD
 	ExistingMeshDataPtr->ExistingClothingAssets = SourceSkeletalMesh->GetMeshClothingAssets();
 	ExistingMeshDataPtr->ExistingSamplingInfo = SourceSkeletalMesh->GetSamplingInfo();
 	ExistingMeshDataPtr->ExistingDefaultAnimatingRig = SourceSkeletalMesh->GetDefaultAnimatingRig();
+	ExistingMeshDataPtr->ExistingDefaultMeshDeformer = SourceSkeletalMesh->GetDefaultMeshDeformer();
 
 	if (ExistingMeshDataPtr->UseMaterialNameSlotWorkflow)
 	{
@@ -914,6 +915,7 @@ void SkeletalMeshImportUtils::RestoreExistingSkelMeshData(const TSharedPtr<const
 		SkeletalMesh->SetSkeleton(MeshData->ExistingSkeleton);
 		SkeletalMesh->SetPostProcessAnimBlueprint(MeshData->ExistingPostProcessAnimBlueprint);
 		SkeletalMesh->SetDefaultAnimatingRig(MeshData->ExistingDefaultAnimatingRig);
+		SkeletalMesh->SetDefaultMeshDeformer(MeshData->ExistingDefaultMeshDeformer);
 
 		SkeletalMesh->GetMorphTargets().Empty(MeshData->ExistingMorphTargets.Num());
 		SkeletalMesh->GetMorphTargets().Append(MeshData->ExistingMorphTargets);
diff --git a/Engine/Source/Editor/UnrealEd/Public/ImportUtils/SkelImport.h b/Engine/Source/Editor/UnrealEd/Public/ImportUtils/SkelImport.h
--- a/Engine/Source/Editor/UnrealEd/Public/ImportUtils/SkelImport.h
+++ b/Engine/Source/Editor/UnrealEd/Public/ImportUtils/SkelImport.h
@@ -85,6 +85,7 @@ struct FExistingSkelMeshData
 	USkeletalMeshLODSettings*				ExistingLODSettings;
 	TSubclassOf<UAnimInstance>				ExistingPostProcessAnimBlueprint;
 	TSoftObjectPtr<UObject>					ExistingDefaultAnimatingRig;
+	UMeshDeformer*							ExistingDefaultMeshDeformer;
 	//////////////////////////////////////////////////////////////////////////
 
 	bool									bExistingUseFullPrecisionUVs;
diff --git a/Engine/Source/Runtime/Engine/Classes/Engine/SkeletalMesh.h b/Engine/Source/Runtime/Engine/Classes/Engine/SkeletalMesh.h
--- a/Engine/Source/Runtime/Engine/Classes/Engine/SkeletalMesh.h
+++ b/Engine/Source/Runtime/Engine/Classes/Engine/SkeletalMesh.h
@@ -2951,14 +2951,6 @@ public:
 		PRAGMA_ENABLE_DEPRECATION_WARNINGS
 	}
 
-	virtual UMeshDeformer* GetDefaultMeshDeformer() const override
-	{
-		WaitUntilAsyncPropertyReleased(ESkeletalMeshAsyncProperties::DefaultMeshDeformer, ESkinnedAssetAsyncPropertyLockType::ReadOnly);
-		PRAGMA_DISABLE_DEPRECATION_WARNINGS
-		return DefaultMeshDeformer;
-		PRAGMA_ENABLE_DEPRECATION_WARNINGS
-	}
-
 	TArray<FSkinWeightProfileInfo>& GetSkinWeightProfiles() 
 	{
 		WaitUntilAsyncPropertyReleased(ESkeletalMeshAsyncProperties::SkinWeightProfiles); 
@@ -2987,7 +2979,7 @@ protected:
 
 	/** Default mesh deformer to use with this mesh. */
 	UE_DEPRECATED(5.1, "This must be protected for async build, always use the accessors even internally.")
-	UPROPERTY(EditAnywhere, BlueprintReadOnly, Category = "Deformer", Meta=(Filter="/Script/Engine.SkinnedMeshComponent"))
+	UPROPERTY(EditAnywhere, BlueprintReadOnly, Category = "Deformer", BlueprintGetter=GetDefaultMeshDeformer)
 	TObjectPtr<UMeshDeformer> DefaultMeshDeformer;
 
 	/** Default translucent material to blend on top of this mesh. Mesh will be rendered twice - once with a base material and once with overlay material */
@@ -3001,6 +2993,27 @@ protected:
 	float OverlayMaterialMaxDrawDistance;
 
 public:
+	/** Set the default mesh deformer used by this mesh. A mesh deformer is used to deform the skeletal mesh at runtime.
+	 *  The default deformer can be overridden by the USkinnedMeshComponent object. */
+	void SetDefaultMeshDeformer(UMeshDeformer* InMeshDeformer)
+	{
+		WaitUntilAsyncPropertyReleased(ESkeletalMeshAsyncProperties::DefaultMeshDeformer, ESkinnedAssetAsyncPropertyLockType::WriteOnly);
+		PRAGMA_DISABLE_DEPRECATION_WARNINGS
+		DefaultMeshDeformer = InMeshDeformer;
+		PRAGMA_ENABLE_DEPRECATION_WARNINGS		
+	}
+	
+	/** Get the default mesh deformer used by this mesh. A mesh deformer is used to deform the skeletal mesh at runtime */
+	UFUNCTION(BlueprintCallable, Category="Rendering|Deformer")
+	virtual UMeshDeformer* GetDefaultMeshDeformer() const override
+	{
+		WaitUntilAsyncPropertyReleased(ESkeletalMeshAsyncProperties::DefaultMeshDeformer, ESkinnedAssetAsyncPropertyLockType::ReadOnly);
+		PRAGMA_DISABLE_DEPRECATION_WARNINGS
+		return DefaultMeshDeformer;
+		PRAGMA_ENABLE_DEPRECATION_WARNINGS
+	}
+
+	
 	/** Get the default overlay material used by this mesh */
 	UFUNCTION(BlueprintCallable, Category="Rendering|Material")
 	virtual class UMaterialInterface* GetOverlayMaterial() const override
