Fix animation attribute registration issues with attributes defined as user-defined structs

#changelist validated
#virtualized

diff --git a/Engine/Source/Runtime/Engine/Private/Animation/AttributeTypes.cpp b/Engine/Source/Runtime/Engine/Private/Animation/AttributeTypes.cpp
--- a/Engine/Source/Runtime/Engine/Private/Animation/AttributeTypes.cpp
+++ b/Engine/Source/Runtime/Engine/Private/Animation/AttributeTypes.cpp
@@ -11,40 +11,41 @@
 namespace UE
 {
 	namespace Anim
-	{		
+	{
 		TArray<TWeakObjectPtr<const UScriptStruct>> AttributeTypes::RegisteredTypes;
 		TArray<TUniquePtr<IAttributeBlendOperator>> AttributeTypes::Operators;
 		TArray<TWeakObjectPtr<const UScriptStruct>> AttributeTypes::InterpolatableTypes;
-		std::atomic<bool> AttributeTypes::bInitialized = false;		
 		AttributeTypes::FOnAttributeTypesChanged AttributeTypes::OnAttributeTypesChangedDelegate;
-		
-		void AttributeTypes::LazyInitialize()
+
+		struct FAttributeTypeRegistrar
 		{
-			bool bWasUninitialized = false;
-			if (bInitialized.compare_exchange_strong(bWasUninitialized, true))
+			static void RegisterBuiltInTypes()
 			{
-				Initialize();
+				AttributeTypes::RegisterType<FFloatAnimationAttribute>();
+				AttributeTypes::RegisterType<FIntegerAnimationAttribute>();
+				AttributeTypes::RegisterType<FStringAnimationAttribute>();
+				AttributeTypes::RegisterType<FTransformAnimationAttribute>();
+				AttributeTypes::RegisterType<FVectorAnimationAttribute>();
+				AttributeTypes::RegisterType<FQuaternionAnimationAttribute>();
 			}
-		}
 
-		void AttributeTypes::Initialize()
-		{
-			RegisterType<FFloatAnimationAttribute>();
-			RegisterType<FIntegerAnimationAttribute>();
-			RegisterType<FStringAnimationAttribute>();
-			RegisterType<FTransformAnimationAttribute>();
-			RegisterType<FVectorAnimationAttribute>();
-			RegisterType<FQuaternionAnimationAttribute>();
-
-			for (const TSoftObjectPtr<UUserDefinedStruct>& UserDefinedStruct : UAnimationSettings::Get()->UserDefinedStructAttributes)
+			static void RegisterUserDefinedStructTypes()
 			{
-				UE::Anim::AttributeTypes::RegisterNonBlendableType(UserDefinedStruct.LoadSynchronous());
+				for (const TSoftObjectPtr<UUserDefinedStruct>& UserDefinedStruct : UAnimationSettings::Get()->UserDefinedStructAttributes)
+				{
+					AttributeTypes::RegisterNonBlendableType(UserDefinedStruct.LoadSynchronous());
+				}
 			}
-		}
+		};
+
+		static FDelayedAutoRegisterHelper DelayedBuiltInTypesInitializationHelper(EDelayedRegisterRunPhase::ObjectSystemReady, []()
+		{
+			FAttributeTypeRegistrar::RegisterBuiltInTypes();
+		});
 		
-		static FDelayedAutoRegisterHelper DelayedAttributeTypesInitializationHelper(EDelayedRegisterRunPhase::ObjectSystemReady, []()
+		static FDelayedAutoRegisterHelper DelayedUserDefinedStructTypesInitializationHelper(EDelayedRegisterRunPhase::EndOfEngineInit, []()
 		{
-			UE::Anim::AttributeTypes::LazyInitialize();
+			FAttributeTypeRegistrar::RegisterUserDefinedStructTypes();
 		});
 
 		bool AttributeTypes::RegisterNonBlendableType(const UScriptStruct* InScriptStruct)
diff --git a/Engine/Source/Runtime/Engine/Public/Animation/AttributeTypes.h b/Engine/Source/Runtime/Engine/Public/Animation/AttributeTypes.h
--- a/Engine/Source/Runtime/Engine/Public/Animation/AttributeTypes.h
+++ b/Engine/Source/Runtime/Engine/Public/Animation/AttributeTypes.h
@@ -50,19 +50,17 @@ namespace UE
 			static ENGINE_API TArray<TWeakObjectPtr<const UScriptStruct>> RegisteredTypes;
 			static ENGINE_API TArray<TUniquePtr<IAttributeBlendOperator>> Operators;
 			static ENGINE_API TArray<TWeakObjectPtr<const UScriptStruct>> InterpolatableTypes;
-			static ENGINE_API std::atomic<bool> bInitialized;
 			static ENGINE_API FOnAttributeTypesChanged OnAttributeTypesChangedDelegate;
 
-			static ENGINE_API void Initialize();
-
 			/** Register user defined structs as non-blendable animation attribute */
 			static ENGINE_API bool RegisterNonBlendableType(const UScriptStruct* InScriptStruct);
 
 			/** Unregisters a specific attribute type and deletes its associated blend operator */
 			static ENGINE_API void UnregisterType(const UScriptStruct* InScriptStruct);
 			
-		public:			
-			static ENGINE_API void LazyInitialize();
+		public:
+			UE_DEPRECATED(5.4, "This function is no longer used, the system will initialize built-in and user-defined types at the appropriate times")
+			static void LazyInitialize() {}
 
 			static FOnAttributeTypesChanged& GetOnAttributeTypesChanged() { return OnAttributeTypesChangedDelegate; };
 			
@@ -128,7 +126,6 @@ namespace UE
 			/** Returns the blend operator for the provided type, asserts when the type is not registered */
 			static const IAttributeBlendOperator* GetTypeOperator(TWeakObjectPtr<const UScriptStruct> WeakStruct)
 			{
-				LazyInitialize();
 				const int32 Index = AttributeTypes::RegisteredTypes.IndexOfByKey(WeakStruct);
 				ensure(WeakStruct.IsValid());
 				checkf(Index != INDEX_NONE, TEXT("Missing operator for attribute, type %s was not registered previously"), *WeakStruct->GetName());
@@ -138,25 +135,23 @@ namespace UE
 			/** Returns whether or not the provided type can be interpolated, defaults to false when the type is not registered */
 			static bool CanInterpolateType(TWeakObjectPtr<const UScriptStruct> WeakStruct)
 			{
-				LazyInitialize();
 				return AttributeTypes::InterpolatableTypes.Contains(WeakStruct);
 			}
 
 			/** Returns whether or not the type is registered */
 			static bool IsTypeRegistered(const UScriptStruct* ScriptStruct)
 			{
-				LazyInitialize();
 				return AttributeTypes::RegisteredTypes.Contains(ScriptStruct);
 			}
 
 			/** Returns all registered types */
 			static TArray<TWeakObjectPtr<const UScriptStruct>>& GetRegisteredTypes()
 			{
-				LazyInitialize();
 				return AttributeTypes::RegisteredTypes;
 			}
 			
 			friend class ::UAnimationSettings;
+			friend struct FAttributeTypeRegistrar;
 		};
 	}
 }
