Use easier to read units in shader compile worker memory printouts

These numbers are usually in the billion byte range, so we might as well show them in GiB for easier parsing.

#changelist validated
#virtualized

diff --git a/Engine/Source/Runtime/Engine/Private/ShaderCompiler/ShaderCompiler.cpp b/Engine/Source/Runtime/Engine/Private/ShaderCompiler/ShaderCompiler.cpp
--- a/Engine/Source/Runtime/Engine/Private/ShaderCompiler/ShaderCompiler.cpp
+++ b/Engine/Source/Runtime/Engine/Private/ShaderCompiler/ShaderCompiler.cpp
@@ -2435,22 +2435,23 @@ namespace ShaderCompileWorkerError
 
 	bool HandleOutOfMemory(const TCHAR* ExceptionInfo, const TCHAR* Hostname, const FPlatformMemoryStats& MemoryStats, const TArray<FShaderCommonCompileJobPtr>& QueuedJobs)
 	{
+		constexpr int64 Gibibyte = 1024 * 1024 * 1024;
 		const FString ErrorReport = FString::Printf(
 			TEXT("ShaderCompileWorker failed with out-of-memory (OOM) exception on machine \"%s\" (%s); MemoryStats:")
-			TEXT("\n\tAvailablePhysical %llu")
-			TEXT("\n\t AvailableVirtual %llu")
-			TEXT("\n\t     UsedPhysical %llu")
-			TEXT("\n\t PeakUsedPhysical %llu")
-			TEXT("\n\t      UsedVirtual %llu")
-			TEXT("\n\t  PeakUsedVirtual %llu"),
+			TEXT("\n\tAvailablePhysical %llu (%.2f GiB)")
+			TEXT("\n\t AvailableVirtual %llu (%.2f GiB)")
+			TEXT("\n\t     UsedPhysical %llu (%.2f GiB)")
+			TEXT("\n\t PeakUsedPhysical %llu (%.2f GiB)")
+			TEXT("\n\t      UsedVirtual %llu (%.2f GiB)")
+			TEXT("\n\t  PeakUsedVirtual %llu (%.2f GiB)"),
 			Hostname,
 			(ExceptionInfo[0] == TEXT('\0') ? TEXT("No exception information") : ExceptionInfo),
-			MemoryStats.AvailablePhysical,
-			MemoryStats.AvailableVirtual,
-			MemoryStats.UsedPhysical,
-			MemoryStats.PeakUsedPhysical,
-			MemoryStats.UsedVirtual,
-			MemoryStats.PeakUsedVirtual
+			MemoryStats.AvailablePhysical, double(MemoryStats.AvailablePhysical) / Gibibyte,
+			MemoryStats.AvailableVirtual, double(MemoryStats.AvailableVirtual) / Gibibyte,
+			MemoryStats.UsedPhysical, double(MemoryStats.UsedPhysical) / Gibibyte,
+			MemoryStats.PeakUsedPhysical, double(MemoryStats.PeakUsedPhysical) / Gibibyte,
+			MemoryStats.UsedVirtual, double(MemoryStats.UsedVirtual) / Gibibyte,
+			MemoryStats.PeakUsedVirtual, double(MemoryStats.PeakUsedVirtual) / Gibibyte
 		);
 
 		if (GShaderCompilingManager->IsRemoteCompilingEnabled())
@@ -4074,6 +4075,7 @@ void FShaderCompileThreadRunnable::PrintWorkerMemoryUsage()
 	FScopeLock WorkerScopeLock(&WorkerInfosLock);
 	FPlatformProcessMemoryStats TotalMemoryStats{};
 	int32 NumValidWorkers = 0;
+	constexpr int64 Gibibyte = 1024 * 1024 * 1024;
 	for (int32 Iter = 0, End = WorkerInfos.Num(); Iter < End; Iter++)
 	{
 		const TUniquePtr<FShaderCompileWorkerInfo>& WorkerInfo = WorkerInfos[Iter];
@@ -4088,16 +4090,16 @@ void FShaderCompileThreadRunnable::PrintWorkerMemoryUsage()
 			NumValidWorkers++;
 			UE_LOG(LogShaderCompilers, Display,
 				TEXT("ShaderCompileWorker [%d/%d] MemoryStats:")
-				TEXT("\n\t     UsedPhysical %llu")
-				TEXT("\n\t PeakUsedPhysical %llu")
-				TEXT("\n\t      UsedVirtual %llu")
-				TEXT("\n\t  PeakUsedVirtual %llu"),
+				TEXT("\n\t     UsedPhysical %llu (%.2f GiB)")
+				TEXT("\n\t PeakUsedPhysical %llu (%.2f GiB)")
+				TEXT("\n\t      UsedVirtual %llu (%.2f GiB)")
+				TEXT("\n\t  PeakUsedVirtual %llu (%.2f GiB)"),
 				Iter + 1,
 				End,
-				MemoryStats.UsedPhysical,
-				MemoryStats.PeakUsedPhysical,
-				MemoryStats.UsedVirtual,
-				MemoryStats.PeakUsedVirtual
+				MemoryStats.UsedPhysical, double(MemoryStats.UsedPhysical) / Gibibyte,
+				MemoryStats.PeakUsedPhysical, double(MemoryStats.PeakUsedPhysical) / Gibibyte,
+				MemoryStats.UsedVirtual, double(MemoryStats.UsedVirtual) / Gibibyte,
+				MemoryStats.PeakUsedVirtual, double(MemoryStats.PeakUsedVirtual) / Gibibyte
 			);
 			TotalMemoryStats.UsedPhysical += MemoryStats.UsedPhysical;
 			TotalMemoryStats.PeakUsedPhysical += MemoryStats.PeakUsedPhysical;
@@ -4111,15 +4113,15 @@ void FShaderCompileThreadRunnable::PrintWorkerMemoryUsage()
 	{
 		UE_LOG(LogShaderCompilers, Display,
 			TEXT("Sum of MemoryStats for %d ShaderCompileWorker(s):")
-			TEXT("\n\t     UsedPhysical %llu")
-			TEXT("\n\t PeakUsedPhysical %llu")
-			TEXT("\n\t      UsedVirtual %llu")
-			TEXT("\n\t  PeakUsedVirtual %llu"),
+			TEXT("\n\t     UsedPhysical %llu (%.2f GiB)")
+			TEXT("\n\t PeakUsedPhysical %llu (%.2f GiB)")
+			TEXT("\n\t      UsedVirtual %llu (%.2f GiB)")
+			TEXT("\n\t  PeakUsedVirtual %llu (%.2f GiB)"),
 			NumValidWorkers,
-			TotalMemoryStats.UsedPhysical,
-			TotalMemoryStats.PeakUsedPhysical,
-			TotalMemoryStats.UsedVirtual,
-			TotalMemoryStats.PeakUsedVirtual
+			TotalMemoryStats.UsedPhysical, double(TotalMemoryStats.UsedPhysical) / Gibibyte,
+			TotalMemoryStats.PeakUsedPhysical, double(TotalMemoryStats.PeakUsedPhysical) / Gibibyte,
+			TotalMemoryStats.UsedVirtual, double(TotalMemoryStats.UsedVirtual) / Gibibyte,
+			TotalMemoryStats.PeakUsedVirtual, double(TotalMemoryStats.PeakUsedVirtual) / Gibibyte
 		);
 	}
 }
