World Partition Actor Descriptor Instance Bounds: editor and runtime bounds now returns world space bounds.

#changelist validated
#virtualized

diff --git a/Engine/Source/Runtime/Engine/Private/WorldPartition/LoaderAdapter/LoaderAdapterSpatial.cpp b/Engine/Source/Runtime/Engine/Private/WorldPartition/LoaderAdapter/LoaderAdapterSpatial.cpp
--- a/Engine/Source/Runtime/Engine/Private/WorldPartition/LoaderAdapter/LoaderAdapterSpatial.cpp
+++ b/Engine/Source/Runtime/Engine/Private/WorldPartition/LoaderAdapter/LoaderAdapterSpatial.cpp
@@ -31,12 +31,9 @@ void ILoaderAdapterSpatial::HandleIntersectingContainer(UWorldPartition* InWorld
 		.SetIncludeSpatiallyLoadedActors(bIncludeSpatiallyLoadedActors)
 		.SetIncludeNonSpatiallyLoadedActors(bIncludeNonSpatiallyLoadedActors);
 
-	const FTransform InstanceTransform = InWorldPartition->GetInstanceTransform();
-	const FBox LocalBoundingBox = InBoundingBox.InverseTransformBy(InstanceTransform);
-	InWorldPartition->EditorHash->ForEachIntersectingActor(LocalBoundingBox, [this, InWorldPartition, &InstanceTransform, &InOperation](FWorldPartitionActorDescInstance* ActorDescInstance)
+	InWorldPartition->EditorHash->ForEachIntersectingActor(InBoundingBox, [this, InWorldPartition, &InOperation](FWorldPartitionActorDescInstance* ActorDescInstance)
 	{
-		const FBox WorldActorEditorBox = ActorDescInstance->GetEditorBounds().TransformBy(InstanceTransform);
-		if (Intersect(WorldActorEditorBox))
+		if (Intersect(ActorDescInstance->GetEditorBounds()))
 		{
 			FWorldPartitionHandle ActorHandle(InWorldPartition, ActorDescInstance->GetGuid());
 			InOperation(ActorHandle);
diff --git a/Engine/Source/Runtime/Engine/Private/WorldPartition/WorldPartitionActorDescInstance.cpp b/Engine/Source/Runtime/Engine/Private/WorldPartition/WorldPartitionActorDescInstance.cpp
--- a/Engine/Source/Runtime/Engine/Private/WorldPartition/WorldPartitionActorDescInstance.cpp
+++ b/Engine/Source/Runtime/Engine/Private/WorldPartition/WorldPartitionActorDescInstance.cpp
@@ -67,7 +67,7 @@ TWeakObjectPtr<AActor>* FWorldPartitionActorDescInstance::GetActorPtr(bool bEven
 
 FSoftObjectPath FWorldPartitionActorDescInstance::GetActorSoftPath() const
 {
-	return ActorPath.IsSet() ? ActorPath.GetValue() : GetActorDesc()->GetActorSoftPath();
+	return ActorPath.IsSet() ? ActorPath.GetValue() : ActorDesc->GetActorSoftPath();
 }
 
 FName FWorldPartitionActorDescInstance::GetActorName() const
@@ -77,17 +77,27 @@ FName FWorldPartitionActorDescInstance::GetActorName() const
 
 bool FWorldPartitionActorDescInstance::IsValid() const
 {
-	return !!GetActorDesc();
+	return !!ActorDesc;
 }
 
 bool FWorldPartitionActorDescInstance::IsEditorRelevant() const
 {
-	return GetActorDesc()->IsEditorRelevant(this);
+	return ActorDesc->IsEditorRelevant(this);
 }
 
 bool FWorldPartitionActorDescInstance::IsRuntimeRelevant() const
 {
-	return GetActorDesc()->IsRuntimeRelevant(this);
+	return ActorDesc->IsRuntimeRelevant(this);
+}
+
+FBox FWorldPartitionActorDescInstance::GetEditorBounds() const
+{
+	return ActorDesc->GetEditorBounds().TransformBy(GetContainerInstance()->GetTransform());
+}
+
+FBox FWorldPartitionActorDescInstance::GetRuntimeBounds() const
+{
+	return ActorDesc->GetRuntimeBounds().TransformBy(GetContainerInstance()->GetTransform());
 }
 
 bool FWorldPartitionActorDescInstance::StartAsyncLoad()
@@ -152,7 +162,7 @@ void FWorldPartitionActorDescInstance::MarkUnload()
 	FlushAsyncLoad();
 
 	// Notify Desc as it can have some custom code to run on the actor depending on type
-	GetActorDesc()->OnUnloadingInstance(this);
+	ActorDesc->OnUnloadingInstance(this);
 
 	if (AActor* Actor = GetActor())
 	{
@@ -206,7 +216,7 @@ const FText& FWorldPartitionActorDescInstance::GetUnloadedReason() const
 
 FString FWorldPartitionActorDescInstance::ToString(FWorldPartitionActorDesc::EToStringMode Mode) const
 {
-	return GetActorDesc()->ToString(Mode);
+	return ActorDesc->ToString(Mode);
 }
 
 void FWorldPartitionActorDescInstance::RegisterChildContainerInstance()
@@ -214,7 +224,7 @@ void FWorldPartitionActorDescInstance::RegisterChildContainerInstance()
 	check(IsChildContainerInstance());
 	check(!ChildContainerInstance);
 
-	ChildContainerInstance = GetActorDesc()->CreateChildContainerInstance(this);
+	ChildContainerInstance = ActorDesc->CreateChildContainerInstance(this);
 	check(ChildContainerInstance);
 	ContainerInstance->OnRegisterChildContainerInstance(GetGuid(), ChildContainerInstance);
 }
@@ -231,7 +241,7 @@ void FWorldPartitionActorDescInstance::UpdateChildContainerInstance()
 {
 	// Create before unregistering so that we benefit from shared containers (use GetActorDesc->IsChildContainerInstance as we want to know if our updated desc should be a Container instance or not)
 	// ChildContainerInstance member might be non null and we don't want IsChildContainerInstance() to return true in this case if the ActorDesc isn't a Container anymore
-	UActorDescContainerInstance* NewChildContainerInstance = GetActorDesc()->IsChildContainerInstance() ? GetActorDesc()->CreateChildContainerInstance(this) : nullptr;
+	UActorDescContainerInstance* NewChildContainerInstance = ActorDesc->IsChildContainerInstance() ? ActorDesc->CreateChildContainerInstance(this) : nullptr;
 	
 	// Unregister previous
 	if (ChildContainerInstance)
diff --git a/Engine/Source/Runtime/Engine/Private/WorldPartition/WorldPartitionStreamingGeneration.cpp b/Engine/Source/Runtime/Engine/Private/WorldPartition/WorldPartitionStreamingGeneration.cpp
--- a/Engine/Source/Runtime/Engine/Private/WorldPartition/WorldPartitionStreamingGeneration.cpp
+++ b/Engine/Source/Runtime/Engine/Private/WorldPartition/WorldPartitionStreamingGeneration.cpp
@@ -501,14 +501,13 @@ class FWorldPartitionStreamingGenerator
 						ActorSetInstance.DataLayers = StreamingGenerator->GetRuntimeDataLayerInstances(PerInstanceData.DataLayers);
 
 						ActorSetInstance.Bounds.Init();
-						const FTransform& ContainerCollectionInstanceDescriptorTransform = ContainerCollectionInstanceDescriptor.Transform;
-						ActorSetInstance.ForEachActor([this, &ActorSetInstance, &ContainerCollectionInstanceDescriptorTransform](const FGuid& ActorGuid)
+						ActorSetInstance.ForEachActor([this, &ActorSetInstance](const FGuid& ActorGuid)
 						{
 							const FStreamingGenerationActorDescView& ActorDescView = ActorSetInstance.ActorSetContainerInstance->ActorDescViewMap->FindByGuidChecked(ActorGuid);
 							const FBox RuntimeBounds = ActorDescView.GetRuntimeBounds();
 							if (RuntimeBounds.IsValid)
 							{
-								ActorSetInstance.Bounds += RuntimeBounds.TransformBy(ContainerCollectionInstanceDescriptorTransform);
+								ActorSetInstance.Bounds += RuntimeBounds;
 							}
 						});
 					}
@@ -874,7 +873,7 @@ class FWorldPartitionStreamingGenerator
 						const FBox RuntimeBounds = ActorDescView.GetRuntimeBounds();
 						check(RuntimeBounds.IsValid);
 
-						ContainerCollectionInstanceDescriptor.Bounds += RuntimeBounds.TransformBy(ContainerCollectionInstanceDescriptor.Transform);
+						ContainerCollectionInstanceDescriptor.Bounds += RuntimeBounds;
 					}
 				});
 
@@ -2057,7 +2056,7 @@ FStreamingGenerationContainerInstanceCollection::FStreamingGenerationContainerIn
 
 UWorld* FStreamingGenerationContainerInstanceCollection::GetWorld() const
 {
-	UWorldPartition* WorldPartition = GetBaseContainerInstance()->GetWorldPartition();
+	UWorldPartition* WorldPartition = GetBaseContainerInstance()->GetOuterWorldPartition();
 	check(WorldPartition);
 	UWorld* World = WorldPartition->GetWorld();
 	check(World);
@@ -2231,7 +2230,7 @@ void UWorldPartition::CheckForErrors(const FCheckForErrorsParams& InParams)
 	{
 		TErrorHandlerSelector<FStreamingGenerationLogErrorHandler> ErrorHandlerSelector(InParams.ErrorHandler);
 		const UActorDescContainerInstance* BaseContainerInstance = InCollection.GetBaseContainerInstance();
-		const UWorldPartition* WorldPartition = BaseContainerInstance->GetWorldPartition();
+		const UWorldPartition* WorldPartition = BaseContainerInstance->GetOuterWorldPartition();
 		const TObjectPtr<UWorldPartitionRuntimeHash>& WorldPartitionRuntimeHash = WorldPartition ? WorldPartition->RuntimeHash : nullptr;
 
 		FWorldPartitionStreamingGenerator::FWorldPartitionStreamingGeneratorParams StreamingGeneratorParams = FWorldPartitionStreamingGenerator::FWorldPartitionStreamingGeneratorParams()
@@ -2267,4 +2266,4 @@ void UWorldPartition::CheckForErrors(const FCheckForErrorsParams& InParams)
 
 #endif // WITH_EDITOR
 
-#undef LOCTEXT_NAMESPACE
+#undef LOCTEXT_NAMESPACE
\ No newline at end of file
diff --git a/Engine/Source/Runtime/Engine/Public/WorldPartition/WorldPartitionActorDescInstance.h b/Engine/Source/Runtime/Engine/Public/WorldPartition/WorldPartitionActorDescInstance.h
--- a/Engine/Source/Runtime/Engine/Public/WorldPartition/WorldPartitionActorDescInstance.h
+++ b/Engine/Source/Runtime/Engine/Public/WorldPartition/WorldPartitionActorDescInstance.h
@@ -61,8 +61,8 @@ public:
 	virtual FName GetFolderPath() const { return ActorDesc->GetFolderPath(); }
 	virtual const FGuid& GetFolderGuid() const { return ActorDesc->GetFolderGuid(); }
 
-	virtual FBox GetEditorBounds() const { return ActorDesc->GetEditorBounds(); }
-	virtual FBox GetRuntimeBounds() const { return ActorDesc->GetRuntimeBounds(); }
+	ENGINE_API virtual FBox GetEditorBounds() const;
+	ENGINE_API virtual FBox GetRuntimeBounds() const;
 
 	virtual bool GetProperty(FName PropertyName, FName* PropertyValue) const { return ActorDesc->GetProperty(PropertyName, PropertyValue); }
 	virtual bool HasProperty(FName PropertyName) const { return ActorDesc->HasProperty(PropertyName); }
