Always write out the memory estimate when cooking a static mesh

Also adds a missing debug name to one of the cook tasks. For textures, we are already always logging the estimated memory usage. Having that available is a huge timesaver when diagnosing out-of-memory conditions on remote machines.

#changelist validated
#virtualized

diff --git a/Engine/Source/Runtime/Engine/Classes/Engine/StaticMesh.h b/Engine/Source/Runtime/Engine/Classes/Engine/StaticMesh.h
--- a/Engine/Source/Runtime/Engine/Classes/Engine/StaticMesh.h
+++ b/Engine/Source/Runtime/Engine/Classes/Engine/StaticMesh.h
@@ -2059,11 +2059,6 @@ private:
 	 */
 	void FinishBuildInternal(const TArray<IStaticMeshComponent*>& InAffectedComponents, bool bHasRenderDataChanged, bool bShouldComputeExtendedBounds = true);
 
-	/**
-	 * Get an estimate of the peak amount of memory required to build this mesh.
-	 */
-	int64 GetBuildRequiredMemory() const;
-
 #if WITH_EDITORONLY_DATA
 	/**
 	 * Deserialize MeshDescription for the specified LodIndex from BulkData, DDC or RawMesh.
@@ -2077,6 +2072,11 @@ public:
 	 */
 	ENGINE_API void CacheDerivedData();
 
+	/**
+	 * Get an estimate of the peak amount of memory required to build this mesh.
+	 */
+	ENGINE_API int64 GetBuildRequiredMemoryEstimate() const;
+
 	/**
 	 * Caches derived renderable for cooked platforms currently active.
 	 */
diff --git a/Engine/Source/Runtime/Engine/Private/StaticMesh.cpp b/Engine/Source/Runtime/Engine/Private/StaticMesh.cpp
--- a/Engine/Source/Runtime/Engine/Private/StaticMesh.cpp
+++ b/Engine/Source/Runtime/Engine/Private/StaticMesh.cpp
@@ -3124,7 +3124,8 @@ void FStaticMeshRenderData::Cache(const ITargetPlatform* TargetPlatform, UStatic
 		{
 			FFormatNamedArguments Args;
 			Args.Add(TEXT("StaticMeshName"), FText::FromString( Owner->GetName() ) );
-			FStaticMeshStatusMessageContext StatusContext( FText::Format( NSLOCTEXT("Engine", "BuildingStaticMeshStatus", "Building static mesh {StaticMeshName}..."), Args ) );
+			Args.Add(TEXT("EstimatedMemory"), FText::FromString( FString::SanitizeFloat(double(Owner->GetBuildRequiredMemoryEstimate()) / (1024.0 * 1024.0), 3) ));
+			FStaticMeshStatusMessageContext StatusContext( FText::Format( NSLOCTEXT("Engine", "BuildingStaticMeshStatus", "Building static mesh {StaticMeshName} (Required Memory Estimate: {EstimatedMemory} MB)..."), Args ) );
 
 			checkf(!Owner->HasAnyFlags(RF_NeedLoad), TEXT("StaticMesh %s being PostLoaded before having been serialized - this suggests an async loading problem."), *GetPathNameSafe(Owner));
 			checkf(Owner->IsMeshDescriptionValid(0), TEXT("Bad MeshDescription on %s"), *GetPathNameSafe(Owner));
@@ -6006,7 +6007,7 @@ void UStaticMesh::ReleaseAsyncProperty(EStaticMeshAsyncProperties AsyncPropertie
 	LockedProperties &= ~(uint32)AsyncProperties;
 }
 
-int64 UStaticMesh::GetBuildRequiredMemory() const
+int64 UStaticMesh::GetBuildRequiredMemoryEstimate() const
 {
 	// We have to base our estimate on something accessible and known before the build, for now use the biggest bulk data size.
 	int64 BiggestBulkDataSize = -1;
@@ -6085,7 +6086,7 @@ void UStaticMesh::PostLoad()
 		}
 
 		AsyncTask = MakeUnique<FStaticMeshAsyncBuildTask>(this, MakeUnique<FStaticMeshPostLoadContext>(MoveTemp(Context)));
-		AsyncTask->StartBackgroundTask(StaticMeshThreadPool, BasePriority, EQueuedWorkFlags::DoNotRunInsideBusyWait, GetBuildRequiredMemory(), TEXT("StaticMesh"));
+		AsyncTask->StartBackgroundTask(StaticMeshThreadPool, BasePriority, EQueuedWorkFlags::DoNotRunInsideBusyWait, GetBuildRequiredMemoryEstimate(), TEXT("StaticMesh"));
 		FStaticMeshCompilingManager::Get().AddStaticMeshes({this});
 	}
 	else
diff --git a/Engine/Source/Runtime/Engine/Private/StaticMeshBuild.cpp b/Engine/Source/Runtime/Engine/Private/StaticMeshBuild.cpp
--- a/Engine/Source/Runtime/Engine/Private/StaticMeshBuild.cpp
+++ b/Engine/Source/Runtime/Engine/Private/StaticMeshBuild.cpp
@@ -249,7 +249,7 @@ void UStaticMesh::BatchBuild(const TArray<UStaticMesh*>& InStaticMeshes, const F
 				// Only launch async compile if errors are not required
 				if (BuildParameters.OutErrors == nullptr && FStaticMeshCompilingManager::Get().IsAsyncCompilationAllowed(StaticMesh))
 				{
-					const int64 BuildRequiredMemory = StaticMesh->GetBuildRequiredMemory();
+					const int64 BuildRequiredMemory = StaticMesh->GetBuildRequiredMemoryEstimate();
 					TUniquePtr<FStaticMeshBuildContext> Context = MakeUnique<FStaticMeshBuildContext>(BuildParameters);
 					StaticMesh->BeginBuildInternal(Context.Get());
 
@@ -257,7 +257,7 @@ void UStaticMesh::BatchBuild(const TArray<UStaticMesh*>& InStaticMeshes, const F
 					EQueuedWorkPriority BasePriority = FStaticMeshCompilingManager::Get().GetBasePriority(StaticMesh);
 					check(StaticMesh->AsyncTask == nullptr);
 					StaticMesh->AsyncTask = MakeUnique<FStaticMeshAsyncBuildTask>(StaticMesh, MoveTemp(Context));
-					StaticMesh->AsyncTask->StartBackgroundTask(StaticMeshThreadPool, BasePriority, EQueuedWorkFlags::DoNotRunInsideBusyWait, BuildRequiredMemory);
+					StaticMesh->AsyncTask->StartBackgroundTask(StaticMeshThreadPool, BasePriority, EQueuedWorkFlags::DoNotRunInsideBusyWait, BuildRequiredMemory, TEXT("StaticMesh"));
 					FStaticMeshCompilingManager::Get().AddStaticMeshes({ StaticMesh });
 					return true;
 				}
