Removed SVT streaming info buffer as it is no longer necessary now that the page table can fall back to coarser mips if a given page/tile isn't resident.
Also moved some of the helper classes of the SVT streaming manager into their own files to separate them and remove clutter in the SVT streamer implementation.

#rnx

diff --git a/Engine/Plugins/FX/Niagara/Shaders/Private/NiagaraDataInterfaceSparseVolumeTextureTemplate.ush b/Engine/Plugins/FX/Niagara/Shaders/Private/NiagaraDataInterfaceSparseVolumeTextureTemplate.ush
--- a/Engine/Plugins/FX/Niagara/Shaders/Private/NiagaraDataInterfaceSparseVolumeTextureTemplate.ush
+++ b/Engine/Plugins/FX/Niagara/Shaders/Private/NiagaraDataInterfaceSparseVolumeTextureTemplate.ush
@@ -6,7 +6,6 @@ SamplerState		{ParameterName}_TileDataTextureSampler;
 Texture3D<uint>		{ParameterName}_PageTableTexture;
 Texture3D			{ParameterName}_PhysicalTileDataATexture;
 Texture3D			{ParameterName}_PhysicalTileDataBTexture;
-ByteAddressBuffer	{ParameterName}_StreamingInfoBuffer;
 uint4				{ParameterName}_PackedUniforms0;
 uint4				{ParameterName}_PackedUniforms1;
 int3				{ParameterName}_TextureSize;
@@ -15,7 +14,7 @@ int					{ParameterName}_MipLevels;
 
 void LoadSparseVolumeTexture_{ParameterName}(in int TexelX, in int TexelY, in int TexelZ, in int MipLevel, out float4 OutAttributesA, out float4 OutAttributesB)
 {
-	const FSparseVolumeTextureUniforms Uniforms = SparseVolumeTextureUnpackUniforms({ParameterName}_PackedUniforms0, {ParameterName}_PackedUniforms1, {ParameterName}_StreamingInfoBuffer);
+	const FSparseVolumeTextureUniforms Uniforms = SparseVolumeTextureUnpackUniforms({ParameterName}_PackedUniforms0, {ParameterName}_PackedUniforms1);
 	const int3 VoxelCoord = SparseVolumeTextureLoadPageTable({ParameterName}_PageTableTexture, Uniforms, int3(TexelX, TexelY, TexelZ), MipLevel);
 	OutAttributesA = {ParameterName}_PhysicalTileDataATexture.Load(int4(VoxelCoord, 0));
 	OutAttributesB = {ParameterName}_PhysicalTileDataBTexture.Load(int4(VoxelCoord, 0));
@@ -23,7 +22,7 @@ void LoadSparseVolumeTexture_{ParameterName}(in int TexelX, in int TexelY, in in
 
 void SampleSparseVolumeTexture_{ParameterName}(in float3 UVW, in int MipLevel, out float4 OutAttributesA, out float4 OutAttributesB)
 {
-	const FSparseVolumeTextureUniforms Uniforms = SparseVolumeTextureUnpackUniforms({ParameterName}_PackedUniforms0, {ParameterName}_PackedUniforms1, {ParameterName}_StreamingInfoBuffer);
+	const FSparseVolumeTextureUniforms Uniforms = SparseVolumeTextureUnpackUniforms({ParameterName}_PackedUniforms0, {ParameterName}_PackedUniforms1);
 	const float3 VoxelUVW = SparseVolumeTextureSamplePageTable({ParameterName}_PageTableTexture, Uniforms, UVW, SVTADDRESSMODE_WRAP, SVTADDRESSMODE_WRAP, SVTADDRESSMODE_WRAP, MipLevel);
 	OutAttributesA = SparseVolumeTextureSamplePhysicalTileData({ParameterName}_PhysicalTileDataATexture, {ParameterName}_PhysicalTileDataBTexture, {ParameterName}_TileDataTextureSampler, VoxelUVW, 0);
 	OutAttributesB = SparseVolumeTextureSamplePhysicalTileData({ParameterName}_PhysicalTileDataATexture, {ParameterName}_PhysicalTileDataBTexture, {ParameterName}_TileDataTextureSampler, VoxelUVW, 1);
diff --git a/Engine/Plugins/FX/Niagara/Shaders/Private/NiagaraSVTToDenseBuffer.usf b/Engine/Plugins/FX/Niagara/Shaders/Private/NiagaraSVTToDenseBuffer.usf
--- a/Engine/Plugins/FX/Niagara/Shaders/Private/NiagaraSVTToDenseBuffer.usf
+++ b/Engine/Plugins/FX/Niagara/Shaders/Private/NiagaraSVTToDenseBuffer.usf
@@ -7,7 +7,6 @@ RWTexture3D<float4>	DestinationBuffer;
 SamplerState		TileDataTextureSampler;
 Texture3D<uint>		SparseVolumeTexturePageTable;
 Texture3D			SparseVolumeTextureA;
-ByteAddressBuffer	StreamingInfoBuffer;
 uint4				PackedSVTUniforms0;
 uint4				PackedSVTUniforms1;
 int3				TextureSize;
@@ -16,8 +15,8 @@ int					MipLevels;
 [numthreads(4, 4, 4)]
 void PerformCopyCS(uint3 Index : SV_DispatchThreadID)
 {
-	const FSparseVolumeTextureUniforms Uniforms = SparseVolumeTextureUnpackUniforms(PackedSVTUniforms0, PackedSVTUniforms1, StreamingInfoBuffer);
-	const int MipLevel = Uniforms.LowestResidentMipLevel; // We want to sample level 0, but we also need to be able to fall back to coarser mips.
+	const FSparseVolumeTextureUniforms Uniforms = SparseVolumeTextureUnpackUniforms(PackedSVTUniforms0, PackedSVTUniforms1);
+	const int MipLevel = 0; // We want to sample level 0 but the page table can actually point to lower resolution mip level voxels
 	const int3 VirtualCoord = Index >> MipLevel;
 	const int3 VoxelCoord = SparseVolumeTextureLoadPageTable(SparseVolumeTexturePageTable, Uniforms, VirtualCoord, MipLevel);
 
diff --git a/Engine/Plugins/FX/Niagara/Source/Niagara/Classes/NiagaraDataInterfaceSparseVolumeTexture.h b/Engine/Plugins/FX/Niagara/Source/Niagara/Classes/NiagaraDataInterfaceSparseVolumeTexture.h
--- a/Engine/Plugins/FX/Niagara/Source/Niagara/Classes/NiagaraDataInterfaceSparseVolumeTexture.h
+++ b/Engine/Plugins/FX/Niagara/Source/Niagara/Classes/NiagaraDataInterfaceSparseVolumeTexture.h
@@ -20,7 +20,6 @@ class UNiagaraDataInterfaceSparseVolumeTexture : public UNiagaraDataInterface
 		SHADER_PARAMETER_TEXTURE(Texture3D<uint>, PageTableTexture)
 		SHADER_PARAMETER_TEXTURE(Texture3D, PhysicalTileDataATexture)
 		SHADER_PARAMETER_TEXTURE(Texture3D, PhysicalTileDataBTexture)
-		SHADER_PARAMETER_SRV(ByteAddressBuffer, StreamingInfoBuffer)
 		SHADER_PARAMETER(FUintVector4, PackedUniforms0)
 		SHADER_PARAMETER(FUintVector4, PackedUniforms1)
 		SHADER_PARAMETER(FIntVector3, TextureSize)
diff --git a/Engine/Plugins/FX/Niagara/Source/Niagara/Private/NiagaraDataInterfaceRenderTargetVolume.cpp b/Engine/Plugins/FX/Niagara/Source/Niagara/Private/NiagaraDataInterfaceRenderTargetVolume.cpp
--- a/Engine/Plugins/FX/Niagara/Source/Niagara/Private/NiagaraDataInterfaceRenderTargetVolume.cpp
+++ b/Engine/Plugins/FX/Niagara/Source/Niagara/Private/NiagaraDataInterfaceRenderTargetVolume.cpp
@@ -846,12 +846,10 @@ bool UNiagaraDataInterfaceRenderTargetVolume::SimCacheReadFrame(UObject* Storage
 
 						FRHITexture* PageTableTexture = RT_SVTRenderResources->GetPageTableTexture();
 						FRHITexture* TextureA = RT_SVTRenderResources->GetPhysicalTileDataATexture();
-						FRHIShaderResourceView* StreamingInfoBufferSRV = RT_SVTRenderResources->GetStreamingInfoBufferSRV();
 
 						PassParameters->TileDataTextureSampler = TStaticSamplerState<SF_Point, AM_Clamp, AM_Clamp, AM_Clamp>::GetRHI();
 						PassParameters->SparseVolumeTexturePageTable = PageTableTexture ? PageTableTexture : GBlackUintVolumeTexture->TextureRHI.GetReference();
 						PassParameters->SparseVolumeTextureA = TextureA ? TextureA : GBlackVolumeTexture->TextureRHI.GetReference();
-						PassParameters->StreamingInfoBuffer = StreamingInfoBufferSRV ? StreamingInfoBufferSRV : GEmptyStructuredBufferWithUAV->ShaderResourceViewRHI.GetReference();
 							
 						PassParameters->PackedSVTUniforms0 = CurrentPackedUniforms0;
 						PassParameters->PackedSVTUniforms1 = CurrentPackedUniforms1;
diff --git a/Engine/Plugins/FX/Niagara/Source/Niagara/Private/NiagaraDataInterfaceSparseVolumeTexture.cpp b/Engine/Plugins/FX/Niagara/Source/Niagara/Private/NiagaraDataInterfaceSparseVolumeTexture.cpp
--- a/Engine/Plugins/FX/Niagara/Source/Niagara/Private/NiagaraDataInterfaceSparseVolumeTexture.cpp
+++ b/Engine/Plugins/FX/Niagara/Source/Niagara/Private/NiagaraDataInterfaceSparseVolumeTexture.cpp
@@ -287,7 +287,6 @@ void UNiagaraDataInterfaceSparseVolumeTexture::SetShaderParameters(const FNiagar
 	Parameters->PageTableTexture = GBlackUintVolumeTexture->TextureRHI;
 	Parameters->PhysicalTileDataATexture = GBlackVolumeTexture->TextureRHI;
 	Parameters->PhysicalTileDataBTexture = GBlackVolumeTexture->TextureRHI;
-	Parameters->StreamingInfoBuffer = GEmptyStructuredBufferWithUAV->ShaderResourceViewRHI;
 	Parameters->PackedUniforms0 = FUintVector4();
 	Parameters->PackedUniforms1 = FUintVector4();
 	Parameters->TextureSize = FIntVector3::ZeroValue;
@@ -298,12 +297,10 @@ void UNiagaraDataInterfaceSparseVolumeTexture::SetShaderParameters(const FNiagar
 		FRHITexture* PageTableTexture = RTInstanceData->RenderResources->GetPageTableTexture();
 		FRHITexture* PhysicalTileDataATexture = RTInstanceData->RenderResources->GetPhysicalTileDataATexture();
 		FRHITexture* PhysicalTileDataBTexture = RTInstanceData->RenderResources->GetPhysicalTileDataBTexture();
-		FRHIShaderResourceView* StreamingInfoBufferSRV = RTInstanceData->RenderResources->GetStreamingInfoBufferSRV();
 
 		Parameters->PageTableTexture = PageTableTexture ? PageTableTexture : Parameters->PageTableTexture;
 		Parameters->PhysicalTileDataATexture = PhysicalTileDataATexture ? PhysicalTileDataATexture : Parameters->PhysicalTileDataATexture;
 		Parameters->PhysicalTileDataBTexture = PhysicalTileDataBTexture ? PhysicalTileDataBTexture : Parameters->PhysicalTileDataBTexture;
-		Parameters->StreamingInfoBuffer = StreamingInfoBufferSRV ? StreamingInfoBufferSRV : Parameters->StreamingInfoBuffer;
 		RTInstanceData->RenderResources->GetPackedUniforms(Parameters->PackedUniforms0, Parameters->PackedUniforms1);
 		Parameters->TextureSize = RTInstanceData->TextureSize;
 		Parameters->MipLevels = RTInstanceData->MipLevels;
diff --git a/Engine/Plugins/FX/Niagara/Source/NiagaraShader/Public/NiagaraSVTShaders.h b/Engine/Plugins/FX/Niagara/Source/NiagaraShader/Public/NiagaraSVTShaders.h
--- a/Engine/Plugins/FX/Niagara/Source/NiagaraShader/Public/NiagaraSVTShaders.h
+++ b/Engine/Plugins/FX/Niagara/Source/NiagaraShader/Public/NiagaraSVTShaders.h
@@ -17,7 +17,6 @@ class FNiagaraCopySVTToDenseBufferCS : public FGlobalShader
 		SHADER_PARAMETER_SAMPLER(SamplerState, TileDataTextureSampler)
 		SHADER_PARAMETER_TEXTURE(Texture3D<uint>, SparseVolumeTexturePageTable)
 		SHADER_PARAMETER_TEXTURE(Texture3D, SparseVolumeTextureA)	
-		SHADER_PARAMETER_SRV(ByteAddressBuffer, StreamingInfoBuffer)
 		SHADER_PARAMETER(FUintVector4, PackedSVTUniforms0)
 		SHADER_PARAMETER(FUintVector4, PackedSVTUniforms1)
 		SHADER_PARAMETER(FIntVector, TextureSize)
diff --git a/Engine/Shaders/Private/SparseVolumeTexture/SparseVolumeTextureCommon.ush b/Engine/Shaders/Private/SparseVolumeTexture/SparseVolumeTextureCommon.ush
--- a/Engine/Shaders/Private/SparseVolumeTexture/SparseVolumeTextureCommon.ush
+++ b/Engine/Shaders/Private/SparseVolumeTexture/SparseVolumeTextureCommon.ush
@@ -48,10 +48,9 @@ struct FSparseVolumeTextureUniforms
 	float3			TileDataTexelSize;
 	int				FrameIndex;
 	int             HighestMipLevel;
-	int				LowestResidentMipLevel;
 };
 
-FSparseVolumeTextureUniforms SparseVolumeTextureUnpackUniforms(const uint4 Packed0, const uint4 Packed1, ByteAddressBuffer StreamingInfoBuffer)
+FSparseVolumeTextureUniforms SparseVolumeTextureUnpackUniforms(const uint4 Packed0, const uint4 Packed1)
 {
 	FSparseVolumeTextureUniforms Result;
 	Result.VolumePageResolution = asfloat(Packed0.xyz);
@@ -61,7 +60,6 @@ FSparseVolumeTextureUniforms SparseVolumeTextureUnpackUniforms(const uint4 Packe
 	Result.TileDataTexelSize = asfloat(Packed1.xyz);
 	Result.FrameIndex = int(Packed1.w & 0xFFFFu);
 	Result.HighestMipLevel = int((Packed1.w >> 16u) & 0xFFFFu);
-	Result.LowestResidentMipLevel = asint(StreamingInfoBuffer.Load(Result.FrameIndex << 2));
 
 	return Result;
 }
diff --git a/Engine/Shaders/Private/SparseVolumeTexture/UpdateSparseVolumeTexture.usf b/Engine/Shaders/Private/SparseVolumeTexture/UpdateSparseVolumeTexture.usf
--- a/Engine/Shaders/Private/SparseVolumeTexture/UpdateSparseVolumeTexture.usf
+++ b/Engine/Shaders/Private/SparseVolumeTexture/UpdateSparseVolumeTexture.usf
@@ -205,22 +205,3 @@ void SparseVolumeTextureUpdatePageTableCS(uint3 DispatchThreadID : SV_DispatchTh
 }
 
 #endif // UPDATE_PAGE_TABLE
-
-#ifdef UPDATE_STREAMING_INFO_BUFFER
-
-RWByteAddressBuffer StreamingInfoBuffer;
-ByteAddressBuffer StreamingInfoBufferUpdates;
-uint UpdateOffset;
-uint NumUpdates;
-
-[numthreads(64, 1, 1)]
-void SparseVolumeTextureUpdateStreamingInfoBufferCS(uint3 DispatchThreadID : SV_DispatchThreadID)
-{
-	if (DispatchThreadID.x < NumUpdates)
-	{
-		uint2 Update = StreamingInfoBufferUpdates.Load2(((DispatchThreadID.x + UpdateOffset) * 2) << 2);
-		StreamingInfoBuffer.Store(Update.x << 2, Update.y);
-	}
-}
-
-#endif // UPDATE_STREAMING_INFO_BUFFER
diff --git a/Engine/Shaders/Private/SparseVolumeTexture/VisualizeSparseVolumeTexture.usf b/Engine/Shaders/Private/SparseVolumeTexture/VisualizeSparseVolumeTexture.usf
--- a/Engine/Shaders/Private/SparseVolumeTexture/VisualizeSparseVolumeTexture.usf
+++ b/Engine/Shaders/Private/SparseVolumeTexture/VisualizeSparseVolumeTexture.usf
@@ -45,7 +45,6 @@ SamplerState TileDataTextureSampler;
 Texture3D<uint> SparseVolumeTexturePageTable;
 Texture3D<float4> SparseVolumeTextureA;
 Texture3D<float4> SparseVolumeTextureB;
-ByteAddressBuffer StreamingInfoBuffer;
 uint4 PackedSVTUniforms0;
 uint4 PackedSVTUniforms1;
 float3 SparseVolumeTextureResolution;
@@ -132,7 +131,7 @@ void VisualizeSparseVolumeTexturePS(
 		tMax.y *= tDelta.y;
 		tMax.z *= tDelta.z;
 
-		const FSparseVolumeTextureUniforms SVTUniforms = SparseVolumeTextureUnpackUniforms(PackedSVTUniforms0, PackedSVTUniforms1, StreamingInfoBuffer);
+		const FSparseVolumeTextureUniforms SVTUniforms = SparseVolumeTextureUnpackUniforms(PackedSVTUniforms0, PackedSVTUniforms1);
 
 		float StepLength = 1.0f;
 
diff --git a/Engine/Source/Runtime/Engine/Classes/SparseVolumeTexture/SparseVolumeTexture.h b/Engine/Source/Runtime/Engine/Classes/SparseVolumeTexture/SparseVolumeTexture.h
--- a/Engine/Source/Runtime/Engine/Classes/SparseVolumeTexture/SparseVolumeTexture.h
+++ b/Engine/Source/Runtime/Engine/Classes/SparseVolumeTexture/SparseVolumeTexture.h
@@ -196,7 +196,6 @@ public:
 	FRHITextureReference* GetPageTableTexture() const			{ check(IsInParallelRenderingThread()); return PageTableTextureReferenceRHI.GetReference(); }
 	FRHITextureReference* GetPhysicalTileDataATexture() const	{ check(IsInParallelRenderingThread()); return PhysicalTileDataATextureReferenceRHI.GetReference(); }
 	FRHITextureReference* GetPhysicalTileDataBTexture() const	{ check(IsInParallelRenderingThread()); return PhysicalTileDataBTextureReferenceRHI.GetReference(); }
-	FRHIShaderResourceView* GetStreamingInfoBufferSRV() const	{ check(IsInParallelRenderingThread()); return StreamingInfoBufferSRVRHI.GetReference(); }
 	ENGINE_API void GetPackedUniforms(FUintVector4& OutPacked0, FUintVector4& OutPacked1) const;
 	// Updates the GlobalVolumeResolution member in a thread-safe way.
 	ENGINE_API void SetGlobalVolumeResolution_GameThread(const FIntVector3& GlobalVolumeResolution);
@@ -215,7 +214,6 @@ private:
 	FTextureReferenceRHIRef PageTableTextureReferenceRHI;
 	FTextureReferenceRHIRef PhysicalTileDataATextureReferenceRHI;
 	FTextureReferenceRHIRef PhysicalTileDataBTextureReferenceRHI;
-	FShaderResourceViewRHIRef StreamingInfoBufferSRVRHI;
 };
 
 }
diff --git a/Engine/Source/Runtime/Engine/Private/Materials/HLSLMaterialTranslator.cpp b/Engine/Source/Runtime/Engine/Private/Materials/HLSLMaterialTranslator.cpp
--- a/Engine/Source/Runtime/Engine/Private/Materials/HLSLMaterialTranslator.cpp
+++ b/Engine/Source/Runtime/Engine/Private/Materials/HLSLMaterialTranslator.cpp
@@ -3946,7 +3946,7 @@ int32 FHLSLMaterialTranslator::AccessUniformExpression(int32 Index)
 	else if (CodeChunk.Type == MCT_SparseVolumeTexture)
 	{
 		int32 TextureInputIndex = UniformTextureExpressions[(uint32)EMaterialTextureParameterType::SparseVolume].AddUnique(TextureUniformExpression);
-		FormattedCode.Appendf(TEXT("SparseVolumeTextureUnpackUniforms(Material.SVTPackedUniform[%d*2], Material.SVTPackedUniform[%d*2+1], Material.SparseVolumeTextureStreamingInfoBuffer_%d)"), TextureInputIndex, TextureInputIndex, TextureInputIndex);
+		FormattedCode.Appendf(TEXT("SparseVolumeTextureUnpackUniforms(Material.SVTPackedUniform[%d*2], Material.SVTPackedUniform[%d*2+1])"), TextureInputIndex, TextureInputIndex);
 	}
 	else
 	{
diff --git a/Engine/Source/Runtime/Engine/Private/Materials/MaterialUniformExpressions.cpp b/Engine/Source/Runtime/Engine/Private/Materials/MaterialUniformExpressions.cpp
--- a/Engine/Source/Runtime/Engine/Private/Materials/MaterialUniformExpressions.cpp
+++ b/Engine/Source/Runtime/Engine/Private/Materials/MaterialUniformExpressions.cpp
@@ -406,7 +406,6 @@ FShaderParametersMetadata* FUniformExpressionSet::CreateBufferStruct()
 	static FString SparseVolumeTexturePageTableNames[128];
 	static FString SparseVolumeTexturePhysicalANames[128];
 	static FString SparseVolumeTexturePhysicalBNames[128];
-	static FString SparseVolumeTextureStreamingInfoBufferNames[128];
 	static FString SparseVolumeTexturePhysicalSamplerNames[128];
 	static bool bInitializedTextureNames = false;
 	if (!bInitializedTextureNames)
@@ -434,7 +433,6 @@ FShaderParametersMetadata* FUniformExpressionSet::CreateBufferStruct()
 			SparseVolumeTexturePageTableNames[i] = FString::Printf(TEXT("SparseVolumeTexturePageTable_%d"), i);
 			SparseVolumeTexturePhysicalANames[i] = FString::Printf(TEXT("SparseVolumeTexturePhysicalA_%d"), i);
 			SparseVolumeTexturePhysicalBNames[i] = FString::Printf(TEXT("SparseVolumeTexturePhysicalB_%d"), i);
-			SparseVolumeTextureStreamingInfoBufferNames[i] = FString::Printf(TEXT("SparseVolumeTextureStreamingInfoBuffer_%d"), i);
 			SparseVolumeTexturePhysicalSamplerNames[i] = FString::Printf(TEXT("SparseVolumeTexturePhysical_%dSampler"), i);
 		}
 	}
@@ -511,11 +509,6 @@ FShaderParametersMetadata* FUniformExpressionSet::CreateBufferStruct()
 		new(Members) FShaderParametersMetadata::FMember(*SparseVolumeTexturePhysicalBNames[i], TEXT("Texture3D"), __LINE__, NextMemberOffset, UBMT_TEXTURE, EShaderPrecisionModifier::Float, 1, 1, 0, NULL);
 		NextMemberOffset += SHADER_PARAMETER_POINTER_ALIGNMENT;
 
-		// Streaming Info Buffer
-		check((NextMemberOffset% SHADER_PARAMETER_POINTER_ALIGNMENT) == 0);
-		new(Members) FShaderParametersMetadata::FMember(*SparseVolumeTextureStreamingInfoBufferNames[i], TEXT("ByteAddressBuffer"), __LINE__, NextMemberOffset, UBMT_SRV, EShaderPrecisionModifier::Float, 1, 1, 0, NULL);
-		NextMemberOffset += SHADER_PARAMETER_POINTER_ALIGNMENT;
-
 		// Sampler
 		check((NextMemberOffset% SHADER_PARAMETER_POINTER_ALIGNMENT) == 0);
 		new(Members) FShaderParametersMetadata::FMember(*SparseVolumeTexturePhysicalSamplerNames[i], TEXT("SamplerState"), __LINE__, NextMemberOffset, UBMT_SAMPLER, EShaderPrecisionModifier::Float, 1, 1, 0, NULL);
@@ -1003,7 +996,7 @@ void FUniformExpressionSet::FillUniformBuffer(const FMaterialRenderContext& Mate
 				+ UniformTextureParameters[(uint32)EMaterialTextureParameterType::Array2D].Num() * 2
 				+ UniformTextureParameters[(uint32)EMaterialTextureParameterType::ArrayCube].Num() * 2
 				+ UniformTextureParameters[(uint32)EMaterialTextureParameterType::Volume].Num() * 2
-				+ UniformTextureParameters[(uint32)EMaterialTextureParameterType::SparseVolume].Num() * 5
+				+ UniformTextureParameters[(uint32)EMaterialTextureParameterType::SparseVolume].Num() * 4
 				+ UniformExternalTextureParameters.Num() * 2
 				+ UniformTextureParameters[(uint32)EMaterialTextureParameterType::Virtual].Num() * 2
 				+ NumPageTableTextures
@@ -1276,16 +1269,14 @@ void FUniformExpressionSet::FillUniformBuffer(const FMaterialRenderContext& Mate
 			void** ResourceTableTexturePageTablePtr = (void**)((uint8*)BufferCursor + 0 * SHADER_PARAMETER_POINTER_ALIGNMENT);
 			void** ResourceTableTexturePhysicalAPtr = (void**)((uint8*)BufferCursor + 1 * SHADER_PARAMETER_POINTER_ALIGNMENT);
 			void** ResourceTableTexturePhysicalBPtr = (void**)((uint8*)BufferCursor + 2 * SHADER_PARAMETER_POINTER_ALIGNMENT);
-			void** ResourceTableStreamingInfoBufferPtr = (void**)((uint8*)BufferCursor + 3 * SHADER_PARAMETER_POINTER_ALIGNMENT);
-			void** ResourceTablePhysicalSamplerPtr = (void**)((uint8*)BufferCursor + 4 * SHADER_PARAMETER_POINTER_ALIGNMENT);
-			BufferCursor = ((uint8*)BufferCursor) + (SHADER_PARAMETER_POINTER_ALIGNMENT * 5);
+			void** ResourceTablePhysicalSamplerPtr = (void**)((uint8*)BufferCursor + 3 * SHADER_PARAMETER_POINTER_ALIGNMENT);
+			BufferCursor = ((uint8*)BufferCursor) + (SHADER_PARAMETER_POINTER_ALIGNMENT * 4);
 			check(BufferCursor <= TempBuffer + TempBufferSize);
 
 			check(GBlackVolumeTexture->TextureRHI);
 			*ResourceTableTexturePageTablePtr = GBlackUintVolumeTexture->TextureRHI;
 			*ResourceTableTexturePhysicalAPtr = GBlackVolumeTexture->TextureRHI;
 			*ResourceTableTexturePhysicalBPtr = GBlackVolumeTexture->TextureRHI;
-			*ResourceTableStreamingInfoBufferPtr = GEmptyStructuredBufferWithUAV->ShaderResourceViewRHI;
 			check(GBlackVolumeTexture->SamplerStateRHI);
 			*ResourceTablePhysicalSamplerPtr = GBlackVolumeTexture->SamplerStateRHI;
 			
@@ -1299,13 +1290,11 @@ void FUniformExpressionSet::FillUniformBuffer(const FMaterialRenderContext& Mate
 					FRHITexture* PageTableTexture = RenderResources->GetPageTableTexture();
 					FRHITexture* TileDataATexture = RenderResources->GetPhysicalTileDataATexture();
 					FRHITexture* TileDataBTexture = RenderResources->GetPhysicalTileDataBTexture();
-					FRHIShaderResourceView* StreamingInfoBufferSRV = RenderResources->GetStreamingInfoBufferSRV();
 
 					// It's possible for RenderResources to be valid, but PageTableTexture still null, so we need to have a fallback (a black uint volume texture)
 					*ResourceTableTexturePageTablePtr = PageTableTexture ? PageTableTexture : *ResourceTableTexturePageTablePtr;
 					*ResourceTableTexturePhysicalAPtr = TileDataATexture ? TileDataATexture : *ResourceTableTexturePhysicalAPtr;
 					*ResourceTableTexturePhysicalBPtr = TileDataBTexture ? TileDataBTexture : *ResourceTableTexturePhysicalBPtr;
-					*ResourceTableStreamingInfoBufferPtr = StreamingInfoBufferSRV ? StreamingInfoBufferSRV : *ResourceTableStreamingInfoBufferPtr;
 					*ResourceTablePhysicalSamplerPtr = TStaticSamplerState<SF_Bilinear, AM_Clamp, AM_Clamp, AM_Clamp>::GetRHI();
 				}
 			}
diff --git a/Engine/Source/Runtime/Engine/Private/SparseVolumeTexture/SparseVolumeTexture.cpp b/Engine/Source/Runtime/Engine/Private/SparseVolumeTexture/SparseVolumeTexture.cpp
--- a/Engine/Source/Runtime/Engine/Private/SparseVolumeTexture/SparseVolumeTexture.cpp
+++ b/Engine/Source/Runtime/Engine/Private/SparseVolumeTexture/SparseVolumeTexture.cpp
@@ -974,7 +974,6 @@ void FTextureRenderResources::ReleaseRHI()
 	PageTableTextureReferenceRHI.SafeRelease();
 	PhysicalTileDataATextureReferenceRHI.SafeRelease();
 	PhysicalTileDataBTextureReferenceRHI.SafeRelease();
-	StreamingInfoBufferSRVRHI.SafeRelease();
 }
 
 }
diff --git a/Engine/Source/Runtime/Engine/Private/SparseVolumeTexture/SparseVolumeTextureData.cpp b/Engine/Source/Runtime/Engine/Private/SparseVolumeTexture/SparseVolumeTextureData.cpp
--- a/Engine/Source/Runtime/Engine/Private/SparseVolumeTexture/SparseVolumeTextureData.cpp
+++ b/Engine/Source/Runtime/Engine/Private/SparseVolumeTexture/SparseVolumeTextureData.cpp
@@ -40,21 +40,6 @@ namespace UE
 namespace SVT
 {
 
-bool IsInBounds(const FIntVector3& Point, const FIntVector3& Min, const FIntVector3& Max)
-{
-	return Point.X >= Min.X && Point.Y >= Min.Y && Point.Z >= Min.Z
-		&& Point.X < Max.X && Point.Y < Max.Y && Point.Z < Max.Z;
-}
-
-FIntVector3 ShiftRightAndMax(const FIntVector3& Value, uint32 ShiftBy, int32 MinValue)
-{
-	FIntVector3 Result = FIntVector3(
-		FMath::Max(Value.X >> ShiftBy, MinValue),
-		FMath::Max(Value.Y >> ShiftBy, MinValue),
-		FMath::Max(Value.Z >> ShiftBy, MinValue));
-	return Result;
-}
-
 bool FTextureData::Create(const ITextureDataProvider& DataProvider)
 {
 	const FTextureDataCreateInfo CreateInfo = DataProvider.GetCreateInfo();
diff --git a/Engine/Source/Runtime/Engine/Private/SparseVolumeTexture/SparseVolumeTextureStreamingManager.cpp b/Engine/Source/Runtime/Engine/Private/SparseVolumeTexture/SparseVolumeTextureStreamingManager.cpp
--- a/Engine/Source/Runtime/Engine/Private/SparseVolumeTexture/SparseVolumeTextureStreamingManager.cpp
+++ b/Engine/Source/Runtime/Engine/Private/SparseVolumeTexture/SparseVolumeTextureStreamingManager.cpp
@@ -9,13 +9,15 @@
 #include "GlobalShader.h"
 #include "ShaderCompilerCore.h" // AllowGlobalShaderLoad()
 #include "Async/ParallelFor.h"
+#include "SparseVolumeTextureTileDataTexture.h"
+#include "SparseVolumeTextureUpload.h"
 
 #if WITH_EDITORONLY_DATA
 #include "DerivedDataCache.h"
 #include "DerivedDataRequestOwner.h"
 #endif
 
-DEFINE_LOG_CATEGORY_STATIC(LogSparseVolumeTextureStreamingManager, Log, All);
+DEFINE_LOG_CATEGORY(LogSparseVolumeTextureStreamingManager);
 
 #ifndef SVT_STREAMING_LOG_VERBOSE
 #define SVT_STREAMING_LOG_VERBOSE 0
@@ -53,14 +55,6 @@ static FAutoConsoleVariableRef CVarSVTStreamingAsync(
 	ECVF_RenderThreadSafe
 );
 
-static int32 GSVTStreamingAsyncCompute = 1;
-static FAutoConsoleVariableRef CVarSVTStreamingAsyncCompute(
-	TEXT("r.SparseVolumeTexture.Streaming.AsyncCompute"),
-	GSVTStreamingAsyncCompute,
-	TEXT("Schedule GPU work in async compute queue."),
-	ECVF_RenderThreadSafe
-);
-
 static int32 GSVTStreamingEmptyPhysicalTileTextures = 0;
 static FAutoConsoleVariableRef CVarSVTStreamingEmptyPhysicalTileTextures(
 	TEXT("r.SparseVolumeTexture.Streaming.EmptyPhysicalTileTextures"),
@@ -104,739 +98,6 @@ static bool DoesPlatformSupportSparseVolumeTexture(EShaderPlatform Platform)
 	return AllowGlobalShaderLoad();
 }
 
-static FIntVector3 ComputeTileDataVolumeResolution(int32 NumAllocatedPages)
-{
-	int32 TileVolumeResolutionCube = 1;
-	while (TileVolumeResolutionCube * TileVolumeResolutionCube * TileVolumeResolutionCube < NumAllocatedPages)
-	{
-		TileVolumeResolutionCube++;				// We use a simple loop to compute the minimum resolution of a cube to store all the tile data
-	}
-	FIntVector3 TileDataVolumeResolution = FIntVector3(TileVolumeResolutionCube, TileVolumeResolutionCube, TileVolumeResolutionCube);
-	
-	// Trim volume to reclaim some space
-	while ((TileDataVolumeResolution.X * TileDataVolumeResolution.Y * (TileDataVolumeResolution.Z - 1)) > NumAllocatedPages)
-	{
-		TileDataVolumeResolution.Z--;
-	}
-	while ((TileDataVolumeResolution.X * (TileDataVolumeResolution.Y - 1) * TileDataVolumeResolution.Z) > NumAllocatedPages)
-	{
-		TileDataVolumeResolution.Y--;
-	}
-	while (((TileDataVolumeResolution.X - 1) * TileDataVolumeResolution.Y * TileDataVolumeResolution.Z) > NumAllocatedPages)
-	{
-		TileDataVolumeResolution.X--;
-	}
-
-	return TileDataVolumeResolution * SPARSE_VOLUME_TILE_RES_PADDED;
-}
-
-static FIntVector3 ComputeLargestPossibleTileDataVolumeResolution(int32 VoxelMemSize)
-{
-	const int64 TileMemSize = SVT::NumVoxelsPerPaddedTile * VoxelMemSize;
-	const int64 NumMaxTiles = SVT::MaxResourceSize / TileMemSize;
-	int64 ResourceSize = NumMaxTiles * TileMemSize;
-
-	// Find a cube with a volume as close to NumMaxTiles as possible
-	int32 TileVolumeResolutionCube = 1;
-	while (((TileVolumeResolutionCube + 1) * (TileVolumeResolutionCube + 1) * (TileVolumeResolutionCube + 1)) <= NumMaxTiles)
-	{
-		++TileVolumeResolutionCube;
-	}
-
-	// Try to add to the sides to get closer to NumMaxTiles
-	FIntVector3 ResolutionInTiles = FIntVector3(TileVolumeResolutionCube, TileVolumeResolutionCube, TileVolumeResolutionCube);
-	if (((ResolutionInTiles.X + 1) * ResolutionInTiles.Y * ResolutionInTiles.Z) <= NumMaxTiles)
-	{
-		++ResolutionInTiles.X;
-	}
-	if ((ResolutionInTiles.X * (ResolutionInTiles.Y + 1) * ResolutionInTiles.Z) <= NumMaxTiles)
-	{
-		++ResolutionInTiles.Y;
-	}
-	if ((ResolutionInTiles.X * ResolutionInTiles.Y * (ResolutionInTiles.Z + 1)) <= NumMaxTiles)
-	{
-		++ResolutionInTiles.Z;
-	}
-
-	const FIntVector3 Resolution = ResolutionInTiles * SPARSE_VOLUME_TILE_RES_PADDED;
-	check(Resolution.X <= SVT::MaxVolumeTextureDim && Resolution.Y <= SVT::MaxVolumeTextureDim && Resolution.Z <= SVT::MaxVolumeTextureDim);
-	check(((int64)Resolution.X * (int64)Resolution.Y * (int64)Resolution.Z) < SVT::MaxResourceSize);
-
-	return Resolution;
-}
-
-class FSparseVolumeTextureUpdateFromBufferCS : public FGlobalShader
-{
-	DECLARE_GLOBAL_SHADER(FSparseVolumeTextureUpdateFromBufferCS);
-	SHADER_USE_PARAMETER_STRUCT(FSparseVolumeTextureUpdateFromBufferCS, FGlobalShader)
-
-	BEGIN_SHADER_PARAMETER_STRUCT(FParameters, )
-		SHADER_PARAMETER_RDG_TEXTURE_UAV(RWTexture3D, DstPhysicalTileTextureA)
-		SHADER_PARAMETER_RDG_TEXTURE_UAV(RWTexture3D, DstPhysicalTileTextureB)
-		SHADER_PARAMETER_RDG_BUFFER_SRV(ByteAddressBuffer, DstTileCoords)
-		SHADER_PARAMETER_RDG_BUFFER_SRV(Buffer<float4>, SrcPhysicalTileBufferA)
-		SHADER_PARAMETER_RDG_BUFFER_SRV(Buffer<float4>, SrcPhysicalTileBufferB)
-		SHADER_PARAMETER(uint32, TileCoordsBufferOffset)
-		SHADER_PARAMETER(uint32, TileDataBufferOffsetInTiles)
-		SHADER_PARAMETER(uint32, NumTilesToCopy)
-		SHADER_PARAMETER(uint32, NumDispatchedGroups)
-		SHADER_PARAMETER(uint32, PaddedTileSize)
-		SHADER_PARAMETER(uint32, bCopyTexureAOnlyUI)
-	END_SHADER_PARAMETER_STRUCT()
-
-	static bool ShouldCompilePermutation(const FGlobalShaderPermutationParameters& Parameters)
-	{
-		return true;
-	}
-
-	static void ModifyCompilationEnvironment(const FGlobalShaderPermutationParameters& Parameters, FShaderCompilerEnvironment& OutEnvironment)
-	{
-		FGlobalShader::ModifyCompilationEnvironment(Parameters, OutEnvironment);
-		OutEnvironment.SetDefine(TEXT("UPDATE_TILE_TEXTURE_FROM_BUFFER"), 1);
-	}
-};
-IMPLEMENT_GLOBAL_SHADER(FSparseVolumeTextureUpdateFromBufferCS, "/Engine/Private/SparseVolumeTexture/UpdateSparseVolumeTexture.usf", "SparseVolumeTextureUpdateFromBufferCS", SF_Compute);
-
-class FSparseVolumeTextureUpdateFromSparseBufferCS : public FGlobalShader
-{
-	DECLARE_GLOBAL_SHADER(FSparseVolumeTextureUpdateFromSparseBufferCS);
-	SHADER_USE_PARAMETER_STRUCT(FSparseVolumeTextureUpdateFromSparseBufferCS, FGlobalShader)
-
-	BEGIN_SHADER_PARAMETER_STRUCT(FParameters, )
-		SHADER_PARAMETER_RDG_TEXTURE_UAV(RWTexture3D, DstPhysicalTileTextureA)
-		SHADER_PARAMETER_RDG_TEXTURE_UAV(RWTexture3D, DstPhysicalTileTextureB)
-		SHADER_PARAMETER_RDG_BUFFER_SRV(Buffer<float4>, SrcPhysicalTileBufferA)
-		SHADER_PARAMETER_RDG_BUFFER_SRV(Buffer<float4>, SrcPhysicalTileBufferB)
-		SHADER_PARAMETER_RDG_BUFFER_SRV(ByteAddressBuffer, OccupancyBitsBuffer)
-		SHADER_PARAMETER_RDG_BUFFER_SRV(ByteAddressBuffer, TileDataOffsetsBuffer)
-		SHADER_PARAMETER_RDG_BUFFER_SRV(ByteAddressBuffer, DstTileCoordsBuffer)
-		SHADER_PARAMETER(FVector4f, FallbackValueA)
-		SHADER_PARAMETER(FVector4f, FallbackValueB)
-		SHADER_PARAMETER(uint32, TileIndexOffset)
-		SHADER_PARAMETER(uint32, SrcVoxelDataOffsetA)
-		SHADER_PARAMETER(uint32, SrcVoxelDataOffsetB)
-		SHADER_PARAMETER(uint32, NumTilesToCopy)
-		SHADER_PARAMETER(uint32, BufferTileStep)
-		SHADER_PARAMETER(uint32, NumDispatchedGroups)
-		SHADER_PARAMETER(uint32, PaddedTileSize)
-		SHADER_PARAMETER(uint32, CopyTexureMask)
-	END_SHADER_PARAMETER_STRUCT()
-
-	static bool ShouldCompilePermutation(const FGlobalShaderPermutationParameters& Parameters)
-	{
-		return true;
-	}
-
-	static void ModifyCompilationEnvironment(const FGlobalShaderPermutationParameters& Parameters, FShaderCompilerEnvironment& OutEnvironment)
-	{
-		FGlobalShader::ModifyCompilationEnvironment(Parameters, OutEnvironment);
-		OutEnvironment.SetDefine(TEXT("UPDATE_TILE_TEXTURE_FROM_SPARSE_BUFFER"), 1);
-	}
-};
-IMPLEMENT_GLOBAL_SHADER(FSparseVolumeTextureUpdateFromSparseBufferCS, "/Engine/Private/SparseVolumeTexture/UpdateSparseVolumeTexture.usf", "SparseVolumeTextureUpdateFromSparseBufferCS", SF_Compute);
-
-class FSparseVolumeTextureUpdatePageTableCS : public FGlobalShader
-{
-	DECLARE_GLOBAL_SHADER(FSparseVolumeTextureUpdatePageTableCS);
-	SHADER_USE_PARAMETER_STRUCT(FSparseVolumeTextureUpdatePageTableCS, FGlobalShader)
-
-	BEGIN_SHADER_PARAMETER_STRUCT(FParameters, )
-		SHADER_PARAMETER_RDG_TEXTURE_UAV(RWTexture3D, PageTable)
-		SHADER_PARAMETER_RDG_BUFFER_SRV(ByteAddressBuffer, PageTableUpdates)
-		SHADER_PARAMETER(uint32, UpdateCoordOffset)
-		SHADER_PARAMETER(uint32, UpdatePayloadOffset)
-		SHADER_PARAMETER(uint32, NumUpdates)
-	END_SHADER_PARAMETER_STRUCT()
-
-	static bool ShouldCompilePermutation(const FGlobalShaderPermutationParameters& Parameters)
-	{
-		return true;
-	}
-
-	static void ModifyCompilationEnvironment(const FGlobalShaderPermutationParameters& Parameters, FShaderCompilerEnvironment& OutEnvironment)
-	{
-		FGlobalShader::ModifyCompilationEnvironment(Parameters, OutEnvironment);
-		OutEnvironment.SetDefine(TEXT("UPDATE_PAGE_TABLE"), 1);
-	}
-};
-IMPLEMENT_GLOBAL_SHADER(FSparseVolumeTextureUpdatePageTableCS, "/Engine/Private/SparseVolumeTexture/UpdateSparseVolumeTexture.usf", "SparseVolumeTextureUpdatePageTableCS", SF_Compute);
-
-class FSparseVolumeTextureUpdateStreamingInfoBufferCS : public FGlobalShader
-{
-	DECLARE_GLOBAL_SHADER(FSparseVolumeTextureUpdateStreamingInfoBufferCS);
-	SHADER_USE_PARAMETER_STRUCT(FSparseVolumeTextureUpdateStreamingInfoBufferCS, FGlobalShader)
-
-	BEGIN_SHADER_PARAMETER_STRUCT(FParameters, )
-		SHADER_PARAMETER_RDG_BUFFER_UAV(RWByteAddressBuffer, StreamingInfoBuffer)
-		SHADER_PARAMETER_RDG_BUFFER_SRV(ByteAddressBuffer, StreamingInfoBufferUpdates)
-		SHADER_PARAMETER(uint32, UpdateOffset)
-		SHADER_PARAMETER(uint32, NumUpdates)
-	END_SHADER_PARAMETER_STRUCT()
-
-	static bool ShouldCompilePermutation(const FGlobalShaderPermutationParameters& Parameters)
-	{
-		return true;
-	}
-
-	static void ModifyCompilationEnvironment(const FGlobalShaderPermutationParameters& Parameters, FShaderCompilerEnvironment& OutEnvironment)
-	{
-		FGlobalShader::ModifyCompilationEnvironment(Parameters, OutEnvironment);
-		OutEnvironment.SetDefine(TEXT("UPDATE_STREAMING_INFO_BUFFER"), 1);
-	}
-};
-IMPLEMENT_GLOBAL_SHADER(FSparseVolumeTextureUpdateStreamingInfoBufferCS, "/Engine/Private/SparseVolumeTexture/UpdateSparseVolumeTexture.usf", "SparseVolumeTextureUpdateStreamingInfoBufferCS", SF_Compute);
-
-// Utility class for uploading tiles to a physical tile data texture
-class FTileUploader
-{
-public:
-
-	struct FAddResult
-	{
-		TStaticArray<uint8*, 2> OccupancyBitsPtrs;
-		TStaticArray<uint8*, 2> TileDataOffsetsPtrs;
-		TStaticArray<uint8*, 2> TileDataPtrs;
-		TStaticArray<uint32, 2> TileDataBaseOffsets; // Caller needs to add this value to all data written to TileDataOffsetsPtrs
-		uint8* PackedPhysicalTileCoordsPtr;
-	};
-
-	FTileUploader()
-	{
-		ResetState();
-	}
-
-	void Init(FRDGBuilder& GraphBuilder, int32 InMaxNumTiles, int32 InMaxNumVoxelsA, int32 InMaxNumVoxelsB, EPixelFormat InFormatA, EPixelFormat InFormatB)
-	{
-		check(InFormatA != PF_Unknown || InFormatB != PF_Unknown);
-		ResetState();
-		MaxNumTiles = InMaxNumTiles;
-		MaxNumVoxelsA = InMaxNumVoxelsA;
-		MaxNumVoxelsB = InMaxNumVoxelsB;
-		FormatA = InFormatA;
-		FormatB = InFormatB;
-		FormatSizeA = GPixelFormats[FormatA].BlockBytes;
-		FormatSizeB = GPixelFormats[FormatB].BlockBytes;
-
-		// Create a new set of buffers if the old set is already queued into RDG.
-		if (IsRegistered(GraphBuilder, DstTileCoordsUploadBuffer))
-		{
-			OccupancyBitsUploadBuffer = nullptr;
-			TileDataOffsetsUploadBuffer = nullptr;
-			DstTileCoordsUploadBuffer = nullptr;
-			TileDataAUploadBuffer = nullptr;
-			TileDataBUploadBuffer = nullptr;
-		}
-
-		const int32 NumTextures = (FormatA == PF_Unknown || FormatB == PF_Unknown) ? 1 : 2;
-
-		if (MaxNumTiles > 0)
-		{
-			FRHICommandListBase& RHICmdList = GraphBuilder.RHICmdList;
-
-			// Occupancy bits
-			{
-				const uint32 BufferSize = NumTextures * MaxNumTiles * SVT::NumOccupancyWordsPerPaddedTile * sizeof(uint32);
-				FRDGBufferDesc BufferDesc = FRDGBufferDesc::CreateByteAddressUploadDesc(BufferSize);
-				BufferDesc.Usage |= EBufferUsageFlags::Dynamic; // Skip the unneeded copy from upload to VRAM resource on d3d12 RHI
-				AllocatePooledBuffer(BufferDesc, OccupancyBitsUploadBuffer, TEXT("SparseVolumeTexture.OccupancyBitsUploadBuffer"));
-
-				OccupancyBitsAPtr = (uint8*)RHICmdList.LockBuffer(OccupancyBitsUploadBuffer->GetRHI(), 0, BufferSize, RLM_WriteOnly);
-				OccupancyBitsBPtr = OccupancyBitsAPtr + (FormatA != PF_Unknown ? (MaxNumTiles * SVT::NumOccupancyWordsPerPaddedTile * sizeof(uint32)) : 0);
-			}
-			// Tile data offsets
-			{
-				const uint32 BufferSize = NumTextures * MaxNumTiles * sizeof(uint32);
-				FRDGBufferDesc BufferDesc = FRDGBufferDesc::CreateByteAddressUploadDesc(BufferSize);
-				BufferDesc.Usage |= EBufferUsageFlags::Dynamic; // Skip the unneeded copy from upload to VRAM resource on d3d12 RHI
-				AllocatePooledBuffer(BufferDesc, TileDataOffsetsUploadBuffer, TEXT("SparseVolumeTexture.TileDataOffsetsUploadBuffer"));
-
-				// Due to a limit on the maximum number of texels in a buffer SRV, we need to upload the data in smaller chunks. In order to figure out the chunk offsets/sizes,
-				// we need to read the TileDataOffset values the caller has written to the returned pointers. We want to avoid reading from a mapped upload buffer pointer,
-				// which is why we use a temporary allocation to write the upload data to.
-				TileDataOffsets.SetNumUninitialized(NumTextures * MaxNumTiles);
-				TileDataOffsetsAPtr = (uint8*)TileDataOffsets.GetData();
-				TileDataOffsetsBPtr = TileDataOffsetsAPtr + (FormatA != PF_Unknown ? (MaxNumTiles * sizeof(uint32)) : 0);
-			}
-			// TileCoords
-			{
-				const uint32 BufferSize = MaxNumTiles * sizeof(uint32);
-				FRDGBufferDesc BufferDesc = FRDGBufferDesc::CreateByteAddressUploadDesc(BufferSize);
-				BufferDesc.Usage |= EBufferUsageFlags::Dynamic; // Skip the unneeded copy from upload to VRAM resource on d3d12 RHI
-				AllocatePooledBuffer(BufferDesc, DstTileCoordsUploadBuffer, TEXT("SparseVolumeTexture.TileCoordsUploadBuffer"));
-
-				TileCoordsPtr = (uint8*)RHICmdList.LockBuffer(DstTileCoordsUploadBuffer->GetRHI(), 0, BufferSize, RLM_WriteOnly);
-			}
-
-			// TileData
-			if (FormatSizeA > 0)
-			{
-				FRDGBufferDesc BufferDesc = FRDGBufferDesc::CreateUploadDesc(FormatSizeA, FMath::Max(MaxNumVoxelsA, 1));
-				BufferDesc.Usage |= EBufferUsageFlags::Dynamic; // Skip the unneeded copy from upload to VRAM resource on d3d12 RHI
-				AllocatePooledBuffer(BufferDesc, TileDataAUploadBuffer, TEXT("SparseVolumeTexture.TileDataAUploadBuffer"));
-
-				TileDataAPtr = (uint8*)RHICmdList.LockBuffer(TileDataAUploadBuffer->GetRHI(), 0, FMath::Max(MaxNumVoxelsA, 1) * FormatSizeA, RLM_WriteOnly);
-			}
-			if (FormatSizeB > 0)
-			{
-				FRDGBufferDesc BufferDesc = FRDGBufferDesc::CreateUploadDesc(FormatSizeB, FMath::Max(MaxNumVoxelsB, 1));
-				BufferDesc.Usage |= EBufferUsageFlags::Dynamic; // Skip the unneeded copy from upload to VRAM resource on d3d12 RHI
-				AllocatePooledBuffer(BufferDesc, TileDataBUploadBuffer, TEXT("SparseVolumeTexture.TileDataBUploadBuffer"));
-
-				TileDataBPtr = (uint8*)RHICmdList.LockBuffer(TileDataBUploadBuffer->GetRHI(), 0, FMath::Max(MaxNumVoxelsB, 1) * FormatSizeB, RLM_WriteOnly);
-			}
-		}
-	}
-
-	FAddResult Add_GetRef(int32 NumTiles, int32 NumVoxelsA, int32 NumVoxelsB)
-	{
-		check((NumWrittenTiles + NumTiles) <= MaxNumTiles);
-		check((NumWrittenVoxelsA + NumVoxelsA) <= MaxNumVoxelsA);
-		check((NumWrittenVoxelsB + NumVoxelsB) <= MaxNumVoxelsB);
-		check(TileCoordsPtr);
-		check(FormatSizeA <= 0 || TileDataAPtr);
-		check(FormatSizeB <= 0 || TileDataBPtr);
-
-		FAddResult Result = {};
-		Result.OccupancyBitsPtrs[0] = TileDataAPtr ? OccupancyBitsAPtr + (NumWrittenTiles * SVT::NumOccupancyWordsPerPaddedTile * sizeof(uint32)) : nullptr;
-		Result.OccupancyBitsPtrs[1] = TileDataBPtr ? OccupancyBitsBPtr + (NumWrittenTiles * SVT::NumOccupancyWordsPerPaddedTile * sizeof(uint32)) : nullptr;
-		Result.TileDataOffsetsPtrs[0] = TileDataAPtr ? TileDataOffsetsAPtr + (NumWrittenTiles * sizeof(uint32)) : nullptr;
-		Result.TileDataOffsetsPtrs[1] = TileDataBPtr ? TileDataOffsetsBPtr + (NumWrittenTiles * sizeof(uint32)) : nullptr;
-		Result.TileDataPtrs[0] = TileDataAPtr ? TileDataAPtr + (NumWrittenVoxelsA * FormatSizeA) : nullptr;
-		Result.TileDataPtrs[1] = TileDataBPtr ? TileDataBPtr + (NumWrittenVoxelsB * FormatSizeB) : nullptr;
-		Result.TileDataBaseOffsets[0] = NumWrittenVoxelsA;
-		Result.TileDataBaseOffsets[1] = NumWrittenVoxelsB;
-		Result.PackedPhysicalTileCoordsPtr = TileCoordsPtr + NumWrittenTiles * sizeof(uint32);
-		
-		NumWrittenTiles += NumTiles;
-		NumWrittenVoxelsA += NumVoxelsA;
-		NumWrittenVoxelsB += NumVoxelsB;
-
-		return Result;
-	}
-
-	void Release()
-	{
-		OccupancyBitsUploadBuffer.SafeRelease();
-		TileDataOffsetsUploadBuffer.SafeRelease();
-		DstTileCoordsUploadBuffer.SafeRelease();
-		TileDataAUploadBuffer.SafeRelease();
-		TileDataBUploadBuffer.SafeRelease();
-		TileDataOffsets.Reset();
-		ResetState();
-	}
-
-	void ResourceUploadTo(FRDGBuilder& GraphBuilder, FRHITexture* DstTextureA, FRHITexture* DstTextureB, const FVector4f& FallbackValueA, const FVector4f& FallbackValueB)
-	{
-		check(DstTextureA || FormatSizeA <= 0);
-		check(DstTextureB || FormatSizeB <= 0);
-		if (MaxNumTiles > 0)
-		{
-			FRHICommandListBase& RHICmdList = GraphBuilder.RHICmdList;
-
-			RHICmdList.UnlockBuffer(OccupancyBitsUploadBuffer->GetRHI());
-			RHICmdList.UnlockBuffer(DstTileCoordsUploadBuffer->GetRHI());
-
-			// TileDataOffset values were written to a temporary allocation so that we can access them later in this function. Unlike the other buffers, we now need to copy that data over to the actual upload buffer.
-			void* TileDataOffsetsUploadPtr = RHICmdList.LockBuffer(TileDataOffsetsUploadBuffer->GetRHI(), 0, TileDataOffsets.Num() * sizeof(TileDataOffsets[0]), RLM_WriteOnly);
-			FMemory::Memcpy(TileDataOffsetsUploadPtr, TileDataOffsets.GetData(), TileDataOffsets.Num() * sizeof(TileDataOffsets[0]));
-			RHICmdList.UnlockBuffer(TileDataOffsetsUploadBuffer->GetRHI());
-
-			if (TileDataAPtr)
-			{
-				RHICmdList.UnlockBuffer(TileDataAUploadBuffer->GetRHI());
-			}
-			if (TileDataBPtr)
-			{
-				RHICmdList.UnlockBuffer(TileDataBUploadBuffer->GetRHI());
-			}
-
-			if (NumWrittenTiles > 0)
-			{
-				FRDGTexture* DummyTexture = GraphBuilder.CreateTexture(FRDGTextureDesc::Create3D(FIntVector(1), PF_R8, FClearValueBinding::None, ETextureCreateFlags::UAV), TEXT("SparseVolumeTexture.DummyTexture"));
-				FRDGTexture* DstTextureARDG = DstTextureA ? GraphBuilder.RegisterExternalTexture(CreateRenderTarget(DstTextureA, TEXT("SparseVolumeTexture.TileDataTextureA"))) : nullptr;
-				FRDGTexture* DstTextureBRDG = DstTextureB ? GraphBuilder.RegisterExternalTexture(CreateRenderTarget(DstTextureB, TEXT("SparseVolumeTexture.TileDataTextureB"))) : nullptr;
-				FRDGBuffer* SrcBufferARDG = FormatSizeA > 0 ? GraphBuilder.RegisterExternalBuffer(TileDataAUploadBuffer) : nullptr;
-				FRDGBuffer* SrcBufferBRDG = FormatSizeB > 0 ? GraphBuilder.RegisterExternalBuffer(TileDataBUploadBuffer) : nullptr;
-				check(SrcBufferARDG || SrcBufferBRDG);
-
-				FRDGTextureUAV* DummyTextureUAV = GraphBuilder.CreateUAV(DummyTexture);
-				FRDGTextureUAV* DstTextureAUAV = DstTextureARDG ? GraphBuilder.CreateUAV(DstTextureARDG) : nullptr;
-				FRDGTextureUAV* DstTextureBUAV = DstTextureBRDG ? GraphBuilder.CreateUAV(DstTextureBRDG) : nullptr;
-
-				FRDGBufferSRV* OccupancyBitsBufferSRV = GraphBuilder.CreateSRV(GraphBuilder.RegisterExternalBuffer(OccupancyBitsUploadBuffer));
-				FRDGBufferSRV* TileDataOffsetsBufferSRV = GraphBuilder.CreateSRV(GraphBuilder.RegisterExternalBuffer(TileDataOffsetsUploadBuffer));
-				FRDGBufferSRV* DstTileCoordsBufferSRV = GraphBuilder.CreateSRV(GraphBuilder.RegisterExternalBuffer(DstTileCoordsUploadBuffer));
-				// Either SrcBufferARDG or SrcBufferBRDG must exist and will have at least 1 element.
-				FRDGBufferSRV* DummySrcBufferSRV = SrcBufferARDG ? GraphBuilder.CreateSRV(SrcBufferARDG, FormatA) : GraphBuilder.CreateSRV(SrcBufferBRDG, FormatB);
-
-				auto ComputeShader = GetGlobalShaderMap(GMaxRHIFeatureLevel)->GetShader<FSparseVolumeTextureUpdateFromSparseBufferCS>();
-
-				int32 NumUploadedTiles = 0;
-				int32 NumUploadedVoxelsA = 0;
-				int32 NumUploadedVoxelsB = 0;
-
-				// This is a limit on some platforms on the maximum number of texels in a texel/typed buffer. Unfortunately for R8 (1 byte) formats, this means that we can upload only 1/16 of all the texels of a 2GB texture.
-				// In order to work around this issue, the data to be uploaded is split into chunks such that this limit is not violated. We can use the TileDataOffsets values to get the number of voxels per tile and then
-				// just add as many tiles to the batch as we can.
-				const int32 MaxNumTexelsPerResource = 1 << 27;
-
-				while (NumUploadedTiles < NumWrittenTiles)
-				{
-					// Determine the number of tiles to upload in this iteration
-					int32 NumTilesInThisBatch = 0;
-					int32 NumVoxelsAInThisBatch = 0;
-					int32 NumVoxelsBInThisBatch = 0;
-					for (int32 TileIndex = NumUploadedTiles; TileIndex < NumWrittenTiles; ++TileIndex)
-					{
-						const int32 VoxelOffsetA = FormatSizeA > 0 ? reinterpret_cast<uint32*>(TileDataOffsetsAPtr)[TileIndex] : 0;
-						const int32 VoxelOffsetB = FormatSizeB > 0 ? reinterpret_cast<uint32*>(TileDataOffsetsBPtr)[TileIndex] : 0;
-						const int32 VoxelEndIndexA = FormatSizeA > 0 && (TileIndex + 1) < NumWrittenTiles ? reinterpret_cast<uint32*>(TileDataOffsetsAPtr)[TileIndex + 1] : NumWrittenVoxelsA;
-						const int32 VoxelEndIndexB = FormatSizeB > 0 && (TileIndex + 1) < NumWrittenTiles ? reinterpret_cast<uint32*>(TileDataOffsetsBPtr)[TileIndex + 1] : NumWrittenVoxelsB;
-						const int32 TileNumVoxelsA = VoxelEndIndexA - VoxelOffsetA;
-						const int32 TileNumVoxelsB = VoxelEndIndexB - VoxelOffsetB;
-						check(TileNumVoxelsA >= 0 && TileNumVoxelsA <= SVT::NumVoxelsPerPaddedTile);
-						check(TileNumVoxelsB >= 0 && TileNumVoxelsB <= SVT::NumVoxelsPerPaddedTile);
-
-						// Adding additional voxels to the batch would exceed the limit, so exit the loop and upload the data.
-						if ((NumVoxelsAInThisBatch + TileNumVoxelsA) > MaxNumTexelsPerResource || (NumVoxelsBInThisBatch + TileNumVoxelsB) > MaxNumTexelsPerResource)
-						{
-							break;
-						}
-
-						NumTilesInThisBatch += 1;
-						NumVoxelsAInThisBatch += TileNumVoxelsA;
-						NumVoxelsBInThisBatch += TileNumVoxelsB;
-					}
-
-					FRDGBufferSRV* TileDataABufferSRV = nullptr;
-					FRDGBufferSRV* TileDataBBufferSRV = nullptr;
-
-					// This is the critical part: For every batch we create a SRV scoped to a range within the voxel data upload buffer still fitting within the typed buffer texel limit.
-					if (FormatSizeA && NumVoxelsAInThisBatch)
-					{
-						FRDGBufferSRVDesc SRVDesc(SrcBufferARDG, FormatA);
-						SRVDesc.StartOffsetBytes = NumUploadedVoxelsA * FormatSizeA;
-						SRVDesc.NumElements = NumVoxelsAInThisBatch;
-						TileDataABufferSRV = GraphBuilder.CreateSRV(SRVDesc);
-					}
-					if (FormatSizeB && NumVoxelsBInThisBatch)
-					{
-						FRDGBufferSRVDesc SRVDesc(SrcBufferBRDG, FormatB);
-						SRVDesc.StartOffsetBytes = NumUploadedVoxelsB * FormatSizeB;
-						SRVDesc.NumElements = NumVoxelsBInThisBatch;
-						TileDataBBufferSRV = GraphBuilder.CreateSRV(SRVDesc);
-					}
-
-					FSparseVolumeTextureUpdateFromSparseBufferCS::FParameters* PassParameters = GraphBuilder.AllocParameters<FSparseVolumeTextureUpdateFromSparseBufferCS::FParameters>();
-					PassParameters->DstPhysicalTileTextureA = DstTextureAUAV ? DstTextureAUAV : DummyTextureUAV;
-					PassParameters->DstPhysicalTileTextureB = DstTextureBUAV ? DstTextureBUAV : DummyTextureUAV;
-					PassParameters->SrcPhysicalTileBufferA = TileDataABufferSRV ? TileDataABufferSRV : DummySrcBufferSRV;
-					PassParameters->SrcPhysicalTileBufferB = TileDataBBufferSRV ? TileDataBBufferSRV : DummySrcBufferSRV;
-					PassParameters->OccupancyBitsBuffer = OccupancyBitsBufferSRV;
-					PassParameters->TileDataOffsetsBuffer = TileDataOffsetsBufferSRV;
-					PassParameters->DstTileCoordsBuffer = DstTileCoordsBufferSRV;
-					PassParameters->FallbackValueA = FallbackValueA;
-					PassParameters->FallbackValueB = FallbackValueB;
-					PassParameters->TileIndexOffset = NumUploadedTiles; // This lets the shader know how many tiles have already been processed in previous dispatches.
-					PassParameters->SrcVoxelDataOffsetA = NumUploadedVoxelsA; // SrcVoxelDataOffsetA and SrcVoxelDataOffsetB are subtracted from the calculated voxel data buffer read indices
-					PassParameters->SrcVoxelDataOffsetB = NumUploadedVoxelsB;
-					PassParameters->NumTilesToCopy = NumTilesInThisBatch;
-					PassParameters->BufferTileStep = MaxNumTiles;
-					PassParameters->NumDispatchedGroups = FMath::Min(NumTilesInThisBatch, GRHIMaxDispatchThreadGroupsPerDimension.X);
-					PassParameters->PaddedTileSize = SPARSE_VOLUME_TILE_RES_PADDED;
-					PassParameters->CopyTexureMask = 0;
-					PassParameters->CopyTexureMask |= FormatSizeA > 0 ? 0x1u : 0x0u;
-					PassParameters->CopyTexureMask |= FormatSizeB > 0 ? 0x2u : 0x0u;
-
-					// Disable async compute for streaming systems when MGPU is active, to work around GPU hangs
-					const bool bAsyncCompute = GSupportsEfficientAsyncCompute && (GSVTStreamingAsyncCompute != 0) && (GNumExplicitGPUsForRendering == 1);
-
-					FComputeShaderUtils::AddPass(
-						GraphBuilder,
-						RDG_EVENT_NAME("Upload SVT Tiles (TileCount: %u)", NumTilesInThisBatch),
-						bAsyncCompute ? ERDGPassFlags::AsyncCompute : ERDGPassFlags::Compute,
-						ComputeShader,
-						PassParameters,
-						FIntVector3(PassParameters->NumDispatchedGroups, 1, 1)
-					);
-
-					NumUploadedTiles += NumTilesInThisBatch;
-					NumUploadedVoxelsA += NumVoxelsAInThisBatch;
-					NumUploadedVoxelsB += NumVoxelsBInThisBatch;
-				}
-
-				check(NumUploadedTiles == NumWrittenTiles);
-				check(NumUploadedVoxelsA == NumWrittenVoxelsA);
-				check(NumUploadedVoxelsB == NumWrittenVoxelsB);
-			}
-		}
-		Release();
-	}
-
-private:
-	TRefCountPtr<FRDGPooledBuffer> OccupancyBitsUploadBuffer;
-	TRefCountPtr<FRDGPooledBuffer> TileDataOffsetsUploadBuffer;
-	TRefCountPtr<FRDGPooledBuffer> DstTileCoordsUploadBuffer;
-	TRefCountPtr<FRDGPooledBuffer> TileDataAUploadBuffer;
-	TRefCountPtr<FRDGPooledBuffer> TileDataBUploadBuffer;
-	TArray<uint32> TileDataOffsets; // CPU-readable per-tile offsets into tile data
-	uint8* OccupancyBitsAPtr = nullptr;
-	uint8* OccupancyBitsBPtr = nullptr;
-	uint8* TileDataOffsetsAPtr = nullptr;
-	uint8* TileDataOffsetsBPtr = nullptr;
-	uint8* TileCoordsPtr = nullptr;
-	uint8* TileDataAPtr = nullptr;
-	uint8* TileDataBPtr = nullptr;
-	int32 MaxNumTiles = 0;
-	int32 MaxNumVoxelsA = 0;
-	int32 MaxNumVoxelsB = 0;
-	EPixelFormat FormatA = PF_Unknown;
-	EPixelFormat FormatB = PF_Unknown;
-	int32 FormatSizeA = 0;
-	int32 FormatSizeB = 0;
-	int32 NumWrittenTiles = 0;
-	int32 NumWrittenVoxelsA = 0;
-	int32 NumWrittenVoxelsB = 0;
-
-	void ResetState()
-	{
-		OccupancyBitsAPtr = nullptr;
-		OccupancyBitsBPtr = nullptr;
-		TileDataOffsetsAPtr = nullptr;
-		TileDataOffsetsBPtr = nullptr;
-		TileCoordsPtr = nullptr;
-		TileDataAPtr = nullptr;
-		TileDataBPtr = nullptr;
-		MaxNumTiles = 0;
-		MaxNumVoxelsA = 0;
-		MaxNumVoxelsB = 0;
-		FormatA = PF_Unknown;
-		FormatB = PF_Unknown;
-		FormatSizeA = 0;
-		FormatSizeB = 0;
-		NumWrittenTiles = 0;
-		NumWrittenVoxelsA = 0;
-		NumWrittenVoxelsB = 0;
-	}
-};
-
-// Utility class for writing page table entries
-class FPageTableUpdater
-{
-public:
-	FPageTableUpdater()
-	{
-		ResetState();
-	}
-
-	void Init(FRDGBuilder& GraphBuilder, int32 InMaxNumUpdates, int32 InEstimatedNumBatches)
-	{
-		ResetState();
-		MaxNumUpdates = InMaxNumUpdates;
-		Batches.Reserve(InEstimatedNumBatches);
-
-		// Create a new buffer if the old one is already queued into RDG.
-		if (IsRegistered(GraphBuilder, UpdatesUploadBuffer))
-		{
-			UpdatesUploadBuffer = nullptr;
-		}
-
-		if (MaxNumUpdates > 0)
-		{
-			// Add EBufferUsageFlags::Dynamic to skip the unneeded copy from upload to VRAM resource on d3d12 RHI
-			FRDGBufferDesc BufferDesc = FRDGBufferDesc::CreateByteAddressUploadDesc(MaxNumUpdates * 2 * sizeof(uint32));
-			BufferDesc.Usage |= EBufferUsageFlags::Dynamic;
-			AllocatePooledBuffer(BufferDesc, UpdatesUploadBuffer, TEXT("SparseVolumeTexture.PageTableUpdatesUploadBuffer"));
-
-			DataPtr = (uint8*)GraphBuilder.RHICmdList.LockBuffer(UpdatesUploadBuffer->GetRHI(), 0, MaxNumUpdates * 2 * sizeof(uint32), RLM_WriteOnly);
-		}
-	}
-
-	void Add_GetRef(FRHITexture* PageTable, int32 MipLevel, int32 NumUpdates, uint8*& OutCoordsPtr, uint8*& OutPayloadPtr)
-	{
-		check((NumWrittenUpdates + NumUpdates) <= MaxNumUpdates);
-		check(DataPtr);
-		FBatch* Batch = Batches.IsEmpty() ? nullptr : &Batches.Last();
-		if (!Batch || Batch->PageTable != PageTable || Batch->MipLevel != MipLevel)
-		{
-			Batch = &Batches.Add_GetRef(FBatch(PageTable, MipLevel));
-		}
-
-		OutCoordsPtr = DataPtr + NumWrittenUpdates * sizeof(uint32);
-		OutPayloadPtr = DataPtr + (MaxNumUpdates + NumWrittenUpdates) * sizeof(uint32);
-
-		Batch->NumUpdates += NumUpdates;
-		NumWrittenUpdates += NumUpdates;
-	}
-
-	void Release()
-	{
-		UpdatesUploadBuffer.SafeRelease();
-		ResetState();
-	}
-
-	void Apply(FRDGBuilder& GraphBuilder)
-	{
-		if (MaxNumUpdates > 0)
-		{
-			GraphBuilder.RHICmdList.UnlockBuffer(UpdatesUploadBuffer->GetRHI());
-
-			if (NumWrittenUpdates > 0)
-			{
-				// Disable async compute for streaming systems when MGPU is active, to work around GPU hangs
-				const bool bAsyncCompute = GSupportsEfficientAsyncCompute && (GSVTStreamingAsyncCompute != 0) && (GNumExplicitGPUsForRendering == 1);
-				auto ComputeShader = GetGlobalShaderMap(GMaxRHIFeatureLevel)->GetShader<FSparseVolumeTextureUpdatePageTableCS>();
-
-				uint32 UpdatesOffset = 0;
-				for (const FBatch& Batch : Batches)
-				{
-					FRDGTexture* PageTableRDG = GraphBuilder.RegisterExternalTexture(CreateRenderTarget(Batch.PageTable, TEXT("SparseVolumeTexture.PageTableTexture")));
-					FRDGTextureUAV* PageTableUAV = GraphBuilder.CreateUAV(FRDGTextureUAVDesc(PageTableRDG, Batch.MipLevel, PF_R32_UINT));
-					FRDGBufferSRV* UpdatesBufferSRV = GraphBuilder.CreateSRV(GraphBuilder.RegisterExternalBuffer(UpdatesUploadBuffer));
-
-					FSparseVolumeTextureUpdatePageTableCS::FParameters* PassParameters = GraphBuilder.AllocParameters<FSparseVolumeTextureUpdatePageTableCS::FParameters>();
-					PassParameters->PageTable = PageTableUAV;
-					PassParameters->PageTableUpdates = UpdatesBufferSRV;
-					PassParameters->UpdateCoordOffset = UpdatesOffset;
-					PassParameters->UpdatePayloadOffset = MaxNumUpdates + UpdatesOffset;
-					PassParameters->NumUpdates = Batch.NumUpdates;
-
-					FComputeShaderUtils::AddPass(
-						GraphBuilder,
-						RDG_EVENT_NAME("Update SVT PageTable (UpdateCount: %u)", Batch.NumUpdates),
-						bAsyncCompute ? ERDGPassFlags::AsyncCompute : ERDGPassFlags::Compute,
-						ComputeShader,
-						PassParameters,
-						FComputeShaderUtils::GetGroupCount(Batch.NumUpdates, 64)
-					);
-
-					UpdatesOffset += Batch.NumUpdates;
-				}
-			}
-		}
-		
-		Release();
-	}
-
-private:
-	struct FBatch
-	{
-		FRHITexture* PageTable;
-		int32 MipLevel;
-		int32 NumUpdates;
-
-		FBatch() = default;
-		FBatch(FRHITexture* InPageTable, int32 InMipLevel) : PageTable(InPageTable), MipLevel(InMipLevel), NumUpdates(0) {}
-	};
-
-	TRefCountPtr<FRDGPooledBuffer> UpdatesUploadBuffer;
-	TArray<FBatch> Batches;
-	uint8* DataPtr = nullptr;
-	int32 NumWrittenUpdates = 0;
-	int32 MaxNumUpdates = 0;
-
-
-	void ResetState()
-	{
-		Batches.Reset();
-		DataPtr = nullptr;
-		NumWrittenUpdates = 0;
-		MaxNumUpdates = 0;
-	}
-};
-
-// Updates entries in the streaming info buffers of multiple SVTs
-class FStreamingInfoBufferUpdater
-{
-public:
-	FStreamingInfoBufferUpdater()
-	{
-		ResetState();
-	}
-
-	void Add(TRefCountPtr<FRDGPooledBuffer> StreamingInfoBuffer, int32 FrameIndex, int32 LowestResidentMipLevel)
-	{
-		FBatch* Batch = Batches.IsEmpty() ? nullptr : &Batches.Last();
-		if (!Batch || Batch->StreamingInfoBuffer != StreamingInfoBuffer)
-		{
-			Batch = &Batches.Add_GetRef(FBatch(StreamingInfoBuffer, Updates.Num()));
-		}
-
-		Updates.Add(FrameIndex);
-		Updates.Add(LowestResidentMipLevel);
-
-		++Batch->NumUpdates;
-	}
-
-	void Apply(FRDGBuilder& GraphBuilder)
-	{
-		if (!Updates.IsEmpty())
-		{
-			TRefCountPtr<FRDGPooledBuffer> UpdatesUploadBuffer;
-			{
-				// Add EBufferUsageFlags::Dynamic to skip the unneeded copy from upload to VRAM resource on d3d12 RHI
-				FRDGBufferDesc BufferDesc = FRDGBufferDesc::CreateByteAddressUploadDesc(Updates.Num() * sizeof(uint32));
-				BufferDesc.Usage |= EBufferUsageFlags::Dynamic;
-				AllocatePooledBuffer(BufferDesc, UpdatesUploadBuffer, TEXT("SparseVolumeTexture.StreamingInfoUploadBuffer"));
-
-				void* DataPtr = GraphBuilder.RHICmdList.LockBuffer(UpdatesUploadBuffer->GetRHI(), 0, Updates.Num() * sizeof(uint32), RLM_WriteOnly);
-				FMemory::Memcpy(DataPtr, Updates.GetData(), Updates.Num() * sizeof(uint32));
-				GraphBuilder.RHICmdList.UnlockBuffer(UpdatesUploadBuffer->GetRHI());
-			}
-
-			// Disable async compute for streaming systems when MGPU is active, to work around GPU hangs
-			const bool bAsyncCompute = GSupportsEfficientAsyncCompute && (GSVTStreamingAsyncCompute != 0) && (GNumExplicitGPUsForRendering == 1);
-			auto ComputeShader = GetGlobalShaderMap(GMaxRHIFeatureLevel)->GetShader<FSparseVolumeTextureUpdateStreamingInfoBufferCS>();
-
-			uint32 UpdatesOffset = 0;
-			for (const FBatch& Batch : Batches)
-			{
-				FRDGBufferUAV* StreamingInfoBufferUAV = GraphBuilder.CreateUAV(GraphBuilder.RegisterExternalBuffer(Batch.StreamingInfoBuffer), PF_R32_UINT);
-				FRDGBufferSRV* UpdatesBufferSRV = GraphBuilder.CreateSRV(GraphBuilder.RegisterExternalBuffer(UpdatesUploadBuffer));
-
-				FSparseVolumeTextureUpdateStreamingInfoBufferCS::FParameters* PassParameters = GraphBuilder.AllocParameters<FSparseVolumeTextureUpdateStreamingInfoBufferCS::FParameters>();
-				PassParameters->StreamingInfoBuffer = StreamingInfoBufferUAV;
-				PassParameters->StreamingInfoBufferUpdates = UpdatesBufferSRV;
-				PassParameters->UpdateOffset = UpdatesOffset;
-				PassParameters->NumUpdates = Batch.NumUpdates;
-
-				FComputeShaderUtils::AddPass(
-					GraphBuilder,
-					RDG_EVENT_NAME("Update SVT StreamingInfoBuffer (UpdateCount: %u)", Batch.NumUpdates),
-					bAsyncCompute ? ERDGPassFlags::AsyncCompute : ERDGPassFlags::Compute,
-					ComputeShader,
-					PassParameters,
-					FComputeShaderUtils::GetGroupCount(Batch.NumUpdates, 64)
-				);
-
-				UpdatesOffset += Batch.NumUpdates;
-			}
-		}
-		
-		ResetState();
-	}
-
-private:
-	struct FBatch
-	{
-		TRefCountPtr<FRDGPooledBuffer> StreamingInfoBuffer;
-		int32 UpdatesOffset;
-		int32 NumUpdates;
-
-		FBatch() = default;
-		FBatch(TRefCountPtr<FRDGPooledBuffer> InStreamingInfoBuffer, int32 InUpdatesOffset) : StreamingInfoBuffer(InStreamingInfoBuffer), UpdatesOffset(InUpdatesOffset), NumUpdates(0) {}
-	};
-
-	TArray<FBatch> Batches;
-	TArray<uint32> Updates;
-
-	void ResetState()
-	{
-		Batches.Reset();
-		Updates.Reset();
-	}
-};
-
 struct FStreamingUpdateParameters
 {
 	FStreamingManager* StreamingManager = nullptr;
@@ -876,7 +137,6 @@ void FStreamingManager::InitRHI(FRHICommandListBase& RHICmdList)
 	MaxPendingMipLevels = GSVTStreamingMaxPendingMipLevels;
 	PendingMipLevels.SetNum(MaxPendingMipLevels);
 	PageTableUpdater = MakeUnique<FPageTableUpdater>();
-	StreamingInfoBufferUpdater = MakeUnique<FStreamingInfoBufferUpdater>();
 
 #if WITH_EDITORONLY_DATA
 	RequestOwner = MakeUnique<FRequestOwner>(EPriority::Normal);
@@ -1173,18 +433,30 @@ void FStreamingManager::BeginAsyncUpdate(FRDGBuilder& GraphBuilder, bool bBlocki
 				continue; // Skip mip level install. SVT no longer exists or mip level was "streamed out" before it was even installed in the first place.
 			}
 
+			// Prepare tile data texture for upload
+			const FTileDataTexture::EUploaderState UploaderState = SVTInfo->TileDataTexture->GetUploaderState();
+			check(UploaderState == FTileDataTexture::EUploaderState::Ready || UploaderState == FTileDataTexture::EUploaderState::Reserving);
+			if (UploaderState == FTileDataTexture::EUploaderState::Ready)
+			{
+				SVTInfo->TileDataTexture->BeginReserveUpload();
+			}
+
 			const int32 FormatSizeA = GPixelFormats[SVTInfo->FormatA].BlockBytes;
 			const int32 FormatSizeB = GPixelFormats[SVTInfo->FormatB].BlockBytes;
 			const FResources* Resources = SVTInfo->PerFrameInfo[PendingMipLevel.FrameIndex].Resources;
-			SVTInfo->TileDataTexture->NumTilesToUpload += Resources->MipLevelStreamingInfo[PendingMipLevel.MipLevelIndex].NumPhysicalTiles;
-			SVTInfo->TileDataTexture->NumVoxelsToUploadA += FormatSizeA > 0 ? Resources->MipLevelStreamingInfo[PendingMipLevel.MipLevelIndex].TileDataSize[0] / FormatSizeA : 0;
-			SVTInfo->TileDataTexture->NumVoxelsToUploadB += FormatSizeB > 0 ? Resources->MipLevelStreamingInfo[PendingMipLevel.MipLevelIndex].TileDataSize[1] / FormatSizeB : 0;
+			const int32 NumTilesToUpload = Resources->MipLevelStreamingInfo[PendingMipLevel.MipLevelIndex].NumPhysicalTiles;
+			const int32 NumVoxelsToUploadA = FormatSizeA > 0 ? Resources->MipLevelStreamingInfo[PendingMipLevel.MipLevelIndex].TileDataSize[0] / FormatSizeA : 0;
+			const int32 NumVoxelsToUploadB = FormatSizeB > 0 ? Resources->MipLevelStreamingInfo[PendingMipLevel.MipLevelIndex].TileDataSize[1] / FormatSizeB : 0;
+
+			SVTInfo->TileDataTexture->ReserveUpload(NumTilesToUpload, NumVoxelsToUploadA, NumVoxelsToUploadB);
+
 			TileDataTexturesToUpdate.Add(SVTInfo->TileDataTexture.Get());
 		}
 
 		for (FTileDataTexture* TileDataTexture : TileDataTexturesToUpdate)
 		{
-			TileDataTexture->TileUploader->Init(GraphBuilder, TileDataTexture->NumTilesToUpload, TileDataTexture->NumVoxelsToUploadA, TileDataTexture->NumVoxelsToUploadB, TileDataTexture->FormatA, TileDataTexture->FormatB);
+			TileDataTexture->EndReserveUpload();
+			TileDataTexture->BeginUpload(GraphBuilder);
 		}
 	}
 
@@ -1223,51 +495,11 @@ void FStreamingManager::EndAsyncUpdate(FRDGBuilder& GraphBuilder)
 	}
 	AsyncTaskEvents.Empty();
 
-	// Clear unused mip levels to 0. SVT_TODO: We can probably skip this because the page table lookup in the shader is clamped
-	if (!PageTableClears.IsEmpty())
-	{
-		for (auto& Clear : PageTableClears)
-		{
-			FRDGTexture* PageTableTextureRDG = GraphBuilder.RegisterExternalTexture(CreateRenderTarget(Clear.PageTableTexture, TEXT("SparseVolumeTexture.PageTableTexture")));
-			FRDGTextureUAVDesc UAVDesc(PageTableTextureRDG, static_cast<uint8>(Clear.MipLevel), PF_R32_UINT);
-			AddClearUAVPass(GraphBuilder, GraphBuilder.CreateUAV(UAVDesc), FUintVector4(ForceInitToZero));
-		}
-		PageTableClears.Reset();
-	}
-
 	// Issue the actual data uploads
 	for (FTileDataTexture* TileDataTexture : TileDataTexturesToUpdate)
 	{
-		TileDataTexture->TileUploader->ResourceUploadTo(GraphBuilder, TileDataTexture->TileDataTextureARHIRef, TileDataTexture->TileDataTextureBRHIRef, TileDataTexture->FallbackValueA, TileDataTexture->FallbackValueB);
-		TileDataTexture->NumTilesToUpload = 0;
-		TileDataTexture->NumVoxelsToUploadA = 0;
-		TileDataTexture->NumVoxelsToUploadB = 0;
-	}
-
-	// Update streaming info buffers
-	for (FStreamingInfo* SVTInfo : SVTsWithInvalidatedStreamingInfoBuffer)
-	{
-#if DO_CHECK
-		bool bSVTInfoExists = false;
-		for (const auto& Pair : StreamingInfo)
-		{
-			if (Pair.Value.Get() == SVTInfo)
-			{
-				bSVTInfoExists = true;
-				break;
-			}
-		}
-		check(bSVTInfoExists);
-#endif
-		
-		for (TConstSetBitIterator It(SVTInfo->DirtyStreamingInfoData); It; ++It)
-		{
-			const int32 FrameIndex = It.GetIndex();
-			StreamingInfoBufferUpdater->Add(SVTInfo->StreamingInfoBuffer, FrameIndex, SVTInfo->PerFrameInfo[FrameIndex].LowestResidentMipLevel);
-		}
+		TileDataTexture->EndUpload(GraphBuilder);
 	}
-	SVTsWithInvalidatedStreamingInfoBuffer.Reset();
-	StreamingInfoBufferUpdater->Apply(GraphBuilder);
 
 	// Update page table with newly streamed in/out pages and make sure descendant pages in the hierarchy have correct fallback values
 	PatchPageTable(GraphBuilder);
@@ -1386,20 +618,11 @@ void FStreamingManager::AddInternal(FRDGBuilder& GraphBuilder, FNewSparseVolumeT
 	{
 		const int32 TileFactor = NumFrames <= 1 ? 1 : 3;
 		const int32 NumPhysicalTilesCapacity = FMath::Max(1, NumRootPhysicalTiles + (TileFactor * MaxNumPhysicalTiles)); // Ensure a minimum size of 1
-		const FIntVector3 TileDataVolumeResolution = ComputeTileDataVolumeResolution(NumPhysicalTilesCapacity);
-		const FIntVector3 TileDataVolumeResolutionInTiles = TileDataVolumeResolution / SPARSE_VOLUME_TILE_RES_PADDED;
+		const FIntVector3 TileDataVolumeResolutionInTiles = FTileDataTexture::GetVolumeResolutionInTiles(NumPhysicalTilesCapacity);
 
 		SVTInfo.TileDataTexture = MakeUnique<FTileDataTexture>(TileDataVolumeResolutionInTiles, SVTInfo.FormatA, SVTInfo.FormatB, SVTInfo.FallbackValueA, SVTInfo.FallbackValueB);
 		SVTInfo.TileDataTexture->InitResource(GraphBuilder.RHICmdList);
 
-		// Create streaming info buffer
-		{
-			FRDGBufferRef StreamingInfoBufferRDG = GraphBuilder.CreateBuffer(FRDGBufferDesc::CreateByteAddressDesc(sizeof(uint32) * NumFrames), TEXT("SparseVolumeTexture.StreamingInfo.RHIBuffer"));
-			SVTInfo.StreamingInfoBuffer = GraphBuilder.ConvertToExternalBuffer(StreamingInfoBufferRDG);
-			SVTInfo.StreamingInfoBufferSRVRHIRef = SVTInfo.StreamingInfoBuffer->GetSRV();
-			SVTInfo.DirtyStreamingInfoData.Init(true, NumFrames);
-		}
-
 		FTileUploader RootTileUploader;
 		RootTileUploader.Init(GraphBuilder, NumRootPhysicalTiles + 1 /*null tile*/, NumRootVoxelsA, NumRootVoxelsB, SVTInfo.FormatA, SVTInfo.FormatB);
 
@@ -1450,11 +673,10 @@ void FStreamingManager::AddInternal(FRDGBuilder& GraphBuilder, FNewSparseVolumeT
 
 			// Initialize TextureRenderResources
 			RHIUpdateTextureReference(FrameInfo.TextureRenderResources->PageTableTextureReferenceRHI, FrameInfo.PageTableTextureRHIRef);
-			RHIUpdateTextureReference(FrameInfo.TextureRenderResources->PhysicalTileDataATextureReferenceRHI, SVTInfo.TileDataTexture->TileDataTextureARHIRef);
-			RHIUpdateTextureReference(FrameInfo.TextureRenderResources->PhysicalTileDataBTextureReferenceRHI, SVTInfo.TileDataTexture->TileDataTextureBRHIRef);
-			FrameInfo.TextureRenderResources->StreamingInfoBufferSRVRHI = SVTInfo.StreamingInfoBufferSRVRHIRef;
+			RHIUpdateTextureReference(FrameInfo.TextureRenderResources->PhysicalTileDataATextureReferenceRHI, SVTInfo.TileDataTexture->GetTileDataTextureA());
+			RHIUpdateTextureReference(FrameInfo.TextureRenderResources->PhysicalTileDataBTextureReferenceRHI, SVTInfo.TileDataTexture->GetTileDataTextureB());
 			FrameInfo.TextureRenderResources->Header = Resources->Header;
-			FrameInfo.TextureRenderResources->TileDataTextureResolution = SVTInfo.TileDataTexture->ResolutionInTiles * SPARSE_VOLUME_TILE_RES_PADDED;
+			FrameInfo.TextureRenderResources->TileDataTextureResolution = SVTInfo.TileDataTexture->GetResolutionInTiles() * SPARSE_VOLUME_TILE_RES_PADDED;
 			FrameInfo.TextureRenderResources->FrameIndex = FrameIdx;
 			FrameInfo.TextureRenderResources->NumLogicalMipLevels = NumMipLevels;
 
@@ -1505,11 +727,9 @@ void FStreamingManager::AddInternal(FRDGBuilder& GraphBuilder, FNewSparseVolumeT
 			InvalidatedSVTFrames.Add(&FrameInfo);
 		}
 
-		RootTileUploader.ResourceUploadTo(GraphBuilder, SVTInfo.TileDataTexture->TileDataTextureARHIRef, SVTInfo.TileDataTexture->TileDataTextureBRHIRef, SVTInfo.FallbackValueA, SVTInfo.FallbackValueB);
+		RootTileUploader.ResourceUploadTo(GraphBuilder, SVTInfo.TileDataTexture->GetTileDataTextureA(), SVTInfo.TileDataTexture->GetTileDataTextureB(), SVTInfo.FallbackValueA, SVTInfo.FallbackValueB);
 	}
 
-	SVTsWithInvalidatedStreamingInfoBuffer.Add(&SVTInfo);
-
 	// Add requests for all mips the first frame. This is necessary for cases where UAnimatedSparseVolumeTexture or UStaticSparseVolumeTexture
 	// are directly bound to the material without getting a specific frame through USparseVolumeTextureFrame::GetFrameAndIssueStreamingRequest().
 	const int32 NumMipLevelsFrame0 = SVTInfo.PerFrameInfo[0].NumMipLevels;
@@ -1566,8 +786,6 @@ void FStreamingManager::RemoveInternal(UStreamableSparseVolumeTexture* SparseVol
 			SVTInfo->TileDataTexture.Reset();
 		}
 
-		SVTsWithInvalidatedStreamingInfoBuffer.Remove(SVTInfo);
-
 		StreamingInfo.Remove(SparseVolumeTexture);
 	}
 }
@@ -1750,7 +968,7 @@ void FStreamingManager::IssueRequests(int32 MaxSelectedRequests)
 		check(TileDataTexture);
 
 		// Ensure that enough tiles are available in the tile texture
-		const int32 TileDataTextureCapacity = TileDataTexture->PhysicalTilesCapacity;
+		const int32 TileDataTextureCapacity = TileDataTexture->GetTileCapacity();
 		const int32 NumAvailableTiles = TileDataTexture->GetNumAvailableTiles();
 		const int32 NumRequiredTiles = MipLevelStreamingInfo.NumPhysicalTiles;
 		if (NumAvailableTiles < NumRequiredTiles)
@@ -1959,25 +1177,15 @@ void FStreamingManager::StreamOutMipLevel(FStreamingInfo* SVTInfo, FLRUNode* LRU
 		FrameInfo.LowestResidentMipLevel, NewLowestResidentMipLevel);
 #endif
 
-	// Only clear page table mip if the mip to be freed is actually resident.
-	const bool bNeedsPageTableClear = MipLevelIndex >= FrameInfo.LowestResidentMipLevel;
 	FrameInfo.LowestRequestedMipLevel = NewLowestRequestedMipLevel;
 	FrameInfo.LowestResidentMipLevel = NewLowestResidentMipLevel;
 	
-	// Update the streaming info buffer data
-	SVTInfo->DirtyStreamingInfoData[FrameIndex] = true;
-	SVTsWithInvalidatedStreamingInfoBuffer.Add(SVTInfo);
 	InvalidatedSVTFrames.Add(&FrameInfo);
 
 	// Unlink
 	LRUNode->Remove();
 	LRUNode->LastRequested = INDEX_NONE;
 
-	if (bNeedsPageTableClear)
-	{
-		PageTableClears.Push({ FrameInfo.PageTableTextureRHIRef, MipLevelIndex });
-	}
-
 	// Free allocated tiles
 	for (uint32& TileCoord : FrameInfo.TileAllocations[MipLevelIndex])
 	{
@@ -2167,7 +1375,7 @@ void FStreamingManager::InstallReadyMipLevels()
 		check((MipLevelStreamingInfo.PageTableSize % (sizeof(uint32) * 2)) == 0);
 		const int32 NumPageTableUpdates = MipLevelStreamingInfo.PageTableSize / (sizeof(uint32) * 2);
 
-		FTileUploader::FAddResult TileDataAddResult = SVTInfo->TileDataTexture->TileUploader->Add_GetRef(NumPhysicalTiles, NumVoxelsA, NumVoxelsB);
+		FTileUploader::FAddResult TileDataAddResult = SVTInfo->TileDataTexture->AddUpload(NumPhysicalTiles, NumVoxelsA, NumVoxelsB);
 
 		// Tile data
 		{
@@ -2219,9 +1427,6 @@ void FStreamingManager::InstallReadyMipLevels()
 		check(FrameInfo.LowestResidentMipLevel == (PendingMipLevel.MipLevelIndex + 1));
 		FrameInfo.LowestResidentMipLevel = PendingMipLevel.MipLevelIndex;
 
-		// Update the streaming info buffer data
-		SVTInfo->DirtyStreamingInfoData[PendingMipLevel.FrameIndex] = true;
-		SVTsWithInvalidatedStreamingInfoBuffer.Add(SVTInfo);
 		InvalidatedSVTFrames.Add(&FrameInfo);
 
 		const int32 LRUNodeIndex = PendingMipLevel.FrameIndex * SVTInfo->NumMipLevelsGlobal + PendingMipLevel.MipLevelIndex;
@@ -2533,83 +1738,5 @@ void FStreamingManager::RequestDDCData(TConstArrayView<UE::DerivedData::FCacheGe
 
 #endif // WITH_EDITORONLY_DATA
 
-FStreamingManager::FTileDataTexture::FTileDataTexture(const FIntVector3& InResolutionInTiles, EPixelFormat InFormatA, EPixelFormat InFormatB, const FVector4f& InFallbackValueA, const FVector4f& InFallbackValueB)
-	: ResolutionInTiles(InResolutionInTiles), 
-	PhysicalTilesCapacity(InResolutionInTiles.X * InResolutionInTiles.Y * InResolutionInTiles.Z), 
-	FormatA(InFormatA),
-	FormatB(InFormatB),
-	FallbackValueA(InFallbackValueA),
-	FallbackValueB(InFallbackValueB),
-	TileUploader(MakeUnique<FTileUploader>()),
-	NumTilesToUpload(0),
-	NumVoxelsToUploadA(0),
-	NumVoxelsToUploadB(0)
-{
-	const int64 MaxFormatSize = FMath::Max(GPixelFormats[FormatA].BlockBytes, GPixelFormats[FormatB].BlockBytes);
-	const FIntVector3 LargestPossibleResolution = ComputeLargestPossibleTileDataVolumeResolution(MaxFormatSize);
-	const int32 LargestPossiblePhysicalTilesCapacity = LargestPossibleResolution.X * LargestPossibleResolution.Y * LargestPossibleResolution.Z;
-	
-	// Ensure that the tile data texture(s) do not exceed the memory size and resolution limits.
-	if (PhysicalTilesCapacity > LargestPossiblePhysicalTilesCapacity
-		|| (ResolutionInTiles.X * SPARSE_VOLUME_TILE_RES_PADDED) > SVT::MaxVolumeTextureDim
-		|| (ResolutionInTiles.Y * SPARSE_VOLUME_TILE_RES_PADDED) > SVT::MaxVolumeTextureDim
-		|| (ResolutionInTiles.Z * SPARSE_VOLUME_TILE_RES_PADDED) > SVT::MaxVolumeTextureDim)
-	{
-		ResolutionInTiles = LargestPossibleResolution;
-		PhysicalTilesCapacity = LargestPossiblePhysicalTilesCapacity;
-
-		UE_LOG(LogSparseVolumeTextureStreamingManager, Warning, TEXT("Requested SparseVolumeTexture tile data texture resolution (in tiles) (%i, %i, %i) exceeds the resource size limit. Using the maximum value of (%i, %i. %i) instead."),
-			InResolutionInTiles.X, InResolutionInTiles.Y, InResolutionInTiles.Z,
-			LargestPossibleResolution.X, LargestPossibleResolution.Y, LargestPossibleResolution.Z);
-	}
-
-	const FIntVector3 Resolution = ResolutionInTiles * SPARSE_VOLUME_TILE_RES_PADDED;
-	check(Resolution.X <= SVT::MaxVolumeTextureDim && Resolution.Y <= SVT::MaxVolumeTextureDim && Resolution.Z <= SVT::MaxVolumeTextureDim);
-	check(((int64)Resolution.X * (int64)Resolution.Y * (int64)Resolution.Z * (int64)GPixelFormats[FormatA].BlockBytes) <= SVT::MaxResourceSize);
-	check(((int64)Resolution.X * (int64)Resolution.Y * (int64)Resolution.Z * (int64)GPixelFormats[FormatB].BlockBytes) <= SVT::MaxResourceSize);
-	
-	TileCoords.SetNum(PhysicalTilesCapacity);
-
-	int32 TileCoordsIndex = 0;
-	for (int32 Z = 0; Z < ResolutionInTiles.Z; ++Z)
-	{
-		for (int32 Y = 0; Y < ResolutionInTiles.Y; ++Y)
-		{
-			for (int32 X = 0; X < ResolutionInTiles.X; ++X)
-			{
-				uint32 PackedCoord = 0;
-				PackedCoord |= (X & 0xFFu);
-				PackedCoord |= (Y & 0xFFu) << 8u;
-				PackedCoord |= (Z & 0xFFu) << 16u;
-				TileCoords[TileCoordsIndex++] = PackedCoord;
-			}
-		}
-	}
-	check(TileCoordsIndex == PhysicalTilesCapacity);
-}
-
-void FStreamingManager::FTileDataTexture::InitRHI(FRHICommandListBase&)
-{
-	const FIntVector3 Resolution = ResolutionInTiles * SPARSE_VOLUME_TILE_RES_PADDED;
-	if (FormatA != PF_Unknown)
-	{
-		const FRHITextureCreateDesc Desc =
-			FRHITextureCreateDesc::Create3D(TEXT("SparseVolumeTexture.PhysicalTileDataA.RHITexture"), Resolution.X, Resolution.Y, Resolution.Z, FormatA)
-			.SetFlags(ETextureCreateFlags::ShaderResource | ETextureCreateFlags::UAV);
-		TileDataTextureARHIRef = RHICreateTexture(Desc);
-	}
-	if (FormatB != PF_Unknown)
-	{
-		const FRHITextureCreateDesc Desc =
-			FRHITextureCreateDesc::Create3D(TEXT("SparseVolumeTexture.PhysicalTileDataB.RHITexture"), Resolution.X, Resolution.Y, Resolution.Z, FormatB)
-			.SetFlags(ETextureCreateFlags::ShaderResource | ETextureCreateFlags::UAV);
-		TileDataTextureBRHIRef = RHICreateTexture(Desc);
-	}
-}
-
-void FStreamingManager::FTileDataTexture::ReleaseRHI()
-{
-}
-
 }
 }
diff --git a/Engine/Source/Runtime/Engine/Private/SparseVolumeTexture/SparseVolumeTextureStreamingManager.h b/Engine/Source/Runtime/Engine/Private/SparseVolumeTexture/SparseVolumeTextureStreamingManager.h
--- a/Engine/Source/Runtime/Engine/Private/SparseVolumeTexture/SparseVolumeTextureStreamingManager.h
+++ b/Engine/Source/Runtime/Engine/Private/SparseVolumeTexture/SparseVolumeTextureStreamingManager.h
@@ -15,6 +15,8 @@
 #include "Containers/Union.h"
 #include "RenderGraphBuilder.h"
 
+DECLARE_LOG_CATEGORY_EXTERN(LogSparseVolumeTextureStreamingManager, Log, All);
+
 class UStreamableSparseVolumeTexture;
 class USparseVolumeTextureFrame;
 
@@ -35,6 +37,9 @@ namespace SVT
 struct FResources;
 struct FMipLevelStreamingInfo;
 class FTextureRenderResources;
+class FTileDataTexture;
+class FTileUploader;
+class FPageTableUpdater;
 
 // Uniquely identifies the mip level of a frame in a static or animated SparseVolumeTexture
 struct FMipLevelKey
@@ -101,68 +106,6 @@ public:
 private:
 	friend class FStreamingUpdateTask;
 
-	// SVT_TODO: actually implement streaming in/out of page table mips greater than this value.
-	static constexpr int32 NumAlwaysResidentPageTableMipLevels = 6;
-
-	// Represents the physical tile data texture that serves as backing memory for the streamed in tiles. While this is treated as a single logical texture,
-	// it currently supports up to two actual RHI textures.
-	struct FTileDataTexture : public FRenderResource
-	{
-	public:
-		static constexpr uint32 PhysicalCoordMask = (1u << 24u) - 1u; // Lower 24 bits are used for storing XYZ in 8 bit each. Upper 8 bit can be used by the caller. 
-
-		FIntVector3 ResolutionInTiles;
-		int32 PhysicalTilesCapacity;
-		EPixelFormat FormatA;
-		EPixelFormat FormatB;
-		FVector4f FallbackValueA;
-		FVector4f FallbackValueB;
-		FTextureRHIRef TileDataTextureARHIRef;
-		FTextureRHIRef TileDataTextureBRHIRef;
-		TUniquePtr<class FTileUploader> TileUploader;
-		int32 NumTilesToUpload;
-		int32 NumVoxelsToUploadA;
-		int32 NumVoxelsToUploadB;
-
-		// Constructor. May change the requested ResolutionInTiles (and resulting PhysicalTilesCapacity) if it exceeds hardware limits.
-		FTileDataTexture(const FIntVector3& ResolutionInTiles, EPixelFormat FormatA, EPixelFormat FormatB, const FVector4f& FallbackValueA, const FVector4f& FallbackValueB);
-
-		// Allocate a tile slot in the texture. The resulting value is a packed coordinate (8 bit per component) of the allocated slot or INDEX_NONE if the allocation failed.
-		// The upper 8 bit are free to be used by the caller.
-		uint32 Allocate() 
-		{ 
-			check(PhysicalTilesCapacity == TileCoords.Num());
-			return TileCoords.IsValidIndex(NextFreeTileCoordIndex) ? TileCoords[NextFreeTileCoordIndex++] : INDEX_NONE; 
-		}
-
-		// Frees a previously allocated tile slot. The upper 8 bit (user data) are automatically cleared by this function.
-		void Free(uint32 PackedPhysicalCoord) 
-		{ 
-			check(PackedPhysicalCoord != INDEX_NONE)
-			PackedPhysicalCoord &= PhysicalCoordMask;
-			check(PhysicalTilesCapacity == TileCoords.Num());
-			check(NextFreeTileCoordIndex > 0);
-#if DO_GUARD_SLOW
-			for (int32 i = NextFreeTileCoordIndex; i < PhysicalTilesCapacity; ++i)
-			{
-				check(TileCoords[i] != PackedPhysicalCoord);
-			}
-#endif
-			TileCoords[--NextFreeTileCoordIndex] = PackedPhysicalCoord;
-		}
-
-		// Number of tiles available for allocation.
-		int32 GetNumAvailableTiles() { return PhysicalTilesCapacity - NextFreeTileCoordIndex; }
-		virtual void InitRHI(FRHICommandListBase& RHICmdList) override;
-		virtual void ReleaseRHI() override;
-
-	private:
-
-		TArray<uint32> TileCoords;
-		int32 NextFreeTileCoordIndex = 0;
-		uint32 NumRefs = 0;
-	};
-
 	// Represents a context or "window" into the frame sequence that moves along the playback direction. It is used to cache the prefetch direction.
 	struct FStreamingWindow
 	{
@@ -218,9 +161,6 @@ private:
 		TArray<FStreamingWindow> StreamingWindows;
 
 		TUniquePtr<FTileDataTexture> TileDataTexture;
-		TRefCountPtr<FRDGPooledBuffer> StreamingInfoBuffer; // One uint32 per frame storing the LowestResidentMipLevel
-		FShaderResourceViewRHIRef StreamingInfoBufferSRVRHIRef;
-		TBitArray<> DirtyStreamingInfoData; // One bit per frame, potentially marking the streaming info data as dirty/in need of an update
 	};
 
 	// Represents an IO request for a mip level
@@ -296,12 +236,6 @@ private:
 		TUnion<FTileDataTask, FPageTableTask> Union;
 	};
 
-	struct FPageTableClear
-	{
-		FTextureRHIRef PageTableTexture;
-		int32 MipLevel;
-	};
-
 	struct FAsyncState
 	{
 		int32 NumReadyMipLevels = 0;
@@ -330,7 +264,6 @@ private:
 #endif
 
 	TUniquePtr<class FPageTableUpdater> PageTableUpdater;
-	TUniquePtr<class FStreamingInfoBufferUpdater> StreamingInfoBufferUpdater;
 	FGraphEventArray AsyncTaskEvents;
 	FAsyncState AsyncState;
 	int32 MaxPendingMipLevels = 0;
@@ -341,11 +274,9 @@ private:
 	// Transient lifetime
 	TArray<FStreamingRequest> ParentRequestsToAdd;
 	TSet<FTileDataTexture*> TileDataTexturesToUpdate;
-	TSet<FStreamingInfo*> SVTsWithInvalidatedStreamingInfoBuffer; // Changes in the lowest resident mip level cause invalidation
 	TSet<FFrameInfo*> InvalidatedSVTFrames; // Set of SVT frames where pages have been streamed in or out. Used in PatchPageTable().
 	TArray<FStreamingRequest> PrioritizedRequestsHeap;
 	TArray<FStreamingRequest> SelectedRequests;
-	TArray<FPageTableClear> PageTableClears;
 	TArray<FUploadTask> UploadTasks; // accessed on the async thread
 	TArray<FPendingMipLevel*> UploadCleanupTasks; // accessed on the async thread
 
diff --git a/Engine/Source/Runtime/Engine/Private/SparseVolumeTexture/SparseVolumeTextureTileDataTexture.cpp b/Engine/Source/Runtime/Engine/Private/SparseVolumeTexture/SparseVolumeTextureTileDataTexture.cpp
new file mode 100644
--- /dev/null
+++ b/Engine/Source/Runtime/Engine/Private/SparseVolumeTexture/SparseVolumeTextureTileDataTexture.cpp
@@ -0,0 +1,192 @@
+// Copyright Epic Games, Inc. All Rights Reserved.
+
+#include "SparseVolumeTextureTileDataTexture.h"
+#include "SparseVolumeTextureUtility.h"
+#include "SparseVolumeTextureStreamingManager.h" // LogSparseVolumeTextureStreamingManager
+
+namespace UE
+{
+namespace SVT
+{
+
+FIntVector3 FTileDataTexture::GetVolumeResolutionInTiles(int32 InNumRequiredTiles)
+{
+	int32 TileVolumeResolutionCube = 1;
+	while (TileVolumeResolutionCube * TileVolumeResolutionCube * TileVolumeResolutionCube < InNumRequiredTiles)
+	{
+		TileVolumeResolutionCube++;				// We use a simple loop to compute the minimum resolution of a cube to store all the tile data
+	}
+	FIntVector3 TileDataVolumeResolution = FIntVector3(TileVolumeResolutionCube, TileVolumeResolutionCube, TileVolumeResolutionCube);
+
+	// Trim volume to reclaim some space
+	while ((TileDataVolumeResolution.X * TileDataVolumeResolution.Y * (TileDataVolumeResolution.Z - 1)) > InNumRequiredTiles)
+	{
+		TileDataVolumeResolution.Z--;
+	}
+	while ((TileDataVolumeResolution.X * (TileDataVolumeResolution.Y - 1) * TileDataVolumeResolution.Z) > InNumRequiredTiles)
+	{
+		TileDataVolumeResolution.Y--;
+	}
+	while (((TileDataVolumeResolution.X - 1) * TileDataVolumeResolution.Y * TileDataVolumeResolution.Z) > InNumRequiredTiles)
+	{
+		TileDataVolumeResolution.X--;
+	}
+
+	return TileDataVolumeResolution;
+}
+
+FIntVector3 FTileDataTexture::GetLargestPossibleVolumeResolutionInTiles(int32 InVoxelMemSize)
+{
+	const int64 TileMemSize = SVT::NumVoxelsPerPaddedTile * InVoxelMemSize;
+	const int64 NumMaxTiles = SVT::MaxResourceSize / TileMemSize;
+
+	// Find a cube with a volume as close to NumMaxTiles as possible
+	int32 TileVolumeResolutionCube = 1;
+	while (((TileVolumeResolutionCube + 1) * (TileVolumeResolutionCube + 1) * (TileVolumeResolutionCube + 1)) <= NumMaxTiles)
+	{
+		++TileVolumeResolutionCube;
+	}
+
+	// Try to add to the sides to get closer to NumMaxTiles
+	FIntVector3 ResolutionInTiles = FIntVector3(TileVolumeResolutionCube, TileVolumeResolutionCube, TileVolumeResolutionCube);
+	if (((ResolutionInTiles.X + 1) * ResolutionInTiles.Y * ResolutionInTiles.Z) <= NumMaxTiles)
+	{
+		++ResolutionInTiles.X;
+	}
+	if ((ResolutionInTiles.X * (ResolutionInTiles.Y + 1) * ResolutionInTiles.Z) <= NumMaxTiles)
+	{
+		++ResolutionInTiles.Y;
+	}
+	if ((ResolutionInTiles.X * ResolutionInTiles.Y * (ResolutionInTiles.Z + 1)) <= NumMaxTiles)
+	{
+		++ResolutionInTiles.Z;
+	}
+
+	const FIntVector3 ResolutionInVoxels = ResolutionInTiles * SPARSE_VOLUME_TILE_RES_PADDED;
+	check(IsInBounds(ResolutionInVoxels, FIntVector3::ZeroValue, FIntVector3(SVT::MaxVolumeTextureDim + 1)));
+	check(ResolutionInVoxels.X <= SVT::MaxVolumeTextureDim && ResolutionInVoxels.Y <= SVT::MaxVolumeTextureDim && ResolutionInVoxels.Z <= SVT::MaxVolumeTextureDim);
+	check(((int64)ResolutionInVoxels.X * (int64)ResolutionInVoxels.Y * (int64)ResolutionInVoxels.Z * (int64)InVoxelMemSize) <= SVT::MaxResourceSize);
+
+	return ResolutionInTiles;
+}
+
+FTileDataTexture::FTileDataTexture(const FIntVector3& InResolutionInTiles, EPixelFormat InFormatA, EPixelFormat InFormatB, const FVector4f& InFallbackValueA, const FVector4f& InFallbackValueB)
+	: TileUploader(MakeUnique<FTileUploader>()),
+	ResolutionInTiles(InResolutionInTiles),
+	PhysicalTilesCapacity(InResolutionInTiles.X* InResolutionInTiles.Y* InResolutionInTiles.Z),
+	FormatA(InFormatA),
+	FormatB(InFormatB),
+	FallbackValueA(InFallbackValueA),
+	FallbackValueB(InFallbackValueB)
+{
+	const int64 MaxFormatSize = FMath::Max(GPixelFormats[FormatA].BlockBytes, GPixelFormats[FormatB].BlockBytes);
+	const FIntVector3 LargestPossibleResolutionInTiles = GetLargestPossibleVolumeResolutionInTiles(MaxFormatSize);
+	const int32 LargestPossiblePhysicalTilesCapacity = LargestPossibleResolutionInTiles.X * LargestPossibleResolutionInTiles.Y * LargestPossibleResolutionInTiles.Z;
+
+	// Ensure that the tile data texture(s) do not exceed the memory size and resolution limits.
+	if (PhysicalTilesCapacity > LargestPossiblePhysicalTilesCapacity
+		|| (ResolutionInTiles.X * SPARSE_VOLUME_TILE_RES_PADDED) > SVT::MaxVolumeTextureDim
+		|| (ResolutionInTiles.Y * SPARSE_VOLUME_TILE_RES_PADDED) > SVT::MaxVolumeTextureDim
+		|| (ResolutionInTiles.Z * SPARSE_VOLUME_TILE_RES_PADDED) > SVT::MaxVolumeTextureDim)
+	{
+		ResolutionInTiles = LargestPossibleResolutionInTiles;
+		PhysicalTilesCapacity = LargestPossiblePhysicalTilesCapacity;
+
+		UE_LOG(LogSparseVolumeTextureStreamingManager, Warning, TEXT("Requested SparseVolumeTexture tile data texture resolution (in tiles) (%i, %i, %i) exceeds the resource size limit. Using the maximum value of (%i, %i. %i) instead."),
+			InResolutionInTiles.X, InResolutionInTiles.Y, InResolutionInTiles.Z,
+			LargestPossibleResolutionInTiles.X, LargestPossibleResolutionInTiles.Y, LargestPossibleResolutionInTiles.Z);
+	}
+
+	const FIntVector3 Resolution = ResolutionInTiles * SPARSE_VOLUME_TILE_RES_PADDED;
+	check(Resolution.X <= SVT::MaxVolumeTextureDim && Resolution.Y <= SVT::MaxVolumeTextureDim && Resolution.Z <= SVT::MaxVolumeTextureDim);
+	check(((int64)Resolution.X * (int64)Resolution.Y * (int64)Resolution.Z * (int64)GPixelFormats[FormatA].BlockBytes) <= SVT::MaxResourceSize);
+	check(((int64)Resolution.X * (int64)Resolution.Y * (int64)Resolution.Z * (int64)GPixelFormats[FormatB].BlockBytes) <= SVT::MaxResourceSize);
+
+	TileCoords.SetNum(PhysicalTilesCapacity);
+
+	int32 TileCoordsIndex = 0;
+	for (int32 Z = 0; Z < ResolutionInTiles.Z; ++Z)
+	{
+		for (int32 Y = 0; Y < ResolutionInTiles.Y; ++Y)
+		{
+			for (int32 X = 0; X < ResolutionInTiles.X; ++X)
+			{
+				uint32 PackedCoord = 0;
+				PackedCoord |= (X & 0xFFu);
+				PackedCoord |= (Y & 0xFFu) << 8u;
+				PackedCoord |= (Z & 0xFFu) << 16u;
+				TileCoords[TileCoordsIndex++] = PackedCoord;
+			}
+		}
+	}
+	check(TileCoordsIndex == PhysicalTilesCapacity);
+}
+
+void FTileDataTexture::BeginReserveUpload()
+{
+	check(UploaderState == EUploaderState::Ready || UploaderState == EUploaderState::Reserved);
+	UploaderState = EUploaderState::Reserving;
+	NumReservedUploadTiles = 0;
+	NumReservedUploadVoxelsA = 0;
+	NumReservedUploadVoxelsB = 0;
+}
+
+void FTileDataTexture::ReserveUpload(int32 NumTiles, int32 NumVoxelsA, int32 NumVoxelsB)
+{
+	check(UploaderState == EUploaderState::Reserving);
+	NumReservedUploadTiles += NumTiles;
+	NumReservedUploadVoxelsA += NumVoxelsA;
+	NumReservedUploadVoxelsB += NumVoxelsB;
+}
+
+void FTileDataTexture::EndReserveUpload()
+{
+	check(UploaderState == EUploaderState::Reserving);
+	UploaderState = EUploaderState::Reserved;
+}
+
+void FTileDataTexture::BeginUpload(FRDGBuilder& GraphBuilder)
+{
+	check(UploaderState == EUploaderState::Reserved);
+	TileUploader->Init(GraphBuilder, NumReservedUploadTiles, NumReservedUploadVoxelsA, NumReservedUploadVoxelsB, FormatA, FormatB);
+	UploaderState = EUploaderState::Uploading;
+}
+
+FTileUploader::FAddResult FTileDataTexture::AddUpload(int32 NumTiles, int32 NumVoxelsA, int32 NumVoxelsB)
+{
+	check(UploaderState == EUploaderState::Uploading);
+	return TileUploader->Add_GetRef(NumTiles, NumVoxelsA, NumVoxelsB);
+}
+
+void FTileDataTexture::EndUpload(FRDGBuilder& GraphBuilder)
+{
+	check(UploaderState == EUploaderState::Uploading);
+	TileUploader->ResourceUploadTo(GraphBuilder, TileDataTextureARHIRef, TileDataTextureBRHIRef, FallbackValueA, FallbackValueB);
+	UploaderState = EUploaderState::Ready;
+}
+
+void FTileDataTexture::InitRHI(FRHICommandListBase&)
+{
+	const FIntVector3 Resolution = ResolutionInTiles * SPARSE_VOLUME_TILE_RES_PADDED;
+	if (FormatA != PF_Unknown)
+	{
+		const FRHITextureCreateDesc Desc =
+			FRHITextureCreateDesc::Create3D(TEXT("SparseVolumeTexture.PhysicalTileDataA.RHITexture"), Resolution.X, Resolution.Y, Resolution.Z, FormatA)
+			.SetFlags(ETextureCreateFlags::ShaderResource | ETextureCreateFlags::UAV);
+		TileDataTextureARHIRef = RHICreateTexture(Desc);
+	}
+	if (FormatB != PF_Unknown)
+	{
+		const FRHITextureCreateDesc Desc =
+			FRHITextureCreateDesc::Create3D(TEXT("SparseVolumeTexture.PhysicalTileDataB.RHITexture"), Resolution.X, Resolution.Y, Resolution.Z, FormatB)
+			.SetFlags(ETextureCreateFlags::ShaderResource | ETextureCreateFlags::UAV);
+		TileDataTextureBRHIRef = RHICreateTexture(Desc);
+	}
+}
+
+void FTileDataTexture::ReleaseRHI()
+{
+}
+
+}
+}
\ No newline at end of file
diff --git a/Engine/Source/Runtime/Engine/Private/SparseVolumeTexture/SparseVolumeTextureTileDataTexture.h b/Engine/Source/Runtime/Engine/Private/SparseVolumeTexture/SparseVolumeTextureTileDataTexture.h
new file mode 100644
--- /dev/null
+++ b/Engine/Source/Runtime/Engine/Private/SparseVolumeTexture/SparseVolumeTextureTileDataTexture.h
@@ -0,0 +1,104 @@
+// Copyright Epic Games, Inc. All Rights Reserved.
+
+#pragma once
+
+#include "CoreMinimal.h"
+#include "RenderResource.h"
+#include "SparseVolumeTextureUpload.h"
+
+class FRHICommandListBase;
+class FRDGBuilder;
+enum EPixelFormat : uint8;
+
+namespace UE
+{
+namespace SVT
+{
+
+class FTileUploader;
+
+// Represents the physical tile data texture that serves as backing memory for the streamed in tiles. While this is treated as a single logical texture,
+// it currently supports up to two actual RHI textures.
+class FTileDataTexture : public FRenderResource
+{
+public:
+	enum class EUploaderState
+	{
+		Ready, Reserving, Reserved, Uploading
+	};
+
+	static constexpr uint32 PhysicalCoordMask = (1u << 24u) - 1u; // Lower 24 bits are used for storing XYZ in 8 bit each. Upper 8 bit can be used by the caller. 
+
+	static FIntVector3 GetVolumeResolutionInTiles(int32 InNumRequiredTiles);
+	static FIntVector3 GetLargestPossibleVolumeResolutionInTiles(int32 InVoxelMemSize);
+
+	// Constructor. May change the requested ResolutionInTiles (and resulting PhysicalTilesCapacity) if it exceeds hardware limits.
+	explicit FTileDataTexture(const FIntVector3& ResolutionInTiles, EPixelFormat FormatA, EPixelFormat FormatB, const FVector4f& FallbackValueA, const FVector4f& FallbackValueB);
+
+	// Allocate a tile slot in the texture. The resulting value is a packed coordinate (8 bit per component) of the allocated slot or INDEX_NONE if the allocation failed.
+	// The upper 8 bit are free to be used by the caller.
+	uint32 Allocate()
+	{
+		check(PhysicalTilesCapacity == TileCoords.Num());
+		return TileCoords.IsValidIndex(NextFreeTileCoordIndex) ? TileCoords[NextFreeTileCoordIndex++] : INDEX_NONE;
+	}
+
+	// Frees a previously allocated tile slot. The upper 8 bit (user data) are automatically cleared by this function.
+	void Free(uint32 PackedPhysicalCoord)
+	{
+		check(PackedPhysicalCoord != INDEX_NONE)
+			PackedPhysicalCoord &= PhysicalCoordMask;
+		check(PhysicalTilesCapacity == TileCoords.Num());
+		check(NextFreeTileCoordIndex > 0);
+#if DO_GUARD_SLOW
+		for (int32 i = NextFreeTileCoordIndex; i < PhysicalTilesCapacity; ++i)
+		{
+			check(TileCoords[i] != PackedPhysicalCoord);
+		}
+#endif
+		TileCoords[--NextFreeTileCoordIndex] = PackedPhysicalCoord;
+	}
+
+	EUploaderState GetUploaderState() const { return UploaderState; }
+	int32 GetTileCapacity() const { return PhysicalTilesCapacity; }
+	int32 GetNumAvailableTiles() const { return PhysicalTilesCapacity - NextFreeTileCoordIndex; } // Number of tiles available for allocation.
+	FIntVector3 GetResolutionInTiles() const { return ResolutionInTiles; }
+	FTextureRHIRef GetTileDataTextureA() { return TileDataTextureARHIRef; }
+	FTextureRHIRef GetTileDataTextureB() { return TileDataTextureBRHIRef; }
+
+	// Transitions from EUploaderState::Ready (or EUploaderState::Reserved) to EUploaderState::Reserving and allows callers to call ReserveUpload() afterwards. Resets number of reserved tiles/voxels.
+	void BeginReserveUpload();
+	// Reserves space in the upload buffer to hold the given number of tiles and voxels in addition to all tiles and voxels reserved with prior calls to this function. Must be in EUploaderState::Reserving.
+	void ReserveUpload(int32 NumTiles, int32 NumVoxelsA, int32 NumVoxelsB);
+	// Transitions from EUploaderState::Reserving to EUploaderState::Reserved.
+	void EndReserveUpload();
+	// Transitions from EUploaderState::Reserved to EUploaderState::Uploading and allows callers to call AddUpload() afterwards.
+	void BeginUpload(FRDGBuilder& GraphBuilder);
+	// Returns pointers/offsets into upload buffer memory for the caller to write data into. Must be in EUploaderState::Uploading.
+	FTileUploader::FAddResult AddUpload(int32 NumTiles, int32 NumVoxelsA, int32 NumVoxelsB);
+	// Actually uploads the data written to the upload buffer and transitions from EUploaderState::Uploading to EUploaderState::Ready.
+	void EndUpload(FRDGBuilder& GraphBuilder);
+
+	virtual void InitRHI(FRHICommandListBase& RHICmdList) override;
+	virtual void ReleaseRHI() override;
+
+private:
+	TUniquePtr<FTileUploader> TileUploader;
+	FIntVector3 ResolutionInTiles;
+	int32 PhysicalTilesCapacity;
+	EPixelFormat FormatA;
+	EPixelFormat FormatB;
+	FVector4f FallbackValueA;
+	FVector4f FallbackValueB;
+	FTextureRHIRef TileDataTextureARHIRef;
+	FTextureRHIRef TileDataTextureBRHIRef;
+	TArray<uint32> TileCoords;
+	int32 NextFreeTileCoordIndex = 0;
+	EUploaderState UploaderState = EUploaderState::Ready;
+	int32 NumReservedUploadTiles = 0;
+	int32 NumReservedUploadVoxelsA = 0;
+	int32 NumReservedUploadVoxelsB = 0;
+};
+
+}
+}
diff --git a/Engine/Source/Runtime/Engine/Private/SparseVolumeTexture/SparseVolumeTextureUpload.cpp b/Engine/Source/Runtime/Engine/Private/SparseVolumeTexture/SparseVolumeTextureUpload.cpp
new file mode 100644
--- /dev/null
+++ b/Engine/Source/Runtime/Engine/Private/SparseVolumeTexture/SparseVolumeTextureUpload.cpp
@@ -0,0 +1,517 @@
+// Copyright Epic Games, Inc. All Rights Reserved.
+
+#include "SparseVolumeTextureUpload.h"
+#include "SparseVolumeTextureUtility.h"
+#include "RenderCore.h"
+#include "RenderGraph.h"
+#include "RenderGraphBuilder.h"
+#include "RenderGraphUtils.h"
+#include "GlobalShader.h"
+
+static int32 GSVTStreamingAsyncCompute = 1;
+static FAutoConsoleVariableRef CVarSVTStreamingAsyncCompute(
+	TEXT("r.SparseVolumeTexture.Streaming.AsyncCompute"),
+	GSVTStreamingAsyncCompute,
+	TEXT("Schedule GPU work in async compute queue."),
+	ECVF_RenderThreadSafe
+);
+
+class FSparseVolumeTextureUpdateFromBufferCS : public FGlobalShader
+{
+	DECLARE_GLOBAL_SHADER(FSparseVolumeTextureUpdateFromBufferCS);
+	SHADER_USE_PARAMETER_STRUCT(FSparseVolumeTextureUpdateFromBufferCS, FGlobalShader)
+
+	BEGIN_SHADER_PARAMETER_STRUCT(FParameters, )
+		SHADER_PARAMETER_RDG_TEXTURE_UAV(RWTexture3D, DstPhysicalTileTextureA)
+		SHADER_PARAMETER_RDG_TEXTURE_UAV(RWTexture3D, DstPhysicalTileTextureB)
+		SHADER_PARAMETER_RDG_BUFFER_SRV(ByteAddressBuffer, DstTileCoords)
+		SHADER_PARAMETER_RDG_BUFFER_SRV(Buffer<float4>, SrcPhysicalTileBufferA)
+		SHADER_PARAMETER_RDG_BUFFER_SRV(Buffer<float4>, SrcPhysicalTileBufferB)
+		SHADER_PARAMETER(uint32, TileCoordsBufferOffset)
+		SHADER_PARAMETER(uint32, TileDataBufferOffsetInTiles)
+		SHADER_PARAMETER(uint32, NumTilesToCopy)
+		SHADER_PARAMETER(uint32, NumDispatchedGroups)
+		SHADER_PARAMETER(uint32, PaddedTileSize)
+		SHADER_PARAMETER(uint32, bCopyTexureAOnlyUI)
+	END_SHADER_PARAMETER_STRUCT()
+
+	static bool ShouldCompilePermutation(const FGlobalShaderPermutationParameters& Parameters)
+	{
+		return true;
+	}
+
+	static void ModifyCompilationEnvironment(const FGlobalShaderPermutationParameters& Parameters, FShaderCompilerEnvironment& OutEnvironment)
+	{
+		FGlobalShader::ModifyCompilationEnvironment(Parameters, OutEnvironment);
+		OutEnvironment.SetDefine(TEXT("UPDATE_TILE_TEXTURE_FROM_BUFFER"), 1);
+	}
+};
+IMPLEMENT_GLOBAL_SHADER(FSparseVolumeTextureUpdateFromBufferCS, "/Engine/Private/SparseVolumeTexture/UpdateSparseVolumeTexture.usf", "SparseVolumeTextureUpdateFromBufferCS", SF_Compute);
+
+class FSparseVolumeTextureUpdateFromSparseBufferCS : public FGlobalShader
+{
+	DECLARE_GLOBAL_SHADER(FSparseVolumeTextureUpdateFromSparseBufferCS);
+	SHADER_USE_PARAMETER_STRUCT(FSparseVolumeTextureUpdateFromSparseBufferCS, FGlobalShader)
+
+	BEGIN_SHADER_PARAMETER_STRUCT(FParameters, )
+		SHADER_PARAMETER_RDG_TEXTURE_UAV(RWTexture3D, DstPhysicalTileTextureA)
+		SHADER_PARAMETER_RDG_TEXTURE_UAV(RWTexture3D, DstPhysicalTileTextureB)
+		SHADER_PARAMETER_RDG_BUFFER_SRV(Buffer<float4>, SrcPhysicalTileBufferA)
+		SHADER_PARAMETER_RDG_BUFFER_SRV(Buffer<float4>, SrcPhysicalTileBufferB)
+		SHADER_PARAMETER_RDG_BUFFER_SRV(ByteAddressBuffer, OccupancyBitsBuffer)
+		SHADER_PARAMETER_RDG_BUFFER_SRV(ByteAddressBuffer, TileDataOffsetsBuffer)
+		SHADER_PARAMETER_RDG_BUFFER_SRV(ByteAddressBuffer, DstTileCoordsBuffer)
+		SHADER_PARAMETER(FVector4f, FallbackValueA)
+		SHADER_PARAMETER(FVector4f, FallbackValueB)
+		SHADER_PARAMETER(uint32, TileIndexOffset)
+		SHADER_PARAMETER(uint32, SrcVoxelDataOffsetA)
+		SHADER_PARAMETER(uint32, SrcVoxelDataOffsetB)
+		SHADER_PARAMETER(uint32, NumTilesToCopy)
+		SHADER_PARAMETER(uint32, BufferTileStep)
+		SHADER_PARAMETER(uint32, NumDispatchedGroups)
+		SHADER_PARAMETER(uint32, PaddedTileSize)
+		SHADER_PARAMETER(uint32, CopyTexureMask)
+	END_SHADER_PARAMETER_STRUCT()
+
+	static bool ShouldCompilePermutation(const FGlobalShaderPermutationParameters& Parameters)
+	{
+		return true;
+	}
+
+	static void ModifyCompilationEnvironment(const FGlobalShaderPermutationParameters& Parameters, FShaderCompilerEnvironment& OutEnvironment)
+	{
+		FGlobalShader::ModifyCompilationEnvironment(Parameters, OutEnvironment);
+		OutEnvironment.SetDefine(TEXT("UPDATE_TILE_TEXTURE_FROM_SPARSE_BUFFER"), 1);
+	}
+};
+IMPLEMENT_GLOBAL_SHADER(FSparseVolumeTextureUpdateFromSparseBufferCS, "/Engine/Private/SparseVolumeTexture/UpdateSparseVolumeTexture.usf", "SparseVolumeTextureUpdateFromSparseBufferCS", SF_Compute);
+
+class FSparseVolumeTextureUpdatePageTableCS : public FGlobalShader
+{
+	DECLARE_GLOBAL_SHADER(FSparseVolumeTextureUpdatePageTableCS);
+	SHADER_USE_PARAMETER_STRUCT(FSparseVolumeTextureUpdatePageTableCS, FGlobalShader)
+
+	BEGIN_SHADER_PARAMETER_STRUCT(FParameters, )
+		SHADER_PARAMETER_RDG_TEXTURE_UAV(RWTexture3D, PageTable)
+		SHADER_PARAMETER_RDG_BUFFER_SRV(ByteAddressBuffer, PageTableUpdates)
+		SHADER_PARAMETER(uint32, UpdateCoordOffset)
+		SHADER_PARAMETER(uint32, UpdatePayloadOffset)
+		SHADER_PARAMETER(uint32, NumUpdates)
+	END_SHADER_PARAMETER_STRUCT()
+
+	static bool ShouldCompilePermutation(const FGlobalShaderPermutationParameters& Parameters)
+	{
+		return true;
+	}
+
+	static void ModifyCompilationEnvironment(const FGlobalShaderPermutationParameters& Parameters, FShaderCompilerEnvironment& OutEnvironment)
+	{
+		FGlobalShader::ModifyCompilationEnvironment(Parameters, OutEnvironment);
+		OutEnvironment.SetDefine(TEXT("UPDATE_PAGE_TABLE"), 1);
+	}
+};
+IMPLEMENT_GLOBAL_SHADER(FSparseVolumeTextureUpdatePageTableCS, "/Engine/Private/SparseVolumeTexture/UpdateSparseVolumeTexture.usf", "SparseVolumeTextureUpdatePageTableCS", SF_Compute);
+
+namespace UE
+{
+namespace SVT
+{
+
+FTileUploader::FTileUploader()
+{
+	ResetState();
+}
+
+void FTileUploader::Init(FRDGBuilder& GraphBuilder, int32 InMaxNumTiles, int32 InMaxNumVoxelsA, int32 InMaxNumVoxelsB, EPixelFormat InFormatA, EPixelFormat InFormatB)
+{
+	check(InFormatA != PF_Unknown || InFormatB != PF_Unknown);
+	ResetState();
+	MaxNumTiles = InMaxNumTiles;
+	MaxNumVoxelsA = InMaxNumVoxelsA;
+	MaxNumVoxelsB = InMaxNumVoxelsB;
+	FormatA = InFormatA;
+	FormatB = InFormatB;
+	FormatSizeA = GPixelFormats[FormatA].BlockBytes;
+	FormatSizeB = GPixelFormats[FormatB].BlockBytes;
+	
+	// Create a new set of buffers if the old set is already queued into RDG.
+	if (IsRegistered(GraphBuilder, DstTileCoordsUploadBuffer))
+	{
+		OccupancyBitsUploadBuffer = nullptr;
+		TileDataOffsetsUploadBuffer = nullptr;
+		DstTileCoordsUploadBuffer = nullptr;
+		TileDataAUploadBuffer = nullptr;
+		TileDataBUploadBuffer = nullptr;
+	}
+	
+	const int32 NumTextures = (FormatA == PF_Unknown || FormatB == PF_Unknown) ? 1 : 2;
+	
+	if (MaxNumTiles > 0)
+	{
+		FRHICommandListBase& RHICmdList = GraphBuilder.RHICmdList;
+	
+		// Occupancy bits
+		{
+			const uint32 BufferSize = NumTextures * MaxNumTiles * SVT::NumOccupancyWordsPerPaddedTile * sizeof(uint32);
+			FRDGBufferDesc BufferDesc = FRDGBufferDesc::CreateByteAddressUploadDesc(BufferSize);
+			BufferDesc.Usage |= EBufferUsageFlags::Dynamic; // Skip the unneeded copy from upload to VRAM resource on d3d12 RHI
+			AllocatePooledBuffer(BufferDesc, OccupancyBitsUploadBuffer, TEXT("SparseVolumeTexture.OccupancyBitsUploadBuffer"));
+	
+			OccupancyBitsAPtr = (uint8*)RHICmdList.LockBuffer(OccupancyBitsUploadBuffer->GetRHI(), 0, BufferSize, RLM_WriteOnly);
+			OccupancyBitsBPtr = OccupancyBitsAPtr + (FormatA != PF_Unknown ? (MaxNumTiles * SVT::NumOccupancyWordsPerPaddedTile * sizeof(uint32)) : 0);
+		}
+		// Tile data offsets
+		{
+			const uint32 BufferSize = NumTextures * MaxNumTiles * sizeof(uint32);
+			FRDGBufferDesc BufferDesc = FRDGBufferDesc::CreateByteAddressUploadDesc(BufferSize);
+			BufferDesc.Usage |= EBufferUsageFlags::Dynamic; // Skip the unneeded copy from upload to VRAM resource on d3d12 RHI
+			AllocatePooledBuffer(BufferDesc, TileDataOffsetsUploadBuffer, TEXT("SparseVolumeTexture.TileDataOffsetsUploadBuffer"));
+	
+			// Due to a limit on the maximum number of texels in a buffer SRV, we need to upload the data in smaller chunks. In order to figure out the chunk offsets/sizes,
+			// we need to read the TileDataOffset values the caller has written to the returned pointers. We want to avoid reading from a mapped upload buffer pointer,
+			// which is why we use a temporary allocation to write the upload data to.
+			TileDataOffsets.SetNumUninitialized(NumTextures * MaxNumTiles);
+			TileDataOffsetsAPtr = (uint8*)TileDataOffsets.GetData();
+			TileDataOffsetsBPtr = TileDataOffsetsAPtr + (FormatA != PF_Unknown ? (MaxNumTiles * sizeof(uint32)) : 0);
+		}
+		// TileCoords
+		{
+			const uint32 BufferSize = MaxNumTiles * sizeof(uint32);
+			FRDGBufferDesc BufferDesc = FRDGBufferDesc::CreateByteAddressUploadDesc(BufferSize);
+			BufferDesc.Usage |= EBufferUsageFlags::Dynamic; // Skip the unneeded copy from upload to VRAM resource on d3d12 RHI
+			AllocatePooledBuffer(BufferDesc, DstTileCoordsUploadBuffer, TEXT("SparseVolumeTexture.TileCoordsUploadBuffer"));
+	
+			TileCoordsPtr = (uint8*)RHICmdList.LockBuffer(DstTileCoordsUploadBuffer->GetRHI(), 0, BufferSize, RLM_WriteOnly);
+		}
+	
+		// TileData
+		if (FormatSizeA > 0)
+		{
+			FRDGBufferDesc BufferDesc = FRDGBufferDesc::CreateUploadDesc(FormatSizeA, FMath::Max(MaxNumVoxelsA, 1));
+			BufferDesc.Usage |= EBufferUsageFlags::Dynamic; // Skip the unneeded copy from upload to VRAM resource on d3d12 RHI
+			AllocatePooledBuffer(BufferDesc, TileDataAUploadBuffer, TEXT("SparseVolumeTexture.TileDataAUploadBuffer"));
+	
+			TileDataAPtr = (uint8*)RHICmdList.LockBuffer(TileDataAUploadBuffer->GetRHI(), 0, FMath::Max(MaxNumVoxelsA, 1) * FormatSizeA, RLM_WriteOnly);
+		}
+		if (FormatSizeB > 0)
+		{
+			FRDGBufferDesc BufferDesc = FRDGBufferDesc::CreateUploadDesc(FormatSizeB, FMath::Max(MaxNumVoxelsB, 1));
+			BufferDesc.Usage |= EBufferUsageFlags::Dynamic; // Skip the unneeded copy from upload to VRAM resource on d3d12 RHI
+			AllocatePooledBuffer(BufferDesc, TileDataBUploadBuffer, TEXT("SparseVolumeTexture.TileDataBUploadBuffer"));
+	
+			TileDataBPtr = (uint8*)RHICmdList.LockBuffer(TileDataBUploadBuffer->GetRHI(), 0, FMath::Max(MaxNumVoxelsB, 1) * FormatSizeB, RLM_WriteOnly);
+		}
+	}
+}
+
+FTileUploader::FAddResult FTileUploader::Add_GetRef(int32 InNumTiles, int32 InNumVoxelsA, int32 InNumVoxelsB)
+{
+	check((NumWrittenTiles + InNumTiles) <= MaxNumTiles);
+	check((NumWrittenVoxelsA + InNumVoxelsA) <= MaxNumVoxelsA);
+	check((NumWrittenVoxelsB + InNumVoxelsB) <= MaxNumVoxelsB);
+	check(TileCoordsPtr);
+	check(FormatSizeA <= 0 || TileDataAPtr);
+	check(FormatSizeB <= 0 || TileDataBPtr);
+	
+	FAddResult Result = {};
+	Result.OccupancyBitsPtrs[0] = TileDataAPtr ? OccupancyBitsAPtr + (NumWrittenTiles * SVT::NumOccupancyWordsPerPaddedTile * sizeof(uint32)) : nullptr;
+	Result.OccupancyBitsPtrs[1] = TileDataBPtr ? OccupancyBitsBPtr + (NumWrittenTiles * SVT::NumOccupancyWordsPerPaddedTile * sizeof(uint32)) : nullptr;
+	Result.TileDataOffsetsPtrs[0] = TileDataAPtr ? TileDataOffsetsAPtr + (NumWrittenTiles * sizeof(uint32)) : nullptr;
+	Result.TileDataOffsetsPtrs[1] = TileDataBPtr ? TileDataOffsetsBPtr + (NumWrittenTiles * sizeof(uint32)) : nullptr;
+	Result.TileDataPtrs[0] = TileDataAPtr ? TileDataAPtr + (NumWrittenVoxelsA * FormatSizeA) : nullptr;
+	Result.TileDataPtrs[1] = TileDataBPtr ? TileDataBPtr + (NumWrittenVoxelsB * FormatSizeB) : nullptr;
+	Result.TileDataBaseOffsets[0] = NumWrittenVoxelsA;
+	Result.TileDataBaseOffsets[1] = NumWrittenVoxelsB;
+	Result.PackedPhysicalTileCoordsPtr = TileCoordsPtr + NumWrittenTiles * sizeof(uint32);
+	
+	NumWrittenTiles += InNumTiles;
+	NumWrittenVoxelsA += InNumVoxelsA;
+	NumWrittenVoxelsB += InNumVoxelsB;
+	
+	return Result;
+}
+
+void FTileUploader::Release()
+{
+	OccupancyBitsUploadBuffer.SafeRelease();
+	TileDataOffsetsUploadBuffer.SafeRelease();
+	DstTileCoordsUploadBuffer.SafeRelease();
+	TileDataAUploadBuffer.SafeRelease();
+	TileDataBUploadBuffer.SafeRelease();
+	TileDataOffsets.Reset();
+	ResetState();
+}
+
+void FTileUploader::ResourceUploadTo(FRDGBuilder& GraphBuilder, FRHITexture* InDstTextureA, FRHITexture* InDstTextureB, const FVector4f& InFallbackValueA, const FVector4f& InFallbackValueB)
+{
+	check(InDstTextureA || FormatSizeA <= 0);
+	check(InDstTextureB || FormatSizeB <= 0);
+	if (MaxNumTiles > 0)
+	{
+		FRHICommandListBase& RHICmdList = GraphBuilder.RHICmdList;
+	
+		RHICmdList.UnlockBuffer(OccupancyBitsUploadBuffer->GetRHI());
+		RHICmdList.UnlockBuffer(DstTileCoordsUploadBuffer->GetRHI());
+	
+		// TileDataOffset values were written to a temporary allocation so that we can access them later in this function. Unlike the other buffers, we now need to copy that data over to the actual upload buffer.
+		void* TileDataOffsetsUploadPtr = RHICmdList.LockBuffer(TileDataOffsetsUploadBuffer->GetRHI(), 0, TileDataOffsets.Num() * sizeof(TileDataOffsets[0]), RLM_WriteOnly);
+		FMemory::Memcpy(TileDataOffsetsUploadPtr, TileDataOffsets.GetData(), TileDataOffsets.Num() * sizeof(TileDataOffsets[0]));
+		RHICmdList.UnlockBuffer(TileDataOffsetsUploadBuffer->GetRHI());
+	
+		if (TileDataAPtr)
+		{
+			RHICmdList.UnlockBuffer(TileDataAUploadBuffer->GetRHI());
+		}
+		if (TileDataBPtr)
+		{
+			RHICmdList.UnlockBuffer(TileDataBUploadBuffer->GetRHI());
+		}
+	
+		if (NumWrittenTiles > 0)
+		{
+			FRDGTexture* DummyTexture = GraphBuilder.CreateTexture(FRDGTextureDesc::Create3D(FIntVector(1), PF_R8, FClearValueBinding::None, ETextureCreateFlags::UAV), TEXT("SparseVolumeTexture.DummyTexture"));
+			FRDGTexture* DstTextureARDG = InDstTextureA ? GraphBuilder.RegisterExternalTexture(CreateRenderTarget(InDstTextureA, TEXT("SparseVolumeTexture.TileDataTextureA"))) : nullptr;
+			FRDGTexture* DstTextureBRDG = InDstTextureB ? GraphBuilder.RegisterExternalTexture(CreateRenderTarget(InDstTextureB, TEXT("SparseVolumeTexture.TileDataTextureB"))) : nullptr;
+			FRDGBuffer* SrcBufferARDG = FormatSizeA > 0 ? GraphBuilder.RegisterExternalBuffer(TileDataAUploadBuffer) : nullptr;
+			FRDGBuffer* SrcBufferBRDG = FormatSizeB > 0 ? GraphBuilder.RegisterExternalBuffer(TileDataBUploadBuffer) : nullptr;
+			check(SrcBufferARDG || SrcBufferBRDG);
+	
+			FRDGTextureUAV* DummyTextureUAV = GraphBuilder.CreateUAV(DummyTexture);
+			FRDGTextureUAV* DstTextureAUAV = DstTextureARDG ? GraphBuilder.CreateUAV(DstTextureARDG) : nullptr;
+			FRDGTextureUAV* DstTextureBUAV = DstTextureBRDG ? GraphBuilder.CreateUAV(DstTextureBRDG) : nullptr;
+	
+			FRDGBufferSRV* OccupancyBitsBufferSRV = GraphBuilder.CreateSRV(GraphBuilder.RegisterExternalBuffer(OccupancyBitsUploadBuffer));
+			FRDGBufferSRV* TileDataOffsetsBufferSRV = GraphBuilder.CreateSRV(GraphBuilder.RegisterExternalBuffer(TileDataOffsetsUploadBuffer));
+			FRDGBufferSRV* DstTileCoordsBufferSRV = GraphBuilder.CreateSRV(GraphBuilder.RegisterExternalBuffer(DstTileCoordsUploadBuffer));
+			// Either SrcBufferARDG or SrcBufferBRDG must exist and will have at least 1 element.
+			FRDGBufferSRV* DummySrcBufferSRV = SrcBufferARDG ? GraphBuilder.CreateSRV(SrcBufferARDG, FormatA) : GraphBuilder.CreateSRV(SrcBufferBRDG, FormatB);
+	
+			auto ComputeShader = GetGlobalShaderMap(GMaxRHIFeatureLevel)->GetShader<FSparseVolumeTextureUpdateFromSparseBufferCS>();
+	
+			int32 NumUploadedTiles = 0;
+			int32 NumUploadedVoxelsA = 0;
+			int32 NumUploadedVoxelsB = 0;
+	
+			// This is a limit on some platforms on the maximum number of texels in a texel/typed buffer. Unfortunately for R8 (1 byte) formats, this means that we can upload only 1/16 of all the texels of a 2GB texture.
+			// In order to work around this issue, the data to be uploaded is split into chunks such that this limit is not violated. We can use the TileDataOffsets values to get the number of voxels per tile and then
+			// just add as many tiles to the batch as we can.
+			const int32 MaxNumTexelsPerResource = 1 << 27;
+	
+			while (NumUploadedTiles < NumWrittenTiles)
+			{
+				// Determine the number of tiles to upload in this iteration
+				int32 NumTilesInThisBatch = 0;
+				int32 NumVoxelsAInThisBatch = 0;
+				int32 NumVoxelsBInThisBatch = 0;
+				for (int32 TileIndex = NumUploadedTiles; TileIndex < NumWrittenTiles; ++TileIndex)
+				{
+					const int32 VoxelOffsetA = FormatSizeA > 0 ? reinterpret_cast<uint32*>(TileDataOffsetsAPtr)[TileIndex] : 0;
+					const int32 VoxelOffsetB = FormatSizeB > 0 ? reinterpret_cast<uint32*>(TileDataOffsetsBPtr)[TileIndex] : 0;
+					const int32 VoxelEndIndexA = FormatSizeA > 0 && (TileIndex + 1) < NumWrittenTiles ? reinterpret_cast<uint32*>(TileDataOffsetsAPtr)[TileIndex + 1] : NumWrittenVoxelsA;
+					const int32 VoxelEndIndexB = FormatSizeB > 0 && (TileIndex + 1) < NumWrittenTiles ? reinterpret_cast<uint32*>(TileDataOffsetsBPtr)[TileIndex + 1] : NumWrittenVoxelsB;
+					const int32 TileNumVoxelsA = VoxelEndIndexA - VoxelOffsetA;
+					const int32 TileNumVoxelsB = VoxelEndIndexB - VoxelOffsetB;
+					check(TileNumVoxelsA >= 0 && TileNumVoxelsA <= SVT::NumVoxelsPerPaddedTile);
+					check(TileNumVoxelsB >= 0 && TileNumVoxelsB <= SVT::NumVoxelsPerPaddedTile);
+	
+					// Adding additional voxels to the batch would exceed the limit, so exit the loop and upload the data.
+					if ((NumVoxelsAInThisBatch + TileNumVoxelsA) > MaxNumTexelsPerResource || (NumVoxelsBInThisBatch + TileNumVoxelsB) > MaxNumTexelsPerResource)
+					{
+						break;
+					}
+	
+					NumTilesInThisBatch += 1;
+					NumVoxelsAInThisBatch += TileNumVoxelsA;
+					NumVoxelsBInThisBatch += TileNumVoxelsB;
+				}
+	
+				FRDGBufferSRV* TileDataABufferSRV = nullptr;
+				FRDGBufferSRV* TileDataBBufferSRV = nullptr;
+	
+				// This is the critical part: For every batch we create a SRV scoped to a range within the voxel data upload buffer still fitting within the typed buffer texel limit.
+				if (FormatSizeA && NumVoxelsAInThisBatch)
+				{
+					FRDGBufferSRVDesc SRVDesc(SrcBufferARDG, FormatA);
+					SRVDesc.StartOffsetBytes = NumUploadedVoxelsA * FormatSizeA;
+					SRVDesc.NumElements = NumVoxelsAInThisBatch;
+					TileDataABufferSRV = GraphBuilder.CreateSRV(SRVDesc);
+				}
+				if (FormatSizeB && NumVoxelsBInThisBatch)
+				{
+					FRDGBufferSRVDesc SRVDesc(SrcBufferBRDG, FormatB);
+					SRVDesc.StartOffsetBytes = NumUploadedVoxelsB * FormatSizeB;
+					SRVDesc.NumElements = NumVoxelsBInThisBatch;
+					TileDataBBufferSRV = GraphBuilder.CreateSRV(SRVDesc);
+				}
+	
+				FSparseVolumeTextureUpdateFromSparseBufferCS::FParameters* PassParameters = GraphBuilder.AllocParameters<FSparseVolumeTextureUpdateFromSparseBufferCS::FParameters>();
+				PassParameters->DstPhysicalTileTextureA = DstTextureAUAV ? DstTextureAUAV : DummyTextureUAV;
+				PassParameters->DstPhysicalTileTextureB = DstTextureBUAV ? DstTextureBUAV : DummyTextureUAV;
+				PassParameters->SrcPhysicalTileBufferA = TileDataABufferSRV ? TileDataABufferSRV : DummySrcBufferSRV;
+				PassParameters->SrcPhysicalTileBufferB = TileDataBBufferSRV ? TileDataBBufferSRV : DummySrcBufferSRV;
+				PassParameters->OccupancyBitsBuffer = OccupancyBitsBufferSRV;
+				PassParameters->TileDataOffsetsBuffer = TileDataOffsetsBufferSRV;
+				PassParameters->DstTileCoordsBuffer = DstTileCoordsBufferSRV;
+				PassParameters->FallbackValueA = InFallbackValueA;
+				PassParameters->FallbackValueB = InFallbackValueB;
+				PassParameters->TileIndexOffset = NumUploadedTiles; // This lets the shader know how many tiles have already been processed in previous dispatches.
+				PassParameters->SrcVoxelDataOffsetA = NumUploadedVoxelsA; // SrcVoxelDataOffsetA and SrcVoxelDataOffsetB are subtracted from the calculated voxel data buffer read indices
+				PassParameters->SrcVoxelDataOffsetB = NumUploadedVoxelsB;
+				PassParameters->NumTilesToCopy = NumTilesInThisBatch;
+				PassParameters->BufferTileStep = MaxNumTiles;
+				PassParameters->NumDispatchedGroups = FMath::Min(NumTilesInThisBatch, GRHIMaxDispatchThreadGroupsPerDimension.X);
+				PassParameters->PaddedTileSize = SPARSE_VOLUME_TILE_RES_PADDED;
+				PassParameters->CopyTexureMask = 0;
+				PassParameters->CopyTexureMask |= FormatSizeA > 0 ? 0x1u : 0x0u;
+				PassParameters->CopyTexureMask |= FormatSizeB > 0 ? 0x2u : 0x0u;
+	
+				// Disable async compute for streaming systems when MGPU is active, to work around GPU hangs
+				const bool bAsyncCompute = GSupportsEfficientAsyncCompute && (GSVTStreamingAsyncCompute != 0) && (GNumExplicitGPUsForRendering == 1);
+	
+				FComputeShaderUtils::AddPass(
+					GraphBuilder,
+					RDG_EVENT_NAME("Upload SVT Tiles (TileCount: %u)", NumTilesInThisBatch),
+					bAsyncCompute ? ERDGPassFlags::AsyncCompute : ERDGPassFlags::Compute,
+					ComputeShader,
+					PassParameters,
+					FIntVector3(PassParameters->NumDispatchedGroups, 1, 1)
+				);
+	
+				NumUploadedTiles += NumTilesInThisBatch;
+				NumUploadedVoxelsA += NumVoxelsAInThisBatch;
+				NumUploadedVoxelsB += NumVoxelsBInThisBatch;
+			}
+	
+			check(NumUploadedTiles == NumWrittenTiles);
+			check(NumUploadedVoxelsA == NumWrittenVoxelsA);
+			check(NumUploadedVoxelsB == NumWrittenVoxelsB);
+		}
+	}
+	Release();
+}
+
+void FTileUploader::ResetState()
+{
+	OccupancyBitsAPtr = nullptr;
+	OccupancyBitsBPtr = nullptr;
+	TileDataOffsetsAPtr = nullptr;
+	TileDataOffsetsBPtr = nullptr;
+	TileCoordsPtr = nullptr;
+	TileDataAPtr = nullptr;
+	TileDataBPtr = nullptr;
+	MaxNumTiles = 0;
+	MaxNumVoxelsA = 0;
+	MaxNumVoxelsB = 0;
+	FormatA = PF_Unknown;
+	FormatB = PF_Unknown;
+	FormatSizeA = 0;
+	FormatSizeB = 0;
+	NumWrittenTiles = 0;
+	NumWrittenVoxelsA = 0;
+	NumWrittenVoxelsB = 0;
+}
+
+FPageTableUpdater::FPageTableUpdater()
+{
+	ResetState();
+}
+
+void FPageTableUpdater::Init(FRDGBuilder& GraphBuilder, int32 InMaxNumUpdates, int32 InEstimatedNumBatches)
+{
+	ResetState();
+	MaxNumUpdates = InMaxNumUpdates;
+	Batches.Reserve(InEstimatedNumBatches);
+	
+	// Create a new buffer if the old one is already queued into RDG.
+	if (IsRegistered(GraphBuilder, UpdatesUploadBuffer))
+	{
+		UpdatesUploadBuffer = nullptr;
+	}
+	
+	if (MaxNumUpdates > 0)
+	{
+		// Add EBufferUsageFlags::Dynamic to skip the unneeded copy from upload to VRAM resource on d3d12 RHI
+		FRDGBufferDesc BufferDesc = FRDGBufferDesc::CreateByteAddressUploadDesc(MaxNumUpdates * 2 * sizeof(uint32));
+		BufferDesc.Usage |= EBufferUsageFlags::Dynamic;
+		AllocatePooledBuffer(BufferDesc, UpdatesUploadBuffer, TEXT("SparseVolumeTexture.PageTableUpdatesUploadBuffer"));
+	
+		DataPtr = (uint8*)GraphBuilder.RHICmdList.LockBuffer(UpdatesUploadBuffer->GetRHI(), 0, MaxNumUpdates * 2 * sizeof(uint32), RLM_WriteOnly);
+	}
+}
+
+void FPageTableUpdater::Add_GetRef(FRHITexture* InPageTable, int32 InMipLevel, int32 InNumUpdates, uint8*& OutCoordsPtr, uint8*& OutPayloadPtr)
+{
+	check((NumWrittenUpdates + InNumUpdates) <= MaxNumUpdates);
+	check(DataPtr);
+	FBatch* Batch = Batches.IsEmpty() ? nullptr : &Batches.Last();
+	if (!Batch || Batch->PageTable != InPageTable || Batch->MipLevel != InMipLevel)
+	{
+		Batch = &Batches.Add_GetRef(FBatch(InPageTable, InMipLevel));
+	}
+	
+	OutCoordsPtr = DataPtr + NumWrittenUpdates * sizeof(uint32);
+	OutPayloadPtr = DataPtr + (MaxNumUpdates + NumWrittenUpdates) * sizeof(uint32);
+	
+	Batch->NumUpdates += InNumUpdates;
+	NumWrittenUpdates += InNumUpdates;
+}
+
+void FPageTableUpdater::Release()
+{
+	UpdatesUploadBuffer.SafeRelease();
+	ResetState();
+}
+
+void FPageTableUpdater::Apply(FRDGBuilder& GraphBuilder)
+{
+	if (MaxNumUpdates > 0)
+	{
+		GraphBuilder.RHICmdList.UnlockBuffer(UpdatesUploadBuffer->GetRHI());
+	
+		if (NumWrittenUpdates > 0)
+		{
+			// Disable async compute for streaming systems when MGPU is active, to work around GPU hangs
+			const bool bAsyncCompute = GSupportsEfficientAsyncCompute && (GSVTStreamingAsyncCompute != 0) && (GNumExplicitGPUsForRendering == 1);
+			auto ComputeShader = GetGlobalShaderMap(GMaxRHIFeatureLevel)->GetShader<FSparseVolumeTextureUpdatePageTableCS>();
+	
+			uint32 UpdatesOffset = 0;
+			for (const FBatch& Batch : Batches)
+			{
+				FRDGTexture* PageTableRDG = GraphBuilder.RegisterExternalTexture(CreateRenderTarget(Batch.PageTable, TEXT("SparseVolumeTexture.PageTableTexture")));
+				FRDGTextureUAV* PageTableUAV = GraphBuilder.CreateUAV(FRDGTextureUAVDesc(PageTableRDG, Batch.MipLevel, PF_R32_UINT));
+				FRDGBufferSRV* UpdatesBufferSRV = GraphBuilder.CreateSRV(GraphBuilder.RegisterExternalBuffer(UpdatesUploadBuffer));
+	
+				FSparseVolumeTextureUpdatePageTableCS::FParameters* PassParameters = GraphBuilder.AllocParameters<FSparseVolumeTextureUpdatePageTableCS::FParameters>();
+				PassParameters->PageTable = PageTableUAV;
+				PassParameters->PageTableUpdates = UpdatesBufferSRV;
+				PassParameters->UpdateCoordOffset = UpdatesOffset;
+				PassParameters->UpdatePayloadOffset = MaxNumUpdates + UpdatesOffset;
+				PassParameters->NumUpdates = Batch.NumUpdates;
+	
+				FComputeShaderUtils::AddPass(
+					GraphBuilder,
+					RDG_EVENT_NAME("Update SVT PageTable (UpdateCount: %u)", Batch.NumUpdates),
+					bAsyncCompute ? ERDGPassFlags::AsyncCompute : ERDGPassFlags::Compute,
+					ComputeShader,
+					PassParameters,
+					FComputeShaderUtils::GetGroupCount(Batch.NumUpdates, 64)
+				);
+	
+				UpdatesOffset += Batch.NumUpdates;
+			}
+		}
+	}
+	
+	Release();
+}
+
+void FPageTableUpdater::ResetState()
+{
+	Batches.Reset();
+	DataPtr = nullptr;
+	NumWrittenUpdates = 0;
+	MaxNumUpdates = 0;
+}
+
+}
+}
\ No newline at end of file
diff --git a/Engine/Source/Runtime/Engine/Private/SparseVolumeTexture/SparseVolumeTextureUpload.h b/Engine/Source/Runtime/Engine/Private/SparseVolumeTexture/SparseVolumeTextureUpload.h
new file mode 100644
--- /dev/null
+++ b/Engine/Source/Runtime/Engine/Private/SparseVolumeTexture/SparseVolumeTextureUpload.h
@@ -0,0 +1,97 @@
+// Copyright Epic Games, Inc. All Rights Reserved.
+
+#pragma once
+#include "CoreMinimal.h"
+#include "Containers/StaticArray.h"
+#include "RenderGraphResources.h"
+
+class FRDGBuilder;
+class FRDGPooledBuffer;
+class FRHITexture;
+enum EPixelFormat : uint8;
+
+namespace UE
+{
+namespace SVT
+{
+
+// Utility class for uploading tiles to a physical tile data texture
+class FTileUploader
+{
+public:
+
+	struct FAddResult
+	{
+		TStaticArray<uint8*, 2> OccupancyBitsPtrs;
+		TStaticArray<uint8*, 2> TileDataOffsetsPtrs;
+		TStaticArray<uint8*, 2> TileDataPtrs;
+		TStaticArray<uint32, 2> TileDataBaseOffsets; // Caller needs to add this value to all data written to TileDataOffsetsPtrs
+		uint8* PackedPhysicalTileCoordsPtr;
+	};
+
+	explicit FTileUploader();
+	void Init(FRDGBuilder& InGraphBuilder, int32 InMaxNumTiles, int32 InMaxNumVoxelsA, int32 InMaxNumVoxelsB, EPixelFormat InFormatA, EPixelFormat InFormatB);
+	FAddResult Add_GetRef(int32 InNumTiles, int32 InNumVoxelsA, int32 InNumVoxelsB);
+	void Release();
+	void ResourceUploadTo(FRDGBuilder& InGraphBuilder, FRHITexture* InDstTextureA, FRHITexture* InDstTextureB, const FVector4f& InFallbackValueA, const FVector4f& InFallbackValueB);
+
+private:
+	TRefCountPtr<FRDGPooledBuffer> OccupancyBitsUploadBuffer;
+	TRefCountPtr<FRDGPooledBuffer> TileDataOffsetsUploadBuffer;
+	TRefCountPtr<FRDGPooledBuffer> DstTileCoordsUploadBuffer;
+	TRefCountPtr<FRDGPooledBuffer> TileDataAUploadBuffer;
+	TRefCountPtr<FRDGPooledBuffer> TileDataBUploadBuffer;
+	TArray<uint32> TileDataOffsets; // CPU-readable per-tile offsets into tile data
+	uint8* OccupancyBitsAPtr;
+	uint8* OccupancyBitsBPtr;
+	uint8* TileDataOffsetsAPtr;
+	uint8* TileDataOffsetsBPtr;
+	uint8* TileCoordsPtr;
+	uint8* TileDataAPtr;
+	uint8* TileDataBPtr;
+	int32 MaxNumTiles;
+	int32 MaxNumVoxelsA;
+	int32 MaxNumVoxelsB;
+	EPixelFormat FormatA;
+	EPixelFormat FormatB;
+	int32 FormatSizeA;
+	int32 FormatSizeB;
+	int32 NumWrittenTiles;
+	int32 NumWrittenVoxelsA;
+	int32 NumWrittenVoxelsB;
+
+	void ResetState();
+};
+
+// Utility class for writing page table entries
+class FPageTableUpdater
+{
+public:
+	explicit FPageTableUpdater();
+	void Init(FRDGBuilder& InGraphBuilder, int32 InMaxNumUpdates, int32 InEstimatedNumBatches);
+	void Add_GetRef(FRHITexture* InPageTable, int32 InMipLevel, int32 InNumUpdates, uint8*& OutCoordsPtr, uint8*& OutPayloadPtr);
+	void Release();
+	void Apply(FRDGBuilder& InGraphBuilder);
+
+private:
+	struct FBatch
+	{
+		FRHITexture* PageTable;
+		int32 MipLevel;
+		int32 NumUpdates;
+
+		FBatch() = default;
+		FBatch(FRHITexture* InPageTable, int32 InMipLevel) : PageTable(InPageTable), MipLevel(InMipLevel), NumUpdates(0) {}
+	};
+
+	TRefCountPtr<FRDGPooledBuffer> UpdatesUploadBuffer;
+	TArray<FBatch> Batches;
+	uint8* DataPtr;
+	int32 NumWrittenUpdates;
+	int32 MaxNumUpdates;
+
+	void ResetState();
+};
+
+}
+}
\ No newline at end of file
diff --git a/Engine/Source/Runtime/Engine/Private/SparseVolumeTexture/SparseVolumeTextureUtility.cpp b/Engine/Source/Runtime/Engine/Private/SparseVolumeTexture/SparseVolumeTextureUtility.cpp
--- a/Engine/Source/Runtime/Engine/Private/SparseVolumeTexture/SparseVolumeTextureUtility.cpp
+++ b/Engine/Source/Runtime/Engine/Private/SparseVolumeTexture/SparseVolumeTextureUtility.cpp
@@ -193,3 +193,18 @@ bool UE::SVT::IsSupportedFormat(EPixelFormat Format)
 		return false;
 	}
 }
+
+bool UE::SVT::IsInBounds(const FIntVector3& Point, const FIntVector3& Min, const FIntVector3& Max)
+{
+	return Point.X >= Min.X && Point.Y >= Min.Y && Point.Z >= Min.Z
+		&& Point.X < Max.X && Point.Y < Max.Y && Point.Z < Max.Z;
+}
+
+FIntVector3 UE::SVT::ShiftRightAndMax(const FIntVector3& Value, uint32 ShiftBy, int32 MinValue)
+{
+	FIntVector3 Result = FIntVector3(
+		FMath::Max(Value.X >> ShiftBy, MinValue),
+		FMath::Max(Value.Y >> ShiftBy, MinValue),
+		FMath::Max(Value.Z >> ShiftBy, MinValue));
+	return Result;
+}
diff --git a/Engine/Source/Runtime/Engine/Private/SparseVolumeTexture/SparseVolumeTextureUtility.h b/Engine/Source/Runtime/Engine/Private/SparseVolumeTexture/SparseVolumeTextureUtility.h
--- a/Engine/Source/Runtime/Engine/Private/SparseVolumeTexture/SparseVolumeTextureUtility.h
+++ b/Engine/Source/Runtime/Engine/Private/SparseVolumeTexture/SparseVolumeTextureUtility.h
@@ -22,6 +22,10 @@ namespace SVT
 	FVector4f ReadVoxel(int64 VoxelIndex, const uint8* TileData, EPixelFormat Format);
 	void WriteVoxel(int64 VoxelIndex, uint8* TileData, EPixelFormat Format, const FVector4f& Value, int32 DstComponent = -1);
 	bool IsSupportedFormat(EPixelFormat Format);
+	// Returns true if Point >= Min && Point < Max for all components
+	bool IsInBounds(const FIntVector3& Point, const FIntVector3& Min, const FIntVector3& Max);
+	// Computes Max(Value >> ShiftBy, MinValue)
+	FIntVector3 ShiftRightAndMax(const FIntVector3& Value, uint32 ShiftBy, int32 MinValue);
 
 namespace Private
 {
diff --git a/Engine/Source/Runtime/Engine/Public/SparseVolumeTexture/SparseVolumeTextureData.h b/Engine/Source/Runtime/Engine/Public/SparseVolumeTexture/SparseVolumeTextureData.h
--- a/Engine/Source/Runtime/Engine/Public/SparseVolumeTexture/SparseVolumeTextureData.h
+++ b/Engine/Source/Runtime/Engine/Public/SparseVolumeTexture/SparseVolumeTextureData.h
@@ -68,11 +68,6 @@ private:
 	void WriteTileDataVoxel(int32 TileIndex, const FIntVector3& TileDataCoord, int32 MipLevel, int32 AttributesIdx, const FVector4f& Value, int32 DstComponent = -1);
 };
 
-// Returns true if Point >= Min && Point < Max for all components
-bool IsInBounds(const FIntVector3& Point, const FIntVector3& Min, const FIntVector3& Max);
-// Computes Max(Value >> ShiftBy, MinValue)
-FIntVector3 ShiftRightAndMax(const FIntVector3& Value, uint32 ShiftBy, int32 MinValue);
-
 }
 }
 
diff --git a/Engine/Source/Runtime/Renderer/Private/SparseVolumeTexture/SparseVolumeTextureViewerRendering.cpp b/Engine/Source/Runtime/Renderer/Private/SparseVolumeTexture/SparseVolumeTextureViewerRendering.cpp
--- a/Engine/Source/Runtime/Renderer/Private/SparseVolumeTexture/SparseVolumeTextureViewerRendering.cpp
+++ b/Engine/Source/Runtime/Renderer/Private/SparseVolumeTexture/SparseVolumeTextureViewerRendering.cpp
@@ -60,7 +60,6 @@ class FVisualizeSparseVolumeTexturePS : public FGlobalShader
 		SHADER_PARAMETER_TEXTURE(Texture3D<uint>, SparseVolumeTexturePageTable)
 		SHADER_PARAMETER_TEXTURE(Texture3D, SparseVolumeTextureA)
 		SHADER_PARAMETER_TEXTURE(Texture3D, SparseVolumeTextureB)
-		SHADER_PARAMETER_SRV(ByteAddressBuffer, StreamingInfoBuffer)
 		SHADER_PARAMETER(FUintVector4, PackedSVTUniforms0)
 		SHADER_PARAMETER(FUintVector4, PackedSVTUniforms1)
 		SHADER_PARAMETER(FVector3f, SparseVolumeTextureResolution)
@@ -135,7 +134,6 @@ void AddSparseVolumeTextureViewerRenderPass(FRDGBuilder& GraphBuilder, FSceneRen
 			PsPassParameters->SparseVolumeTexturePageTable = GBlackUintVolumeTexture->TextureRHI;
 			PsPassParameters->SparseVolumeTextureA = GBlackVolumeTexture->TextureRHI;
 			PsPassParameters->SparseVolumeTextureB = GBlackVolumeTexture->TextureRHI;
-			PsPassParameters->StreamingInfoBuffer = GEmptyStructuredBufferWithUAV->ShaderResourceViewRHI;
 			PsPassParameters->PackedSVTUniforms0 = FUintVector4(0);
 			PsPassParameters->PackedSVTUniforms1 = FUintVector4(0);
 			PsPassParameters->SparseVolumeTextureResolution = SVTSceneProxy->VolumeResolution;
@@ -155,12 +153,10 @@ void AddSparseVolumeTextureViewerRenderPass(FRDGBuilder& GraphBuilder, FSceneRen
 				FRHITexture* PageTableTexture = RenderResources->GetPageTableTexture();
 				FRHITexture* TextureA = RenderResources->GetPhysicalTileDataATexture();
 				FRHITexture* TextureB = RenderResources->GetPhysicalTileDataBTexture();
-				FRHIShaderResourceView* StreamingInfoBufferSRV = RenderResources->GetStreamingInfoBufferSRV();
 
 				PsPassParameters->SparseVolumeTexturePageTable = PageTableTexture ? PageTableTexture : PsPassParameters->SparseVolumeTexturePageTable;
 				PsPassParameters->SparseVolumeTextureA = TextureA ? TextureA : PsPassParameters->SparseVolumeTextureA;
 				PsPassParameters->SparseVolumeTextureB = TextureB ? TextureB : PsPassParameters->SparseVolumeTextureB;
-				PsPassParameters->StreamingInfoBuffer = StreamingInfoBufferSRV ? StreamingInfoBufferSRV : PsPassParameters->StreamingInfoBuffer;
 
 				SVTSceneProxy->TextureRenderResources->GetPackedUniforms(PsPassParameters->PackedSVTUniforms0, PsPassParameters->PackedSVTUniforms1);
 			}
