World Partition Runtiome HashSet: fix bad edit condition.

diff --git a/Engine/Source/Runtime/Engine/Public/WorldPartition/RuntimeHashSet/RuntimePartition.h b/Engine/Source/Runtime/Engine/Public/WorldPartition/RuntimeHashSet/RuntimePartition.h
--- a/Engine/Source/Runtime/Engine/Public/WorldPartition/RuntimeHashSet/RuntimePartition.h
+++ b/Engine/Source/Runtime/Engine/Public/WorldPartition/RuntimeHashSet/RuntimePartition.h
@@ -75,13 +75,13 @@ public:
 	UPROPERTY()
 	FName Name;
 
-	UPROPERTY(EditAnywhere, Category = RuntimeSettings, Meta = (EditCondition = "HLODIndex != INDEX_NONE", EditConditionHides, HideEditConditionToggle))
+	UPROPERTY(EditAnywhere, Category = RuntimeSettings, Meta = (EditCondition = "HLODIndex == INDEX_NONE", EditConditionHides, HideEditConditionToggle))
 	bool bBlockOnSlowStreaming;
 
-	UPROPERTY(EditAnywhere, Category = RuntimeSettings, Meta = (EditCondition = "HLODIndex != INDEX_NONE", EditConditionHides, HideEditConditionToggle))
+	UPROPERTY(EditAnywhere, Category = RuntimeSettings, Meta = (EditCondition = "HLODIndex == INDEX_NONE", EditConditionHides, HideEditConditionToggle))
 	bool bClientOnlyVisible;
 
-	UPROPERTY(EditAnywhere, Category = RuntimeSettings, Meta = (EditCondition = "HLODIndex != INDEX_NONE", EditConditionHides, HideEditConditionToggle))
+	UPROPERTY(EditAnywhere, Category = RuntimeSettings, Meta = (EditCondition = "HLODIndex == INDEX_NONE", EditConditionHides, HideEditConditionToggle))
 	int32 Priority;
 
 	UPROPERTY(EditAnywhere, Category = RuntimeSettings)
