Put back check that was removed in a previous CL that triggers during AutoTest Cook (will investigate later).

diff --git a/Engine/Source/Runtime/Engine/Private/WorldPartition/DataLayer/DataLayerManager.cpp b/Engine/Source/Runtime/Engine/Private/WorldPartition/DataLayer/DataLayerManager.cpp
--- a/Engine/Source/Runtime/Engine/Private/WorldPartition/DataLayer/DataLayerManager.cpp
+++ b/Engine/Source/Runtime/Engine/Private/WorldPartition/DataLayer/DataLayerManager.cpp
@@ -852,6 +852,12 @@ void UDataLayerManager::ResolveActorDescContainerInstanceDataLayersInternal(UAct
 	check(!InActorDescInstance || (InActorDescInstance->GetContainerInstance() == InActorDescContainerInstance));
 		
 	const UWorldPartition* ContainerOuterWorldPartition = InActorDescContainerInstance->GetTopWorldPartition();
+	// Skip resolving for template containers (will be done on ActorDescViews)
+	if (!ContainerOuterWorldPartition)
+	{
+		return;
+	}
+
 	const ULevelStreaming* ContainerLevelStreaming = FLevelUtils::FindStreamingLevel(ContainerOuterWorldPartition->GetTypedOuter<UWorld>()->PersistentLevel);
 	const UWorld* ContainerLevelStreamingWorld = ContainerLevelStreaming ? ContainerLevelStreaming->GetWorld() : nullptr;
 	const UWorldPartition* ContainerOwningWorldPartition = ContainerLevelStreamingWorld && !ContainerLevelStreamingWorld->IsGameWorld() ? ContainerLevelStreamingWorld->GetWorldPartition() : ContainerOuterWorldPartition;
