CanvasRenderTarget2D - Exposed Get/SetSampleCount to external callers

diff --git a/Engine/Source/Runtime/Engine/Classes/Engine/CanvasRenderTarget2D.h b/Engine/Source/Runtime/Engine/Classes/Engine/CanvasRenderTarget2D.h
--- a/Engine/Source/Runtime/Engine/Classes/Engine/CanvasRenderTarget2D.h
+++ b/Engine/Source/Runtime/Engine/Classes/Engine/CanvasRenderTarget2D.h
@@ -92,10 +92,10 @@ public:
 	}
 
 	UFUNCTION(BlueprintGetter)
-	virtual ETextureRenderTargetSampleCount GetSampleCount() const override;
+	ENGINE_API virtual ETextureRenderTargetSampleCount GetSampleCount() const override;
 
 	UFUNCTION(BlueprintSetter)
-	void SetSampleCount(ETextureRenderTargetSampleCount InSampleCount);
+	ENGINE_API void SetSampleCount(ETextureRenderTargetSampleCount InSampleCount);
 
 #if WITH_EDITOR
 	virtual bool CanEditChange(const FProperty* InProperty) const override;
