Fix DebugViewModeHelpers.Enable cvar not being checked for all cases.

diff --git a/Engine/Source/Runtime/Engine/Private/DebugViewModeHelpers.cpp b/Engine/Source/Runtime/Engine/Private/DebugViewModeHelpers.cpp
--- a/Engine/Source/Runtime/Engine/Private/DebugViewModeHelpers.cpp
+++ b/Engine/Source/Runtime/Engine/Private/DebugViewModeHelpers.cpp
@@ -118,6 +118,11 @@ bool AllowDebugViewShaderMode(EDebugViewShaderMode ShaderMode, EShaderPlatform P
 
 bool ShouldCompileDebugViewModeShader(const FMeshMaterialShaderPermutationParameters& Parameters)
 {
+	if(!CVarEnableDebugViewModeHelpers.GetValueOnAnyThread())
+	{
+		return false;
+	}
+
 	if (!PlatformSupportsDebugViewShaders(Parameters.Platform))
 	{
 		return false;
