Chaos - fix scale of welded trriangle meshes

This didn't matter before but we now allow welding of kinematics and if we weld a scaled kinematic to another kinematic, we need to scale the child mesh by its local scale.

diff --git a/Engine/Source/Runtime/Engine/Private/PhysicsEngine/Experimental/ChaosInterfaceUtils.cpp b/Engine/Source/Runtime/Engine/Private/PhysicsEngine/Experimental/ChaosInterfaceUtils.cpp
--- a/Engine/Source/Runtime/Engine/Private/PhysicsEngine/Experimental/ChaosInterfaceUtils.cpp
+++ b/Engine/Source/Runtime/Engine/Private/PhysicsEngine/Experimental/ChaosInterfaceUtils.cpp
@@ -297,7 +297,11 @@ namespace ChaosInterface
 		{
 			const FTransform MeshTransform = FTransform(InParams.LocalTransform.GetRotation(), Scale * InParams.LocalTransform.GetTranslation(), FVector(1, 1, 1));
 			const bool bHasTranslationOrRotation = !MeshTransform.GetTranslation().IsNearlyZero() || !MeshTransform.GetRotation().IsIdentity();
-			const bool bNoScale = Scale == FVector(1);
+
+			// Extract the scale from the transform - we have separate wrapper classes for scale versus translate/rotate 
+			FVector NetScale = Scale * InParams.LocalTransform.GetScale3D();
+			const bool bNoScale = FVector::PointsAreNear(NetScale, FVector(1), UE_KINDA_SMALL_NUMBER);
+
 			for (auto& ChaosTriMesh : InParams.TriMeshGeometries)
 			{
 				Chaos::FImplicitObjectPtr Implicit;
@@ -307,7 +311,7 @@ namespace ChaosInterface
 				}
 				else
 				{
-					Implicit = MakeImplicitObjectPtr<Chaos::TImplicitObjectScaled<Chaos::FTriangleMeshImplicitObject>>(ChaosTriMesh, Scale);
+					Implicit = MakeImplicitObjectPtr<Chaos::TImplicitObjectScaled<Chaos::FTriangleMeshImplicitObject>>(ChaosTriMesh, NetScale);
 				}
 
 				// Wrap the mesh in a non-scaled transform if necessary (the scale is pulled out above)
