Remove bOwnerSelected from DecalProxy

This does not appear to be used anymore. Removing this lets us remove the decal proxy invalidation which was occuring on selection in the editor.

This fixes the bug that selecting a decal proxy causes the path tracer to invalidate.

diff --git a/Engine/Source/Runtime/Engine/Classes/Components/DecalComponent.h b/Engine/Source/Runtime/Engine/Classes/Components/DecalComponent.h
--- a/Engine/Source/Runtime/Engine/Classes/Components/DecalComponent.h
+++ b/Engine/Source/Runtime/Engine/Classes/Components/DecalComponent.h
@@ -137,6 +137,7 @@ public:
 	/**
 	 * Pushes new selection state to the render thread primitive proxy
 	 */
+	UE_DEPRECATED(5.4, "This method has no effect, calls to this method may be safely removed.")
 	ENGINE_API void PushSelectionToProxy();
 
 protected:
diff --git a/Engine/Source/Runtime/Engine/Private/Actor.cpp b/Engine/Source/Runtime/Engine/Private/Actor.cpp
--- a/Engine/Source/Runtime/Engine/Private/Actor.cpp
+++ b/Engine/Source/Runtime/Engine/Private/Actor.cpp
@@ -5229,11 +5229,6 @@ void AActor::PushSelectionToProxies()
 				PrimComponent->PushSelectionToProxy();
 			}
 
-			if(UDecalComponent* DecalComponent = Cast<UDecalComponent>(Component))
-			{
-				DecalComponent->PushSelectionToProxy();
-			}
-
 			if (ULightComponent* LightComponent = Cast<ULightComponent>(Component))
 			{
 				LightComponent->PushSelectionToProxy();
diff --git a/Engine/Source/Runtime/Engine/Private/Components/DecalComponent.cpp b/Engine/Source/Runtime/Engine/Private/Components/DecalComponent.cpp
--- a/Engine/Source/Runtime/Engine/Private/Components/DecalComponent.cpp
+++ b/Engine/Source/Runtime/Engine/Private/Components/DecalComponent.cpp
@@ -50,7 +50,6 @@ FDeferredDecalProxy::FDeferredDecalProxy(const UDecalComponent* InComponent)
 	Component = InComponent;
 	DecalMaterial = EffectiveMaterial;
 	SetTransformIncludingDecalSize(InComponent->GetTransformIncludingDecalSize(), InComponent->CalcBounds(InComponent->GetComponentTransform()));
-	bOwnerSelected = InComponent->IsOwnerSelected();
 	SortOrder = InComponent->SortOrder;
 
 #if WITH_EDITOR
@@ -87,7 +86,6 @@ FDeferredDecalProxy::FDeferredDecalProxy(const USceneComponent* InComponent, UMa
 	}
 
 	SetTransformIncludingDecalSize(FTransform::Identity, InComponent->CalcBounds(InComponent->GetComponentTransform()));
-	bOwnerSelected = InComponent->IsOwnerSelected();
 	SortOrder = 0;
 
 #if WITH_EDITOR
@@ -344,7 +342,8 @@ void UDecalComponent::PrecachePSOs()
 
 void UDecalComponent::PushSelectionToProxy()
 {
-	MarkRenderStateDirty();	
+	// The decal's proxy does not actually need to know if the decal is selected or not, so there is nothing to do here.
+	// This function has been marked as deprecated and can eventually be removed.
 }
 
 class UMaterialInterface* UDecalComponent::GetDecalMaterial() const
diff --git a/Engine/Source/Runtime/Engine/Public/SceneManagement.h b/Engine/Source/Runtime/Engine/Public/SceneManagement.h
--- a/Engine/Source/Runtime/Engine/Public/SceneManagement.h
+++ b/Engine/Source/Runtime/Engine/Public/SceneManagement.h
@@ -1404,8 +1404,6 @@ private:
 
 public:
 
-	bool bOwnerSelected;
-
 	/** Larger values draw later (on top). */
 	int32 SortOrder;
 
