World Partition: bring back FWorldPartitionHandle/FWorldPartitionReference arrow and dereference operators.

diff --git a/Engine/Plugins/Editor/ObjectMixer/ObjectMixer/Source/ObjectMixer/Private/Views/List/Modes/ObjectMixerOutlinerHierarchy.cpp b/Engine/Plugins/Editor/ObjectMixer/ObjectMixer/Source/ObjectMixer/Private/Views/List/Modes/ObjectMixerOutlinerHierarchy.cpp
--- a/Engine/Plugins/Editor/ObjectMixer/ObjectMixer/Source/ObjectMixer/Private/Views/List/Modes/ObjectMixerOutlinerHierarchy.cpp
+++ b/Engine/Plugins/Editor/ObjectMixer/ObjectMixer/Source/ObjectMixer/Private/Views/List/Modes/ObjectMixerOutlinerHierarchy.cpp
@@ -312,7 +312,7 @@ FSceneOutlinerTreeItemPtr FObjectMixerOutlinerHierarchy::FindOrCreateParentItem(
 	}
 	else if (const FActorDescTreeItem* ActorDescItem = Item.CastTo<FActorDescTreeItem>())
 	{
-		if (const FWorldPartitionActorDescInstance* ActorDescInstance = ActorDescItem->ActorDescHandle.GetInstance())
+		if (const FWorldPartitionActorDescInstance* ActorDescInstance = *ActorDescItem->ActorDescHandle)
 		{
 			if (UWorld* RepresentingWorldPtr = RepresentingWorld.Get())
 			{
diff --git a/Engine/Plugins/Editor/ObjectMixer/ObjectMixer/Source/ObjectMixer/Private/Views/List/Modes/ObjectMixerOutlinerMode.cpp b/Engine/Plugins/Editor/ObjectMixer/ObjectMixer/Source/ObjectMixer/Private/Views/List/Modes/ObjectMixerOutlinerMode.cpp
--- a/Engine/Plugins/Editor/ObjectMixer/ObjectMixer/Source/ObjectMixer/Private/Views/List/Modes/ObjectMixerOutlinerMode.cpp
+++ b/Engine/Plugins/Editor/ObjectMixer/ObjectMixer/Source/ObjectMixer/Private/Views/List/Modes/ObjectMixerOutlinerMode.cpp
@@ -3295,7 +3295,7 @@ bool FObjectMixerOutlinerMode::CompareItemWithClassName(SceneOutliner::FilterBar
 	}
 	else if (const FActorDescTreeItem* ActorDescItem = InItem.CastTo<FActorDescTreeItem>())
 	{
-		if (const FWorldPartitionActorDescInstance* ActorDescInstance = ActorDescItem->ActorDescHandle.GetInstance())
+		if (const FWorldPartitionActorDescInstance* ActorDescInstance = *ActorDescItem->ActorDescHandle)
 		{
 			// For Unloaded Actors, grab the native class 
 			FTopLevelAssetPath ClassPath = ActorDescInstance->GetNativeClass();
diff --git a/Engine/Plugins/Editor/ObjectMixer/ObjectMixer/Source/ObjectMixer/Private/Views/List/SObjectMixerEditorList.cpp b/Engine/Plugins/Editor/ObjectMixer/ObjectMixer/Source/ObjectMixer/Private/Views/List/SObjectMixerEditorList.cpp
--- a/Engine/Plugins/Editor/ObjectMixer/ObjectMixer/Source/ObjectMixer/Private/Views/List/SObjectMixerEditorList.cpp
+++ b/Engine/Plugins/Editor/ObjectMixer/ObjectMixer/Source/ObjectMixer/Private/Views/List/SObjectMixerEditorList.cpp
@@ -903,7 +903,7 @@ void SObjectMixerEditorList::CreateActorTextInfoColumns(FSceneOutlinerInitializa
 			}
 			else if (const FActorDescTreeItem* ActorDescItem = Item.CastTo<FActorDescTreeItem>())
 			{
-				if (const FWorldPartitionActorDescInstance* ActorDescInstance = ActorDescItem->ActorDescHandle.GetInstance(); ActorDescInstance && !ActorDescInstance->GetDataLayerInstanceNames().IsEmpty())
+				if (const FWorldPartitionActorDescInstance* ActorDescInstance = *ActorDescItem->ActorDescHandle; ActorDescInstance && !ActorDescInstance->GetDataLayerInstanceNames().IsEmpty())
 				{
 					if (const UActorDescContainerInstance* ActorDescContainerInstance = ActorDescInstance->GetContainerInstance())
 					{
@@ -952,7 +952,7 @@ void SObjectMixerEditorList::CreateActorTextInfoColumns(FSceneOutlinerInitializa
 		}
 		else if (const FActorDescTreeItem* ActorDescItem = Item.CastTo<FActorDescTreeItem>())
 		{
-			if (const FWorldPartitionActorDescInstance* ActorDescInstance = ActorDescItem->ActorDescHandle.GetInstance())
+			if (const FWorldPartitionActorDescInstance* ActorDescInstance = *ActorDescItem->ActorDescHandle)
 			{
 				if (const UContentBundleDescriptor* Descriptor = ContentBundleEngineSubsystem->GetContentBundleDescriptor(ActorDescInstance->GetContentBundleGuid()))
 				{
@@ -978,7 +978,7 @@ void SObjectMixerEditorList::CreateActorTextInfoColumns(FSceneOutlinerInitializa
 		}
 		else if (const FActorDescTreeItem* ActorDescItem = Item.CastTo<FActorDescTreeItem>())
 		{
-			if (const FWorldPartitionActorDescInstance* ActorDescInstance = ActorDescItem->ActorDescHandle.GetInstance())
+			if (const FWorldPartitionActorDescInstance* ActorDescInstance = *ActorDescItem->ActorDescHandle)
 			{
 				if (FName LevelPackage = ActorDescInstance->GetChildContainerPackage(); !LevelPackage.IsNone())
 				{
@@ -1023,7 +1023,7 @@ void SObjectMixerEditorList::CreateActorTextInfoColumns(FSceneOutlinerInitializa
 		}
 		else if (const FActorDescTreeItem* ActorDescItem = Item.CastTo<FActorDescTreeItem>())
 		{
-			if (const FWorldPartitionActorDescInstance* ActorDescInstance = ActorDescItem->ActorDescHandle.GetInstance())
+			if (const FWorldPartitionActorDescInstance* ActorDescInstance = *ActorDescItem->ActorDescHandle)
 			{
 				return ActorDescInstance->GetActorName().ToString();
 			}
@@ -1050,7 +1050,7 @@ void SObjectMixerEditorList::CreateActorTextInfoColumns(FSceneOutlinerInitializa
 		}
 		else if (const FActorDescTreeItem* ActorDescItem = Item.CastTo<FActorDescTreeItem>())
 		{
-			if (const FWorldPartitionActorDescInstance* ActorDescInstance = ActorDescItem->ActorDescHandle.GetInstance())
+			if (const FWorldPartitionActorDescInstance* ActorDescInstance = *ActorDescItem->ActorDescHandle)
 			{
 				return FPackageName::GetShortName(ActorDescInstance->GetActorPackage());
 			}
diff --git a/Engine/Source/Editor/DataLayerEditor/Private/DataLayer/DataLayerEditorSubsystem.cpp b/Engine/Source/Editor/DataLayerEditor/Private/DataLayer/DataLayerEditorSubsystem.cpp
--- a/Engine/Source/Editor/DataLayerEditor/Private/DataLayer/DataLayerEditorSubsystem.cpp
+++ b/Engine/Source/Editor/DataLayerEditor/Private/DataLayer/DataLayerEditorSubsystem.cpp
@@ -1320,14 +1320,14 @@ bool UDataLayerEditorSubsystem::PassDataLayersFilter(UWorld* World, const FWorld
 		// If Actor is loaded and dirty, use a newly resolved DataLayerInstanceNames array
 		auto GetLatestDataLayerInstanceNames = [DataLayerManager, ActorHandle]() -> FDataLayerInstanceNames
 		{
-			if(AActor* Actor = ActorHandle.GetInstance()->GetActor(false); Actor && Actor->GetPackage()->IsDirty())
+			if(AActor* Actor = ActorHandle->GetActor(false); Actor && Actor->GetPackage()->IsDirty())
 			{
 				TUniquePtr<FWorldPartitionActorDesc> NewActorDesc = Actor->CreateActorDesc();
 				return FDataLayerUtils::ResolveDataLayerInstanceNames(DataLayerManager, NewActorDesc.Get());
 			}
 			else
 			{
-				return ActorHandle.GetInstance()->GetDataLayerInstanceNames();
+				return ActorHandle->GetDataLayerInstanceNames();
 			}
 		};
 
diff --git a/Engine/Source/Editor/DataLayerEditor/Private/DataLayer/DataLayersActorDescTreeItem.h b/Engine/Source/Editor/DataLayerEditor/Private/DataLayer/DataLayersActorDescTreeItem.h
--- a/Engine/Source/Editor/DataLayerEditor/Private/DataLayer/DataLayersActorDescTreeItem.h
+++ b/Engine/Source/Editor/DataLayerEditor/Private/DataLayer/DataLayersActorDescTreeItem.h
@@ -49,7 +49,7 @@ public:
 	{
 		if (ActorDescHandle.IsValid())
 		{
-			UActorDescContainerInstance* ContainerInstance = ActorDescHandle.GetContainerInstance();
+			UActorDescContainerInstance* ContainerInstance = ActorDescHandle->GetContainerInstance();
 			UWorld* OwningWorld = ContainerInstance->GetOuterWorldPartition()->GetWorld();
 			ULevelInstanceSubsystem* LevelInstanceSubsystem = UWorld::GetSubsystem<ULevelInstanceSubsystem>(OwningWorld);
 			ULevel* Level = ContainerInstance->GetTypedOuter<UWorld>()->PersistentLevel;
@@ -75,12 +75,12 @@ public:
 
 	bool Filter(FFilterPredicate Pred) const
 	{
-		return Pred.Execute(ActorDescHandle.GetInstance(), DataLayer.Get());
+		return Pred.Execute(*ActorDescHandle, DataLayer.Get());
 	}
 
 	bool GetInteractiveState(FInteractivePredicate Pred) const
 	{
-		return Pred.Execute(ActorDescHandle.GetInstance(), DataLayer.Get());
+		return Pred.Execute(*ActorDescHandle, DataLayer.Get());
 	}
 
 	/* Begin ISceneOutlinerTreeItem Implementation */
diff --git a/Engine/Source/Editor/DataLayerEditor/Private/DataLayer/SDataLayerBrowser.cpp b/Engine/Source/Editor/DataLayerEditor/Private/DataLayer/SDataLayerBrowser.cpp
--- a/Engine/Source/Editor/DataLayerEditor/Private/DataLayer/SDataLayerBrowser.cpp
+++ b/Engine/Source/Editor/DataLayerEditor/Private/DataLayer/SDataLayerBrowser.cpp
@@ -99,7 +99,7 @@ void SDataLayerBrowser::Construct(const FArguments& InArgs)
 		}
 		else if (const FDataLayerActorDescTreeItem* ActorDescItem = Item.CastTo<FDataLayerActorDescTreeItem>())
 		{
-			if (const FWorldPartitionActorDescInstance* ActorDescInstance = ActorDescItem->ActorDescHandle.GetInstance())
+			if (const FWorldPartitionActorDescInstance* ActorDescInstance = *ActorDescItem->ActorDescHandle)
 			{
 				return ActorDescInstance->GetActorName().ToString();
 			}
diff --git a/Engine/Source/Editor/SceneOutliner/Private/ActorBrowsingMode.cpp b/Engine/Source/Editor/SceneOutliner/Private/ActorBrowsingMode.cpp
--- a/Engine/Source/Editor/SceneOutliner/Private/ActorBrowsingMode.cpp
+++ b/Engine/Source/Editor/SceneOutliner/Private/ActorBrowsingMode.cpp
@@ -1719,7 +1719,7 @@ void FActorBrowsingMode::PinItems(const TArray<FSceneOutlinerTreeItemPtr>& InIte
 		{
 			if (FWorldPartitionHandle ActorHandle(WorldPartition, ActorGuid); ActorHandle.IsValid())
 			{
-				if (AActor* PinnedActor = ActorHandle.GetInstance()->GetActor())
+				if (AActor* PinnedActor = ActorHandle->GetActor())
 				{
 					GEditor->SelectActor(PinnedActor, /*bInSelected=*/true, /*bNotify=*/false);
 					LastPinnedActor = PinnedActor;
@@ -1923,7 +1923,7 @@ bool FActorBrowsingMode::CompareItemWithClassName(SceneOutliner::FilterBarType I
 	}
 	else if (const FActorDescTreeItem* ActorDescItem = InItem.CastTo<FActorDescTreeItem>())
 	{
-		if (const FWorldPartitionActorDescInstance* ActorDescInstance = ActorDescItem->ActorDescHandle.GetInstance())
+		if (const FWorldPartitionActorDescInstance* ActorDescInstance = *ActorDescItem->ActorDescHandle)
 		{
 			// For Unloaded Actors, grab the native class 
 			FTopLevelAssetPath ClassPath = ActorDescInstance->GetNativeClass();
diff --git a/Engine/Source/Editor/SceneOutliner/Private/ActorDescTreeItem.cpp b/Engine/Source/Editor/SceneOutliner/Private/ActorDescTreeItem.cpp
--- a/Engine/Source/Editor/SceneOutliner/Private/ActorDescTreeItem.cpp
+++ b/Engine/Source/Editor/SceneOutliner/Private/ActorDescTreeItem.cpp
@@ -123,7 +123,7 @@ private:
 	{
 		if (TSharedPtr<FActorDescTreeItem> TreeItem = TreeItemPtr.Pin())
 		{
-			if (FWorldPartitionActorDescInstance* ActorDescInstance = TreeItem->ActorDescHandle.GetInstance())
+			if (FWorldPartitionActorDescInstance* ActorDescInstance = *TreeItem->ActorDescHandle)
 			{
 				FFormatNamedArguments Args;
 				Args.Add(TEXT("ActorLabel"), FText::FromString(TreeItem->GetDisplayString()));
@@ -144,7 +144,7 @@ private:
 	{
 		if (TSharedPtr<FActorDescTreeItem> TreeItem = TreeItemPtr.Pin())
 		{
-			if (const FWorldPartitionActorDescInstance* ActorDescInstance = TreeItem->ActorDescHandle.GetInstance())
+			if (const FWorldPartitionActorDescInstance* ActorDescInstance = *TreeItem->ActorDescHandle)
 			{
 				return FText::FromName(ActorDescInstance->GetDisplayClassName());
 			}
@@ -164,7 +164,7 @@ private:
 
 		if (TreeItem.IsValid() && WeakSceneOutliner.IsValid())
 		{
-			if (const FWorldPartitionActorDescInstance* ActorDescInstance = TreeItem->ActorDescHandle.GetInstance())
+			if (const FWorldPartitionActorDescInstance* ActorDescInstance = *TreeItem->ActorDescHandle)
 			{
 				const FName IconName = ActorDescInstance->GetDisplayClassName();
 
@@ -235,7 +235,7 @@ FActorDescTreeItem::FActorDescTreeItem(const FGuid& InActorGuid, UActorDescConta
 	, ID(ComputeTreeItemID(InActorGuid, InContainerInstance))
 	, ActorGuid(InActorGuid)
 {
-	if (const FWorldPartitionActorDescInstance* const ActorDescInstance = ActorDescHandle.GetInstance())
+	if (const FWorldPartitionActorDescInstance* const ActorDescInstance = *ActorDescHandle)
 	{
 		DisplayString = ActorDescInstance->GetActorLabel().ToString();
 	}
@@ -251,7 +251,7 @@ FActorDescTreeItem::FActorDescTreeItem(const FWorldPartitionActorDescInstance* I
 	, ID(ComputeTreeItemID(InActorDescInstance->GetGuid(), InActorDescInstance->GetContainerInstance()))
 	, ActorGuid(InActorDescInstance->GetGuid())
 {
-	if (const FWorldPartitionActorDescInstance* const ActorDescInstance = ActorDescHandle.GetInstance())
+	if (const FWorldPartitionActorDescInstance* const ActorDescInstance = *ActorDescHandle)
 	{
 		DisplayString = ActorDescInstance->GetActorLabel().ToString();
 	}
@@ -299,7 +299,7 @@ TSharedRef<SWidget> FActorDescTreeItem::GenerateLabelWidget(ISceneOutliner& Outl
 
 void FActorDescTreeItem::FocusActorBounds() const
 {
-	if (FWorldPartitionActorDescInstance const* ActorDescInstance = ActorDescHandle.GetInstance())
+	if (FWorldPartitionActorDescInstance const* ActorDescInstance = *ActorDescHandle)
 	{
 		const bool bActiveViewportOnly = true;
 		const FBox EditorBounds = ActorDescInstance->GetEditorBounds();
@@ -312,7 +312,7 @@ void FActorDescTreeItem::FocusActorBounds() const
 
 void FActorDescTreeItem::CopyActorFilePathtoClipboard() const
 {
-	if (FWorldPartitionActorDescInstance const* ActorDescInstance = ActorDescHandle.GetInstance())
+	if (FWorldPartitionActorDescInstance const* ActorDescInstance = *ActorDescHandle)
 	{
 		FString PackageFilename;
 		if (FPackageName::TryConvertLongPackageNameToFilename(ActorDescInstance->GetActorPackage().ToString(), PackageFilename, FPackageName::GetAssetPackageExtension()))
@@ -337,14 +337,14 @@ bool FActorDescTreeItem::GetVisibility() const
 
 bool FActorDescTreeItem::ShouldShowPinnedState() const
 {
-	return FLoaderAdapterPinnedActors::SupportsPinning(ActorDescHandle.GetInstance());
+	return FLoaderAdapterPinnedActors::SupportsPinning(*ActorDescHandle);
 }
 
 bool FActorDescTreeItem::GetPinnedState() const
 {
-	if (ActorDescHandle.IsValid() && ActorDescHandle.GetContainerInstance())
+	if (ActorDescHandle.IsValid() && ActorDescHandle->GetContainerInstance())
 	{
-		UWorldPartition* WorldPartition = ActorDescHandle.GetContainerInstance()->GetOuterWorldPartition();
+		UWorldPartition* WorldPartition = ActorDescHandle->GetContainerInstance()->GetOuterWorldPartition();
 		return WorldPartition ? WorldPartition->IsActorPinned(GetGuid()) : false;
 	}
 	return false;
@@ -352,7 +352,7 @@ bool FActorDescTreeItem::GetPinnedState() const
 
 UExternalDataLayerAsset* FActorDescTreeItem::GetExternalDataLayerAsset() const
 {
-	if (const FWorldPartitionActorDescInstance* ActorDescInstance = ActorDescHandle.GetInstance())
+	if (const FWorldPartitionActorDescInstance* ActorDescInstance = *ActorDescHandle)
 	{
 		const FSoftObjectPath& ExternalDataLayerAsset = ActorDescInstance->GetExternalDataLayerAsset();
 		if (CachedExternalDataLayerAsset.ToSoftObjectPath() != ExternalDataLayerAsset)
diff --git a/Engine/Source/Editor/SceneOutliner/Private/ActorHierarchy.cpp b/Engine/Source/Editor/SceneOutliner/Private/ActorHierarchy.cpp
--- a/Engine/Source/Editor/SceneOutliner/Private/ActorHierarchy.cpp
+++ b/Engine/Source/Editor/SceneOutliner/Private/ActorHierarchy.cpp
@@ -274,7 +274,7 @@ FSceneOutlinerTreeItemPtr FActorHierarchy::FindOrCreateParentItem(const ISceneOu
 	}
 	else if (const FActorDescTreeItem* ActorDescItem = Item.CastTo<FActorDescTreeItem>())
 	{
-		if (const FWorldPartitionActorDescInstance* ActorDescInstance = ActorDescItem->ActorDescHandle.GetInstance())
+		if (const FWorldPartitionActorDescInstance* ActorDescInstance = *ActorDescItem->ActorDescHandle)
 		{
 			if (UWorld* RepresentingWorldPtr = RepresentingWorld.Get())
 			{
diff --git a/Engine/Source/Editor/SceneOutliner/Private/SceneOutlinerActorInfoColumn.cpp b/Engine/Source/Editor/SceneOutliner/Private/SceneOutlinerActorInfoColumn.cpp
--- a/Engine/Source/Editor/SceneOutliner/Private/SceneOutlinerActorInfoColumn.cpp
+++ b/Engine/Source/Editor/SceneOutliner/Private/SceneOutlinerActorInfoColumn.cpp
@@ -61,7 +61,7 @@ struct FGetInfo
 		}
 		else if (const FActorDescTreeItem* ActorDescItem = Item.CastTo<FActorDescTreeItem>())
 		{
-			if (const FWorldPartitionActorDescInstance* ActorDescInstance = ActorDescItem->ActorDescHandle.GetInstance())
+			if (const FWorldPartitionActorDescInstance* ActorDescInstance = *ActorDescItem->ActorDescHandle)
 			{
 				return ActorDescInstance->GetDisplayClassName().ToString();
 			}
diff --git a/Engine/Source/Editor/SceneOutliner/Private/SceneOutlinerHelpers.cpp b/Engine/Source/Editor/SceneOutliner/Private/SceneOutlinerHelpers.cpp
--- a/Engine/Source/Editor/SceneOutliner/Private/SceneOutlinerHelpers.cpp
+++ b/Engine/Source/Editor/SceneOutliner/Private/SceneOutlinerHelpers.cpp
@@ -40,7 +40,7 @@ namespace SceneOutliner
 		}
 		else if (const FActorDescTreeItem* ActorDescItem = TreeItem.CastTo<FActorDescTreeItem>())
 		{
-			if (const FWorldPartitionActorDescInstance* ActorDescInstance = ActorDescItem->ActorDescHandle.GetInstance())
+			if (const FWorldPartitionActorDescInstance* ActorDescInstance = *ActorDescItem->ActorDescHandle)
 			{
 				return ActorDescInstance->GetActorPackage().ToString();
 			}
@@ -73,7 +73,7 @@ namespace SceneOutliner
 		}
 		else if (const FActorDescTreeItem* ActorDescItem = TreeItem.CastTo<FActorDescTreeItem>())
 		{
-			if (const FWorldPartitionActorDescInstance* ActorDescInstance = ActorDescItem->ActorDescHandle.GetInstance())
+			if (const FWorldPartitionActorDescInstance* ActorDescInstance = *ActorDescItem->ActorDescHandle)
 			{
 				return FindPackage(nullptr, *ActorDescInstance->GetActorPackage().ToString());
 			}
diff --git a/Engine/Source/Editor/SceneOutliner/Private/SceneOutlinerModule.cpp b/Engine/Source/Editor/SceneOutliner/Private/SceneOutlinerModule.cpp
--- a/Engine/Source/Editor/SceneOutliner/Private/SceneOutlinerModule.cpp
+++ b/Engine/Source/Editor/SceneOutliner/Private/SceneOutlinerModule.cpp
@@ -375,7 +375,7 @@ void FSceneOutlinerModule::CreateActorInfoColumns(FSceneOutlinerInitializationOp
 			}
 			else if (const FActorDescTreeItem* ActorDescItem = Item.CastTo<FActorDescTreeItem>())
 			{
-				if (const FWorldPartitionActorDescInstance* ActorDescInstance = ActorDescItem->ActorDescHandle.GetInstance(); ActorDescInstance && !ActorDescInstance->GetDataLayerInstanceNames().IsEmpty())
+				if (const FWorldPartitionActorDescInstance* ActorDescInstance = *ActorDescItem->ActorDescHandle; ActorDescInstance && !ActorDescInstance->GetDataLayerInstanceNames().IsEmpty())
 				{
 					if (const UActorDescContainerInstance* ActorDescContainerInstance = ActorDescInstance->GetContainerInstance())
 					{
@@ -424,7 +424,7 @@ void FSceneOutlinerModule::CreateActorInfoColumns(FSceneOutlinerInitializationOp
 		}
 		else if (const FActorDescTreeItem* ActorDescItem = Item.CastTo<FActorDescTreeItem>())
 		{
-			if (const FWorldPartitionActorDescInstance* ActorDescInstance = ActorDescItem->ActorDescHandle.GetInstance())
+			if (const FWorldPartitionActorDescInstance* ActorDescInstance = *ActorDescItem->ActorDescHandle)
 			{
 				if (const UContentBundleDescriptor* Descriptor = ContentBundleEngineSubsystem->GetContentBundleDescriptor(ActorDescInstance->GetContentBundleGuid()))
 				{
@@ -450,7 +450,7 @@ void FSceneOutlinerModule::CreateActorInfoColumns(FSceneOutlinerInitializationOp
 		}
 		else if (const FActorDescTreeItem* ActorDescItem = Item.CastTo<FActorDescTreeItem>())
 		{
-			if (const FWorldPartitionActorDescInstance* ActorDescInstance = ActorDescItem->ActorDescHandle.GetInstance())
+			if (const FWorldPartitionActorDescInstance* ActorDescInstance = *ActorDescItem->ActorDescHandle)
 			{
 				if (FName LevelPackage = ActorDescInstance->GetChildContainerPackage(); !LevelPackage.IsNone())
 				{
@@ -495,7 +495,7 @@ void FSceneOutlinerModule::CreateActorInfoColumns(FSceneOutlinerInitializationOp
 		}
 		else if (const FActorDescTreeItem* ActorDescItem = Item.CastTo<FActorDescTreeItem>())
 		{
-			if (const FWorldPartitionActorDescInstance* ActorDescInstance = ActorDescItem->ActorDescHandle.GetInstance())
+			if (const FWorldPartitionActorDescInstance* ActorDescInstance = *ActorDescItem->ActorDescHandle)
 			{
 				return ActorDescInstance->GetActorName().ToString();
 			}
@@ -522,7 +522,7 @@ void FSceneOutlinerModule::CreateActorInfoColumns(FSceneOutlinerInitializationOp
 		}
 		else if (const FActorDescTreeItem* ActorDescItem = Item.CastTo<FActorDescTreeItem>())
 		{
-			if (const FWorldPartitionActorDescInstance* ActorDescInstance = ActorDescItem->ActorDescHandle.GetInstance())
+			if (const FWorldPartitionActorDescInstance* ActorDescInstance = *ActorDescItem->ActorDescHandle)
 			{
 				return FPackageName::GetShortName(ActorDescInstance->GetActorPackage());
 			}
diff --git a/Engine/Source/Editor/SceneOutliner/Public/ActorDescTreeItem.h b/Engine/Source/Editor/SceneOutliner/Public/ActorDescTreeItem.h
--- a/Engine/Source/Editor/SceneOutliner/Public/ActorDescTreeItem.h
+++ b/Engine/Source/Editor/SceneOutliner/Public/ActorDescTreeItem.h
@@ -22,12 +22,12 @@ public:
 
 	bool Filter(FFilterPredicate Pred) const
 	{
-		return Pred.Execute(ActorDescHandle.GetInstance());
+		return Pred.Execute(*ActorDescHandle);
 	}
 
 	bool GetInteractiveState(FInteractivePredicate Pred) const
 	{
-		return Pred.Execute(ActorDescHandle.GetInstance());
+		return Pred.Execute(*ActorDescHandle);
 	}
 
 	/** The actor desc this tree item is associated with. */
@@ -56,7 +56,7 @@ public:
 	FActorDescTreeItem(const FGuid& InActorGuid, class UActorDescContainer* InContainer) : IActorBaseTreeItem(Type) {}
 
 	/* Begin ISceneOutlinerTreeItem Implementation */
-	virtual bool IsValid() const override { return ActorDescHandle.GetInstance() != nullptr; }
+	virtual bool IsValid() const override { return *ActorDescHandle != nullptr; }
 	virtual FSceneOutlinerTreeItemID GetID() const override;
 	virtual FString GetDisplayString() const override;
 	virtual bool CanInteract() const override;
diff --git a/Engine/Source/Editor/UnrealEd/Private/WorldPartition/WorldPartitionResaveActorsBuilder.cpp b/Engine/Source/Editor/UnrealEd/Private/WorldPartition/WorldPartitionResaveActorsBuilder.cpp
--- a/Engine/Source/Editor/UnrealEd/Private/WorldPartition/WorldPartitionResaveActorsBuilder.cpp
+++ b/Engine/Source/Editor/UnrealEd/Private/WorldPartition/WorldPartitionResaveActorsBuilder.cpp
@@ -265,7 +265,7 @@ bool UWorldPartitionResaveActorsBuilder::RunInternal(UWorld* World, const FCellI
 			// Change packaging of all actors in the current cluster
 			for (FWorldPartitionReference& ActorReference : ActorReferences)
 			{
-				const FWorldPartitionActorDescInstance* ActorDescInstance = ActorReference.GetInstance();
+				const FWorldPartitionActorDescInstance* ActorDescInstance = *ActorReference;
 				AActor* Actor = ActorReference.GetActor();
 
 				if (!Actor)
diff --git a/Engine/Source/Editor/WorldPartitionEditor/Private/WorldPartition/HLOD/HLODEditorData.cpp b/Engine/Source/Editor/WorldPartitionEditor/Private/WorldPartition/HLOD/HLODEditorData.cpp
--- a/Engine/Source/Editor/WorldPartitionEditor/Private/WorldPartition/HLOD/HLODEditorData.cpp
+++ b/Engine/Source/Editor/WorldPartitionEditor/Private/WorldPartition/HLOD/HLODEditorData.cpp
@@ -199,7 +199,7 @@ void FWorldPartitionHLODEditorData::UpdateLoadedActorsState()
 		{
 			if (PinnedActor.IsValid())
 			{
-				GatherLoadedActorsBounds(LoadedBounds, PinnedActor.GetInstance(), WorldPartition->GetActorDescContainerInstance());
+				GatherLoadedActorsBounds(LoadedBounds, *PinnedActor, WorldPartition->GetActorDescContainerInstance());
 			}
 		}
 	}
diff --git a/Engine/Source/Editor/WorldPartitionEditor/Private/WorldPartition/SWorldPartitionEditorGrid2D.cpp b/Engine/Source/Editor/WorldPartitionEditor/Private/WorldPartition/SWorldPartitionEditorGrid2D.cpp
--- a/Engine/Source/Editor/WorldPartitionEditor/Private/WorldPartition/SWorldPartitionEditorGrid2D.cpp
+++ b/Engine/Source/Editor/WorldPartitionEditor/Private/WorldPartition/SWorldPartitionEditorGrid2D.cpp
@@ -1368,8 +1368,8 @@ void SWorldPartitionEditorGrid2D::Tick(const FGeometry& AllottedGeometry, const
 	UWorldPartitionSubsystem* WorldPartitionSubsystem = UWorld::GetSubsystem<UWorldPartitionSubsystem>(GetWorld());
 	for (const FWorldPartitionHandle& SelectedActorHandle : WorldPartitionSubsystem->SelectedActorHandles)
 	{
-		ShownActorGuids.Add(SelectedActorHandle.GetInstance()->GetGuid());
-		SelectedActorGuids.Add(SelectedActorHandle.GetInstance()->GetGuid());
+		ShownActorGuids.Add(SelectedActorHandle->GetGuid());
+		SelectedActorGuids.Add(SelectedActorHandle->GetGuid());
 	}
 }
 
@@ -2198,7 +2198,7 @@ void SWorldPartitionEditorGrid2D::FocusSelection()
 	UWorldPartitionSubsystem* WorldPartitionSubsystem = UWorld::GetSubsystem<UWorldPartitionSubsystem>(GetWorld());
 	for (const FWorldPartitionHandle& SelectedActorHandle : WorldPartitionSubsystem->SelectedActorHandles)
 	{
-		SelectionBox += SelectedActorHandle.GetInstance()->GetEditorBounds();
+		SelectionBox += SelectedActorHandle->GetEditorBounds();
 	}
 
 	if (SelectionBox.IsValid)
diff --git a/Engine/Source/Runtime/Engine/Private/LevelInstance/LevelInstanceSubsystem.cpp b/Engine/Source/Runtime/Engine/Private/LevelInstance/LevelInstanceSubsystem.cpp
--- a/Engine/Source/Runtime/Engine/Private/LevelInstance/LevelInstanceSubsystem.cpp
+++ b/Engine/Source/Runtime/Engine/Private/LevelInstance/LevelInstanceSubsystem.cpp
@@ -2583,7 +2583,7 @@ bool ULevelInstanceSubsystem::CheckForLoop(const ILevelInstanceInterface* LevelI
 
 bool ULevelInstanceSubsystem::PassLevelInstanceFilter(UWorld* World, const FWorldPartitionHandle& ActorHandle) const
 {
-	UWorld* ContainerOuterWorld = ActorHandle.GetContainerInstance()->GetOuterWorldPartition()->GetTypedOuter<UWorld>();
+	UWorld* ContainerOuterWorld = ActorHandle->GetContainerInstance()->GetOuterWorldPartition()->GetTypedOuter<UWorld>();
 	check(ContainerOuterWorld);
 	if (const ILevelInstanceInterface* TopAncestor = GetOwningLevelInstance(ContainerOuterWorld->PersistentLevel))
 	{
@@ -2597,7 +2597,7 @@ bool ULevelInstanceSubsystem::PassLevelInstanceFilter(UWorld* World, const FWorl
 		const TMap<FActorContainerID, TSet<FGuid>>& FilteredActors = TopAncestor->GetFilteredActorsPerContainer();
 		if (const TSet<FGuid>* FilteredActorsForContainer = FilteredActors.Find(ContainerID))
 		{
-			if (FilteredActorsForContainer->Contains(ActorHandle.GetInstance()->GetGuid()))
+			if (FilteredActorsForContainer->Contains(ActorHandle->GetGuid()))
 			{
 				return false;
 			}
diff --git a/Engine/Source/Runtime/Engine/Private/Tests/WorldPartition/WorldPartitionHandleTests.cpp b/Engine/Source/Runtime/Engine/Private/Tests/WorldPartition/WorldPartitionHandleTests.cpp
--- a/Engine/Source/Runtime/Engine/Private/Tests/WorldPartition/WorldPartitionHandleTests.cpp
+++ b/Engine/Source/Runtime/Engine/Private/Tests/WorldPartition/WorldPartitionHandleTests.cpp
@@ -79,68 +79,68 @@ namespace WorldPartitionTests
 			Reference = FWorldPartitionReference(ActorDescContainerInstance, FGuid(TEXT("0D2B04D240BE5DE58FE437A8D2DBF5C9")));
 		}
 
-		TestTrue(TEXT("Handle container"), Handle.GetContainerInstance() == ActorDescContainerInstance);
-		TestTrue(TEXT("Reference container"), Reference.GetContainerInstance() == ActorDescContainerInstance);
+		TestTrue(TEXT("Handle container"), Handle->GetContainerInstance() == ActorDescContainerInstance);
+		TestTrue(TEXT("Reference container"), Reference->GetContainerInstance() == ActorDescContainerInstance);
 
-		TestTrue(TEXT("Handle soft refcount"), FWorldPartitionActorDescUnitTestAcccessor::GetSoftRefCount(Handle.GetInstance()) == 1);
-		TestTrue(TEXT("Handle hard refcount"), FWorldPartitionActorDescUnitTestAcccessor::GetHardRefCount(Handle.GetInstance()) == 0);
-		TestTrue(TEXT("Reference soft refcount"), FWorldPartitionActorDescUnitTestAcccessor::GetSoftRefCount(Reference.GetInstance()) == 0);
-		TestTrue(TEXT("Reference hard refcount"), FWorldPartitionActorDescUnitTestAcccessor::GetHardRefCount(Reference.GetInstance()) == 1);
+		TestTrue(TEXT("Handle soft refcount"), FWorldPartitionActorDescUnitTestAcccessor::GetSoftRefCount(*Handle) == 1);
+		TestTrue(TEXT("Handle hard refcount"), FWorldPartitionActorDescUnitTestAcccessor::GetHardRefCount(*Handle) == 0);
+		TestTrue(TEXT("Reference soft refcount"), FWorldPartitionActorDescUnitTestAcccessor::GetSoftRefCount(*Reference) == 0);
+		TestTrue(TEXT("Reference hard refcount"), FWorldPartitionActorDescUnitTestAcccessor::GetHardRefCount(*Reference) == 1);
 
 		// Pin handle scope test
 		{
 			FWorldPartitionHandlePinRefScope PinRefScopeHandle(Handle);
-			TestTrue(TEXT("Pin to Handle soft refcount"), FWorldPartitionActorDescUnitTestAcccessor::GetSoftRefCount(Handle.GetInstance()) == 1);
-			TestTrue(TEXT("Pin to Handle hard refcount"), FWorldPartitionActorDescUnitTestAcccessor::GetHardRefCount(Handle.GetInstance()) == 0);
+			TestTrue(TEXT("Pin to Handle soft refcount"), FWorldPartitionActorDescUnitTestAcccessor::GetSoftRefCount(*Handle) == 1);
+			TestTrue(TEXT("Pin to Handle hard refcount"), FWorldPartitionActorDescUnitTestAcccessor::GetHardRefCount(*Handle) == 0);
 		}
 
 		// Pin reference scope test
 		{
 			FWorldPartitionHandlePinRefScope PinRefScopeReference(Reference);
-			TestTrue(TEXT("Pin to Reference soft refcount"), FWorldPartitionActorDescUnitTestAcccessor::GetSoftRefCount(Reference.GetInstance()) == 0);
-			TestTrue(TEXT("Pin to Reference hard refcount"), FWorldPartitionActorDescUnitTestAcccessor::GetHardRefCount(Reference.GetInstance()) == 2);
+			TestTrue(TEXT("Pin to Reference soft refcount"), FWorldPartitionActorDescUnitTestAcccessor::GetSoftRefCount(*Reference) == 0);
+			TestTrue(TEXT("Pin to Reference hard refcount"), FWorldPartitionActorDescUnitTestAcccessor::GetHardRefCount(*Reference) == 2);
 		}
 
 		TestTrue(TEXT("Handle/Reference equality"), Handle != Reference);
 		TestTrue(TEXT("Reference/Handle equality"), Reference != Handle);
 		
-		TestTrue(TEXT("Handle soft refcount"), FWorldPartitionActorDescUnitTestAcccessor::GetSoftRefCount(Handle.GetInstance()) == 1);
-		TestTrue(TEXT("Handle hard refcount"), FWorldPartitionActorDescUnitTestAcccessor::GetHardRefCount(Handle.GetInstance()) == 0);
-		TestTrue(TEXT("Reference soft refcount"), FWorldPartitionActorDescUnitTestAcccessor::GetSoftRefCount(Reference.GetInstance()) == 0);
-		TestTrue(TEXT("Reference hard refcount"), FWorldPartitionActorDescUnitTestAcccessor::GetHardRefCount(Reference.GetInstance()) == 1);
+		TestTrue(TEXT("Handle soft refcount"), FWorldPartitionActorDescUnitTestAcccessor::GetSoftRefCount(*Handle) == 1);
+		TestTrue(TEXT("Handle hard refcount"), FWorldPartitionActorDescUnitTestAcccessor::GetHardRefCount(*Handle) == 0);
+		TestTrue(TEXT("Reference soft refcount"), FWorldPartitionActorDescUnitTestAcccessor::GetSoftRefCount(*Reference) == 0);
+		TestTrue(TEXT("Reference hard refcount"), FWorldPartitionActorDescUnitTestAcccessor::GetHardRefCount(*Reference) == 1);
 
 		// Conversions
 		{
 			FWorldPartitionHandle HandleToReference = Reference.ToHandle();
 			TestTrue(TEXT("Handle/Reference equality"), HandleToReference == Reference);
 			TestTrue(TEXT("Reference/Handle equality"), Reference == HandleToReference);
-			TestTrue(TEXT("Reference soft refcount"), FWorldPartitionActorDescUnitTestAcccessor::GetSoftRefCount(Reference.GetInstance()) == 1);
-			TestTrue(TEXT("Reference hard refcount"), FWorldPartitionActorDescUnitTestAcccessor::GetHardRefCount(Reference.GetInstance()) == 1);
+			TestTrue(TEXT("Reference soft refcount"), FWorldPartitionActorDescUnitTestAcccessor::GetSoftRefCount(*Reference) == 1);
+			TestTrue(TEXT("Reference hard refcount"), FWorldPartitionActorDescUnitTestAcccessor::GetHardRefCount(*Reference) == 1);
 
 			FWorldPartitionReference ReferenceToHandle = Handle.ToReference();
 			TestTrue(TEXT("Handle/Reference equality"), ReferenceToHandle == Handle);
 			TestTrue(TEXT("Handle/Reference equality"), Handle == ReferenceToHandle);
-			TestTrue(TEXT("Reference soft refcount"), FWorldPartitionActorDescUnitTestAcccessor::GetSoftRefCount(Reference.GetInstance()) == 1);
-			TestTrue(TEXT("Reference hard refcount"), FWorldPartitionActorDescUnitTestAcccessor::GetHardRefCount(Reference.GetInstance()) == 1);
+			TestTrue(TEXT("Reference soft refcount"), FWorldPartitionActorDescUnitTestAcccessor::GetSoftRefCount(*Reference) == 1);
+			TestTrue(TEXT("Reference hard refcount"), FWorldPartitionActorDescUnitTestAcccessor::GetHardRefCount(*Reference) == 1);
 		}
 
-		TestTrue(TEXT("Handle soft refcount"), FWorldPartitionActorDescUnitTestAcccessor::GetSoftRefCount(Handle.GetInstance()) == 1);
-		TestTrue(TEXT("Handle hard refcount"), FWorldPartitionActorDescUnitTestAcccessor::GetHardRefCount(Handle.GetInstance()) == 0);
-		TestTrue(TEXT("Reference soft refcount"), FWorldPartitionActorDescUnitTestAcccessor::GetSoftRefCount(Reference.GetInstance()) == 0);
-		TestTrue(TEXT("Reference hard refcount"), FWorldPartitionActorDescUnitTestAcccessor::GetHardRefCount(Reference.GetInstance()) == 1);
+		TestTrue(TEXT("Handle soft refcount"), FWorldPartitionActorDescUnitTestAcccessor::GetSoftRefCount(*Handle) == 1);
+		TestTrue(TEXT("Handle hard refcount"), FWorldPartitionActorDescUnitTestAcccessor::GetHardRefCount(*Handle) == 0);
+		TestTrue(TEXT("Reference soft refcount"), FWorldPartitionActorDescUnitTestAcccessor::GetSoftRefCount(*Reference) == 0);
+		TestTrue(TEXT("Reference hard refcount"), FWorldPartitionActorDescUnitTestAcccessor::GetHardRefCount(*Reference) == 1);
 
 		// inplace new test
 		{
 			uint8 Buffer[sizeof(FWorldPartitionHandle)];
 			FWorldPartitionHandle* HandlePtr = new (Buffer) FWorldPartitionHandle(Reference.ToHandle());
 
-			TestTrue(TEXT("Handle array soft refcount"), FWorldPartitionActorDescUnitTestAcccessor::GetSoftRefCount(Reference.GetInstance()) == 1);
-			TestTrue(TEXT("Handle array hard refcount"), FWorldPartitionActorDescUnitTestAcccessor::GetHardRefCount(Reference.GetInstance()) == 1);
+			TestTrue(TEXT("Handle array soft refcount"), FWorldPartitionActorDescUnitTestAcccessor::GetSoftRefCount(*Reference) == 1);
+			TestTrue(TEXT("Handle array hard refcount"), FWorldPartitionActorDescUnitTestAcccessor::GetHardRefCount(*Reference) == 1);
 
 			HandlePtr->~FWorldPartitionHandle();
 
-			TestTrue(TEXT("Handle array soft refcount"), FWorldPartitionActorDescUnitTestAcccessor::GetSoftRefCount(Reference.GetInstance()) == 0);
-			TestTrue(TEXT("Handle array hard refcount"), FWorldPartitionActorDescUnitTestAcccessor::GetHardRefCount(Reference.GetInstance()) == 1);
+			TestTrue(TEXT("Handle array soft refcount"), FWorldPartitionActorDescUnitTestAcccessor::GetSoftRefCount(*Reference) == 0);
+			TestTrue(TEXT("Handle array hard refcount"), FWorldPartitionActorDescUnitTestAcccessor::GetHardRefCount(*Reference) == 1);
 		}
 
 		// TArray test
@@ -148,14 +148,14 @@ namespace WorldPartitionTests
 			TArray<FWorldPartitionHandle> HandleList;
 			HandleList.Add(Handle);
 
-			TestTrue(TEXT("Handle array soft refcount"), FWorldPartitionActorDescUnitTestAcccessor::GetSoftRefCount(Handle.GetInstance()) == 2);
-			TestTrue(TEXT("Handle array hard refcount"), FWorldPartitionActorDescUnitTestAcccessor::GetHardRefCount(Handle.GetInstance()) == 0);
+			TestTrue(TEXT("Handle array soft refcount"), FWorldPartitionActorDescUnitTestAcccessor::GetSoftRefCount(*Handle) == 2);
+			TestTrue(TEXT("Handle array hard refcount"), FWorldPartitionActorDescUnitTestAcccessor::GetHardRefCount(*Handle) == 0);
 
 			FWorldPartitionReference ReferenceToHandle = Handle.ToReference();
 			TestTrue(TEXT("Handle/Reference equality"), ReferenceToHandle == Handle);
 			TestTrue(TEXT("Handle/Reference equality"), Handle == ReferenceToHandle);
-			TestTrue(TEXT("Handle soft refcount"), FWorldPartitionActorDescUnitTestAcccessor::GetSoftRefCount(Handle.GetInstance()) == 2);
-			TestTrue(TEXT("Handle hard refcount"), FWorldPartitionActorDescUnitTestAcccessor::GetHardRefCount(Handle.GetInstance()) == 1);
+			TestTrue(TEXT("Handle soft refcount"), FWorldPartitionActorDescUnitTestAcccessor::GetSoftRefCount(*Handle) == 2);
+			TestTrue(TEXT("Handle hard refcount"), FWorldPartitionActorDescUnitTestAcccessor::GetHardRefCount(*Handle) == 1);
 
 			TestTrue(TEXT("Handle array contains handle"), HandleList.Contains(Handle));
 			TestTrue(TEXT("Handle array contains reference"), HandleList.Contains(ReferenceToHandle));
@@ -163,36 +163,36 @@ namespace WorldPartitionTests
 			HandleList.Add(Reference.ToHandle());
 			
 			TestTrue(TEXT("Handle array contains reference"), HandleList.Contains(Reference));
-			TestTrue(TEXT("Handle array soft refcount"), FWorldPartitionActorDescUnitTestAcccessor::GetSoftRefCount(Reference.GetInstance()) == 1);
-			TestTrue(TEXT("Handle array hard refcount"), FWorldPartitionActorDescUnitTestAcccessor::GetHardRefCount(Reference.GetInstance()) == 1);
+			TestTrue(TEXT("Handle array soft refcount"), FWorldPartitionActorDescUnitTestAcccessor::GetSoftRefCount(*Reference) == 1);
+			TestTrue(TEXT("Handle array hard refcount"), FWorldPartitionActorDescUnitTestAcccessor::GetHardRefCount(*Reference) == 1);
 			
 			HandleList.Remove(Handle);
-			TestTrue(TEXT("Handle array soft refcount"), FWorldPartitionActorDescUnitTestAcccessor::GetSoftRefCount(Handle.GetInstance()) == 1);
-			TestTrue(TEXT("Handle array hard refcount"), FWorldPartitionActorDescUnitTestAcccessor::GetHardRefCount(Handle.GetInstance()) == 1);
+			TestTrue(TEXT("Handle array soft refcount"), FWorldPartitionActorDescUnitTestAcccessor::GetSoftRefCount(*Handle) == 1);
+			TestTrue(TEXT("Handle array hard refcount"), FWorldPartitionActorDescUnitTestAcccessor::GetHardRefCount(*Handle) == 1);
 
 			HandleList.Remove(Reference.ToHandle());
-			TestTrue(TEXT("Handle array soft refcount"), FWorldPartitionActorDescUnitTestAcccessor::GetSoftRefCount(Reference.GetInstance()) == 0);
-			TestTrue(TEXT("Handle array hard refcount"), FWorldPartitionActorDescUnitTestAcccessor::GetHardRefCount(Reference.GetInstance()) == 1);
+			TestTrue(TEXT("Handle array soft refcount"), FWorldPartitionActorDescUnitTestAcccessor::GetSoftRefCount(*Reference) == 0);
+			TestTrue(TEXT("Handle array hard refcount"), FWorldPartitionActorDescUnitTestAcccessor::GetHardRefCount(*Reference) == 1);
 		}
 
-		TestTrue(TEXT("Handle soft refcount"), FWorldPartitionActorDescUnitTestAcccessor::GetSoftRefCount(Handle.GetInstance()) == 1);
-		TestTrue(TEXT("Handle hard refcount"), FWorldPartitionActorDescUnitTestAcccessor::GetHardRefCount(Handle.GetInstance()) == 0);
-		TestTrue(TEXT("Reference soft refcount"), FWorldPartitionActorDescUnitTestAcccessor::GetSoftRefCount(Reference.GetInstance()) == 0);
-		TestTrue(TEXT("Reference hard refcount"), FWorldPartitionActorDescUnitTestAcccessor::GetHardRefCount(Reference.GetInstance()) == 1);
+		TestTrue(TEXT("Handle soft refcount"), FWorldPartitionActorDescUnitTestAcccessor::GetSoftRefCount(*Handle) == 1);
+		TestTrue(TEXT("Handle hard refcount"), FWorldPartitionActorDescUnitTestAcccessor::GetHardRefCount(*Handle) == 0);
+		TestTrue(TEXT("Reference soft refcount"), FWorldPartitionActorDescUnitTestAcccessor::GetSoftRefCount(*Reference) == 0);
+		TestTrue(TEXT("Reference hard refcount"), FWorldPartitionActorDescUnitTestAcccessor::GetHardRefCount(*Reference) == 1);
 
 		// TSet test
 		{
 			TSet<FWorldPartitionHandle> HandleSet;
 			HandleSet.Add(Handle);
 
-			TestTrue(TEXT("Handle set soft refcount"), FWorldPartitionActorDescUnitTestAcccessor::GetSoftRefCount(Handle.GetInstance()) == 2);
-			TestTrue(TEXT("Handle set hard refcount"), FWorldPartitionActorDescUnitTestAcccessor::GetHardRefCount(Handle.GetInstance()) == 0);
+			TestTrue(TEXT("Handle set soft refcount"), FWorldPartitionActorDescUnitTestAcccessor::GetSoftRefCount(*Handle) == 2);
+			TestTrue(TEXT("Handle set hard refcount"), FWorldPartitionActorDescUnitTestAcccessor::GetHardRefCount(*Handle) == 0);
 
 			FWorldPartitionReference ReferenceToHandle = Handle.ToReference();
 			TestTrue(TEXT("Handle/Reference equality"), ReferenceToHandle == Handle);
 			TestTrue(TEXT("Handle/Reference equality"), Handle == ReferenceToHandle);
-			TestTrue(TEXT("Handle soft refcount"), FWorldPartitionActorDescUnitTestAcccessor::GetSoftRefCount(Handle.GetInstance()) == 2);
-			TestTrue(TEXT("Handle hard refcount"), FWorldPartitionActorDescUnitTestAcccessor::GetHardRefCount(Handle.GetInstance()) == 1);
+			TestTrue(TEXT("Handle soft refcount"), FWorldPartitionActorDescUnitTestAcccessor::GetSoftRefCount(*Handle) == 2);
+			TestTrue(TEXT("Handle hard refcount"), FWorldPartitionActorDescUnitTestAcccessor::GetHardRefCount(*Handle) == 1);
 
 			TestTrue(TEXT("Handle set contains handle"), HandleSet.Contains(Handle));
 			TestTrue(TEXT("Handle set contains reference"), HandleSet.Contains(ReferenceToHandle.ToHandle()));
@@ -200,49 +200,49 @@ namespace WorldPartitionTests
 			HandleSet.Add(Reference.ToHandle());
 			TestTrue(TEXT("Handle set contains reference"), HandleSet.Contains(Reference.ToHandle()));
 
-			TestTrue(TEXT("Reference soft refcount"), FWorldPartitionActorDescUnitTestAcccessor::GetSoftRefCount(Reference.GetInstance()) == 1);
-			TestTrue(TEXT("Reference hard refcount"), FWorldPartitionActorDescUnitTestAcccessor::GetHardRefCount(Reference.GetInstance()) == 1);
+			TestTrue(TEXT("Reference soft refcount"), FWorldPartitionActorDescUnitTestAcccessor::GetSoftRefCount(*Reference) == 1);
+			TestTrue(TEXT("Reference hard refcount"), FWorldPartitionActorDescUnitTestAcccessor::GetHardRefCount(*Reference) == 1);
 		}
 
-		TestTrue(TEXT("Handle soft refcount"), FWorldPartitionActorDescUnitTestAcccessor::GetSoftRefCount(Handle.GetInstance()) == 1);
-		TestTrue(TEXT("Handle hard refcount"), FWorldPartitionActorDescUnitTestAcccessor::GetHardRefCount(Handle.GetInstance()) == 0);
-		TestTrue(TEXT("Reference soft refcount"), FWorldPartitionActorDescUnitTestAcccessor::GetSoftRefCount(Reference.GetInstance()) == 0);
-		TestTrue(TEXT("Reference hard refcount"), FWorldPartitionActorDescUnitTestAcccessor::GetHardRefCount(Reference.GetInstance()) == 1);
+		TestTrue(TEXT("Handle soft refcount"), FWorldPartitionActorDescUnitTestAcccessor::GetSoftRefCount(*Handle) == 1);
+		TestTrue(TEXT("Handle hard refcount"), FWorldPartitionActorDescUnitTestAcccessor::GetHardRefCount(*Handle) == 0);
+		TestTrue(TEXT("Reference soft refcount"), FWorldPartitionActorDescUnitTestAcccessor::GetSoftRefCount(*Reference) == 0);
+		TestTrue(TEXT("Reference hard refcount"), FWorldPartitionActorDescUnitTestAcccessor::GetHardRefCount(*Reference) == 1);
 
 		// Move tests
 		{
 			// Handle move
-			TestTrue(TEXT("Handle soft refcount"), FWorldPartitionActorDescUnitTestAcccessor::GetSoftRefCount(Handle.GetInstance()) == 1);
-			TestTrue(TEXT("Handle hard refcount"), FWorldPartitionActorDescUnitTestAcccessor::GetHardRefCount(Handle.GetInstance()) == 0);
+			TestTrue(TEXT("Handle soft refcount"), FWorldPartitionActorDescUnitTestAcccessor::GetSoftRefCount(*Handle) == 1);
+			TestTrue(TEXT("Handle hard refcount"), FWorldPartitionActorDescUnitTestAcccessor::GetHardRefCount(*Handle) == 0);
 			{
 				FWorldPartitionHandle HandleCopy(MoveTemp(Handle));
 				TestTrue(TEXT("Handle move src not valid"), !Handle.IsValid());
 				TestTrue(TEXT("Handle move dst valid"), HandleCopy.IsValid());
-				TestTrue(TEXT("Handle soft refcount"), FWorldPartitionActorDescUnitTestAcccessor::GetSoftRefCount(HandleCopy.GetInstance()) == 1);
-				TestTrue(TEXT("Handle hard refcount"), FWorldPartitionActorDescUnitTestAcccessor::GetHardRefCount(HandleCopy.GetInstance()) == 0);
+				TestTrue(TEXT("Handle soft refcount"), FWorldPartitionActorDescUnitTestAcccessor::GetSoftRefCount(*HandleCopy) == 1);
+				TestTrue(TEXT("Handle hard refcount"), FWorldPartitionActorDescUnitTestAcccessor::GetHardRefCount(*HandleCopy) == 0);
 
 				Handle = MoveTemp(HandleCopy);
 				TestTrue(TEXT("Handle move src not valid"), !HandleCopy.IsValid());
 				TestTrue(TEXT("Handle move dst valid"), Handle.IsValid());
-				TestTrue(TEXT("Handle soft refcount"), FWorldPartitionActorDescUnitTestAcccessor::GetSoftRefCount(Handle.GetInstance()) == 1);
-				TestTrue(TEXT("Handle hard refcount"), FWorldPartitionActorDescUnitTestAcccessor::GetHardRefCount(Handle.GetInstance()) == 0);
+				TestTrue(TEXT("Handle soft refcount"), FWorldPartitionActorDescUnitTestAcccessor::GetSoftRefCount(*Handle) == 1);
+				TestTrue(TEXT("Handle hard refcount"), FWorldPartitionActorDescUnitTestAcccessor::GetHardRefCount(*Handle) == 0);
 			}
 
 			// Reference move
-			TestTrue(TEXT("Reference soft refcount"), FWorldPartitionActorDescUnitTestAcccessor::GetSoftRefCount(Reference.GetInstance()) == 0);
-			TestTrue(TEXT("Reference hard refcount"), FWorldPartitionActorDescUnitTestAcccessor::GetHardRefCount(Reference.GetInstance()) == 1);
+			TestTrue(TEXT("Reference soft refcount"), FWorldPartitionActorDescUnitTestAcccessor::GetSoftRefCount(*Reference) == 0);
+			TestTrue(TEXT("Reference hard refcount"), FWorldPartitionActorDescUnitTestAcccessor::GetHardRefCount(*Reference) == 1);
 			{
 				FWorldPartitionReference ReferenceCopy(MoveTemp(Reference));
 				TestTrue(TEXT("Reference move src not valid"), !Reference.IsValid());
 				TestTrue(TEXT("Reference move dst valid"), ReferenceCopy.IsValid());
-				TestTrue(TEXT("Reference soft refcount"), FWorldPartitionActorDescUnitTestAcccessor::GetSoftRefCount(ReferenceCopy.GetInstance()) == 0);
-				TestTrue(TEXT("Reference hard refcount"), FWorldPartitionActorDescUnitTestAcccessor::GetHardRefCount(ReferenceCopy.GetInstance()) == 1);
+				TestTrue(TEXT("Reference soft refcount"), FWorldPartitionActorDescUnitTestAcccessor::GetSoftRefCount(*ReferenceCopy) == 0);
+				TestTrue(TEXT("Reference hard refcount"), FWorldPartitionActorDescUnitTestAcccessor::GetHardRefCount(*ReferenceCopy) == 1);
 
 				Reference = MoveTemp(ReferenceCopy);
 				TestTrue(TEXT("Reference move src not valid"), !ReferenceCopy.IsValid());
 				TestTrue(TEXT("Reference move dst valid"), Reference.IsValid());
-				TestTrue(TEXT("Reference soft refcount"), FWorldPartitionActorDescUnitTestAcccessor::GetSoftRefCount(Reference.GetInstance()) == 0);
-				TestTrue(TEXT("Reference hard refcount"), FWorldPartitionActorDescUnitTestAcccessor::GetHardRefCount(Reference.GetInstance()) == 1);
+				TestTrue(TEXT("Reference soft refcount"), FWorldPartitionActorDescUnitTestAcccessor::GetSoftRefCount(*Reference) == 0);
+				TestTrue(TEXT("Reference hard refcount"), FWorldPartitionActorDescUnitTestAcccessor::GetHardRefCount(*Reference) == 1);
 			}
 
 			// Handle reference move
@@ -250,14 +250,14 @@ namespace WorldPartitionTests
 				FWorldPartitionHandle HandleFromReference(MoveTemp(Reference));
 				TestTrue(TEXT("Handle move src not valid"), !Reference.IsValid());
 				TestTrue(TEXT("Handle move dst valid"), HandleFromReference.IsValid());
-				TestTrue(TEXT("Handle soft refcount"), FWorldPartitionActorDescUnitTestAcccessor::GetSoftRefCount(HandleFromReference.GetInstance()) == 1);
-				TestTrue(TEXT("Handle hard refcount"), FWorldPartitionActorDescUnitTestAcccessor::GetHardRefCount(HandleFromReference.GetInstance()) == 0);
+				TestTrue(TEXT("Handle soft refcount"), FWorldPartitionActorDescUnitTestAcccessor::GetSoftRefCount(*HandleFromReference) == 1);
+				TestTrue(TEXT("Handle hard refcount"), FWorldPartitionActorDescUnitTestAcccessor::GetHardRefCount(*HandleFromReference) == 0);
 
 				Reference = MoveTemp(HandleFromReference);
 				TestTrue(TEXT("Handle move src not valid"), !HandleFromReference.IsValid());
 				TestTrue(TEXT("Reference move dst valid"), Reference.IsValid());
-				TestTrue(TEXT("Reference soft refcount"), FWorldPartitionActorDescUnitTestAcccessor::GetSoftRefCount(Reference.GetInstance()) == 0);
-				TestTrue(TEXT("Reference hard refcount"), FWorldPartitionActorDescUnitTestAcccessor::GetHardRefCount(Reference.GetInstance()) == 1);
+				TestTrue(TEXT("Reference soft refcount"), FWorldPartitionActorDescUnitTestAcccessor::GetSoftRefCount(*Reference) == 0);
+				TestTrue(TEXT("Reference hard refcount"), FWorldPartitionActorDescUnitTestAcccessor::GetHardRefCount(*Reference) == 1);
 			}
 
 			// Reference handle move
@@ -265,27 +265,27 @@ namespace WorldPartitionTests
 				FWorldPartitionReference ReferenceFromHandle(MoveTemp(Handle));
 				TestTrue(TEXT("Handle move src not valid"), !Handle.IsValid());
 				TestTrue(TEXT("Reference move dst valid"), ReferenceFromHandle.IsValid());
-				TestTrue(TEXT("Reference soft refcount"), FWorldPartitionActorDescUnitTestAcccessor::GetSoftRefCount(ReferenceFromHandle.GetInstance()) == 0);
-				TestTrue(TEXT("Reference hard refcount"), FWorldPartitionActorDescUnitTestAcccessor::GetHardRefCount(ReferenceFromHandle.GetInstance()) == 1);
+				TestTrue(TEXT("Reference soft refcount"), FWorldPartitionActorDescUnitTestAcccessor::GetSoftRefCount(*ReferenceFromHandle) == 0);
+				TestTrue(TEXT("Reference hard refcount"), FWorldPartitionActorDescUnitTestAcccessor::GetHardRefCount(*ReferenceFromHandle) == 1);
 
 				Handle = MoveTemp(ReferenceFromHandle);
 				TestTrue(TEXT("Reference move src not valid"), !ReferenceFromHandle.IsValid());
 				TestTrue(TEXT("Handle move dst valid"), Handle.IsValid());
-				TestTrue(TEXT("Handle soft refcount"), FWorldPartitionActorDescUnitTestAcccessor::GetSoftRefCount(Handle.GetInstance()) == 1);
-				TestTrue(TEXT("Handle hard refcount"), FWorldPartitionActorDescUnitTestAcccessor::GetHardRefCount(Handle.GetInstance()) == 0);
+				TestTrue(TEXT("Handle soft refcount"), FWorldPartitionActorDescUnitTestAcccessor::GetSoftRefCount(*Handle) == 1);
+				TestTrue(TEXT("Handle hard refcount"), FWorldPartitionActorDescUnitTestAcccessor::GetHardRefCount(*Handle) == 0);
 			}
 
 			// Reset
 			{
 				FWorldPartitionReference ReferenceFromHandle(Handle.ToReference());
 				TestTrue(TEXT("Reference valid"), ReferenceFromHandle.IsValid());
-				TestTrue(TEXT("Reference soft refcount"), FWorldPartitionActorDescUnitTestAcccessor::GetSoftRefCount(ReferenceFromHandle.GetInstance()) == 1);
-				TestTrue(TEXT("Reference hard refcount"), FWorldPartitionActorDescUnitTestAcccessor::GetHardRefCount(ReferenceFromHandle.GetInstance()) == 1);
+				TestTrue(TEXT("Reference soft refcount"), FWorldPartitionActorDescUnitTestAcccessor::GetSoftRefCount(*ReferenceFromHandle) == 1);
+				TestTrue(TEXT("Reference hard refcount"), FWorldPartitionActorDescUnitTestAcccessor::GetHardRefCount(*ReferenceFromHandle) == 1);
 
 				ReferenceFromHandle.Reset();
 				TestTrue(TEXT("Reference not valid"), !ReferenceFromHandle.IsValid());
-				TestTrue(TEXT("Handle soft refcount"), FWorldPartitionActorDescUnitTestAcccessor::GetSoftRefCount(Handle.GetInstance()) == 1);
-				TestTrue(TEXT("Handle hard refcount"), FWorldPartitionActorDescUnitTestAcccessor::GetHardRefCount(Handle.GetInstance()) == 0);
+				TestTrue(TEXT("Handle soft refcount"), FWorldPartitionActorDescUnitTestAcccessor::GetSoftRefCount(*Handle) == 1);
+				TestTrue(TEXT("Handle hard refcount"), FWorldPartitionActorDescUnitTestAcccessor::GetHardRefCount(*Handle) == 0);
 			}
 		}
 
diff --git a/Engine/Source/Runtime/Engine/Private/Tests/WorldPartition/WorldPartitionPathConversionTests.cpp b/Engine/Source/Runtime/Engine/Private/Tests/WorldPartition/WorldPartitionPathConversionTests.cpp
--- a/Engine/Source/Runtime/Engine/Private/Tests/WorldPartition/WorldPartitionPathConversionTests.cpp
+++ b/Engine/Source/Runtime/Engine/Private/Tests/WorldPartition/WorldPartitionPathConversionTests.cpp
@@ -90,7 +90,7 @@ namespace WorldPartitionTests
 			return false;
 		}
 
-		UObject* ResolvedObject = ActorHandle.GetInstance()->GetActorSoftPath().TryLoad();
+		UObject* ResolvedObject = ActorHandle->GetActorSoftPath().TryLoad();
 		if (!TestNotNull(TEXT("Actor Handle Loaded"), ActorHandle.GetActor()))
 		{
 			return false;
diff --git a/Engine/Source/Runtime/Engine/Private/WorldPartition/LoaderAdapter/LoaderAdapterActorList.cpp b/Engine/Source/Runtime/Engine/Private/WorldPartition/LoaderAdapter/LoaderAdapterActorList.cpp
--- a/Engine/Source/Runtime/Engine/Private/WorldPartition/LoaderAdapter/LoaderAdapterActorList.cpp
+++ b/Engine/Source/Runtime/Engine/Private/WorldPartition/LoaderAdapter/LoaderAdapterActorList.cpp
@@ -65,7 +65,7 @@ void FLoaderAdapterActorList::RemoveActors(const TArray<FWorldPartitionHandle>&
 
 	for (const FWorldPartitionHandle& ActorHandle : ActorHandles)
 	{
-		if (ActorHandle.GetInstance()->IsChildContainerInstance())
+		if (ActorHandle->IsChildContainerInstance())
 		{
 			if (UWorldPartition* ContainerWorldPartition = GetLoadedChildWorldPartition(ActorHandle); ContainerWorldPartition && ContainerWorldPartition->IsStreamingEnabledInEditor())
 			{
diff --git a/Engine/Source/Runtime/Engine/Private/WorldPartition/LoaderAdapter/LoaderAdapterList.cpp b/Engine/Source/Runtime/Engine/Private/WorldPartition/LoaderAdapter/LoaderAdapterList.cpp
--- a/Engine/Source/Runtime/Engine/Private/WorldPartition/LoaderAdapter/LoaderAdapterList.cpp
+++ b/Engine/Source/Runtime/Engine/Private/WorldPartition/LoaderAdapter/LoaderAdapterList.cpp
@@ -26,7 +26,7 @@ void FLoaderAdapterList::ForEachActor(TFunctionRef<void(const FWorldPartitionHan
 
 void FLoaderAdapterList::HandleActorContainer(const FWorldPartitionHandle& InActor, TFunctionRef<void(const FWorldPartitionHandle&)> InOperation) const
 {
-	if (InActor.GetInstance()->IsChildContainerInstance())
+	if (InActor->IsChildContainerInstance())
 	{
 		if (UWorldPartition* ContainerWorldPartition = GetLoadedChildWorldPartition(InActor); ContainerWorldPartition && ContainerWorldPartition->IsStreamingEnabledInEditor())
 		{
diff --git a/Engine/Source/Runtime/Engine/Private/WorldPartition/LoaderAdapter/LoaderAdapterPinnedActors.cpp b/Engine/Source/Runtime/Engine/Private/WorldPartition/LoaderAdapter/LoaderAdapterPinnedActors.cpp
--- a/Engine/Source/Runtime/Engine/Private/WorldPartition/LoaderAdapter/LoaderAdapterPinnedActors.cpp
+++ b/Engine/Source/Runtime/Engine/Private/WorldPartition/LoaderAdapter/LoaderAdapterPinnedActors.cpp
@@ -47,11 +47,9 @@ bool FLoaderAdapterPinnedActors::PassActorDescFilter(const FWorldPartitionHandle
 	// We want to be able to pin any type of actors (HLODs, etc).
 	// Allow recursive pinning by setting bCheckIsMainWorldPartition = false
 	const bool bCheckIsMainWorldPartition = false;
-	return ActorHandle.IsValid() && !ActorsToRemove.Contains(ActorHandle) && LoaderAdapterPinnedActorsUtils::SupportsPinning(ActorHandle.GetInstance(), bCheckIsMainWorldPartition);
+	return ActorHandle.IsValid() && !ActorsToRemove.Contains(ActorHandle) && LoaderAdapterPinnedActorsUtils::SupportsPinning(*ActorHandle, bCheckIsMainWorldPartition);
 }
 
-
-
 bool FLoaderAdapterPinnedActors::SupportsPinning(FWorldPartitionActorDescInstance* InActorDescInstance)
 {
 	// Public api doesn't allow pinning of non main world partition actors
diff --git a/Engine/Source/Runtime/Engine/Private/WorldPartition/RuntimeHashSet/WorldPartitionRuntimeHashSetHLODGeneration.cpp b/Engine/Source/Runtime/Engine/Private/WorldPartition/RuntimeHashSet/WorldPartitionRuntimeHashSetHLODGeneration.cpp
--- a/Engine/Source/Runtime/Engine/Private/WorldPartition/RuntimeHashSet/WorldPartitionRuntimeHashSetHLODGeneration.cpp
+++ b/Engine/Source/Runtime/Engine/Private/WorldPartition/RuntimeHashSet/WorldPartitionRuntimeHashSetHLODGeneration.cpp
@@ -137,8 +137,8 @@ namespace PrivateUtils
 		}
 		else
 		{
-			DeletePackage(Handle.GetInstance()->GetActorPackage().ToString(), SourceControlHelper);
-			WorldPartition->RemoveActor(Handle.GetInstance()->GetGuid());
+			DeletePackage(Handle->GetActorPackage().ToString(), SourceControlHelper);
+			WorldPartition->RemoveActor(Handle->GetGuid());
 		}
 	}
 }
diff --git a/Engine/Source/Runtime/Engine/Private/WorldPartition/RuntimeSpatialHash/RuntimeSpatialHashHLOD.cpp b/Engine/Source/Runtime/Engine/Private/WorldPartition/RuntimeSpatialHash/RuntimeSpatialHashHLOD.cpp
--- a/Engine/Source/Runtime/Engine/Private/WorldPartition/RuntimeSpatialHash/RuntimeSpatialHashHLOD.cpp
+++ b/Engine/Source/Runtime/Engine/Private/WorldPartition/RuntimeSpatialHash/RuntimeSpatialHashHLOD.cpp
@@ -583,7 +583,7 @@ bool UWorldPartitionRuntimeSpatialHash::SetupHLODActors(const IStreamingGenerati
 
 	auto DeleteHLODActor = [&SourceControlHelper, WorldPartition](FWorldPartitionHandle ActorHandle)
 	{
-		FWorldPartitionActorDescInstance* ActorDescInstance = ActorHandle.GetInstance();
+		FWorldPartitionActorDescInstance* ActorDescInstance = *ActorHandle;
 		check(ActorDescInstance);
 
 		DeletePackage(WorldPartition, ActorDescInstance, SourceControlHelper);
diff --git a/Engine/Source/Runtime/Engine/Private/WorldPartition/WorldPartition.cpp b/Engine/Source/Runtime/Engine/Private/WorldPartition/WorldPartition.cpp
--- a/Engine/Source/Runtime/Engine/Private/WorldPartition/WorldPartition.cpp
+++ b/Engine/Source/Runtime/Engine/Private/WorldPartition/WorldPartition.cpp
@@ -279,7 +279,7 @@ void UWorldPartition::FWorldPartitionExternalDirtyActorsTracker::Tick(float InDe
 		}
 
 		// Transfer ownership of our last ref if actor can be pinned
-		if (Reference.IsValid() && Reference.GetInstance()->GetHardRefCount() <= 1 && Owner->PinnedActors && FLoaderAdapterPinnedActors::SupportsPinning(Reference.GetInstance()))
+		if (Reference.IsValid() && Reference->GetHardRefCount() <= 1 && Owner->PinnedActors && FLoaderAdapterPinnedActors::SupportsPinning(*Reference))
 		{
 			Owner->PinnedActors->AddActors({ Reference.ToHandle() });
 		}
diff --git a/Engine/Source/Runtime/Engine/Private/WorldPartition/WorldPartitionActorLoaderInterface.cpp b/Engine/Source/Runtime/Engine/Private/WorldPartition/WorldPartitionActorLoaderInterface.cpp
--- a/Engine/Source/Runtime/Engine/Private/WorldPartition/WorldPartitionActorLoaderInterface.cpp
+++ b/Engine/Source/Runtime/Engine/Private/WorldPartition/WorldPartitionActorLoaderInterface.cpp
@@ -122,7 +122,7 @@ void IWorldPartitionActorLoaderInterface::ILoaderAdapter::UnregisterDelegates()
 
 bool IWorldPartitionActorLoaderInterface::ILoaderAdapter::PassActorDescFilter(const FWorldPartitionHandle& Actor) const
 {
-	return Actor.GetInstance()->IsEditorRelevant();
+	return Actor->IsEditorRelevant();
 }
 
 void IWorldPartitionActorLoaderInterface::ILoaderAdapter::RefreshLoadedState()
@@ -135,16 +135,16 @@ void IWorldPartitionActorLoaderInterface::ILoaderAdapter::RefreshLoadedState()
 			TArray<FWorldPartitionHandle> ActorsToUnload;			
 			ForEachActor([this, &ActorsToLoad, &ActorsToUnload](const FWorldPartitionHandle& ActorHandle)
 			{
-				const FActorReferenceMap* ActorReferences = GetContainerReferencesConst(ActorHandle.GetContainerInstance());
+				const FActorReferenceMap* ActorReferences = GetContainerReferencesConst(ActorHandle->GetContainerInstance());
 
 				if (ShouldActorBeLoaded(ActorHandle))
 				{
-					if (!ActorReferences || !ActorReferences->Contains(ActorHandle.GetInstance()->GetGuid()))
+					if (!ActorReferences || !ActorReferences->Contains(ActorHandle->GetGuid()))
 					{
 						ActorsToLoad.Add(ActorHandle);
 					}
 				}
-				else if (ActorReferences && ActorReferences->Contains(ActorHandle.GetInstance()->GetGuid()))
+				else if (ActorReferences && ActorReferences->Contains(ActorHandle->GetGuid()))
 				{
 					ActorsToUnload.Add(ActorHandle);
 				}
@@ -246,7 +246,7 @@ bool IWorldPartitionActorLoaderInterface::ILoaderAdapter::ShouldActorBeLoaded(co
 {
 	check(Actor.IsValid());
 
-	Actor.GetInstance()->SetUnloadedReason(nullptr);
+	Actor->SetUnloadedReason(nullptr);
 
 	if (!PassActorDescFilter(Actor))
 	{
@@ -257,7 +257,7 @@ bool IWorldPartitionActorLoaderInterface::ILoaderAdapter::ShouldActorBeLoaded(co
 	{
 		if (!ActorDescFilter.Get().PassFilter(World, Actor))
 		{
-			Actor.GetInstance()->SetUnloadedReason(ActorDescFilter.Get().GetFilterReason());
+			Actor->SetUnloadedReason(ActorDescFilter.Get().GetFilterReason());
 			return false;
 		}
 	}
@@ -289,14 +289,14 @@ void IWorldPartitionActorLoaderInterface::ILoaderAdapter::AddReferenceToActor(FW
 {
 	TFunction<void(const FWorldPartitionHandle&, FReferenceMap&)> AddReferences = [this, &AddReferences](const FWorldPartitionHandle& Handle, FReferenceMap& ReferenceMap)
 	{
-		const FWorldPartitionActorDescInstance* ActorDescInstance = Handle.GetInstance();
+		const FWorldPartitionActorDescInstance* ActorDescInstance = *Handle;
 		if (!ReferenceMap.Contains(ActorDescInstance->GetGuid()))
 		{
 			ReferenceMap.Emplace(ActorDescInstance->GetGuid(), Handle.ToReference());
 			
 			for (const FGuid& ReferencedActorGuid : ActorDescInstance->GetReferences())
 			{
-				FWorldPartitionHandle ReferenceActorHandle(Handle.GetContainerInstance(), ReferencedActorGuid);
+				FWorldPartitionHandle ReferenceActorHandle(Handle->GetContainerInstance(), ReferencedActorGuid);
 
 				if (ReferenceActorHandle.IsValid())
 				{
@@ -306,14 +306,14 @@ void IWorldPartitionActorLoaderInterface::ILoaderAdapter::AddReferenceToActor(FW
 		}
 	};
 
-	FActorReferenceMap& ActorReferences = GetContainerReferences(ActorHandle.GetContainerInstance());
-	AddReferences(ActorHandle, ActorReferences.Emplace(ActorHandle.GetInstance()->GetGuid()));
+	FActorReferenceMap& ActorReferences = GetContainerReferences(ActorHandle->GetContainerInstance());
+	AddReferences(ActorHandle, ActorReferences.Emplace(ActorHandle->GetGuid()));
 }
 
 void IWorldPartitionActorLoaderInterface::ILoaderAdapter::RemoveReferenceToActor(FWorldPartitionHandle& ActorHandle)
 {
-	FActorReferenceMap& ActorReferences = GetContainerReferences(ActorHandle.GetContainerInstance());
-	ActorReferences.Remove(ActorHandle.GetInstance()->GetGuid());
+	FActorReferenceMap& ActorReferences = GetContainerReferences(ActorHandle->GetContainerInstance());
+	ActorReferences.Remove(ActorHandle->GetGuid());
 }
 
 void IWorldPartitionActorLoaderInterface::ILoaderAdapter::OnRefreshLoadedState(bool bFromUserOperation)
@@ -358,7 +358,7 @@ UWorldPartition* IWorldPartitionActorLoaderInterface::ILoaderAdapter::GetLoadedC
 
 UWorldPartition* IWorldPartitionActorLoaderInterface::GetLoadedChildWorldPartition(const FWorldPartitionHandle& Handle)
 {
-	if (FWorldPartitionActorDescInstance* Instance = Handle.GetInstance())
+	if (FWorldPartitionActorDescInstance* Instance = *Handle)
 	{
 		return Instance->GetLoadedChildWorldPartition();
 	}
diff --git a/Engine/Source/Runtime/Engine/Private/WorldPartition/WorldPartitionEditorSpatialHash.cpp b/Engine/Source/Runtime/Engine/Private/WorldPartition/WorldPartitionEditorSpatialHash.cpp
--- a/Engine/Source/Runtime/Engine/Private/WorldPartition/WorldPartitionEditorSpatialHash.cpp
+++ b/Engine/Source/Runtime/Engine/Private/WorldPartition/WorldPartitionEditorSpatialHash.cpp
@@ -69,14 +69,14 @@ void UWorldPartitionEditorSpatialHash::Tick(float DeltaSeconds)
 
 			for (const FWorldPartitionHandle& ActorHandle : Cell->Actors)
 			{
-				RuntimeBounds += ActorHandle.GetInstance()->GetEditorBounds();
+				RuntimeBounds += ActorHandle->GetEditorBounds();
 			}
 		}
 
 		NonSpatialBounds.Init();
 		for (const FWorldPartitionHandle& ActorHandle : AlwaysLoadedCell->Actors)
 		{
-			NonSpatialBounds += ActorHandle.GetInstance()->GetEditorBounds();
+			NonSpatialBounds += ActorHandle->GetEditorBounds();
 		}
 
 		const int32 OldLevel = GetLevelForBox(EditorBounds);
@@ -126,7 +126,7 @@ void UWorldPartitionEditorSpatialHash::HashActor(FWorldPartitionHandle& InActorH
 {
 	check(InActorHandle.IsValid());
 
-	const FWorldPartitionActorDescInstance* ActorDescInstance = InActorHandle.GetInstance();
+	const FWorldPartitionActorDescInstance* ActorDescInstance = *InActorHandle;
 	const bool bConsiderActorSpatiallyLoaded = ActorDescInstance->GetIsSpatiallyLoaded();
 	const FBox ActorBounds = bConsiderActorSpatiallyLoaded ? ActorDescInstance->GetEditorBounds() : FBox(ForceInit);
 
@@ -239,7 +239,7 @@ void UWorldPartitionEditorSpatialHash::UnhashActor(FWorldPartitionHandle& InActo
 {
 	check(InActorHandle.IsValid());
 
-	const FWorldPartitionActorDescInstance* ActorDescInstance = InActorHandle.GetInstance();
+	const FWorldPartitionActorDescInstance* ActorDescInstance = *InActorHandle;
 	const bool bConsiderActorSpatiallyLoaded = ActorDescInstance->GetIsSpatiallyLoaded();
 	const FBox ActorBounds = bConsiderActorSpatiallyLoaded ? ActorDescInstance->GetEditorBounds() : FBox(ForceInit);
 
@@ -314,13 +314,13 @@ int32 UWorldPartitionEditorSpatialHash::ForEachIntersectingActor(const FBox& Box
 				if (ActorHandle.IsValid())
 				{
 					bool bWasAlreadyInSet;
-					IntersectedActors.Add(ActorHandle.GetInstance()->GetGuid(), &bWasAlreadyInSet);
+					IntersectedActors.Add(ActorHandle->GetGuid(), &bWasAlreadyInSet);
 
 					if (!bWasAlreadyInSet)
 					{
-						if (Box.Intersect(ActorHandle.GetInstance()->GetEditorBounds()))
+						if (Box.Intersect(ActorHandle->GetEditorBounds()))
 						{
-							InOperation(ActorHandle.GetInstance());
+							InOperation(*ActorHandle);
 						}
 					}
 				}
@@ -334,14 +334,14 @@ int32 UWorldPartitionEditorSpatialHash::ForEachIntersectingActor(const FBox& Box
 		{
 			if (ActorHandle.IsValid())
 			{
-				if (Box.Intersect(ActorHandle.GetInstance()->GetEditorBounds()))
+				if (Box.Intersect(ActorHandle->GetEditorBounds()))
 				{
 					bool bWasAlreadyInSet;
-					IntersectedActors.Add(ActorHandle.GetInstance()->GetGuid(), &bWasAlreadyInSet);
+					IntersectedActors.Add(ActorHandle->GetGuid(), &bWasAlreadyInSet);
 				
 					if (!bWasAlreadyInSet)
 					{
-						InOperation(ActorHandle.GetInstance());
+						InOperation(*ActorHandle);
 					}
 				}
 			}
diff --git a/Engine/Source/Runtime/Engine/Private/WorldPartition/WorldPartitionHelpers.cpp b/Engine/Source/Runtime/Engine/Private/WorldPartition/WorldPartitionHelpers.cpp
--- a/Engine/Source/Runtime/Engine/Private/WorldPartition/WorldPartitionHelpers.cpp
+++ b/Engine/Source/Runtime/Engine/Private/WorldPartition/WorldPartitionHelpers.cpp
@@ -185,7 +185,7 @@ void FWorldPartitionHelpers::ForEachActorWithLoading(UWorldPartition* WorldParti
 				WorldPartitionHelpers::LoadReferences(WorldPartition, ActorDescInstance->GetGuid(), Result.ActorReferences);
 
 				FWorldPartitionReference ActorReference(WorldPartition, ActorDescInstance->GetGuid());
-				if (!Func(ActorReference.GetInstance()))
+				if (!Func(*ActorReference))
 				{
 					return false;
 				}
diff --git a/Engine/Source/Runtime/Engine/Public/WorldPartition/WorldPartitionHandle.h b/Engine/Source/Runtime/Engine/Public/WorldPartition/WorldPartitionHandle.h
--- a/Engine/Source/Runtime/Engine/Public/WorldPartition/WorldPartitionHandle.h
+++ b/Engine/Source/Runtime/Engine/Public/WorldPartition/WorldPartitionHandle.h
@@ -318,20 +318,14 @@ public:
 		return *this;
 	}
 
-	UE_DEPRECATED(5.4, "Use GetInstance instead")
-	FORCEINLINE FWorldPartitionActorDesc* operator->() const
+	FORCEINLINE FWorldPartitionActorDescInstance* operator->() const
 	{
-		PRAGMA_DISABLE_DEPRECATION_WARNINGS
-		return Get();
-		PRAGMA_ENABLE_DEPRECATION_WARNINGS
+		return GetInstance();
 	}
 
-	UE_DEPRECATED(5.4, "Use GetInstance instead")
-	FORCEINLINE FWorldPartitionActorDesc* operator*() const
+	FORCEINLINE FWorldPartitionActorDescInstance* operator*() const
 	{
-		PRAGMA_DISABLE_DEPRECATION_WARNINGS
-		return Get();
-		PRAGMA_ENABLE_DEPRECATION_WARNINGS
+		return GetInstance();
 	}
 
 	FORCEINLINE FWorldPartitionActorDescInstance* GetInstance() const
@@ -339,11 +333,6 @@ public:
 		return IsValid() ? ActorDescInstance->Get() : nullptr;
 	}
 
-	FORCEINLINE UActorDescContainerInstance* GetContainerInstance() const
-	{
-		return IsValid() ? Impl::GetActorDescContainerInstance(ActorDescInstance) : nullptr;
-	}
-
 	FORCEINLINE bool IsValid() const
 	{
 		return ContainerInstance.IsValid() && ActorDescInstance && ActorDescInstance->IsValid();
