#ue5 - Fix root motion replication to proxies when finishing root motion to ensure the replicated shadowed state is correct on the frame that root motion finishes. Replication was being canceled when it became inactive, so some state was stale and might not trigger future updates. Adapted from GitHub #10331

#tests PIE with root motion and multiple players

diff --git a/Engine/Source/Runtime/Engine/Private/Character.cpp b/Engine/Source/Runtime/Engine/Private/Character.cpp
--- a/Engine/Source/Runtime/Engine/Private/Character.cpp
+++ b/Engine/Source/Runtime/Engine/Private/Character.cpp
@@ -1548,9 +1548,12 @@ void ACharacter::PreReplication( IRepChangedPropertyTracker & ChangedPropertyTra
 	}
 	else
 	{
+		const bool bWasRootMotionPreviouslyActive = RepRootMotion.bIsActive;
 		RepRootMotion.Clear();
 
-		DOREPLIFETIME_ACTIVE_OVERRIDE_FAST( ACharacter, RepRootMotion, false );
+		// Replicate RepRootMotion one last time when root motion ends, so that clients see the change.
+		// Then deactivate subsequent property comparisons and replication updates until root motion starts again.
+		DOREPLIFETIME_ACTIVE_OVERRIDE_FAST( ACharacter, RepRootMotion, bWasRootMotionPreviouslyActive );
 	}
 
 	bProxyIsJumpForceApplied = (JumpForceTimeRemaining > 0.0f);
