- Set the ResolutionStatus to Resolved when we are loading from serialization.

diff --git a/Engine/Source/Runtime/Engine/Private/Engine/ActorInstanceHandle.cpp b/Engine/Source/Runtime/Engine/Private/Engine/ActorInstanceHandle.cpp
--- a/Engine/Source/Runtime/Engine/Private/Engine/ActorInstanceHandle.cpp
+++ b/Engine/Source/Runtime/Engine/Private/Engine/ActorInstanceHandle.cpp
@@ -468,6 +468,7 @@ FArchive& operator<<(FArchive& Ar, FActorInstanceHandle& Handle)
 	if (Ar.IsLoading())
 	{
 		Handle.ManagerInterface = FActorInstanceManagerInterface(WeakManagerObject.Get());
+		Handle.ResolutionStatus = FActorInstanceHandle::EResolutionStatus::Resolved;
 	}
 
 	return Ar;
