PR #10735: Fix the error rotation in `ApplyDeltaToAllPhysicsTransforms`. This affects movement applied to the character through the character movement component. Previously the rotation was applied in the wrong space.

diff --git a/Engine/Source/Runtime/Engine/Private/SkeletalMeshComponentPhysics.cpp b/Engine/Source/Runtime/Engine/Private/SkeletalMeshComponentPhysics.cpp
--- a/Engine/Source/Runtime/Engine/Private/SkeletalMeshComponentPhysics.cpp
+++ b/Engine/Source/Runtime/Engine/Private/SkeletalMeshComponentPhysics.cpp
@@ -464,7 +464,7 @@ void USkeletalMeshComponent::ApplyDeltaToAllPhysicsTransforms(const FVector& Del
 		{
 			// move the root body
 			FTransform RootBodyTM = RootBI->GetUnrealWorldTransform();
-			RootBodyTM.SetRotation(RootBodyTM.GetRotation() * DeltaRotation);
+			RootBodyTM.SetRotation(DeltaRotation * RootBodyTM.GetRotation());
 			RootBodyTM.SetTranslation(RootBodyTM.GetTranslation() + DeltaLocation);
 			RootBI->SetBodyTransform(RootBodyTM, ETeleportType::TeleportPhysics);
 
@@ -477,7 +477,7 @@ void USkeletalMeshComponent::ApplyDeltaToAllPhysicsTransforms(const FVector& Del
 					check(BI);
 
 					FTransform BodyTM = BI->GetUnrealWorldTransform();
-					BodyTM.SetRotation(BodyTM.GetRotation() * DeltaRotation);
+					BodyTM.SetRotation(DeltaRotation * BodyTM.GetRotation());
 					BodyTM.SetTranslation(BodyTM.GetTranslation() + DeltaLocation);
 					BI->SetBodyTransform( BodyTM, ETeleportType::TeleportPhysics );
 				}
