Automatic evaluation system for Orthographic Near and Far depth planes.
Uses a units to pixel calculation to determine the necessary precision of the depth buffer wrt 16bit downscaling.
Can be set per camera, and also introduced several CVars for adjusting settings.

Also added debug cameras so all cameras in projects can easily be set to Orthographic for testing.

diff --git a/Engine/Source/Editor/DetailCustomizations/Private/CameraDetails.cpp b/Engine/Source/Editor/DetailCustomizations/Private/CameraDetails.cpp
--- a/Engine/Source/Editor/DetailCustomizations/Private/CameraDetails.cpp
+++ b/Engine/Source/Editor/DetailCustomizations/Private/CameraDetails.cpp
@@ -107,6 +107,9 @@ void FCameraDetails::CustomizeDetails( IDetailLayoutBuilder& DetailLayout )
 		IDetailPropertyRow& OrthoWidthRow = CameraCategory.AddProperty(DetailLayout.GetProperty(GET_MEMBER_NAME_CHECKED(UCameraComponent, OrthoWidth)));
 		OrthoWidthRow.Visibility(OrthographicVisibility);
 
+		IDetailPropertyRow& AutoCalculateOrthoPlanesRow = CameraCategory.AddProperty(DetailLayout.GetProperty(GET_MEMBER_NAME_CHECKED(UCameraComponent, bAutoCalculateOrthoPlanes)));
+		AutoCalculateOrthoPlanesRow.Visibility(OrthographicVisibility);
+
 		IDetailPropertyRow& OrthoNearClipPlaneRow = CameraCategory.AddProperty(DetailLayout.GetProperty(GET_MEMBER_NAME_CHECKED(UCameraComponent, OrthoNearClipPlane)));
 		OrthoNearClipPlaneRow.Visibility(OrthographicVisibility);
 
diff --git a/Engine/Source/Runtime/Core/Public/UObject/UE5ReleaseStreamObjectVersion.h b/Engine/Source/Runtime/Core/Public/UObject/UE5ReleaseStreamObjectVersion.h
--- a/Engine/Source/Runtime/Core/Public/UObject/UE5ReleaseStreamObjectVersion.h
+++ b/Engine/Source/Runtime/Core/Public/UObject/UE5ReleaseStreamObjectVersion.h
@@ -162,6 +162,9 @@ struct FUE5ReleaseStreamObjectVersion
 
 		// Store physics thread particles data in single precision
 		SinglePrecisonParticleDataPT,
+
+		//Orthographic Near and Far Plane Auto-resolve enabled by default
+		OrthographicAutoNearFarPlane,
 		
 		// -----<new versions can be added above this line>-------------------------------------------------
 		VersionPlusOne,
diff --git a/Engine/Source/Runtime/Engine/Classes/Camera/CameraComponent.h b/Engine/Source/Runtime/Engine/Classes/Camera/CameraComponent.h
--- a/Engine/Source/Runtime/Engine/Classes/Camera/CameraComponent.h
+++ b/Engine/Source/Runtime/Engine/Classes/Camera/CameraComponent.h
@@ -40,14 +40,20 @@ class UCameraComponent : public USceneComponent
 	UFUNCTION(BlueprintCallable, Category = Camera)
 	void SetOrthoWidth(float InOrthoWidth) { OrthoWidth = InOrthoWidth; }
 
-	/** The near plane distance of the orthographic view (in world units) */
+	/** Automatically determine a min/max Near/Far clip plane position depending on OrthoWidth value*/
 	UPROPERTY(Interp, EditAnywhere, BlueprintReadWrite, Category = CameraSettings)
+	bool bAutoCalculateOrthoPlanes;
+	UFUNCTION(BlueprintCallable, Category = Camera)
+	void SetAutoCalculateOrthoPlanes(bool bAutoCalculate) { bAutoCalculateOrthoPlanes = bAutoCalculate; }
+
+	/** The near plane distance of the orthographic view (in world units) */
+	UPROPERTY(Interp, EditAnywhere, BlueprintReadWrite, Category = CameraSettings, meta = (EditCondition = "!bAutoCalculateOrthoPlanes"))
 	float OrthoNearClipPlane;
 	UFUNCTION(BlueprintCallable, Category = Camera)
 	void SetOrthoNearClipPlane(float InOrthoNearClipPlane) { OrthoNearClipPlane = InOrthoNearClipPlane; }
 
 	/** The far plane distance of the orthographic view (in world units) */
-	UPROPERTY(Interp, EditAnywhere, BlueprintReadWrite, Category = CameraSettings)
+	UPROPERTY(Interp, EditAnywhere, BlueprintReadWrite, Category = CameraSettings, meta = (EditCondition = "!bAutoCalculateOrthoPlanes"))
 	float OrthoFarClipPlane;
 	UFUNCTION(BlueprintCallable, Category = Camera)
 	void SetOrthoFarClipPlane(float InOrthoFarClipPlane) { OrthoFarClipPlane = InOrthoFarClipPlane; }
diff --git a/Engine/Source/Runtime/Engine/Classes/Camera/CameraTypes.h b/Engine/Source/Runtime/Engine/Classes/Camera/CameraTypes.h
--- a/Engine/Source/Runtime/Engine/Classes/Camera/CameraTypes.h
+++ b/Engine/Source/Runtime/Engine/Classes/Camera/CameraTypes.h
@@ -57,6 +57,10 @@ struct FMinimalViewInfo
 	UPROPERTY(EditAnywhere, BlueprintReadWrite, Category=Camera)
 	float OrthoWidth;
 
+	/** Option for the Ortho camera to automatically calculated the near/far plane */
+	UPROPERTY(EditAnywhere, BlueprintReadWrite, Category=Camera)
+	bool bAutoCalculateOrthoPlanes;
+
 	/** The near plane distance of the orthographic view (in world units) */
 	UPROPERTY(Interp, EditAnywhere, BlueprintReadWrite, Category=Camera)
 	float OrthoNearClipPlane;
@@ -103,12 +107,19 @@ struct FMinimalViewInfo
 	/** Optional transform to be considered as this view's previous transform */
 	TOptional<FTransform> PreviousViewTransform;
 
+private:
+	// Only used for Ortho camera auto plane calculations, tells the Near plane of the extra distance that needs to be added.
+	float OrthoCameraArmLength;
+
+public:
+
 	FMinimalViewInfo()
 		: Location(ForceInit)
 		, Rotation(ForceInit)
 		, FOV(90.0f)
 		, DesiredFOV(90.0f)
 		, OrthoWidth(512.0f)
+		, bAutoCalculateOrthoPlanes(true)
 		, OrthoNearClipPlane(0.0f)
 		, OrthoFarClipPlane(UE_OLD_WORLD_MAX)
 		, PerspectiveNearClipPlane(-1.0f)
@@ -118,6 +129,7 @@ struct FMinimalViewInfo
 		, ProjectionMode(ECameraProjectionMode::Perspective)
 		, PostProcessBlendWeight(0.0f)
 		, OffCenterProjectionOffset(ForceInitToZero)
+		, OrthoCameraArmLength(0.0f)
 	{
 	}
 
@@ -138,15 +150,24 @@ struct FMinimalViewInfo
 	ENGINE_API FMatrix CalculateProjectionMatrix(float ForceOrthoNegativeNearPlane = 0.0f) const;
 
 	/** Calculates the projection matrix (and potentially a constrained view rectangle) given a FMinimalViewInfo and partially configured projection data (must have the view rect already set) */
-	ENGINE_API static void CalculateProjectionMatrixGivenView(const FMinimalViewInfo& ViewInfo, TEnumAsByte<enum EAspectRatioAxisConstraint> AspectRatioAxisConstraint, class FViewport* Viewport, struct FSceneViewProjectionData& InOutProjectionData);
+	ENGINE_API static void CalculateProjectionMatrixGivenView(FMinimalViewInfo& ViewInfo, TEnumAsByte<enum EAspectRatioAxisConstraint> AspectRatioAxisConstraint, class FViewport* Viewport, struct FSceneViewProjectionData& InOutProjectionData);
 	/** Calculates the projection matrix (and potentially a constrained view rectangle) given a FMinimalViewInfo and partially configured projection data (must have the view rect already set). ConstrainedViewRectangle is only used if the ViewInfo.bConstrainAspectRatio is set. */
-	ENGINE_API static void CalculateProjectionMatrixGivenViewRectangle(const FMinimalViewInfo& ViewInfo, TEnumAsByte<enum EAspectRatioAxisConstraint> AspectRatioAxisConstraint, const FIntRect& ConstrainedViewRectangle, FSceneViewProjectionData& InOutProjectionData);
+	ENGINE_API static void CalculateProjectionMatrixGivenViewRectangle(FMinimalViewInfo& ViewInfo, TEnumAsByte<enum EAspectRatioAxisConstraint> AspectRatioAxisConstraint, const FIntRect& ConstrainedViewRectangle, FSceneViewProjectionData& InOutProjectionData);
 
 	/** The near plane distance of the perspective view (in world units). Returns the value of PerspectiveNearClipPlane if positive, and GNearClippingPlane otherwise */
 	FORCEINLINE float GetFinalPerspectiveNearClipPlane() const
 	{
 		return PerspectiveNearClipPlane > 0.0f ? PerspectiveNearClipPlane : GNearClippingPlane;
 	}
+
+	/** Automatically calculates the Near/Far plane values for an Ortho camera */
+	ENGINE_API bool AutoCalculateOrthoPlanes(const FSceneViewProjectionData& InOutProjectionData);
+
+	/** Sets the camera distance from view target for AutoCalculateOrthoPlanes */
+	ENGINE_API inline void SetOrthoCameraArmLengthFromOwnerLocation(const FVector ActorLocation)
+	{
+		OrthoCameraArmLength = (ActorLocation - Location).Length();
+	}
 };
 
 #if UE_ENABLE_INCLUDE_ORDER_DEPRECATED_IN_5_2
diff --git a/Engine/Source/Runtime/Engine/Classes/Camera/PlayerCameraManager.h b/Engine/Source/Runtime/Engine/Classes/Camera/PlayerCameraManager.h
--- a/Engine/Source/Runtime/Engine/Classes/Camera/PlayerCameraManager.h
+++ b/Engine/Source/Runtime/Engine/Classes/Camera/PlayerCameraManager.h
@@ -408,6 +408,10 @@ public:
 	UPROPERTY(EditAnywhere, BlueprintReadWrite, Category = PlayerCameraManager)
 	uint32 bIsOrthographic : 1;
 
+	/** True when this camera should automatically calculated the Near+Far planes */
+	UPROPERTY(EditAnywhere, BlueprintReadWrite, Category = PlayerCameraManager)
+	uint32 bAutoCalculateOrthoPlanes : 1;
+
 	/** True if black bars should be added if the destination view has a different aspect ratio (only used when a view target doesn't specify whether or not to constrain the aspect ratio; most of the time the value from a camera component is used instead) */
 	UPROPERTY(EditAnywhere, BlueprintReadWrite, Category=PlayerCameraManager)
 	uint32 bDefaultConstrainAspectRatio : 1;
diff --git a/Engine/Source/Runtime/Engine/Classes/Components/SceneCaptureComponent2D.h b/Engine/Source/Runtime/Engine/Classes/Components/SceneCaptureComponent2D.h
--- a/Engine/Source/Runtime/Engine/Classes/Components/SceneCaptureComponent2D.h
+++ b/Engine/Source/Runtime/Engine/Classes/Components/SceneCaptureComponent2D.h
@@ -34,6 +34,10 @@ public:
 	UPROPERTY(EditAnywhere, BlueprintReadWrite, Category=Projection)
 	float OrthoWidth;
 
+	/** Automatically determine a min/max Near/Far clip plane position depending on OrthoWidth value*/
+	UPROPERTY(EditAnywhere, BlueprintReadWrite, Category=Projection)
+	bool bAutoCalculateOrthoPlanes;
+
 	/** Output render target of the scene capture that can be read in materials. */
 	UPROPERTY(EditAnywhere, BlueprintReadWrite, Category=SceneCapture)
 	TObjectPtr<class UTextureRenderTarget2D> TextureTarget;
diff --git a/Engine/Source/Runtime/Engine/Private/Camera/CameraComponent.cpp b/Engine/Source/Runtime/Engine/Private/Camera/CameraComponent.cpp
--- a/Engine/Source/Runtime/Engine/Private/Camera/CameraComponent.cpp
+++ b/Engine/Source/Runtime/Engine/Private/Camera/CameraComponent.cpp
@@ -20,33 +20,13 @@
 #include "IXRCamera.h"
 #include "Math/UnitConversion.h"
 #include "UObject/FortniteMainBranchObjectVersion.h"
+#include "UObject/UE5ReleaseStreamObjectVersion.h"
 #include "UObject/UnrealType.h"
 
 #include UE_INLINE_GENERATED_CPP_BY_NAME(CameraComponent)
 
 #define LOCTEXT_NAMESPACE "CameraComponent"
 
-static TAutoConsoleVariable<float> CVarDefaultCameraOrthoWidth(
-	TEXT("r.Ortho.DefaultCameraWidth"),
-	DEFAULT_ORTHOWIDTH,
-	TEXT("Default Ortho Width when creating a new camera actor"),
-	ECVF_ReadOnly
-);
-
-static TAutoConsoleVariable<float> CVarDefaultCameraNearPlane(
-	TEXT("r.Ortho.DefaultCameraNearPlane"),
-	DEFAULT_ORTHONEARPLANE,
-	TEXT("Default Ortho Near Plane when creating a new camera actor"),
-	ECVF_ReadOnly
-);
-
-static TAutoConsoleVariable<float> CVarDefaultCameraFarPlane(
-	TEXT("r.Ortho.DefaultCameraFarPlane"),
-	UE_OLD_WORLD_MAX,
-	TEXT("Default Ortho Far Plane when creating a new camera actor"),
-	ECVF_ReadOnly
-);
-
 //////////////////////////////////////////////////////////////////////////
 // UCameraComponent
 
@@ -66,9 +46,10 @@ UCameraComponent::UCameraComponent(const FObjectInitializer& ObjectInitializer)
 
 	FieldOfView = 90.0f;
 	AspectRatio = 1.777778f;
-	OrthoWidth = CVarDefaultCameraOrthoWidth.GetValueOnAnyThread();
-	OrthoNearClipPlane = CVarDefaultCameraNearPlane.GetValueOnAnyThread();
-	OrthoFarClipPlane = CVarDefaultCameraFarPlane.GetValueOnAnyThread();
+	OrthoWidth = DEFAULT_ORTHOWIDTH;
+	bAutoCalculateOrthoPlanes = true;
+	OrthoNearClipPlane = DEFAULT_ORTHONEARPLANE;
+	OrthoFarClipPlane = DEFAULT_ORTHOFARPLANE;
 	bConstrainAspectRatio = false;
 	bOverrideAspectRatioAxisConstraint = false;
 	bUseFieldOfViewForLOD = true;
@@ -319,6 +300,12 @@ void UCameraComponent::Serialize(FArchive& Ar)
 		OrthoFarClipPlane = UE_OLD_WORLD_MAX;
 	}
 
+	Ar.UsingCustomVersion(FUE5ReleaseStreamObjectVersion::GUID);
+	if (Ar.CustomVer(FUE5ReleaseStreamObjectVersion::GUID) < FUE5ReleaseStreamObjectVersion::OrthographicAutoNearFarPlane)
+	{
+		bAutoCalculateOrthoPlanes = false;
+	}
+
 	Super::Serialize(Ar);
 
 	if (Ar.IsLoading())
@@ -433,6 +420,15 @@ void UCameraComponent::GetCameraView(float DeltaTime, FMinimalViewInfo& DesiredV
 	DesiredView.OrthoNearClipPlane = OrthoNearClipPlane;
 	DesiredView.OrthoFarClipPlane = OrthoFarClipPlane;
 
+	if (bAutoCalculateOrthoPlanes)
+	{
+		DesiredView.bAutoCalculateOrthoPlanes = true;
+		if (const AActor* ViewTarget = GetOwner())
+		{
+			DesiredView.SetOrthoCameraArmLengthFromOwnerLocation(ViewTarget->GetActorLocation());
+		}
+	}
+
 	if (bOverrideAspectRatioAxisConstraint)
 	{
 		DesiredView.AspectRatioAxisConstraint = AspectRatioAxisConstraint;
diff --git a/Engine/Source/Runtime/Engine/Private/Camera/CameraStackTypes.cpp b/Engine/Source/Runtime/Engine/Private/Camera/CameraStackTypes.cpp
--- a/Engine/Source/Runtime/Engine/Private/Camera/CameraStackTypes.cpp
+++ b/Engine/Source/Runtime/Engine/Private/Camera/CameraStackTypes.cpp
@@ -17,10 +17,9 @@ static TAutoConsoleVariable<bool> CVarUseLegacyMaintainYFOV(
 	ECVF_Default
 );
 
-bool GAllowOrthoNearPlaneCorrection = true;
-static FAutoConsoleVariableRef CVarAllowOrthoNearPlaneCorrection(
+static TAutoConsoleVariable<bool> CVarAllowOrthoNearPlaneCorrection(
 	TEXT("r.Ortho.AllowNearPlaneCorrection"),
-	GAllowOrthoNearPlaneCorrection,
+	true,
 	TEXT("Orthographic near planes can be behind the camera position, which causes some issues with Unreal resolving lighting behind the camera position ")
 	TEXT("This CVar enables the Orthographic cameras globally to automatically update the camera location to match the NearPlane location, ")
 	TEXT("and force the pseudo camera position to be the replaced near plane location for the projection matrix calculation. ")
@@ -28,6 +27,85 @@ static FAutoConsoleVariableRef CVarAllowOrthoNearPlaneCorrection(
 	ECVF_Default
 );
 
+static TAutoConsoleVariable<bool> CVarOrthoAllowAutoPlanes(
+	TEXT("r.Ortho.AutoPlanes"),
+	true,
+	TEXT("Globally allow Ortho cameras to utilise the automatic Near/Far plane evaluations."),
+	ECVF_RenderThreadSafe
+);
+
+static TAutoConsoleVariable<int> CVarOrthoClampToMaxFPBuffer(
+	TEXT("r.Ortho.AutoPlanes.ClampToMaxFPBuffer"),
+	1,
+	TEXT("When auto evaluating clip planes, determines whether 16bit depth scaling should be used.")
+	TEXT("16bit scaling is advantageous for any depth downscaling that occurs (e.g. HZB downscaling uses 16 bit textures instead of 32).")
+	TEXT("This feature will calculate the maximum depth scale needed based on the Unreal Unit (cm by default) to Pixel ratio.")
+	TEXT("It assumes that we don't need 32bit depth range for smaller scenes, because most actors will be within a reasonable visible frustum")
+	TEXT("However it does still scale up to a maximum of UE_OLD_WORLD_MAX which is the typical full range of the depth buffer, so larger scenes still work too."),
+	ECVF_RenderThreadSafe
+);
+
+static TAutoConsoleVariable<bool> CVarOrthoScaleIncrementingUnits(
+	TEXT("r.Ortho.AutoPlanes.ScaleIncrementingUnits"),
+	true,
+	TEXT("Select whether to scale the Near/Far plane Min/Max values as we increase in unit to pixel ratio (i.e. as we go from CM to M to KM)"),
+	ECVF_RenderThreadSafe
+);
+
+static TAutoConsoleVariable<float>  CVarOrthoAutoDepthScale(
+	TEXT("r.Ortho.AutoPlanes.DepthScale"),
+	-1.0f,
+	TEXT("Allows the 16 bit depth scaling to be adjusted from the  default +FP16 Max (66504.0f)")
+	TEXT("This is useful if the far plane doesn't need to be as far away, so it will improve depth deltas"),
+	ECVF_RenderThreadSafe
+);
+
+static TAutoConsoleVariable<float>  CVarOrthoAutoPlaneShift(
+	TEXT("r.Ortho.AutoPlanes.ShiftPlanes"),
+	0.0f,
+	TEXT("Shifts the whole frustum in the Z direction.")
+	TEXT("This can be useful if, for example you need the Near plane closer to the camera, at the reduction of the Far plane value (e.g. a horizontal 2.5D scene)."),
+	ECVF_RenderThreadSafe
+);
+
+#if !(UE_BUILD_SHIPPING)
+static TAutoConsoleVariable<bool> CVarDebugForceAllCamerasToOrtho(
+	TEXT("r.Ortho.Debug.ForceAllCamerasToOrtho"),
+	false,
+	TEXT("Debug Force all cameras in the scene to use Orthographic views"),
+	ECVF_RenderThreadSafe
+);
+
+
+static TAutoConsoleVariable<float> CVarDebugForceCameraOrthoWidth(
+	TEXT("r.Ortho.Debug.ForceOrthoWidth"),
+	DEFAULT_ORTHOWIDTH,
+	TEXT("Debug Force Ortho Width when creating a new camera actor"),
+	ECVF_RenderThreadSafe
+);
+
+static TAutoConsoleVariable<bool> CVarDebugForceUseOrthoAutoPlanes(
+	TEXT("r.Ortho.Debug.ForceUseAutoPlanes"),
+	true,
+	TEXT("Debug Force boolean for whether to use the automatic near and far plane evaluation"),
+	ECVF_RenderThreadSafe
+);
+
+static TAutoConsoleVariable<float> CVarDebugForceCameraOrthoNearPlane(
+	TEXT("r.Ortho.Debug.ForceCameraNearPlane"),
+	DEFAULT_ORTHONEARPLANE,
+	TEXT("Debug Force Ortho Near Plane when creating a new camera actor"),
+	ECVF_RenderThreadSafe
+);
+
+static TAutoConsoleVariable<float> CVarDebugForceCameraOrthoFarPlane(
+	TEXT("r.Ortho.Debug.ForceCameraFarPlane"),
+	UE_OLD_WORLD_MAX,
+	TEXT("Debug Force Ortho Far Plane when creating a new camera actor"),
+	ECVF_RenderThreadSafe
+);
+#endif
+
 //////////////////////////////////////////////////////////////////////////
 // FMinimalViewInfo
 
@@ -112,7 +190,7 @@ FMatrix FMinimalViewInfo::CalculateProjectionMatrix(float NegativeNearPlaneOverr
 		const float HalfOrthoWidth = OrthoWidth / 2.0f;
 		const float ScaledOrthoHeight = OrthoWidth / 2.0f * YScale;
 
-		const float NearPlane = GAllowOrthoNearPlaneCorrection && NegativeNearPlaneOverride < 0.0f ? NegativeNearPlaneOverride : OrthoNearClipPlane;
+		const float NearPlane = CVarAllowOrthoNearPlaneCorrection.GetValueOnAnyThread() && NegativeNearPlaneOverride < 0.0f ? NegativeNearPlaneOverride : OrthoNearClipPlane;
 		const float FarPlane = OrthoFarClipPlane;
 
 		const float ZScale = 1.0f / (FarPlane - NearPlane);
@@ -149,11 +227,23 @@ FMatrix FMinimalViewInfo::CalculateProjectionMatrix(float NegativeNearPlaneOverr
 	return ProjectionMatrix;
 }
 
-void FMinimalViewInfo::CalculateProjectionMatrixGivenViewRectangle(const FMinimalViewInfo& ViewInfo, TEnumAsByte<enum EAspectRatioAxisConstraint> AspectRatioAxisConstraint, const FIntRect& ConstrainedViewRectangle, FSceneViewProjectionData& InOutProjectionData)
+void FMinimalViewInfo::CalculateProjectionMatrixGivenViewRectangle(FMinimalViewInfo& ViewInfo, TEnumAsByte<enum EAspectRatioAxisConstraint> AspectRatioAxisConstraint, const FIntRect& ConstrainedViewRectangle, FSceneViewProjectionData& InOutProjectionData)
 {
+#if !(UE_BUILD_SHIPPING)
+	if (CVarDebugForceAllCamerasToOrtho.GetValueOnAnyThread())
+	{
+		ViewInfo.ProjectionMode = ECameraProjectionMode::Orthographic;
+		ViewInfo.OrthoWidth = CVarDebugForceCameraOrthoWidth.GetValueOnAnyThread();
+		ViewInfo.bAutoCalculateOrthoPlanes = CVarDebugForceUseOrthoAutoPlanes.GetValueOnAnyThread();
+		ViewInfo.OrthoNearClipPlane = CVarDebugForceCameraOrthoNearPlane.GetValueOnAnyThread();
+		ViewInfo.OrthoFarClipPlane = CVarDebugForceCameraOrthoFarPlane.GetValueOnAnyThread();
+	}
+#endif
+
 	bool bOrthographicNearPlaneCorrection = ViewInfo.ProjectionMode == ECameraProjectionMode::Orthographic
-											&& GAllowOrthoNearPlaneCorrection
-											&& ViewInfo.OrthoNearClipPlane < 0.0f;
+											&& CVarAllowOrthoNearPlaneCorrection.GetValueOnAnyThread();
+
+	ViewInfo.AutoCalculateOrthoPlanes(InOutProjectionData);	
 
 	// Create the projection matrix (and possibly constrain the view rectangle)
 	if (ViewInfo.bConstrainAspectRatio)
@@ -161,17 +251,12 @@ void FMinimalViewInfo::CalculateProjectionMatrixGivenViewRectangle(const FMinima
 		// Enforce a particular aspect ratio for the render of the scene. 
 		// Results in black bars at top/bottom etc.
 		InOutProjectionData.SetConstrainedViewRectangle(ConstrainedViewRectangle);
+		InOutProjectionData.ProjectionMatrix = ViewInfo.CalculateProjectionMatrix();
 		if (bOrthographicNearPlaneCorrection)
 		{
 			float NearPlane = ViewInfo.OrthoNearClipPlane;
-			InOutProjectionData.UpdateOrthoNearPlane(NearPlane);
-			InOutProjectionData.ProjectionMatrix = ViewInfo.CalculateProjectionMatrix(NearPlane);
-		}
-		else
-		{
-			InOutProjectionData.ProjectionMatrix = ViewInfo.CalculateProjectionMatrix();
+			InOutProjectionData.UpdateOrthoNearPlane(NearPlane, true);
 		}
-		
 	}
 	else
 	{
@@ -187,9 +272,8 @@ void FMinimalViewInfo::CalculateProjectionMatrixGivenViewRectangle(const FMinima
 
 		// If x is bigger, and we're respecting x or major axis, AND mobile isn't forcing us to be Y axis aligned
 		const bool bMaintainXFOV = 
-			((SizeX > SizeY) && (AspectRatioAxisConstraint == AspectRatio_MajorAxisFOV)) || 
-			(AspectRatioAxisConstraint == AspectRatio_MaintainXFOV) || 
-			(ViewInfo.ProjectionMode == ECameraProjectionMode::Orthographic);
+			((SizeX > SizeY) && ((AspectRatioAxisConstraint == AspectRatio_MajorAxisFOV) || (ViewInfo.ProjectionMode == ECameraProjectionMode::Orthographic))) ||
+			(AspectRatioAxisConstraint == AspectRatio_MaintainXFOV);
 		if (bMaintainXFOV)
 		{
 			// If the viewport is wider than it is tall
@@ -267,9 +351,96 @@ void FMinimalViewInfo::CalculateProjectionMatrixGivenViewRectangle(const FMinima
 	}
 }
 
-void FMinimalViewInfo::CalculateProjectionMatrixGivenView(const FMinimalViewInfo& ViewInfo, TEnumAsByte<enum EAspectRatioAxisConstraint> AspectRatioAxisConstraint, FViewport* Viewport, FSceneViewProjectionData& InOutProjectionData)
+void FMinimalViewInfo::CalculateProjectionMatrixGivenView(FMinimalViewInfo& ViewInfo, TEnumAsByte<enum EAspectRatioAxisConstraint> AspectRatioAxisConstraint, FViewport* Viewport, FSceneViewProjectionData& InOutProjectionData)
 {
 	FIntRect ViewExtents = Viewport->CalculateViewExtents(ViewInfo.AspectRatio, InOutProjectionData.GetViewRect());
 	CalculateProjectionMatrixGivenViewRectangle(ViewInfo, AspectRatioAxisConstraint, ViewExtents, InOutProjectionData);
 }
 
+bool FMinimalViewInfo::AutoCalculateOrthoPlanes(const FSceneViewProjectionData& InOutProjectionData)
+{	
+	if (ProjectionMode == ECameraProjectionMode::Orthographic && CVarOrthoAllowAutoPlanes.GetValueOnAnyThread() && bAutoCalculateOrthoPlanes)
+	{
+		//First check if we are using 16bit buffer and unit scaling, then set the min/max values accordingly
+		const bool bUse16bitDepth = CVarOrthoClampToMaxFPBuffer.GetValueOnAnyThread() == 1;
+		const bool bScaleIncrementingUnits = CVarOrthoScaleIncrementingUnits.GetValueOnAnyThread() && bUse16bitDepth;
+		const float MaxFPValue = bScaleIncrementingUnits ? UE_LARGE_WORLD_MAX : UE_OLD_WORLD_MAX;
+		float FPScale = bUse16bitDepth ? 66504.0f : UE_OLD_WORLD_MAX;
+
+		const float AutoDepthScale = CVarOrthoAutoDepthScale.GetValueOnAnyThread();
+		if (AutoDepthScale > 0.0f)
+		{
+			//This allows the user to override the FP scaling value, where the default is 16bit.
+			FPScale = FMath::Clamp(AutoDepthScale, 1.0f, FPScale);
+		}
+
+		//Get the OrthoHeight, with Ortho the depth is typically bound to the Y axis so we use that
+		const float OrthoHeight = OrthoWidth / (AspectRatio == 0.0f ? UE_DELTA : AspectRatio);
+
+		//Get the normalized view forward vector of the camera
+		const FRotationMatrix RotMat(Rotation);
+		FVector ViewForward = RotMat.GetColumn(2);
+		ViewForward.Normalize();	
+
+		/**
+		 * The CosAngle is the cosine of the angle between the ViewForward and camera down.
+		 * Forcing the absolute value for this means that Up/Down is 1.0f and Forward (90 degrees) is 0.
+		 * We use this to scale the Near Plane, and the far plane if 16 bit scaling is disabled.
+		 */
+		float CosAngle = FMath::Abs((ViewForward.Dot(FVector(0, 0, -1.0f))));
+		
+		/** 
+		 * We still max out at UE_OLD_WORLD_MAX or Max32FP, but we scale the FarPlane depending on ratio of the pixel size to the world unit size.
+		 * Details below, but the reasoning is, we can't visibly see smaller than a pixel, 
+		 * so the passes that need 16bit buffers such as HZB have their plane distances scaled automatically depending on this ratio.
+		 */		
+		float FarPlane;
+		float UnitPerPixelRatio = 1.0f;
+		float ViewportSizeX = InOutProjectionData.GetViewRect().Width();
+		if (bUse16bitDepth && OrthoWidth > 0  && ViewportSizeX > 0)
+		{
+			//The CmPerPixelRatio determines the far plane depth scale required for the scene
+			UnitPerPixelRatio = OrthoWidth / (float)ViewportSizeX;
+			if (bScaleIncrementingUnits)
+			{	
+				//This scales the min/max depending on the dynamic scale of the unit to pixel as the ortho width increases at the sacrifice of the max FarPlane, allowing scaling to LWC
+				UnitPerPixelRatio = FMath::Log2(UnitPerPixelRatio);
+			}
+			
+			FarPlane = FMath::Clamp(FPScale * UnitPerPixelRatio, FPScale, MaxFPValue);
+		}
+		else
+		{
+			/** 
+			 * Default path if the 16bit scaling depth is disabled or not usable.
+			 * Note: this path does not scale for 16 bit buffers, it only calculates Near/Far plane min/max automatically
+			 */			
+			FarPlane = OrthoHeight / (CosAngle == 0.0f ? UE_DELTA : CosAngle);
+		}
+		/** 
+		* The camera arm length is adjusted depending on the CosAngle as the horizontal view typically has a significantly larger plane range, 
+		* so it becomes irrelevant, whereas it is necessary to account for in a top down view. Note: a small scene camera arm length will become irrelevant for a large ortho width.
+		*/
+		float CameraArmLength = OrthoCameraArmLength * CosAngle;
+
+		/**
+		 * The NearPlane calculation is a scaled OrthoHeight depending on the camera angle, 
+		 * which maxes out at 45 degrees by default as this captures the entire scene for the majority of angles.
+		 * r.Ortho.AutoPlanes.ShiftPlanes should be used to account for views outside of this.
+		 * 
+		 * The FarPlane is the required depth precision interpretation for the UnitPerPixelRatio. 
+		 * We clamp this to remove the Near plane difference, and also max out at the previously set maximum FPValue. 
+		 * This setup should help for possible future implementations where we can increase the depth range (i.e. LWC + double float depth buffers).
+		 */
+		float SinAngle = FMath::Clamp(1.0f - CosAngle, 0.707f, 1.0f);
+		float NearPlane = OrthoHeight * (FMath::Clamp(CosAngle, 0.707f, 1.0f) - (1.0f/SinAngle)) - CameraArmLength;
+		FarPlane = FMath::Clamp(FarPlane, OrthoHeight, MaxFPValue + NearPlane);
+
+		//The Planes can be scaled in the Z axis without restriction to ensure a user can capture their entire view.
+		const float AutoPlaneShift = CVarOrthoAutoPlaneShift.GetValueOnAnyThread();
+		OrthoNearClipPlane = NearPlane + AutoPlaneShift;
+		OrthoFarClipPlane = FarPlane + AutoPlaneShift;
+		return true;
+	}
+	return false;
+}
diff --git a/Engine/Source/Runtime/Engine/Private/Components/SceneCaptureComponent.cpp b/Engine/Source/Runtime/Engine/Private/Components/SceneCaptureComponent.cpp
--- a/Engine/Source/Runtime/Engine/Private/Components/SceneCaptureComponent.cpp
+++ b/Engine/Source/Runtime/Engine/Private/Components/SceneCaptureComponent.cpp
@@ -8,6 +8,7 @@
 #include "Camera/CameraTypes.h"
 #include "UObject/EditorObjectVersion.h"
 #include "UObject/FortniteMainBranchObjectVersion.h"
+#include "UObject/UE5ReleaseStreamObjectVersion.h"
 #include "SceneInterface.h"
 #include "UObject/ConstructorHelpers.h"
 #include "Components/StaticMeshComponent.h"
@@ -566,8 +567,8 @@ USceneCaptureComponent2D::USceneCaptureComponent2D(const FObjectInitializer& Obj
 {
 	FOVAngle = 90.0f;
 
-	const auto CVar = IConsoleManager::Get().FindConsoleVariable(TEXT("r.Ortho.DefaultCameraWidth"));
-	OrthoWidth = CVar ? CVar->GetFloat() : DEFAULT_ORTHOWIDTH;
+	OrthoWidth = DEFAULT_ORTHOWIDTH;
+	bAutoCalculateOrthoPlanes = true;
 
 	bUseCustomProjectionMatrix = false;
 	bAutoActivate = true;
@@ -700,6 +701,15 @@ void USceneCaptureComponent2D::GetCameraView(float DeltaTime, FMinimalViewInfo&
 	OutMinimalViewInfo.bConstrainAspectRatio = false;
 	OutMinimalViewInfo.ProjectionMode = ProjectionType;
 	OutMinimalViewInfo.OrthoWidth = OrthoWidth;
+		
+	if (bAutoCalculateOrthoPlanes)
+	{
+		OutMinimalViewInfo.bAutoCalculateOrthoPlanes = true;
+		if(const AActor* ViewTarget = GetOwner())
+		{
+			OutMinimalViewInfo.SetOrthoCameraArmLengthFromOwnerLocation(ViewTarget->GetActorLocation());
+		}
+	}
 }
 
 void USceneCaptureComponent2D::CaptureSceneDeferred()
@@ -846,6 +856,12 @@ void USceneCaptureComponent2D::Serialize(FArchive& Ar)
 		OrthoWidth = 512.0f;
 	}
 
+	Ar.UsingCustomVersion(FUE5ReleaseStreamObjectVersion::GUID);
+	if (Ar.CustomVer(FUE5ReleaseStreamObjectVersion::GUID) < FUE5ReleaseStreamObjectVersion::OrthographicAutoNearFarPlane)
+	{
+		bAutoCalculateOrthoPlanes = false;
+	}
+
 	Super::Serialize(Ar);
 
 	if (Ar.IsLoading())
diff --git a/Engine/Source/Runtime/Engine/Private/LocalPlayer.cpp b/Engine/Source/Runtime/Engine/Private/LocalPlayer.cpp
--- a/Engine/Source/Runtime/Engine/Private/LocalPlayer.cpp
+++ b/Engine/Source/Runtime/Engine/Private/LocalPlayer.cpp
@@ -687,6 +687,14 @@ void ULocalPlayer::GetViewPoint(FMinimalViewInfo& OutViewInfo) const
 			OutViewInfo = PlayerController->PlayerCameraManager->GetCameraCacheView();
 			OutViewInfo.FOV = PlayerController->PlayerCameraManager->GetFOVAngle();
 			PlayerController->GetPlayerViewPoint(/*out*/ OutViewInfo.Location, /*out*/ OutViewInfo.Rotation);
+			if (PlayerController->PlayerCameraManager->bAutoCalculateOrthoPlanes)
+			{
+				OutViewInfo.bAutoCalculateOrthoPlanes = true;
+				if (const APawn* Owner = PlayerController->GetPawn())
+				{
+					OutViewInfo.SetOrthoCameraArmLengthFromOwnerLocation(Owner->GetActorLocation());
+				}
+			}
 		}
 		else
 		{
diff --git a/Engine/Source/Runtime/Engine/Private/PlayerCameraManager.cpp b/Engine/Source/Runtime/Engine/Private/PlayerCameraManager.cpp
--- a/Engine/Source/Runtime/Engine/Private/PlayerCameraManager.cpp
+++ b/Engine/Source/Runtime/Engine/Private/PlayerCameraManager.cpp
@@ -39,6 +39,7 @@ APlayerCameraManager::APlayerCameraManager(const FObjectInitializer& ObjectIniti
 	DefaultAspectRatio = 1.33333f;
 	bDefaultConstrainAspectRatio = false;
 	DefaultOrthoWidth = 512.0f;
+	bAutoCalculateOrthoPlanes = true;
 	SetHidden(true);
 	bReplicates = false;
 	FreeCamDistance = 256.0f;
@@ -364,6 +365,15 @@ void APlayerCameraManager::UpdateViewTarget(FTViewTarget& OutVT, float DeltaTime
 	OutVT.POV.ProjectionMode = bIsOrthographic ? ECameraProjectionMode::Orthographic : ECameraProjectionMode::Perspective;
 	OutVT.POV.PostProcessBlendWeight = 1.0f;
 
+	if (bAutoCalculateOrthoPlanes)
+	{
+		OutVT.POV.bAutoCalculateOrthoPlanes = true;
+		if (OutVT.Target)
+		{
+			OutVT.POV.SetOrthoCameraArmLengthFromOwnerLocation(OutVT.Target->GetActorLocation());
+		}
+	}
+
 	bool bDoNotApplyModifiers = false;
 
 	if (ACameraActor* CamActor = Cast<ACameraActor>(OutVT.Target))
diff --git a/Engine/Source/Runtime/Engine/Public/EngineDefines.h b/Engine/Source/Runtime/Engine/Public/EngineDefines.h
--- a/Engine/Source/Runtime/Engine/Public/EngineDefines.h
+++ b/Engine/Source/Runtime/Engine/Public/EngineDefines.h
@@ -56,6 +56,7 @@
 #define DEFAULT_ORTHOZOOM			10000.0						/* Default 2D viewport zoom */
 #define DEFAULT_ORTHOWIDTH			1536.0f						/* Default 2D ortho viewport width */
 #define DEFAULT_ORTHONEARPLANE		-DEFAULT_ORTHOWIDTH/2.0f	/* Default 2D ortho viewport nearplane */
+#define DEFAULT_ORTHOFARPLANE		UE_OLD_WORLD_MAX			/* Default 2D ortho viewport farplane */
 
 /** bits needed to store DPG value */
 #define SDPG_NumBits 3
