Fix cast warning when calculating camera arm length for Ortho Near/Far plane correction.

diff --git a/Engine/Source/Runtime/Engine/Classes/Camera/CameraTypes.h b/Engine/Source/Runtime/Engine/Classes/Camera/CameraTypes.h
--- a/Engine/Source/Runtime/Engine/Classes/Camera/CameraTypes.h
+++ b/Engine/Source/Runtime/Engine/Classes/Camera/CameraTypes.h
@@ -166,7 +166,7 @@ public:
 	/** Sets the camera distance from view target for AutoCalculateOrthoPlanes */
 	ENGINE_API inline void SetOrthoCameraArmLengthFromOwnerLocation(const FVector ActorLocation)
 	{
-		OrthoCameraArmLength = (ActorLocation - Location).Length();
+		OrthoCameraArmLength = static_cast<float>((ActorLocation - Location).Length());
 	}
 };
 
