Deprecated some ISM public functions & data that will be removed (part 2/2).

diff --git a/Engine/Source/Runtime/Engine/Classes/Components/InstancedStaticMeshComponent.h b/Engine/Source/Runtime/Engine/Classes/Components/InstancedStaticMeshComponent.h
--- a/Engine/Source/Runtime/Engine/Classes/Components/InstancedStaticMeshComponent.h
+++ b/Engine/Source/Runtime/Engine/Classes/Components/InstancedStaticMeshComponent.h
@@ -317,6 +317,7 @@ class UInstancedStaticMeshComponent : public UStaticMeshComponent, public ISMIns
 	* @param UpdateInstancePreviousTransforms	The transforms of the new instances to update.
 	* @return									True on success
 	*/
+	UE_DEPRECATED(5.4, "Use the new ID-based APIs instead as this enables persistence tracking for incremental updates.")
 	ENGINE_API virtual bool UpdateInstances(
 		const TArray<int32>& UpdateInstanceIds, 
 		const TArray<FTransform>& UpdateInstanceTransforms, 
@@ -566,12 +567,14 @@ public:
 	UE_DEPRECATED(5.4, "This function has been added only for the purposes of moving LWI code outside of the engine. Don't use it, it will be removed soon.")
 	ENGINE_API void OnPostPopulatePerInstanceData() { OnPostLoadPerInstanceData(); }
 
+	UE_DEPRECATED(5.4, "Use the new ID-based APIs instead as this enables persistence tracking for incremental updates.")
 	TArray<int32> PerInstanceIds;
 
 	/** Used to cache a unique identifier for each instance.  These are provided
 	*	by the interface UpdateInstances.  This is a map from unique id to index
 	*	into the PerInstanceSMData array.
 	*/
+	UE_DEPRECATED(5.4, "Use the new ID-based APIs instead as this enables persistence tracking for incremental updates.")
 	TMap<int32, int32> InstanceIdToInstanceIndexMap;
 
 	/** Request to navigation system to update only part of navmesh occupied by specified instance. */
diff --git a/Engine/Source/Runtime/Engine/Private/InstancedStaticMesh.cpp b/Engine/Source/Runtime/Engine/Private/InstancedStaticMesh.cpp
--- a/Engine/Source/Runtime/Engine/Private/InstancedStaticMesh.cpp
+++ b/Engine/Source/Runtime/Engine/Private/InstancedStaticMesh.cpp
@@ -2032,9 +2032,11 @@ UInstancedStaticMeshComponent::UInstancedStaticMeshComponent(FVTableHelper& Help
 {
 }
 
+PRAGMA_DISABLE_DEPRECATION_WARNINGS
 UInstancedStaticMeshComponent::~UInstancedStaticMeshComponent()
 {
 }
+PRAGMA_ENABLE_DEPRECATION_WARNINGS
 
 TStructOnScope<FActorComponentInstanceData> UInstancedStaticMeshComponent::GetComponentInstanceData() const
 {
@@ -3861,6 +3863,7 @@ bool UInstancedStaticMeshComponent::UpdateInstances(
 	int32 InNumCustomDataFloats,
 	const TArray<float>& CustomFloatData)
 {
+PRAGMA_DISABLE_DEPRECATION_WARNINGS
 	CSV_SCOPED_TIMING_STAT_EXCLUSIVE(UInstancedStaticMeshComponent_UpdateInstances);
 	TRACE_CPUPROFILER_EVENT_SCOPE_STR("UInstancedStaticMeshComponent::UpdateInstances");
 
@@ -4046,6 +4049,7 @@ bool UInstancedStaticMeshComponent::UpdateInstances(
 	FCsvProfiler::Get()->RecordCustomStat(NumBytesUpdated, CSV_CATEGORY_INDEX(InstancedStaticMeshComponent), TotalSizeUpdateBytes, ECsvCustomStatOp::Accumulate);
 #endif
 
+PRAGMA_ENABLE_DEPRECATION_WARNINGS
 	return true;
 }
 
