[DisplayCVarList] - Draw cvars yellow if their value differs from the default.

todo: Reduce the use of string compares if it becomes an issue.

diff --git a/Engine/Source/Runtime/Engine/Private/UnrealEngine.cpp b/Engine/Source/Runtime/Engine/Private/UnrealEngine.cpp
--- a/Engine/Source/Runtime/Engine/Private/UnrealEngine.cpp
+++ b/Engine/Source/Runtime/Engine/Private/UnrealEngine.cpp
@@ -12644,7 +12644,23 @@ void DrawStatsHUD( UWorld* World, FViewport* Viewport, FCanvas* Canvas, UCanvas*
 			{
 				for (const auto& CVarToVisualize : GDisplayCVarListExecHelper.ConsoleVariablesToVisualize)
 				{
-					Canvas->DrawShadowedString(MessageX, MessageY, *FString::Printf(TEXT("%s : %s"), *CVarToVisualize.Key, *CVarToVisualize.Value->GetString()), GEngine->GetSmallFont(), FLinearColor::White);
+					const FString CurrentValue = CVarToVisualize.Value->GetString();
+					const FString DefaultValue = CVarToVisualize.Value->GetDefaultValue();
+					bool bSame = true;
+
+					// Floats sometimes return true erroneously because they can be stringified as e.g '1' or '1.0' by different functions.
+					if (CVarToVisualize.Value->IsVariableFloat())
+					{
+						const float A = CVarToVisualize.Value->GetFloat();
+						const float B = FCString::Atof(*DefaultValue);
+						bSame = FMath::IsNearlyEqual(A, B);
+					}
+					else
+					{
+						bSame = CurrentValue.Equals(DefaultValue);
+					}
+
+					Canvas->DrawShadowedString(MessageX, MessageY, *FString::Printf(TEXT("%s : %s"), *CVarToVisualize.Key, *CurrentValue), GEngine->GetSmallFont(), bSame ? FLinearColor::White : FLinearColor::Yellow);
  					MessageY += FontSizeY;
 				}
 			}
