Added tool tip details about lighting channels and translucent materials.

diff --git a/Engine/Source/Runtime/Engine/Classes/Components/LightComponent.h b/Engine/Source/Runtime/Engine/Classes/Components/LightComponent.h
--- a/Engine/Source/Runtime/Engine/Classes/Components/LightComponent.h
+++ b/Engine/Source/Runtime/Engine/Classes/Components/LightComponent.h
@@ -150,6 +150,7 @@ class ULightComponent : public ULightComponentBase
 	/** 
 	 * Channels that this light should affect.  
 	 * These channels only apply to opaque materials, direct lighting, and dynamic lighting and shadowing.
+	 * Lighting channels are only supported on translucent materials using forward shading (i.e. when not using the translucency lighting volume).
 	 */
 	UPROPERTY(EditAnywhere, AdvancedDisplay, BlueprintReadOnly, Category=Light)
 	FLightingChannels LightingChannels;
diff --git a/Engine/Source/Runtime/Engine/Classes/Components/PrimitiveComponent.h b/Engine/Source/Runtime/Engine/Classes/Components/PrimitiveComponent.h
--- a/Engine/Source/Runtime/Engine/Classes/Components/PrimitiveComponent.h
+++ b/Engine/Source/Runtime/Engine/Classes/Components/PrimitiveComponent.h
@@ -734,6 +734,7 @@ public:
 	/** 
 	 * Channels that this component should be in.  Lights with matching channels will affect the component.  
 	 * These channels only apply to opaque materials, direct lighting, and dynamic lighting and shadowing.
+	 * Lighting channels are only supported on translucent materials using forward shading (i.e. when not using the translucency lighting volume).
 	 */
 	UPROPERTY(EditAnywhere, AdvancedDisplay, BlueprintReadOnly, Category=Lighting)
 	FLightingChannels LightingChannels;
