Avoid failing compression if we can't resample due to a low number of keys

diff --git a/Engine/Source/Runtime/Engine/Private/Animation/AnimCompress_PerTrackCompression.cpp b/Engine/Source/Runtime/Engine/Private/Animation/AnimCompress_PerTrackCompression.cpp
--- a/Engine/Source/Runtime/Engine/Private/Animation/AnimCompress_PerTrackCompression.cpp
+++ b/Engine/Source/Runtime/Engine/Private/Animation/AnimCompress_PerTrackCompression.cpp
@@ -1568,11 +1568,6 @@ bool UAnimCompress_PerTrackCompression::DoReduction(const FCompressibleAnimData&
 		return false;
 	}
 
-	if (bResampleAnimation && CompressibleAnimData.NumberOfKeys < MinKeysForResampling)
-	{
-		return false;
-	}
-
 	ensure((MaxPosDiffBitwise > 0.0f) && (MaxAngleDiffBitwise > 0.0f) && (MaxScaleDiffBitwise > 0.0f) && (MaxZeroingThreshold >= 0.0f));
 	ensure(MaxZeroingThreshold <= MaxPosDiffBitwise);
 	ensure(!(bUseAdaptiveError2 && bUseAdaptiveError));
