Exporting some additional functions.

diff --git a/Engine/Source/Runtime/Engine/Public/Rendering/SkeletalMeshRenderData.h b/Engine/Source/Runtime/Engine/Public/Rendering/SkeletalMeshRenderData.h
--- a/Engine/Source/Runtime/Engine/Public/Rendering/SkeletalMeshRenderData.h
+++ b/Engine/Source/Runtime/Engine/Public/Rendering/SkeletalMeshRenderData.h
@@ -74,7 +74,7 @@ public:
 	ENGINE_API void GetResourceSizeEx(FResourceSizeEx& CumulativeResourceSize);
 
 	/** Get the estimated memory overhead of buffers marked as NeedsCPUAccess. */
-	SIZE_T GetCPUAccessMemoryOverhead() const;
+	ENGINE_API SIZE_T GetCPUAccessMemoryOverhead() const;
 
 	/** Returns true if this resource must be skinned on the CPU for the given feature level. */
 	ENGINE_API bool RequiresCPUSkinning(ERHIFeatureLevel::Type FeatureLevel) const;
diff --git a/Engine/Source/Runtime/Engine/Public/StaticMeshResources.h b/Engine/Source/Runtime/Engine/Public/StaticMeshResources.h
--- a/Engine/Source/Runtime/Engine/Public/StaticMeshResources.h
+++ b/Engine/Source/Runtime/Engine/Public/StaticMeshResources.h
@@ -522,7 +522,7 @@ public:
 	/** Serialize. */
 	void Serialize(FArchive& Ar, UObject* Owner, int32 Idx);
 
-	void GetResourceSizeEx(FResourceSizeEx& CumulativeResourceSize) const;
+	ENGINE_API void GetResourceSizeEx(FResourceSizeEx& CumulativeResourceSize) const;
 
 #if RHI_RAYTRACING
 	void SetupRayTracingGeometryInitializer(FRayTracingGeometryInitializer& Initializer, const FName& DebugName, const FName& OwnerName) const;
@@ -780,7 +780,7 @@ public:
 	ENGINE_API void GetResourceSizeEx(FResourceSizeEx& CumulativeResourceSize) const;
 
 	/** Get the estimated memory overhead of buffers marked as NeedsCPUAccess. */
-	SIZE_T GetCPUAccessMemoryOverhead() const;
+	ENGINE_API SIZE_T GetCPUAccessMemoryOverhead() const;
 
 	/** Allocate LOD resources. */
 	ENGINE_API void AllocateLODResources(int32 NumLODs);
