nop Texture.cpp comments only ; note about why rendertarget textures go black, not fixing for now.

see also CL 21456593

diff --git a/Engine/Source/Runtime/Engine/Private/Texture.cpp b/Engine/Source/Runtime/Engine/Private/Texture.cpp
--- a/Engine/Source/Runtime/Engine/Private/Texture.cpp
+++ b/Engine/Source/Runtime/Engine/Private/Texture.cpp
@@ -912,6 +912,14 @@ PRAGMA_ENABLE_DEPRECATION_WARNINGS
 	{
 		// Update the texture resource. This will recache derived data if necessary
 		// which may involve recompressing the texture.
+
+		// Note for RenderTarget :
+		// if PIE is running, this will cause the RenderTarget to refresh
+		// if PIE is not running, this will change the RenderTarget to black
+		// in some cases you must always do this even if PIE is not running (eg. if changing size or format)
+		// but in other cases you could skip this UpdateResource to leave the existing rendertarget contents valid
+		// -> not attempting to do that for now
+
 		UpdateResource();
 	}
 
