Fix skylightcomponent not getting mark as dirty when we change quality settings (since CL 30477845). DestroyRenderState is now getting call before the cvar OnChanged callback, making SkylightComponent->IsRenderStateCreated() returning false in OnChangeSkylightRealTimeReflectionCapture. This result in the skylight not beeing put in the SkyCapturesToUpdate array and never gets updated properly. To fix, just remove the IsRenderStateCreated() in the callback, since we check if the component is visible inside USkyLightComponent::SetCaptureIsDirty

#tests Editor

diff --git a/Engine/Source/Runtime/Engine/Private/Components/SkyLightComponent.cpp b/Engine/Source/Runtime/Engine/Private/Components/SkyLightComponent.cpp
--- a/Engine/Source/Runtime/Engine/Private/Components/SkyLightComponent.cpp
+++ b/Engine/Source/Runtime/Engine/Private/Components/SkyLightComponent.cpp
@@ -92,7 +92,7 @@ void OnChangeSkylightRealTimeReflectionCapture(IConsoleVariable* Var)
 	for (TObjectIterator<USkyLightComponent> It; It; ++It)
 	{
 		USkyLightComponent* SkylightComponent = *It;
-		if (IsValid(SkylightComponent) && SkylightComponent->IsRenderStateCreated())
+		if (IsValid(SkylightComponent))
 		{
 			SkylightComponent->SetCaptureIsDirty();
 		}
