Restore changes to GameSession MaxPlayers and MaxSpectators when config is reloaded.

diff --git a/Engine/Source/Runtime/Engine/Classes/GameFramework/GameSession.h b/Engine/Source/Runtime/Engine/Classes/GameFramework/GameSession.h
--- a/Engine/Source/Runtime/Engine/Classes/GameFramework/GameSession.h
+++ b/Engine/Source/Runtime/Engine/Classes/GameFramework/GameSession.h
@@ -271,6 +271,18 @@ protected:
 	 */
 	ENGINE_API virtual void OnEndSessionComplete(FName InSessionName, bool bWasSuccessful);
 
+	/**
+	 * PostReloadConfig override to reapply config property overrides.
+	 */
+	ENGINE_API virtual void PostReloadConfig(FProperty* PropertyThatWasLoaded) override;
+
+private:
+
+	/** Override for the default value of MaxPlayers passed to InitOptions. */
+	TOptional<int32> MaxPlayersOptionOverride;
+
+	/** Override for the default value of MaxSpectators passed to InitOptions. */
+	TOptional<int32> MaxSpectatorsOptionOverride;
 };
 
 /** 
diff --git a/Engine/Source/Runtime/Engine/Private/GameSession.cpp b/Engine/Source/Runtime/Engine/Private/GameSession.cpp
--- a/Engine/Source/Runtime/Engine/Private/GameSession.cpp
+++ b/Engine/Source/Runtime/Engine/Private/GameSession.cpp
@@ -124,6 +124,24 @@ void AGameSession::OnEndSessionComplete(FName InSessionName, bool bWasSuccessful
 	UE_LOG(LogGameSession, Verbose, TEXT("OnEndSessionComplete %s bSuccess: %d"), *InSessionName.ToString(), bWasSuccessful);
 }
 
+void AGameSession::PostReloadConfig(FProperty* PropertyThatWasLoaded)
+{
+	Super::PostReloadConfig(PropertyThatWasLoaded);
+
+	if (!IsTemplate())
+	{
+		if (MaxPlayersOptionOverride.IsSet())
+		{
+			MaxPlayers = *MaxPlayersOptionOverride;
+		}
+
+		if (MaxSpectatorsOptionOverride.IsSet())
+		{
+			MaxSpectators = *MaxSpectatorsOptionOverride;
+		}
+	}
+}
+
 bool AGameSession::HandleStartMatchRequest()
 {
 	return false;
@@ -135,8 +153,17 @@ void AGameSession::InitOptions( const FString& Options )
 	check(World);
 	AGameModeBase* const GameMode = World ? World->GetAuthGameMode() : nullptr;
 
-	MaxPlayers = UGameplayStatics::GetIntOption( Options, TEXT("MaxPlayers"), MaxPlayers );
-	MaxSpectators = UGameplayStatics::GetIntOption( Options, TEXT("MaxSpectators"), MaxSpectators );
+	if (UGameplayStatics::HasOption(Options, TEXT("MaxPlayers")))
+	{
+		MaxPlayers = UGameplayStatics::GetIntOption(Options, TEXT("MaxPlayers"), MaxPlayers);
+		MaxPlayersOptionOverride = MaxPlayers;
+	}
+
+	if (UGameplayStatics::HasOption(Options, TEXT("MaxSpectators")))
+	{
+		MaxSpectators = UGameplayStatics::GetIntOption(Options, TEXT("MaxSpectators"), MaxSpectators);
+		MaxSpectatorsOptionOverride = MaxSpectators;
+	}
 	
 	if (GameMode)
 	{
