RHI validation fixes and related cleanup for SVT:
-Fixed validation error where a (transient) dummy UAV that is created and used during voxel data upload wasn't properly acquired according to validation. Fixed by specializing the upload shader depending on the texture configuration of the SVT and removing the dummy UAV.
-Fixed missing transitions back to SRV after uploading data to the page table and/or tile data texture.
-Fixed validation errors caused by not setting a 3D texture as fallback to the SVT texture references when the SVT doesn't have the full set of textures.
-Since they regularly need to be imported into RDG, SVT tile data texture and page table texture are now IPooledRenderTargets instead of raw RHI textures.
-Replaced (implicit) usages of the immediate command list with the current command list in the scope when creating RHI resources

#rnx

diff --git a/Engine/Shaders/Private/SparseVolumeTexture/UpdateSparseVolumeTexture.usf b/Engine/Shaders/Private/SparseVolumeTexture/UpdateSparseVolumeTexture.usf
--- a/Engine/Shaders/Private/SparseVolumeTexture/UpdateSparseVolumeTexture.usf
+++ b/Engine/Shaders/Private/SparseVolumeTexture/UpdateSparseVolumeTexture.usf
@@ -23,56 +23,25 @@ uint3 UnflattenCoord(uint Index, uint3 VolumeSize)
 	return Result;
 }
 
-#ifdef UPDATE_TILE_TEXTURE_FROM_BUFFER
-
-RWTexture3D<float4> DstPhysicalTileTextureA;
-RWTexture3D<float4> DstPhysicalTileTextureB;
-ByteAddressBuffer DstTileCoords;
-Buffer<float4> SrcPhysicalTileBufferA;
-Buffer<float4> SrcPhysicalTileBufferB;
-uint TileCoordsBufferOffset;
-uint TileDataBufferOffsetInTiles;
-uint NumTilesToCopy;
-uint NumDispatchedGroups;
-uint PaddedTileSize;
-uint bCopyTexureAOnlyUI;
-
-[numthreads(64, 1, 1)]
-void SparseVolumeTextureUpdateFromBufferCS(uint3 GroupId : SV_GroupID, uint GroupThreadIndex : SV_GroupIndex)
-{
-	const uint GroupIdx = GroupId.x;
-	const bool bCopyTextureAOnly = (bCopyTexureAOnlyUI != 0);
-	const uint NumVoxelsPerPaddedTile = PaddedTileSize * PaddedTileSize * PaddedTileSize;
-
-	// Process one tile per workgroup
-	for (uint TileIdx = GroupIdx; TileIdx < NumTilesToCopy; TileIdx += NumDispatchedGroups)
-	{
-		const uint PackedDstCoord = DstTileCoords.Load((TileIdx + TileCoordsBufferOffset) << 2u);
-		const uint3 DstBaseCoord = uint3(PackedDstCoord & 0xFFu, (PackedDstCoord >> 8u) & 0xFFu, (PackedDstCoord >> 16u) & 0xFFu) * PaddedTileSize;
-		const uint TileDataBufferBaseIndex = (TileIdx + TileDataBufferOffsetInTiles) * NumVoxelsPerPaddedTile;
-
-		// Iterate over all texels in the tile and copy them
-		for (uint VoxelIdx = GroupThreadIndex; VoxelIdx < NumVoxelsPerPaddedTile; VoxelIdx += 64)
-		{
-			const uint LoadIdx = TileDataBufferBaseIndex + VoxelIdx;
-			const uint3 DstCoord = DstBaseCoord + UnflattenCoord(VoxelIdx, PaddedTileSize);
-			DstPhysicalTileTextureA[DstCoord] = SrcPhysicalTileBufferA.Load(LoadIdx);
-			if (!bCopyTextureAOnly)
-			{
-				DstPhysicalTileTextureB[DstCoord] = SrcPhysicalTileBufferB.Load(LoadIdx);
-			}
-		}
-	}
-}
+#ifdef UPDATE_TILE_TEXTURE_FROM_SPARSE_BUFFER
 
-#endif // UPDATE_TILE_TEXTURE_FROM_BUFFER
+#ifndef TEXTURE_UPDATE_MASK
+#define TEXTURE_UPDATE_MASK 0
+#endif // TEXTURE_UPDATE_MASK
 
-#ifdef UPDATE_TILE_TEXTURE_FROM_SPARSE_BUFFER
+#define UPDATE_TEXTURE_A ((TEXTURE_UPDATE_MASK & 1) != 0)
+#define UPDATE_TEXTURE_B ((TEXTURE_UPDATE_MASK & 2) != 0)
 
+#if UPDATE_TEXTURE_A
 RWTexture3D<float4> DstPhysicalTileTextureA;
-RWTexture3D<float4> DstPhysicalTileTextureB;
 Buffer<float4> SrcPhysicalTileBufferA;
+#endif // UPDATE_TEXTURE_A
+
+#if UPDATE_TEXTURE_B
+RWTexture3D<float4> DstPhysicalTileTextureB;
 Buffer<float4> SrcPhysicalTileBufferB;
+#endif // UPDATE_TEXTURE_B
+
 ByteAddressBuffer OccupancyBitsBuffer;
 ByteAddressBuffer TileDataOffsetsBuffer;
 ByteAddressBuffer DstTileCoordsBuffer;
@@ -85,7 +54,6 @@ uint NumTilesToCopy;
 uint BufferTileStep;
 uint NumDispatchedGroups;
 uint PaddedTileSize;
-uint CopyTexureMask;
 
 uint ComputeLocalVoxelDataLoadIndex(uint VoxelIndex, uint2 OccupancyBits, out bool bValidVoxel)
 {
@@ -106,14 +74,43 @@ uint ComputeLocalVoxelDataLoadIndex(uint VoxelIndex, uint2 OccupancyBits, out bo
 	return PrefixSum;
 }
 
+void CopyToTileDataTexture(RWTexture3D<float4> DstTileTexture, Buffer<float4> SrcBuffer, uint SrcVoxelDataOffset, uint TileIdx, uint TextureTileOffset, uint3 DstBaseCoord, float4 FallbackValue, uint GroupThreadIndex)
+{
+	const uint NumVoxelsPerPaddedTile = PaddedTileSize * PaddedTileSize * PaddedTileSize;
+	const uint NumOccupancyWordsPerPaddedTile = (NumVoxelsPerPaddedTile + 31u) / 32u;
+
+	// The offset in the source buffer at which the tile data for this tile starts.
+	const uint TileDataOffset = TileDataOffsetsBuffer.Load((TileIdx + TextureTileOffset) << 2u);
+	uint OccupancyLoadIdx = (TileIdx + TextureTileOffset) * NumOccupancyWordsPerPaddedTile;
+	// Keep track of the number of voxels read from the source buffer (for this tile only).
+	uint NumReadVoxels = 0;
+
+	// Iterate over all voxels in the tile and copy them.
+	for (uint VoxelIdx = GroupThreadIndex; VoxelIdx < NumVoxelsPerPaddedTile; VoxelIdx += 64)
+	{
+		const uint2 OccupancyBits = OccupancyBitsBuffer.Load2(OccupancyLoadIdx << 2u);
+
+		bool bValidVoxel = false;
+		const uint LocalLoadIndex = ComputeLocalVoxelDataLoadIndex(VoxelIdx, OccupancyBits, bValidVoxel);
+		const uint VoxelLoadIndex = TileDataOffset + NumReadVoxels + LocalLoadIndex;
+		const uint3 DstCoord = DstBaseCoord + UnflattenCoord(VoxelIdx, PaddedTileSize);
+
+		float4 Voxel = FallbackValue;
+		BRANCH if (bValidVoxel)
+		{
+			Voxel = SrcBuffer.Load(VoxelLoadIndex - SrcVoxelDataOffset);
+		}
+		DstTileTexture[DstCoord] = Voxel;
+
+		OccupancyLoadIdx += 2; // We just consumed 64 bits (2 uints)
+		NumReadVoxels += countbits(OccupancyBits.x) + countbits(OccupancyBits.y);
+	}
+}
+
 [numthreads(64, 1, 1)]
 void SparseVolumeTextureUpdateFromSparseBufferCS(uint3 GroupId : SV_GroupID, uint GroupThreadIndex : SV_GroupIndex)
 {
-	const bool bCopyTextureA = CopyTexureMask & 1u;
-	const bool bCopyTextureB = CopyTexureMask & 2u;
 	const uint GroupIdx = GroupId.x;
-	const uint NumVoxelsPerPaddedTile = PaddedTileSize * PaddedTileSize * PaddedTileSize;
-	const uint NumOccupancyWordsPerPaddedTile = (NumVoxelsPerPaddedTile + 31u) / 32u;
 
 	// Process one tile per workgroup
 	for (uint TileIdx = (GroupIdx + TileIndexOffset); TileIdx < (NumTilesToCopy + TileIndexOffset); TileIdx += NumDispatchedGroups)
@@ -125,53 +122,14 @@ void SparseVolumeTextureUpdateFromSparseBufferCS(uint3 GroupId : SV_GroupID, uin
 		// This offset is used to jump to the offset for the current texture.
 		uint TextureTileOffset = 0;
 
-		// Iterate over both possible textures.
-		UNROLL for (uint TextureIdx = 0; TextureIdx < 2; ++TextureIdx)
-		{
-			BRANCH if ((TextureIdx == 0 && bCopyTextureA) || (TextureIdx == 1 && bCopyTextureB))
-			{
-				// The offset in the source buffer at which the tile data for this tile starts.
-				const uint TileDataOffset = TileDataOffsetsBuffer.Load((TileIdx + TextureTileOffset) << 2u);
-				uint OccupancyLoadIdx = (TileIdx + TextureTileOffset) * NumOccupancyWordsPerPaddedTile;
-				// Keep track of the number of voxels read from the source buffer (for this tile only).
-				uint NumReadVoxels = 0;
-
-				// Iterate over all voxels in the tile and copy them.
-				for (uint VoxelIdx = GroupThreadIndex; VoxelIdx < NumVoxelsPerPaddedTile; VoxelIdx += 64)
-				{
-					const uint2 OccupancyBits = OccupancyBitsBuffer.Load2(OccupancyLoadIdx << 2u);
-
-					bool bValidVoxel = false;
-					const uint LocalLoadIndex = ComputeLocalVoxelDataLoadIndex(VoxelIdx, OccupancyBits, bValidVoxel);
-					const uint VoxelLoadIndex = TileDataOffset + NumReadVoxels + LocalLoadIndex;
-					const uint3 DstCoord = DstBaseCoord + UnflattenCoord(VoxelIdx, PaddedTileSize);
-
-					BRANCH if (TextureIdx == 0)
-					{
-						float4 Voxel = FallbackValueA;
-						BRANCH if (bValidVoxel)
-						{
-							Voxel = SrcPhysicalTileBufferA.Load(VoxelLoadIndex - SrcVoxelDataOffsetA);
-						}
-						DstPhysicalTileTextureA[DstCoord] = Voxel;
-					}
-					else
-					{
-						float4 Voxel = FallbackValueB;
-						BRANCH if (bValidVoxel)
-						{
-							Voxel = SrcPhysicalTileBufferB.Load(VoxelLoadIndex - SrcVoxelDataOffsetB);
-						}
-						DstPhysicalTileTextureB[DstCoord] = Voxel;
-					}
-					
-					OccupancyLoadIdx += 2; // We just consumed 64 bits (2 uints)
-					NumReadVoxels += countbits(OccupancyBits.x) + countbits(OccupancyBits.y);
-				}
-
-				TextureTileOffset += BufferTileStep;
-			}
-		}
+#if UPDATE_TEXTURE_A
+		CopyToTileDataTexture(DstPhysicalTileTextureA, SrcPhysicalTileBufferA, SrcVoxelDataOffsetA, TileIdx, TextureTileOffset, DstBaseCoord, FallbackValueA, GroupThreadIndex);
+		TextureTileOffset += BufferTileStep;
+#endif // UPDATE_TEXTURE_A
+
+#if UPDATE_TEXTURE_B
+		CopyToTileDataTexture(DstPhysicalTileTextureB, SrcPhysicalTileBufferB, SrcVoxelDataOffsetB, TileIdx, TextureTileOffset, DstBaseCoord, FallbackValueB, GroupThreadIndex);
+#endif // UPDATE_TEXTURE_B
 	}
 }
 
diff --git a/Engine/Source/Runtime/Engine/Private/SparseVolumeTexture/SparseVolumeTexture.cpp b/Engine/Source/Runtime/Engine/Private/SparseVolumeTexture/SparseVolumeTexture.cpp
--- a/Engine/Source/Runtime/Engine/Private/SparseVolumeTexture/SparseVolumeTexture.cpp
+++ b/Engine/Source/Runtime/Engine/Private/SparseVolumeTexture/SparseVolumeTexture.cpp
@@ -962,11 +962,11 @@ void FTextureRenderResources::SetGlobalVolumeResolution_GameThread(const FIntVec
 		});
 }
 
-void FTextureRenderResources::InitRHI(FRHICommandListBase&)
+void FTextureRenderResources::InitRHI(FRHICommandListBase& RHICmdList)
 {
-	PageTableTextureReferenceRHI = RHICreateTextureReference(GBlackUintVolumeTexture->TextureRHI);
-	PhysicalTileDataATextureReferenceRHI = RHICreateTextureReference(GBlackVolumeTexture->TextureRHI);
-	PhysicalTileDataBTextureReferenceRHI = RHICreateTextureReference(GBlackVolumeTexture->TextureRHI);
+	PageTableTextureReferenceRHI = RHICmdList.CreateTextureReference(GBlackUintVolumeTexture->TextureRHI);
+	PhysicalTileDataATextureReferenceRHI = RHICmdList.CreateTextureReference(GBlackVolumeTexture->TextureRHI);
+	PhysicalTileDataBTextureReferenceRHI = RHICmdList.CreateTextureReference(GBlackVolumeTexture->TextureRHI);
 }
 
 void FTextureRenderResources::ReleaseRHI()
diff --git a/Engine/Source/Runtime/Engine/Private/SparseVolumeTexture/SparseVolumeTextureStreamingManager.cpp b/Engine/Source/Runtime/Engine/Private/SparseVolumeTexture/SparseVolumeTextureStreamingManager.cpp
--- a/Engine/Source/Runtime/Engine/Private/SparseVolumeTexture/SparseVolumeTextureStreamingManager.cpp
+++ b/Engine/Source/Runtime/Engine/Private/SparseVolumeTexture/SparseVolumeTextureStreamingManager.cpp
@@ -11,6 +11,7 @@
 #include "Async/ParallelFor.h"
 #include "SparseVolumeTextureTileDataTexture.h"
 #include "SparseVolumeTextureUpload.h"
+#include "GlobalRenderResources.h"
 
 #if WITH_EDITORONLY_DATA
 #include "DerivedDataCache.h"
@@ -657,24 +658,28 @@ void FStreamingManager::AddInternal(FRDGBuilder& GraphBuilder, FNewSparseVolumeT
 
 			// Create page table
 			{
-				// SVT_TODO: Currently we keep all mips of the page table resident. It would be better to stream in/out page table mips.
-				const int32 NumResidentMipLevels = NumMipLevels;
 				FIntVector3 PageTableResolution = Resources->Header.PageTableVolumeResolution;
 				PageTableResolution = FIntVector3(FMath::Max(1, PageTableResolution.X), FMath::Max(1, PageTableResolution.Y), FMath::Max(1, PageTableResolution.Z));
 
 				const EPixelFormat PageEntryFormat = PF_R32_UINT;
-				const FRHITextureCreateDesc Desc =
-					FRHITextureCreateDesc::Create3D(TEXT("SparseVolumeTexture.PageTable.RHITexture"), PageTableResolution.X, PageTableResolution.Y, PageTableResolution.Z, PageEntryFormat)
-					.SetFlags(ETextureCreateFlags::ShaderResource | ETextureCreateFlags::UAV)
-					.SetNumMips((uint8)NumResidentMipLevels);
+				const ETextureCreateFlags Flags = TexCreate_ShaderResource | TexCreate_UAV | TexCreate_3DTiling | TexCreate_ReduceMemoryWithTilingMode;
+				FRDGTexture* PageTableRDG = GraphBuilder.CreateTexture(FRDGTextureDesc::Create3D(PageTableResolution, PageEntryFormat, FClearValueBinding::Black, Flags, (uint8)NumMipLevels), TEXT("SparseVolumeTexture.PageTableTexture"));
 
-				FrameInfo.PageTableTextureRHIRef = RHICreateTexture(Desc);
+				// Clear page table to zero
+				for (int32 MipLevelIndex = 0; MipLevelIndex < NumMipLevels; ++MipLevelIndex)
+				{
+					FRDGTextureUAV* UAV = GraphBuilder.CreateUAV(FRDGTextureUAVDesc(PageTableRDG, MipLevelIndex, PageEntryFormat));
+					AddClearUAVPass(GraphBuilder, UAV, FUintVector4(ForceInitToZero));
+				}
+
+				FrameInfo.PageTableTexture = GraphBuilder.ConvertToExternalTexture(PageTableRDG);
+				GraphBuilder.UseExternalAccessMode(PageTableRDG, ERHIAccess::SRVMask, ERHIPipeline::All);
 			}
 
 			// Initialize TextureRenderResources
-			RHIUpdateTextureReference(FrameInfo.TextureRenderResources->PageTableTextureReferenceRHI, FrameInfo.PageTableTextureRHIRef);
-			RHIUpdateTextureReference(FrameInfo.TextureRenderResources->PhysicalTileDataATextureReferenceRHI, SVTInfo.TileDataTexture->GetTileDataTextureA());
-			RHIUpdateTextureReference(FrameInfo.TextureRenderResources->PhysicalTileDataBTextureReferenceRHI, SVTInfo.TileDataTexture->GetTileDataTextureB());
+			GraphBuilder.RHICmdList.UpdateTextureReference(FrameInfo.TextureRenderResources->PageTableTextureReferenceRHI, FrameInfo.PageTableTexture->GetRHI());
+			GraphBuilder.RHICmdList.UpdateTextureReference(FrameInfo.TextureRenderResources->PhysicalTileDataATextureReferenceRHI, SVTInfo.TileDataTexture->GetTileDataTextureA() ? SVTInfo.TileDataTexture->GetTileDataTextureA()->GetRHI() : GBlackVolumeTexture->TextureRHI.GetReference());
+			GraphBuilder.RHICmdList.UpdateTextureReference(FrameInfo.TextureRenderResources->PhysicalTileDataBTextureReferenceRHI, SVTInfo.TileDataTexture->GetTileDataTextureB() ? SVTInfo.TileDataTexture->GetTileDataTextureB()->GetRHI() : GBlackVolumeTexture->TextureRHI.GetReference());
 			FrameInfo.TextureRenderResources->Header = Resources->Header;
 			FrameInfo.TextureRenderResources->TileDataTextureResolution = SVTInfo.TileDataTexture->GetResolutionInTiles() * SPARSE_VOLUME_TILE_RES_PADDED;
 			FrameInfo.TextureRenderResources->FrameIndex = FrameIdx;
@@ -715,10 +720,6 @@ void FStreamingManager::AddInternal(FRDGBuilder& GraphBuilder, FNewSparseVolumeT
 						FMemory::Memcpy(AddResult.TileDataPtrs[AttributesIdx], SrcTileData, RootStreamingInfo->TileDataSize[AttributesIdx]);
 					}
 				}
-
-				// Update highest mip (1x1x1) in page table
-				const FUpdateTextureRegion3D UpdateRegion(0, 0, 0, 0, 0, 0, 1, 1, 1);
-				RHIUpdateTexture3D(FrameInfo.PageTableTextureRHIRef, NumMipLevels - 1, UpdateRegion, sizeof(uint32), sizeof(uint32), (uint8*)&TileCoord);
 			}
 
 			const FPageTopology::FMip& TopologyMipInfo = Resources->Topology.MipInfo[NumMipLevels - 1];
@@ -777,7 +778,7 @@ void FStreamingManager::RemoveInternal(UStreamableSparseVolumeTexture* SparseVol
 		// Release resources
 		for (FFrameInfo& FrameInfo : SVTInfo->PerFrameInfo)
 		{
-			FrameInfo.PageTableTextureRHIRef.SafeRelease();
+			FrameInfo.PageTableTexture.SafeRelease();
 			InvalidatedSVTFrames.Remove(&FrameInfo);
 		}
 		if (SVTInfo->TileDataTexture)
@@ -1620,7 +1621,7 @@ void FStreamingManager::PatchPageTable(FRDGBuilder& GraphBuilder)
 					MipUpdateWriteIndex = 0;
 					bEnteredNewMipRange = false;
 
-					PageTableUpdater->Add_GetRef(FrameInfo.PageTableTextureRHIRef, MipLevel, NumUpdatesThisMip, DstCoordsPtr, DstEntryPtr);
+					PageTableUpdater->Add_GetRef(FrameInfo.PageTableTexture, MipLevel, NumUpdatesThisMip, DstCoordsPtr, DstEntryPtr);
 				}
 
 				uint32 PageTableEntry = 0;
diff --git a/Engine/Source/Runtime/Engine/Private/SparseVolumeTexture/SparseVolumeTextureStreamingManager.h b/Engine/Source/Runtime/Engine/Private/SparseVolumeTexture/SparseVolumeTextureStreamingManager.h
--- a/Engine/Source/Runtime/Engine/Private/SparseVolumeTexture/SparseVolumeTextureStreamingManager.h
+++ b/Engine/Source/Runtime/Engine/Private/SparseVolumeTexture/SparseVolumeTextureStreamingManager.h
@@ -132,7 +132,7 @@ private:
 		TBitArray<> ResidentPages; //  One bit for every (non-zero) page (in the sparse page octree) for all mip levels, starting at the highest mip
 		TBitArray<> ResidentPagesNew; // Reflects changes made during an update of the streaming system
 		TBitArray<> InvalidatedPages; // All pages that need updated page table entries
-		FTextureRHIRef PageTableTextureRHIRef;
+		TRefCountPtr<IPooledRenderTarget> PageTableTexture;
 	};
 
 	// Used to keep track of the least-recently-used order of mip levels. This is done per mip level and not per frame, so that lower resolution mip levels are kept in memory for longer.
diff --git a/Engine/Source/Runtime/Engine/Private/SparseVolumeTexture/SparseVolumeTextureTileDataTexture.cpp b/Engine/Source/Runtime/Engine/Private/SparseVolumeTexture/SparseVolumeTextureTileDataTexture.cpp
--- a/Engine/Source/Runtime/Engine/Private/SparseVolumeTexture/SparseVolumeTextureTileDataTexture.cpp
+++ b/Engine/Source/Runtime/Engine/Private/SparseVolumeTexture/SparseVolumeTextureTileDataTexture.cpp
@@ -3,6 +3,7 @@
 #include "SparseVolumeTextureTileDataTexture.h"
 #include "SparseVolumeTextureUtility.h"
 #include "SparseVolumeTextureStreamingManager.h" // LogSparseVolumeTextureStreamingManager
+#include "RenderTargetPool.h"
 
 namespace UE
 {
@@ -161,26 +162,21 @@ FTileUploader::FAddResult FTileDataTexture::AddUpload(int32 NumTiles, int32 NumV
 void FTileDataTexture::EndUpload(FRDGBuilder& GraphBuilder)
 {
 	check(UploaderState == EUploaderState::Uploading);
-	TileUploader->ResourceUploadTo(GraphBuilder, TileDataTextureARHIRef, TileDataTextureBRHIRef, FallbackValueA, FallbackValueB);
+	TileUploader->ResourceUploadTo(GraphBuilder, TileDataTextureA, TileDataTextureB, FallbackValueA, FallbackValueB);
 	UploaderState = EUploaderState::Ready;
 }
 
-void FTileDataTexture::InitRHI(FRHICommandListBase&)
+void FTileDataTexture::InitRHI(FRHICommandListBase& RHICmdList)
 {
 	const FIntVector3 Resolution = ResolutionInTiles * SPARSE_VOLUME_TILE_RES_PADDED;
+	const ETextureCreateFlags Flags = TexCreate_ShaderResource | TexCreate_UAV | TexCreate_3DTiling | TexCreate_ReduceMemoryWithTilingMode;
 	if (FormatA != PF_Unknown)
 	{
-		const FRHITextureCreateDesc Desc =
-			FRHITextureCreateDesc::Create3D(TEXT("SparseVolumeTexture.PhysicalTileDataA.RHITexture"), Resolution.X, Resolution.Y, Resolution.Z, FormatA)
-			.SetFlags(ETextureCreateFlags::ShaderResource | ETextureCreateFlags::UAV);
-		TileDataTextureARHIRef = RHICreateTexture(Desc);
+		TileDataTextureA = GRenderTargetPool.FindFreeElement(RHICmdList, FRDGTextureDesc::Create3D(Resolution, FormatA, FClearValueBinding::Black, Flags), TEXT("SparseVolumeTexture.PhysicalTileDataA.RHITexture"));
 	}
 	if (FormatB != PF_Unknown)
 	{
-		const FRHITextureCreateDesc Desc =
-			FRHITextureCreateDesc::Create3D(TEXT("SparseVolumeTexture.PhysicalTileDataB.RHITexture"), Resolution.X, Resolution.Y, Resolution.Z, FormatB)
-			.SetFlags(ETextureCreateFlags::ShaderResource | ETextureCreateFlags::UAV);
-		TileDataTextureBRHIRef = RHICreateTexture(Desc);
+		TileDataTextureB = GRenderTargetPool.FindFreeElement(RHICmdList, FRDGTextureDesc::Create3D(Resolution, FormatB, FClearValueBinding::Black, Flags), TEXT("SparseVolumeTexture.PhysicalTileDataB.RHITexture"));
 	}
 }
 
diff --git a/Engine/Source/Runtime/Engine/Private/SparseVolumeTexture/SparseVolumeTextureTileDataTexture.h b/Engine/Source/Runtime/Engine/Private/SparseVolumeTexture/SparseVolumeTextureTileDataTexture.h
--- a/Engine/Source/Runtime/Engine/Private/SparseVolumeTexture/SparseVolumeTextureTileDataTexture.h
+++ b/Engine/Source/Runtime/Engine/Private/SparseVolumeTexture/SparseVolumeTextureTileDataTexture.h
@@ -63,8 +63,8 @@ public:
 	int32 GetTileCapacity() const { return PhysicalTilesCapacity; }
 	int32 GetNumAvailableTiles() const { return PhysicalTilesCapacity - NextFreeTileCoordIndex; } // Number of tiles available for allocation.
 	FIntVector3 GetResolutionInTiles() const { return ResolutionInTiles; }
-	FTextureRHIRef GetTileDataTextureA() { return TileDataTextureARHIRef; }
-	FTextureRHIRef GetTileDataTextureB() { return TileDataTextureBRHIRef; }
+	TRefCountPtr<IPooledRenderTarget> GetTileDataTextureA() { return TileDataTextureA; }
+	TRefCountPtr<IPooledRenderTarget> GetTileDataTextureB() { return TileDataTextureB; }
 
 	// Transitions from EUploaderState::Ready (or EUploaderState::Reserved) to EUploaderState::Reserving and allows callers to call ReserveUpload() afterwards. Resets number of reserved tiles/voxels.
 	void BeginReserveUpload();
@@ -90,8 +90,8 @@ private:
 	EPixelFormat FormatB;
 	FVector4f FallbackValueA;
 	FVector4f FallbackValueB;
-	FTextureRHIRef TileDataTextureARHIRef;
-	FTextureRHIRef TileDataTextureBRHIRef;
+	TRefCountPtr<IPooledRenderTarget> TileDataTextureA;
+	TRefCountPtr<IPooledRenderTarget> TileDataTextureB;
 	TArray<uint32> TileCoords;
 	int32 NextFreeTileCoordIndex = 0;
 	EUploaderState UploaderState = EUploaderState::Ready;
diff --git a/Engine/Source/Runtime/Engine/Private/SparseVolumeTexture/SparseVolumeTextureUpload.cpp b/Engine/Source/Runtime/Engine/Private/SparseVolumeTexture/SparseVolumeTextureUpload.cpp
--- a/Engine/Source/Runtime/Engine/Private/SparseVolumeTexture/SparseVolumeTextureUpload.cpp
+++ b/Engine/Source/Runtime/Engine/Private/SparseVolumeTexture/SparseVolumeTextureUpload.cpp
@@ -16,43 +16,14 @@ static FAutoConsoleVariableRef CVarSVTStreamingAsyncCompute(
 	ECVF_RenderThreadSafe
 );
 
-class FSparseVolumeTextureUpdateFromBufferCS : public FGlobalShader
-{
-	DECLARE_GLOBAL_SHADER(FSparseVolumeTextureUpdateFromBufferCS);
-	SHADER_USE_PARAMETER_STRUCT(FSparseVolumeTextureUpdateFromBufferCS, FGlobalShader)
-
-	BEGIN_SHADER_PARAMETER_STRUCT(FParameters, )
-		SHADER_PARAMETER_RDG_TEXTURE_UAV(RWTexture3D, DstPhysicalTileTextureA)
-		SHADER_PARAMETER_RDG_TEXTURE_UAV(RWTexture3D, DstPhysicalTileTextureB)
-		SHADER_PARAMETER_RDG_BUFFER_SRV(ByteAddressBuffer, DstTileCoords)
-		SHADER_PARAMETER_RDG_BUFFER_SRV(Buffer<float4>, SrcPhysicalTileBufferA)
-		SHADER_PARAMETER_RDG_BUFFER_SRV(Buffer<float4>, SrcPhysicalTileBufferB)
-		SHADER_PARAMETER(uint32, TileCoordsBufferOffset)
-		SHADER_PARAMETER(uint32, TileDataBufferOffsetInTiles)
-		SHADER_PARAMETER(uint32, NumTilesToCopy)
-		SHADER_PARAMETER(uint32, NumDispatchedGroups)
-		SHADER_PARAMETER(uint32, PaddedTileSize)
-		SHADER_PARAMETER(uint32, bCopyTexureAOnlyUI)
-	END_SHADER_PARAMETER_STRUCT()
-
-	static bool ShouldCompilePermutation(const FGlobalShaderPermutationParameters& Parameters)
-	{
-		return true;
-	}
-
-	static void ModifyCompilationEnvironment(const FGlobalShaderPermutationParameters& Parameters, FShaderCompilerEnvironment& OutEnvironment)
-	{
-		FGlobalShader::ModifyCompilationEnvironment(Parameters, OutEnvironment);
-		OutEnvironment.SetDefine(TEXT("UPDATE_TILE_TEXTURE_FROM_BUFFER"), 1);
-	}
-};
-IMPLEMENT_GLOBAL_SHADER(FSparseVolumeTextureUpdateFromBufferCS, "/Engine/Private/SparseVolumeTexture/UpdateSparseVolumeTexture.usf", "SparseVolumeTextureUpdateFromBufferCS", SF_Compute);
-
 class FSparseVolumeTextureUpdateFromSparseBufferCS : public FGlobalShader
 {
 	DECLARE_GLOBAL_SHADER(FSparseVolumeTextureUpdateFromSparseBufferCS);
 	SHADER_USE_PARAMETER_STRUCT(FSparseVolumeTextureUpdateFromSparseBufferCS, FGlobalShader)
 
+	class FTextureUpdateMask : SHADER_PERMUTATION_SPARSE_INT("TEXTURE_UPDATE_MASK", 1, 2, 3);
+	using FPermutationDomain = TShaderPermutationDomain<FTextureUpdateMask>;
+
 	BEGIN_SHADER_PARAMETER_STRUCT(FParameters, )
 		SHADER_PARAMETER_RDG_TEXTURE_UAV(RWTexture3D, DstPhysicalTileTextureA)
 		SHADER_PARAMETER_RDG_TEXTURE_UAV(RWTexture3D, DstPhysicalTileTextureB)
@@ -70,7 +41,6 @@ class FSparseVolumeTextureUpdateFromSparseBufferCS : public FGlobalShader
 		SHADER_PARAMETER(uint32, BufferTileStep)
 		SHADER_PARAMETER(uint32, NumDispatchedGroups)
 		SHADER_PARAMETER(uint32, PaddedTileSize)
-		SHADER_PARAMETER(uint32, CopyTexureMask)
 	END_SHADER_PARAMETER_STRUCT()
 
 	static bool ShouldCompilePermutation(const FGlobalShaderPermutationParameters& Parameters)
@@ -242,7 +212,7 @@ void FTileUploader::Release()
 	ResetState();
 }
 
-void FTileUploader::ResourceUploadTo(FRDGBuilder& GraphBuilder, FRHITexture* InDstTextureA, FRHITexture* InDstTextureB, const FVector4f& InFallbackValueA, const FVector4f& InFallbackValueB)
+void FTileUploader::ResourceUploadTo(FRDGBuilder& GraphBuilder, const TRefCountPtr<IPooledRenderTarget>& InDstTextureA, const TRefCountPtr<IPooledRenderTarget>& InDstTextureB, const FVector4f& InFallbackValueA, const FVector4f& InFallbackValueB)
 {
 	check(InDstTextureA || FormatSizeA <= 0);
 	check(InDstTextureB || FormatSizeB <= 0);
@@ -269,14 +239,23 @@ void FTileUploader::ResourceUploadTo(FRDGBuilder& GraphBuilder, FRHITexture* InD
 	
 		if (NumWrittenTiles > 0)
 		{
-			FRDGTexture* DummyTexture = GraphBuilder.CreateTexture(FRDGTextureDesc::Create3D(FIntVector(1), PF_R8, FClearValueBinding::None, ETextureCreateFlags::UAV), TEXT("SparseVolumeTexture.DummyTexture"));
-			FRDGTexture* DstTextureARDG = InDstTextureA ? GraphBuilder.RegisterExternalTexture(CreateRenderTarget(InDstTextureA, TEXT("SparseVolumeTexture.TileDataTextureA"))) : nullptr;
-			FRDGTexture* DstTextureBRDG = InDstTextureB ? GraphBuilder.RegisterExternalTexture(CreateRenderTarget(InDstTextureB, TEXT("SparseVolumeTexture.TileDataTextureB"))) : nullptr;
+			FRDGTexture* DstTextureARDG = nullptr;
+			FRDGTexture* DstTextureBRDG = nullptr;
+			if (InDstTextureA)
+			{
+				DstTextureARDG = GraphBuilder.RegisterExternalTexture(InDstTextureA, ERDGTextureFlags::ForceImmediateFirstBarrier);
+				GraphBuilder.UseInternalAccessMode(DstTextureARDG);
+			}
+			if (InDstTextureB)
+			{
+				DstTextureBRDG = GraphBuilder.RegisterExternalTexture(InDstTextureB, ERDGTextureFlags::ForceImmediateFirstBarrier);
+				GraphBuilder.UseInternalAccessMode(DstTextureBRDG);
+			}
+
 			FRDGBuffer* SrcBufferARDG = FormatSizeA > 0 ? GraphBuilder.RegisterExternalBuffer(TileDataAUploadBuffer) : nullptr;
 			FRDGBuffer* SrcBufferBRDG = FormatSizeB > 0 ? GraphBuilder.RegisterExternalBuffer(TileDataBUploadBuffer) : nullptr;
 			check(SrcBufferARDG || SrcBufferBRDG);
 	
-			FRDGTextureUAV* DummyTextureUAV = GraphBuilder.CreateUAV(DummyTexture);
 			FRDGTextureUAV* DstTextureAUAV = DstTextureARDG ? GraphBuilder.CreateUAV(DstTextureARDG) : nullptr;
 			FRDGTextureUAV* DstTextureBUAV = DstTextureBRDG ? GraphBuilder.CreateUAV(DstTextureBRDG) : nullptr;
 	
@@ -286,7 +265,14 @@ void FTileUploader::ResourceUploadTo(FRDGBuilder& GraphBuilder, FRHITexture* InD
 			// Either SrcBufferARDG or SrcBufferBRDG must exist and will have at least 1 element.
 			FRDGBufferSRV* DummySrcBufferSRV = SrcBufferARDG ? GraphBuilder.CreateSRV(SrcBufferARDG, FormatA) : GraphBuilder.CreateSRV(SrcBufferBRDG, FormatB);
 	
-			auto ComputeShader = GetGlobalShaderMap(GMaxRHIFeatureLevel)->GetShader<FSparseVolumeTextureUpdateFromSparseBufferCS>();
+			uint32 TextureUpdateMask = 0;
+			TextureUpdateMask |= FormatSizeA > 0 ? 0x1u : 0x0u;
+			TextureUpdateMask |= FormatSizeB > 0 ? 0x2u : 0x0u;
+			check(TextureUpdateMask != 0);
+
+			FSparseVolumeTextureUpdateFromSparseBufferCS::FPermutationDomain CSPermutationDomain;
+			CSPermutationDomain.Set<FSparseVolumeTextureUpdateFromSparseBufferCS::FTextureUpdateMask>(TextureUpdateMask);
+			auto ComputeShader = GetGlobalShaderMap(GMaxRHIFeatureLevel)->GetShader<FSparseVolumeTextureUpdateFromSparseBufferCS>(CSPermutationDomain);
 	
 			int32 NumUploadedTiles = 0;
 			int32 NumUploadedVoxelsA = 0;
@@ -345,9 +331,9 @@ void FTileUploader::ResourceUploadTo(FRDGBuilder& GraphBuilder, FRHITexture* InD
 				}
 	
 				FSparseVolumeTextureUpdateFromSparseBufferCS::FParameters* PassParameters = GraphBuilder.AllocParameters<FSparseVolumeTextureUpdateFromSparseBufferCS::FParameters>();
-				PassParameters->DstPhysicalTileTextureA = DstTextureAUAV ? DstTextureAUAV : DummyTextureUAV;
-				PassParameters->DstPhysicalTileTextureB = DstTextureBUAV ? DstTextureBUAV : DummyTextureUAV;
-				PassParameters->SrcPhysicalTileBufferA = TileDataABufferSRV ? TileDataABufferSRV : DummySrcBufferSRV;
+				PassParameters->DstPhysicalTileTextureA = DstTextureAUAV;
+				PassParameters->DstPhysicalTileTextureB = DstTextureBUAV;
+				PassParameters->SrcPhysicalTileBufferA = TileDataABufferSRV ? TileDataABufferSRV : DummySrcBufferSRV; // TileDataABufferSRV and TileDataBBufferSRV can be null if there are no explicitly stored voxels to upload
 				PassParameters->SrcPhysicalTileBufferB = TileDataBBufferSRV ? TileDataBBufferSRV : DummySrcBufferSRV;
 				PassParameters->OccupancyBitsBuffer = OccupancyBitsBufferSRV;
 				PassParameters->TileDataOffsetsBuffer = TileDataOffsetsBufferSRV;
@@ -361,9 +347,6 @@ void FTileUploader::ResourceUploadTo(FRDGBuilder& GraphBuilder, FRHITexture* InD
 				PassParameters->BufferTileStep = MaxNumTiles;
 				PassParameters->NumDispatchedGroups = FMath::Min(NumTilesInThisBatch, GRHIMaxDispatchThreadGroupsPerDimension.X);
 				PassParameters->PaddedTileSize = SPARSE_VOLUME_TILE_RES_PADDED;
-				PassParameters->CopyTexureMask = 0;
-				PassParameters->CopyTexureMask |= FormatSizeA > 0 ? 0x1u : 0x0u;
-				PassParameters->CopyTexureMask |= FormatSizeB > 0 ? 0x2u : 0x0u;
 	
 				// Disable async compute for streaming systems when MGPU is active, to work around GPU hangs
 				const bool bAsyncCompute = GSupportsEfficientAsyncCompute && (GSVTStreamingAsyncCompute != 0) && (GNumExplicitGPUsForRendering == 1);
@@ -385,6 +368,16 @@ void FTileUploader::ResourceUploadTo(FRDGBuilder& GraphBuilder, FRHITexture* InD
 			check(NumUploadedTiles == NumWrittenTiles);
 			check(NumUploadedVoxelsA == NumWrittenVoxelsA);
 			check(NumUploadedVoxelsB == NumWrittenVoxelsB);
+
+			// Transition back to SRV in case we want to access them using their raw RHI resources within this graph
+			if (DstTextureARDG)
+			{
+				GraphBuilder.UseExternalAccessMode(DstTextureARDG, ERHIAccess::SRVMask, ERHIPipeline::All);
+			}
+			if (DstTextureBRDG)
+			{
+				GraphBuilder.UseExternalAccessMode(DstTextureBRDG, ERHIAccess::SRVMask, ERHIPipeline::All);
+			}
 		}
 	}
 	Release();
@@ -439,7 +432,7 @@ void FPageTableUpdater::Init(FRDGBuilder& GraphBuilder, int32 InMaxNumUpdates, i
 	}
 }
 
-void FPageTableUpdater::Add_GetRef(FRHITexture* InPageTable, int32 InMipLevel, int32 InNumUpdates, uint8*& OutCoordsPtr, uint8*& OutPayloadPtr)
+void FPageTableUpdater::Add_GetRef(const TRefCountPtr<IPooledRenderTarget>& InPageTable, int32 InMipLevel, int32 InNumUpdates, uint8*& OutCoordsPtr, uint8*& OutPayloadPtr)
 {
 	check((NumWrittenUpdates + InNumUpdates) <= MaxNumUpdates);
 	check(DataPtr);
@@ -473,11 +466,23 @@ void FPageTableUpdater::Apply(FRDGBuilder& GraphBuilder)
 			// Disable async compute for streaming systems when MGPU is active, to work around GPU hangs
 			const bool bAsyncCompute = GSupportsEfficientAsyncCompute && (GSVTStreamingAsyncCompute != 0) && (GNumExplicitGPUsForRendering == 1);
 			auto ComputeShader = GetGlobalShaderMap(GMaxRHIFeatureLevel)->GetShader<FSparseVolumeTextureUpdatePageTableCS>();
+
+			// Register page table textures with RDG and set it to internal access mode
+			for (const FBatch& Batch : Batches)
+			{
+				FRDGTexture* PageTableRDG = GraphBuilder.FindExternalTexture(Batch.PageTable);
+				if (!PageTableRDG)
+				{
+					PageTableRDG = GraphBuilder.RegisterExternalTexture(Batch.PageTable, ERDGTextureFlags::ForceImmediateFirstBarrier);
+				}
+				GraphBuilder.UseInternalAccessMode(PageTableRDG); // Make sure to use graph tracking in case we previously called UseExternalAccessMode() within the current graph
+			}
 	
 			uint32 UpdatesOffset = 0;
 			for (const FBatch& Batch : Batches)
 			{
-				FRDGTexture* PageTableRDG = GraphBuilder.RegisterExternalTexture(CreateRenderTarget(Batch.PageTable, TEXT("SparseVolumeTexture.PageTableTexture")));
+				FRDGTexture* PageTableRDG = GraphBuilder.FindExternalTexture(Batch.PageTable);
+				check(PageTableRDG);
 				FRDGTextureUAV* PageTableUAV = GraphBuilder.CreateUAV(FRDGTextureUAVDesc(PageTableRDG, Batch.MipLevel, PF_R32_UINT));
 				FRDGBufferSRV* UpdatesBufferSRV = GraphBuilder.CreateSRV(GraphBuilder.RegisterExternalBuffer(UpdatesUploadBuffer));
 	
@@ -499,6 +504,14 @@ void FPageTableUpdater::Apply(FRDGBuilder& GraphBuilder)
 	
 				UpdatesOffset += Batch.NumUpdates;
 			}
+
+			// Transition back to SRV in case we want to access the page table texture using its raw RHI resource within this graph
+			for (const FBatch& Batch : Batches)
+			{
+				FRDGTexture* PageTableRDG = GraphBuilder.FindExternalTexture(Batch.PageTable);
+				check(PageTableRDG);
+				GraphBuilder.UseExternalAccessMode(PageTableRDG, ERHIAccess::SRVMask, ERHIPipeline::All);
+			}
 		}
 	}
 	
diff --git a/Engine/Source/Runtime/Engine/Private/SparseVolumeTexture/SparseVolumeTextureUpload.h b/Engine/Source/Runtime/Engine/Private/SparseVolumeTexture/SparseVolumeTextureUpload.h
--- a/Engine/Source/Runtime/Engine/Private/SparseVolumeTexture/SparseVolumeTextureUpload.h
+++ b/Engine/Source/Runtime/Engine/Private/SparseVolumeTexture/SparseVolumeTextureUpload.h
@@ -33,7 +33,7 @@ public:
 	void Init(FRDGBuilder& InGraphBuilder, int32 InMaxNumTiles, int32 InMaxNumVoxelsA, int32 InMaxNumVoxelsB, EPixelFormat InFormatA, EPixelFormat InFormatB);
 	FAddResult Add_GetRef(int32 InNumTiles, int32 InNumVoxelsA, int32 InNumVoxelsB);
 	void Release();
-	void ResourceUploadTo(FRDGBuilder& InGraphBuilder, FRHITexture* InDstTextureA, FRHITexture* InDstTextureB, const FVector4f& InFallbackValueA, const FVector4f& InFallbackValueB);
+	void ResourceUploadTo(FRDGBuilder& InGraphBuilder, const TRefCountPtr<IPooledRenderTarget>& InDstTextureA, const TRefCountPtr<IPooledRenderTarget>& InDstTextureB, const FVector4f& InFallbackValueA, const FVector4f& InFallbackValueB);
 
 private:
 	TRefCountPtr<FRDGPooledBuffer> OccupancyBitsUploadBuffer;
@@ -69,19 +69,19 @@ class FPageTableUpdater
 public:
 	explicit FPageTableUpdater();
 	void Init(FRDGBuilder& InGraphBuilder, int32 InMaxNumUpdates, int32 InEstimatedNumBatches);
-	void Add_GetRef(FRHITexture* InPageTable, int32 InMipLevel, int32 InNumUpdates, uint8*& OutCoordsPtr, uint8*& OutPayloadPtr);
+	void Add_GetRef(const TRefCountPtr<IPooledRenderTarget>& InPageTable, int32 InMipLevel, int32 InNumUpdates, uint8*& OutCoordsPtr, uint8*& OutPayloadPtr);
 	void Release();
 	void Apply(FRDGBuilder& InGraphBuilder);
 
 private:
 	struct FBatch
 	{
-		FRHITexture* PageTable;
+		TRefCountPtr<IPooledRenderTarget> PageTable;
 		int32 MipLevel;
 		int32 NumUpdates;
 
 		FBatch() = default;
-		FBatch(FRHITexture* InPageTable, int32 InMipLevel) : PageTable(InPageTable), MipLevel(InMipLevel), NumUpdates(0) {}
+		FBatch(const TRefCountPtr<IPooledRenderTarget>& InPageTable, int32 InMipLevel) : PageTable(InPageTable), MipLevel(InMipLevel), NumUpdates(0) {}
 	};
 
 	TRefCountPtr<FRDGPooledBuffer> UpdatesUploadBuffer;
