Skeletal Mesh: Detect broken normals/tangents on import data and meshdescription and reject them if so.

#rnx

diff --git a/Engine/Source/Runtime/Core/Private/UObject/UE5ReleaseStreamObjectVersion.cpp b/Engine/Source/Runtime/Core/Private/UObject/UE5ReleaseStreamObjectVersion.cpp
--- a/Engine/Source/Runtime/Core/Private/UObject/UE5ReleaseStreamObjectVersion.cpp
+++ b/Engine/Source/Runtime/Core/Private/UObject/UE5ReleaseStreamObjectVersion.cpp
@@ -13,7 +13,7 @@ TMap<FGuid, FGuid> FUE5ReleaseStreamObjectVersion::GetSystemGuids()
 	SystemGuids.Add(DevGuids.NANITE_DERIVEDDATA_VER, FGuid("AC88CFBDDA614A9C8488F64D2DAF699D"));
 	SystemGuids.Add(DevGuids.NIAGARASHADERMAP_DERIVEDDATA_VER, FGuid("6BE35B6FACB34568970120F9BC8DAB80"));
 	SystemGuids.Add(DevGuids.Niagara_LatestScriptCompileVersion, FGuid("43A777C54EB24E9894525B8D04529F23"));
-	SystemGuids.Add(DevGuids.SkeletalMeshDerivedDataVersion, FGuid("25C49E579B3142DDA2A8C14037267679"));
+	SystemGuids.Add(DevGuids.SkeletalMeshDerivedDataVersion, FGuid("E55E25709ADD40B385EFDB6E3BB393A7"));
 	SystemGuids.Add(DevGuids.STATICMESH_DERIVEDDATA_VER, FGuid("7A556175518D458AA8786ED69AB8DDE7"));
 
 	return SystemGuids;
diff --git a/Engine/Source/Runtime/Engine/Private/Rendering/SkeletalMeshLODImporterData.cpp b/Engine/Source/Runtime/Engine/Private/Rendering/SkeletalMeshLODImporterData.cpp
--- a/Engine/Source/Runtime/Engine/Private/Rendering/SkeletalMeshLODImporterData.cpp
+++ b/Engine/Source/Runtime/Engine/Private/Rendering/SkeletalMeshLODImporterData.cpp
@@ -2200,6 +2200,33 @@ bool FSkeletalMeshImportData::GetMeshDescription(const USkeletalMesh* InSkeletal
 		MaterialGroups.Add(PolygonGroupID);
 	}
 
+	// These are not always honest.
+	bool bHaveValidNormals = bHasNormals;
+	bool bHaveValidTangents = bHasNormals;
+	
+	for (int32 TriangleIndex = 0; TriangleIndex < Faces.Num(); TriangleIndex++)
+	{
+		const SkeletalMeshImportData::FTriangle &Triangle = Faces[TriangleIndex];
+
+		for (int32 Corner = 0; Corner < 3; Corner++)
+		{
+			if (!Triangle.TangentZ[Corner].IsNormalized())
+			{
+				bHaveValidNormals = false;
+			}
+			if (!Triangle.TangentX[Corner].IsNormalized() ||
+			    !Triangle.TangentY[Corner].IsNormalized())
+			{
+				bHaveValidTangents = false;
+			}
+
+			if (!bHaveValidNormals && !bHaveValidTangents)
+			{
+				break;
+			}
+		}
+	}
+
 	for (int32 TriangleIndex = 0; TriangleIndex < Faces.Num(); TriangleIndex++)
 	{
 		const SkeletalMeshImportData::FTriangle &Triangle = Faces[TriangleIndex];
@@ -2224,10 +2251,22 @@ bool FSkeletalMeshImportData::GetMeshDescription(const USkeletalMesh* InSkeletal
 				{
 					VertexInstanceUVs.Set(VertexInstanceID, UVIndex, Wedge.UVs[UVIndex]);
 				}
-				VertexInstanceTangents.Set(VertexInstanceID, Triangle.TangentX[Corner]);
-				VertexInstanceNormals.Set(VertexInstanceID, Triangle.TangentZ[Corner]);
-				VertexInstanceBinormalSigns.Set(VertexInstanceID,
-					((Triangle.TangentZ[Corner] ^ Triangle.TangentX[Corner]) | Triangle.TangentY[Corner]) < 0 ? -1.0f : 1.0f);
+
+				if (bHaveValidNormals)
+				{
+					VertexInstanceNormals.Set(VertexInstanceID, Triangle.TangentZ[Corner]);
+				}
+				if (bHaveValidTangents)
+				{
+					VertexInstanceTangents.Set(VertexInstanceID, Triangle.TangentX[Corner]);
+
+					// We can only divine the bi-tangent sign if the normal is also given. 
+					if (bHaveValidNormals)
+					{
+						VertexInstanceBinormalSigns.Set(VertexInstanceID,
+							((Triangle.TangentZ[Corner] ^ Triangle.TangentX[Corner]) | Triangle.TangentY[Corner]) < 0 ? -1.0f : 1.0f);
+					}
+				}
 				
 				VertexInstanceIDMap[WedgeId] = VertexInstanceID;
 			}
@@ -2656,7 +2695,7 @@ FSkeletalMeshImportData FSkeletalMeshImportData::CreateFromMeshDescription(const
 	SkelMeshImportData.Points.SetNumUninitialized(InMeshDescription.Vertices().Num());
 	SkelMeshImportData.Influences.Reserve(InMeshDescription.Vertices().Num() * 4);
 	SkelMeshImportData.PointToRawMap.SetNumUninitialized(InMeshDescription.Vertices().Num());
-	
+
 	for (FVertexID VertexID : InMeshDescription.Vertices().GetElementIDs())
 	{
 		// We can use GetValue here because the mesh description was compacted before the copy
@@ -2687,11 +2726,32 @@ FSkeletalMeshImportData FSkeletalMeshImportData::CreateFromMeshDescription(const
 		}
 	}
 
+	
+	bool bHaveValidNormals = true;
+	bool bHaveValidTangents = true;
+	for (FVertexInstanceID VertexInstanceID: InMeshDescription.VertexInstances().GetElementIDs())
+	{
+		if (!VertexInstanceNormals.Get(VertexInstanceID).IsNormalized())
+		{
+			bHaveValidNormals = false;
+		}
+		if (!VertexInstanceTangents.Get(VertexInstanceID).IsNormalized())
+		{
+			bHaveValidTangents = false;
+		}
+		if (!bHaveValidNormals && !bHaveValidTangents)
+		{
+			break;
+		}
+	}
+	
 	//////////////////////////////////////////////////////////////////////////
 	//Copy the triangle and vertex instances
 	SkelMeshImportData.Faces.AddZeroed(InMeshDescription.Triangles().Num());
 	SkelMeshImportData.Wedges.Reserve(InMeshDescription.VertexInstances().Num());
 	SkelMeshImportData.NumTexCoords = FMath::Min<int32>(VertexInstanceUVs.GetNumChannels(), (int32)MAX_TEXCOORDS);
+	SkelMeshImportData.bHasNormals = bHaveValidNormals;
+	SkelMeshImportData.bHasTangents = bHaveValidTangents;
 	for (FTriangleID TriangleID : InMeshDescription.Triangles().GetElementIDs())
 	{
 		FPolygonGroupID PolygonGroupID = InMeshDescription.GetTrianglePolygonGroup(TriangleID);
@@ -2722,18 +2782,24 @@ FSkeletalMeshImportData FSkeletalMeshImportData::CreateFromMeshDescription(const
 			{
 				Wedge.UVs[UVChannelIndex] = VertexInstanceUVs.Get(VertexInstanceID, UVChannelIndex);
 			}
-			Face.TangentX[Corner] = VertexInstanceTangents[VertexInstanceID];
-			Face.TangentZ[Corner] = VertexInstanceNormals[VertexInstanceID];
-			Face.TangentY[Corner] = FVector3f::CrossProduct(VertexInstanceNormals[VertexInstanceID], VertexInstanceTangents[VertexInstanceID]).GetSafeNormal() * VertexInstanceBiNormalSigns[VertexInstanceID];
 
-			if (Face.TangentX[Corner] != FVector3f::ZeroVector)
+			if (bHaveValidNormals)
 			{
-				SkelMeshImportData.bHasTangents = true;
+				Face.TangentZ[Corner] = VertexInstanceNormals[VertexInstanceID];
 			}
-			
-			if (Face.TangentZ[Corner] != FVector3f::ZeroVector)
+			else
+			{
+				Face.TangentZ[Corner] = FVector3f::ZeroVector;
+			}
+
+			if (bHaveValidTangents && bHaveValidNormals)
+			{
+				Face.TangentX[Corner] = VertexInstanceTangents[VertexInstanceID];
+				Face.TangentY[Corner] = FVector3f::CrossProduct(VertexInstanceNormals[VertexInstanceID], VertexInstanceTangents[VertexInstanceID]).GetSafeNormal() * VertexInstanceBiNormalSigns[VertexInstanceID];
+			}
+			else
 			{
-				SkelMeshImportData.bHasNormals = true;
+				Face.TangentX[Corner] = FVector3f::ZeroVector;
 			}
 
 			Face.WedgeIndex[Corner] = SkelMeshImportData.Wedges.Add(Wedge);
