Reduce the maximum size of Iris active objects on the client to reduce memory.

#localization none

diff --git a/Engine/Source/Runtime/Engine/Classes/Engine/NetDriver.h b/Engine/Source/Runtime/Engine/Classes/Engine/NetDriver.h
--- a/Engine/Source/Runtime/Engine/Classes/Engine/NetDriver.h
+++ b/Engine/Source/Runtime/Engine/Classes/Engine/NetDriver.h
@@ -627,9 +627,13 @@ struct FNetDriverReplicationSystemConfig
 {
 	GENERATED_USTRUCT_BODY()
 
-	/** Override the max object count. If 0 use the default system value. */
+	/** Override the max object count when running as a client. If 0 use the default system value. */
 	UPROPERTY()
-	uint32 MaxReplicatedObjectCount = 0;
+	uint32 MaxReplicatedObjectClientCount = 0;
+
+	/** Override the max object count when running as a server. If 0 use the default system value. */
+	UPROPERTY()
+	uint32 MaxReplicatedObjectServerCount = 0;
 	
 	/** Override the max compressed object count. If 0 use the default system value. */
 	UPROPERTY()
diff --git a/Engine/Source/Runtime/Engine/Private/NetDriver.cpp b/Engine/Source/Runtime/Engine/Private/NetDriver.cpp
--- a/Engine/Source/Runtime/Engine/Private/NetDriver.cpp
+++ b/Engine/Source/Runtime/Engine/Private/NetDriver.cpp
@@ -247,11 +247,21 @@ namespace UE::Net::Private
 	static bool bIgnoreStaticActorDestruction = false;
 
 
-	void ApplyReplicationSystemConfig(const FNetDriverReplicationSystemConfig& ReplicationSystemConfig, UReplicationSystem::FReplicationSystemParams& OutParams)
+	void ApplyReplicationSystemConfig(const FNetDriverReplicationSystemConfig& ReplicationSystemConfig, UReplicationSystem::FReplicationSystemParams& OutParams, bool bIsServer)
 	{
-		if (ReplicationSystemConfig.MaxReplicatedObjectCount != 0)
+		if (bIsServer)
+		{
+			if (ReplicationSystemConfig.MaxReplicatedObjectServerCount != 0)
+			{
+				OutParams.MaxReplicatedObjectCount = ReplicationSystemConfig.MaxReplicatedObjectServerCount;
+			}
+		}
+		else
 		{
-			OutParams.MaxReplicatedObjectCount = ReplicationSystemConfig.MaxReplicatedObjectCount;
+			if (ReplicationSystemConfig.MaxReplicatedObjectClientCount != 0)
+			{
+				OutParams.MaxReplicatedObjectCount = ReplicationSystemConfig.MaxReplicatedObjectClientCount;
+			}
 		}
 
 		if (ReplicationSystemConfig.MaxDeltaCompressedObjectCount != 0)
@@ -6788,7 +6798,7 @@ void UNetDriver::CreateReplicationSystem(bool bInitAsClient)
 		Params.bAllowObjectReplication = !bInitAsClient;
 		Params.ForwardNetRPCCallDelegate.BindUObject(this, &UNetDriver::ForwardRemoteFunction);
 
-		UE::Net::Private::ApplyReplicationSystemConfig(ReplicationSystemConfig, Params);
+		UE::Net::Private::ApplyReplicationSystemConfig(ReplicationSystemConfig, Params, !bInitAsClient);
 
 		SetReplicationSystem(UE::Net::FReplicationSystemFactory::CreateReplicationSystem(Params));
 	}
diff --git a/Engine/Source/Runtime/Experimental/Iris/Core/Private/Iris/ReplicationSystem/NetRefHandleManager.cpp b/Engine/Source/Runtime/Experimental/Iris/Core/Private/Iris/ReplicationSystem/NetRefHandleManager.cpp
--- a/Engine/Source/Runtime/Experimental/Iris/Core/Private/Iris/ReplicationSystem/NetRefHandleManager.cpp
+++ b/Engine/Source/Runtime/Experimental/Iris/Core/Private/Iris/ReplicationSystem/NetRefHandleManager.cpp
@@ -77,6 +77,8 @@ FInternalNetRefIndex FNetRefHandleManager::InternalCreateNetObject(const FNetRef
 {
 	if (ActiveObjectCount >= MaxActiveObjectCount)
 	{
+		UE_LOG(LogIris, Error, TEXT("NetRefHandleManager: Maximum active object count reached (%d/%d)."), ActiveObjectCount, MaxActiveObjectCount);
+		ensureMsgf(false, TEXT("NetRefHandleManager: Maximum active object count reached (%d/%d)."), ActiveObjectCount, MaxActiveObjectCount);
 		return InvalidInternalIndex;
 	}
 
@@ -133,6 +135,11 @@ FInternalNetRefIndex FNetRefHandleManager::InternalCreateNetObject(const FNetRef
 
 		return InternalIndex;
 	}
+	else
+	{
+		UE_LOG(LogIris, Error, TEXT("NetRefHandleManager: Cannot find free active object (%d/%d)."), ActiveObjectCount, MaxActiveObjectCount);
+		ensureMsgf(false, TEXT("NetRefHandleManager: Cannot find free active object (%d/%d)."), ActiveObjectCount, MaxActiveObjectCount);
+	}
 		
 	return InvalidInternalIndex;
 }
