[UBA] Report highest number of agents and cores in shader compilation stats.

This adds more information about the distributed shader compiler to the compile stats. As an example from the UBA default settings when launching Lyra:
"LogShaderCompilers: Display: Highest number of remote agents active in parallel: 15 (480 active cores peak)"

#rnx

diff --git a/Engine/Plugins/UbaController/Source/UbaController/Private/UbaControllerModule.cpp b/Engine/Plugins/UbaController/Source/UbaController/Private/UbaControllerModule.cpp
--- a/Engine/Plugins/UbaController/Source/UbaController/Private/UbaControllerModule.cpp
+++ b/Engine/Plugins/UbaController/Source/UbaController/Private/UbaControllerModule.cpp
@@ -212,6 +212,11 @@ TFuture<FDistributedBuildTaskResult> FUbaControllerModule::EnqueueTask(const FTa
 	return MoveTemp(Future);
 }
 
+bool FUbaControllerModule::PollStats(FDistributedBuildStats& OutStats)
+{
+	return JobDispatcherThread != nullptr && JobDispatcherThread->PollStats(OutStats);
+}
+
 void FUbaControllerModule::ReportJobProcessed(const FTaskResponse& InTaskResponse, FTask* CompileTask)
 {
 	if (CompileTask)
diff --git a/Engine/Plugins/UbaController/Source/UbaController/Private/UbaHordeAgentManager.cpp b/Engine/Plugins/UbaController/Source/UbaController/Private/UbaHordeAgentManager.cpp
--- a/Engine/Plugins/UbaController/Source/UbaController/Private/UbaHordeAgentManager.cpp
+++ b/Engine/Plugins/UbaController/Source/UbaController/Private/UbaHordeAgentManager.cpp
@@ -37,6 +37,7 @@ FUbaHordeAgentManager::FUbaHordeAgentManager(const FString& InWorkingDir, uba::N
 	,	LastRequestFailTime(1)
 	,	TargetCoreCount(0)
 	,	EstimatedCoreCount(0)
+	,   ActiveCoreCount(0)
 	,	AskForAgents(true)
 {
 	ParseConfig();
@@ -86,6 +87,17 @@ void FUbaHordeAgentManager::SetTargetCoreCount(uint32 Count)
 	}
 }
 
+int32 FUbaHordeAgentManager::GetAgentCount() const
+{
+	FScopeLock AgentsScopeLock(&AgentsLock);
+	return Agents.Num();
+}
+
+uint32 FUbaHordeAgentManager::GetActiveCoreCount() const
+{
+	return ActiveCoreCount;
+}
+
 // Creates a bundle blob (one of several chunks of a file) to be uploaded to Horde
 // This code has been adopted from the HordeTest project.
 // See 'Engine/Source/Programs/Horde/Samples/HordeTest/Main.cpp'.
diff --git a/Engine/Plugins/UbaController/Source/UbaController/Private/UbaJobProcessor.cpp b/Engine/Plugins/UbaController/Source/UbaController/Private/UbaJobProcessor.cpp
--- a/Engine/Plugins/UbaController/Source/UbaController/Private/UbaJobProcessor.cpp
+++ b/Engine/Plugins/UbaController/Source/UbaController/Private/UbaJobProcessor.cpp
@@ -407,6 +407,8 @@ uint32 FUbaJobProcessor::Run()
 			
 			// TODO: Not sure this is a good idea in a cooking scenario where number of queued processes are going up and down
 			SessionServer_SetMaxRemoteProcessCount(UbaSessionServer, TargetCoreCount);
+
+			UpdateStats();
 		}
 
 		lock.Unlock();
@@ -522,4 +524,22 @@ bool FUbaJobProcessor::HasJobsInFlight() const
 	uint32 finished = 0;
 	Scheduler_GetStats(UbaScheduler, queued, activeLocal, activeRemote, finished);	
 	return (queued + activeLocal + activeRemote) != 0;
-}
\ No newline at end of file
+}
+
+bool FUbaJobProcessor::PollStats(FDistributedBuildStats& OutStats)
+{
+	// Return current stats and reset internal data
+	FScopeLock StatsLockGuard(&StatsLock);
+	OutStats = Stats;
+	Stats = FDistributedBuildStats();
+	return true;
+}
+
+void FUbaJobProcessor::UpdateStats()
+{
+	FScopeLock StatsLockGuard(&StatsLock);
+
+	// Update maximum
+	Stats.MaxRemoteAgents = FMath::Max(Stats.MaxRemoteAgents, (uint32)HordeAgentManager->GetAgentCount());
+	Stats.MaxActiveAgentCores = FMath::Max(Stats.MaxActiveAgentCores, HordeAgentManager->GetActiveCoreCount());
+}
diff --git a/Engine/Plugins/UbaController/Source/UbaController/Public/UbaControllerModule.h b/Engine/Plugins/UbaController/Source/UbaController/Public/UbaControllerModule.h
--- a/Engine/Plugins/UbaController/Source/UbaController/Public/UbaControllerModule.h
+++ b/Engine/Plugins/UbaController/Source/UbaController/Public/UbaControllerModule.h
@@ -26,7 +26,8 @@ public:
 	virtual void InitializeController() override final;
 	virtual FString CreateUniqueFilePath() override final;
 	virtual TFuture<FDistributedBuildTaskResult> EnqueueTask(const FTaskCommandData& CommandData) override final;
-	
+	virtual bool PollStats(FDistributedBuildStats& OutStats) override;
+
 	void ReportJobProcessed(const FTaskResponse& InTaskResponse, FTask* CompileTask);
 	void CleanWorkingDirectory() const;
 
diff --git a/Engine/Plugins/UbaController/Source/UbaController/Public/UbaHordeAgentManager.h b/Engine/Plugins/UbaController/Source/UbaController/Public/UbaHordeAgentManager.h
--- a/Engine/Plugins/UbaController/Source/UbaController/Public/UbaHordeAgentManager.h
+++ b/Engine/Plugins/UbaController/Source/UbaController/Public/UbaHordeAgentManager.h
@@ -16,6 +16,12 @@ public:
 
 	void SetTargetCoreCount(uint32 Count);
 
+	// Returns the number of agents currently handled by this agent manager.
+	int32 GetAgentCount() const;
+
+	// Returns the active number of cores allocated across all agents.
+	uint32 GetActiveCoreCount() const;
+
 private:
 	struct FHordeAgentWrapper
 	{
@@ -40,7 +46,7 @@ private:
 	FCriticalSection UbaAgentBundleFilePathLock;
 	FString UbaAgentBundleFilePath;
 
-	FCriticalSection AgentsLock;
+	mutable FCriticalSection AgentsLock;
 	TArray<TUniquePtr<FHordeAgentWrapper>> Agents;
 
 	TAtomic<uint64> LastRequestFailTime;
diff --git a/Engine/Plugins/UbaController/Source/UbaController/Public/UbaJobProcessor.h b/Engine/Plugins/UbaController/Source/UbaController/Public/UbaJobProcessor.h
--- a/Engine/Plugins/UbaController/Source/UbaController/Public/UbaJobProcessor.h
+++ b/Engine/Plugins/UbaController/Source/UbaController/Public/UbaJobProcessor.h
@@ -6,6 +6,7 @@
 #include "Templates/SharedPointer.h"
 #include "Templates/UniquePtr.h"
 #include "UbaExports.h"
+#include "DistributedBuildControllerInterface.h"
 
 class FUbaControllerModule;
 class FUbaHordeAgentManager;
@@ -37,6 +38,8 @@ public:
 
 	bool HasJobsInFlight() const;
 
+	bool PollStats(FDistributedBuildStats& OutStats);
+
 private:
 
 	void CalculateKnownInputs();
@@ -45,6 +48,8 @@ private:
 	void StartUba();
 	void ShutDownUba();
 
+	void UpdateStats();
+
 	/** The runnable thread */
 	FRunnableThread* Thread;
 	
@@ -75,4 +80,7 @@ private:
 	TArray<uba::tchar> KnownInputsBuffer;
 
 	uba::CallbackLogWriter LogWriter;
+
+	FCriticalSection StatsLock;
+	FDistributedBuildStats Stats;
 };
\ No newline at end of file
diff --git a/Engine/Source/Developer/DistributedBuildInterface/Public/DistributedBuildControllerInterface.h b/Engine/Source/Developer/DistributedBuildInterface/Public/DistributedBuildControllerInterface.h
--- a/Engine/Source/Developer/DistributedBuildInterface/Public/DistributedBuildControllerInterface.h
+++ b/Engine/Source/Developer/DistributedBuildInterface/Public/DistributedBuildControllerInterface.h
@@ -12,6 +12,12 @@ struct FDistributedBuildTaskResult
 	bool bCompleted;
 };
 
+struct FDistributedBuildStats
+{
+	uint32 MaxRemoteAgents = 0;
+	uint32 MaxActiveAgentCores = 0;
+};
+
 struct FTaskCommandData
 {	
 	FString Command;
@@ -64,6 +70,9 @@ public:
 	// Returns a new file path to be used for writing input data to.
 	virtual FString CreateUniqueFilePath() = 0;
 
+	// Returns the distributed build statistics since the last call and resets its internal values. Returns false if there are no statistics provided.
+	virtual bool PollStats(FDistributedBuildStats& OutStats) { return false; }
+
 	// Launches a task. Returns a future which can be waited on for the results.
 	virtual TFuture<FDistributedBuildTaskResult> EnqueueTask(const FTaskCommandData& CommandData) = 0;
 	
diff --git a/Engine/Source/Runtime/Engine/Private/ShaderCompiler/ShaderCompiler.cpp b/Engine/Source/Runtime/Engine/Private/ShaderCompiler/ShaderCompiler.cpp
--- a/Engine/Source/Runtime/Engine/Private/ShaderCompiler/ShaderCompiler.cpp
+++ b/Engine/Source/Runtime/Engine/Private/ShaderCompiler/ShaderCompiler.cpp
@@ -4720,6 +4720,11 @@ void FShaderCompilerStats::WriteStatSummary()
 		UE_LOG(LogShaderCompilers, Display, TEXT("Total thread preprocess time: %s s"), *FormatNumber(TotalThreadPreprocessTimeForAllShaders));
 		UE_LOG(LogShaderCompilers, Display, TEXT("Percentage time preprocessing: %.2f%%"), TotalThreadTimeForAllShaders > 0.0 ? (TotalThreadPreprocessTimeForAllShaders / TotalThreadTimeForAllShaders) * 100.0 : 0.0);
 
+		if (Counters.MaxRemoteAgents > 0)
+		{
+			UE_LOG(LogShaderCompilers, Display, TEXT("Highest number of remote agents active in parallel: %u (%u active cores peak)"), Counters.MaxRemoteAgents, Counters.MaxActiveAgentCores);
+		}
+
 		if (TotalTimeAtLeastOneJobWasInFlight > 0.0)
 		{
 			double EffectiveParallelization = TotalThreadTimeForAllShaders / TotalTimeAtLeastOneJobWasInFlight;
@@ -5274,6 +5279,13 @@ void FShaderCompilerStats::RegisterJobBatch(int32 NumJobs, EExecutionType ExecTy
 	}
 }
 
+void FShaderCompilerStats::RegisterDistributedBuildStats(const FDistributedBuildStats& InStats)
+{
+	FScopeLock Lock(&CompileStatsLock);
+	Counters.MaxRemoteAgents = FMath::Max(Counters.MaxRemoteAgents, InStats.MaxRemoteAgents);
+	Counters.MaxActiveAgentCores = FMath::Max(Counters.MaxActiveAgentCores, InStats.MaxActiveAgentCores);
+}
+
 void FShaderCompilerStats::FMaterialCounters::WriteStatSummary(const TCHAR* AggregatedSuffix)
 {
 	auto CalcTimePercentage = [&](double Val) {
diff --git a/Engine/Source/Runtime/Engine/Private/ShaderCompiler/ShaderCompilerDistributed.cpp b/Engine/Source/Runtime/Engine/Private/ShaderCompiler/ShaderCompilerDistributed.cpp
--- a/Engine/Source/Runtime/Engine/Private/ShaderCompiler/ShaderCompilerDistributed.cpp
+++ b/Engine/Source/Runtime/Engine/Private/ShaderCompiler/ShaderCompilerDistributed.cpp
@@ -102,6 +102,12 @@ void FShaderCompileDistributedThreadRunnable_Interface::DispatchShaderCompileJob
 			MoveTemp(OutputFilePath)
 		)
 	);
+
+	FDistributedBuildStats Stats;
+	if (CachedController.PollStats(Stats))
+	{
+		GShaderCompilerStats->RegisterDistributedBuildStats(Stats);
+	}
 }
 
 TArray<FString> FShaderCompileDistributedThreadRunnable_Interface::GetDependencyFilesForJobs(
diff --git a/Engine/Source/Runtime/Engine/Public/ShaderCompiler.h b/Engine/Source/Runtime/Engine/Public/ShaderCompiler.h
--- a/Engine/Source/Runtime/Engine/Public/ShaderCompiler.h
+++ b/Engine/Source/Runtime/Engine/Public/ShaderCompiler.h
@@ -308,6 +308,8 @@ private:
 /** Results for a single compiled and finalized shader map. */
 using FShaderMapFinalizeResults = FShaderMapCompileResults;
 
+struct FDistributedBuildStats;
+
 class FShaderCompilerStats
 {
 public:
@@ -435,6 +437,9 @@ public:
 	/** Informs statistics about a new job batch, so we can tally up batches. */
 	void RegisterJobBatch(int32 NumJobs, EExecutionType ExecType);
 
+	/** Informs about current distributed build statistics. */
+	void RegisterDistributedBuildStats(const FDistributedBuildStats& InStats);
+
 	ENGINE_API void GatherAnalytics(const FString& BaseName, TArray<FAnalyticsEventAttribute>& Attributes);
 
 private:
@@ -540,6 +545,12 @@ private:
 		/** Memory budget allocated for the job cache */
 		uint64 CacheMemBudget = 0;
 
+		/** Maximum number of remote agents used during compilation. */
+		uint32 MaxRemoteAgents = 0;
+
+		/** Maximum number of CPU cores active across all remote agents. */
+		uint32 MaxActiveAgentCores = 0;
+
 		FCounters& operator+=(const FCounters& Other)
 		{
 			AccumulatedLocalWorkerIdleTime += Other.AccumulatedLocalWorkerIdleTime;
@@ -576,6 +587,8 @@ private:
 			UniqueCacheOutputs += Other.UniqueCacheOutputs;
 			CacheMemUsed += Other.CacheMemUsed;
 			CacheMemBudget += Other.CacheMemBudget;
+			MaxRemoteAgents += Other.MaxRemoteAgents;
+			MaxActiveAgentCores += Other.MaxActiveAgentCores;
 
 			return *this;
 		}
