Common Code changes for TP refactor

diff --git a/Engine/Source/Developer/TargetPlatform/Public/Common/TargetPlatformBase.h b/Engine/Source/Developer/TargetPlatform/Public/Common/TargetPlatformBase.h
--- a/Engine/Source/Developer/TargetPlatform/Public/Common/TargetPlatformBase.h
+++ b/Engine/Source/Developer/TargetPlatform/Public/Common/TargetPlatformBase.h
@@ -861,7 +861,7 @@ public:
 	}
 	virtual void GetShaderFormatModuleHints(TArray<FName>& OutModuleNames) const override 
 	{
-		TargetPlatformControls->GetShaderFormatModuleHints(OutModuleNames);
+		TargetPlatformSettings->GetShaderFormatModuleHints(OutModuleNames);
 	}
 	virtual void GetTextureFormats(const class UTexture* Texture, TArray< TArray<FName> >& OutFormats) const override 
 	{
diff --git a/Engine/Source/Developer/TargetPlatform/Public/Common/TargetPlatformControlsBase.h b/Engine/Source/Developer/TargetPlatform/Public/Common/TargetPlatformControlsBase.h
--- a/Engine/Source/Developer/TargetPlatform/Public/Common/TargetPlatformControlsBase.h
+++ b/Engine/Source/Developer/TargetPlatform/Public/Common/TargetPlatformControlsBase.h
@@ -158,10 +158,6 @@ public:
 		return NAME_MeshBuilder;
 	}
 
-	virtual void GetShaderFormatModuleHints(TArray<FName>& OutModuleNames) const override
-	{
-	}
-
 	virtual void GetTextureFormatModuleHints(TArray<FName>& OutModuleNames) const override
 	{
 		// these are the default texture format modules, since many platforms 
@@ -249,6 +245,22 @@ public:
 		check(TPlatformProperties::HasEditorOnlyData() != TPlatformProperties::RequiresCookedData());
 	}
 
+	/**
+	 * Constructor that makes a TPCI based solely on TPlatformProperties
+	 */
+	TTargetPlatformControlsBase(ITargetPlatformSettings* TargetPlatformSettings)
+		: TTargetPlatformControlsBase(new PlatformInfo::FTargetPlatformInfo(
+			TPlatformProperties::IniPlatformName(),
+			TPlatformProperties::HasEditorOnlyData() ? EBuildTargetType::Editor :
+			TPlatformProperties::IsServerOnly() ? EBuildTargetType::Server :
+			TPlatformProperties::IsClientOnly() ? EBuildTargetType::Client :
+			EBuildTargetType::Game,
+			TEXT("")),
+			TargetPlatformSettings
+		)
+	{
+	}
+
 	virtual bool HasEditorOnlyData() const override
 	{
 		return TPlatformProperties::HasEditorOnlyData();
diff --git a/Engine/Source/Developer/TargetPlatform/Public/Common/TargetPlatformSettingsBase.h b/Engine/Source/Developer/TargetPlatform/Public/Common/TargetPlatformSettingsBase.h
--- a/Engine/Source/Developer/TargetPlatform/Public/Common/TargetPlatformSettingsBase.h
+++ b/Engine/Source/Developer/TargetPlatform/Public/Common/TargetPlatformSettingsBase.h
@@ -38,6 +38,10 @@ public:
 	TARGETPLATFORM_API virtual void GetRayTracingShaderFormats(TArray<FName>& OutFormats) const override;
 
 #if WITH_ENGINE
+	virtual void GetShaderFormatModuleHints(TArray<FName>& OutModuleNames) const override
+	{
+	}
+
 	virtual void GetReflectionCaptureFormats(TArray<FName>& OutFormats) const override
 	{
 		OutFormats.Add(FName(TEXT("FullHDR")));
diff --git a/Engine/Source/Developer/TargetPlatform/Public/Interfaces/ITargetPlatformControls.h b/Engine/Source/Developer/TargetPlatform/Public/Interfaces/ITargetPlatformControls.h
--- a/Engine/Source/Developer/TargetPlatform/Public/Interfaces/ITargetPlatformControls.h
+++ b/Engine/Source/Developer/TargetPlatform/Public/Interfaces/ITargetPlatformControls.h
@@ -337,13 +337,6 @@ public:
 	 */
 	virtual FName GetPhysicsFormat(class UBodySetup* Body) const = 0;
 
-	/**
-	 * Gets a list of modules that may contain the GetAllTargetedShaderFormats. This is optional -
-	 * if any required shader format isn't found in this list, then it will use the old path
-	 * of loading all shader format modules to gather all available shader formats
-	 */
-	virtual void GetShaderFormatModuleHints(TArray<FName>& OutModuleNames) const = 0;
-
 	/**
 	 * Gets the texture format to use for each layer in the given texture, for each of the platform's formats.
 	 * _Most_ platforms only supply one format for a given texture, so OutFormats.Num() is usually 1. The exception is Android_Multi,
diff --git a/Engine/Source/Developer/TargetPlatform/Public/Interfaces/ITargetPlatformSettings.h b/Engine/Source/Developer/TargetPlatform/Public/Interfaces/ITargetPlatformSettings.h
--- a/Engine/Source/Developer/TargetPlatform/Public/Interfaces/ITargetPlatformSettings.h
+++ b/Engine/Source/Developer/TargetPlatform/Public/Interfaces/ITargetPlatformSettings.h
@@ -241,6 +241,13 @@ public:
 	virtual bool SupportsFeature(ETargetPlatformFeatures Feature) const = 0;
 
 #if WITH_ENGINE
+	/**
+	 * Gets a list of modules that may contain the GetAllTargetedShaderFormats. This is optional -
+	 * if any required shader format isn't found in this list, then it will use the old path
+	 * of loading all shader format modules to gather all available shader formats
+	 */
+	virtual void GetShaderFormatModuleHints(TArray<FName>& OutModuleNames) const = 0;
+
 	/**
 	 * Gets the reflection capture formats this platform needs.
 	 *
diff --git a/Engine/Source/Runtime/Engine/Classes/Engine/Texture.h b/Engine/Source/Runtime/Engine/Classes/Engine/Texture.h
--- a/Engine/Source/Runtime/Engine/Classes/Engine/Texture.h
+++ b/Engine/Source/Runtime/Engine/Classes/Engine/Texture.h
@@ -1929,10 +1929,15 @@ public:
 	 * Checks whether this texture should be tiled to a platform-specific format during cook, or whether the bNotOfflineProcessed flag 
 	 * should be set to true at runtime because it has not been tiled at cook
 	 * 
-	 * @param  TargetPlatform	The platform for which the texture is being cooked and texture group info will be extracted from. 
-	 *                          If null, this info will be extracted from UDeviceProfileManager::Get().GetActiveProfile(), possibly at runtime
+	 * @param  TargetPlatformSettings	The platform for which the texture is being cooked and texture group info will be extracted from. 
+	 *									If null, this info will be extracted from UDeviceProfileManager::Get().GetActiveProfile(), possibly at runtime
 	 * @return true if platform tiling during cook is disabled for this texture 
 	 */
+	ENGINE_API bool IsCookPlatformTilingDisabled(const ITargetPlatformSettings* TargetPlatformSettings) const;
+
+	/**
+	 * Legacy function to support the old API, will call IsCookPlatformTilingDisabled with ITargetPlatformSettings
+	 */
 	ENGINE_API bool IsCookPlatformTilingDisabled(const ITargetPlatform* TargetPlatform) const;
 
 	/**
diff --git a/Engine/Source/Runtime/Engine/Private/Texture.cpp b/Engine/Source/Runtime/Engine/Private/Texture.cpp
--- a/Engine/Source/Runtime/Engine/Private/Texture.cpp
+++ b/Engine/Source/Runtime/Engine/Private/Texture.cpp
@@ -1225,7 +1225,7 @@ void UTexture::PostLoad()
 
 #endif
 
-	if (IsCookPlatformTilingDisabled(nullptr)) // nullptr TargetPlatform means it will use UDeviceProfileManager::Get().GetActiveProfile() to get the tiling settings
+	if (IsCookPlatformTilingDisabled(static_cast<ITargetPlatformSettings*>(nullptr))) // nullptr TargetPlatform means it will use UDeviceProfileManager::Get().GetActiveProfile() to get the tiling settings
 	{
 		// The texture was not processed/tiled during cook, so it has to be tiled when uploaded to the GPU if necessary
 		bNotOfflineProcessed = true;
@@ -1516,15 +1516,15 @@ TextureMipGenSettings UTexture::GetMipGenSettingsFromString(const TCHAR* InStr,
 	return bTextureGroup ? TMGS_SimpleAverage : TMGS_FromTextureGroup;
 }
 
-bool UTexture::IsCookPlatformTilingDisabled(const ITargetPlatform* TargetPlatform) const
+bool UTexture::IsCookPlatformTilingDisabled(const ITargetPlatformSettings* TargetPlatformSettings) const
 {
 	if (CookPlatformTilingSettings.GetValue() == TextureCookPlatformTilingSettings::TCPTS_FromTextureGroup)
 	{
 		const UTextureLODSettings* TextureLODSettings = nullptr;
 
-		if (TargetPlatform)
+		if (TargetPlatformSettings)
 		{
-			TextureLODSettings = &TargetPlatform->GetTextureLODSettings();
+			TextureLODSettings = &TargetPlatformSettings->GetTextureLODSettings();
 		}
 		else
 		{
@@ -1548,6 +1548,11 @@ bool UTexture::IsCookPlatformTilingDisabled(const ITargetPlatform* TargetPlatfor
 	return CookPlatformTilingSettings.GetValue() == TextureCookPlatformTilingSettings::TCPTS_DoNotTile;
 }
 
+bool UTexture::IsCookPlatformTilingDisabled(const ITargetPlatform* TargetPlatform) const
+{
+	return IsCookPlatformTilingDisabled(TargetPlatform ? TargetPlatform->GetTargetPlatformSettings(): static_cast<ITargetPlatformSettings*>(nullptr));
+}
+
 void UTexture::SetDeterministicLightingGuid()
 {
 #if WITH_EDITORONLY_DATA
