[Audio Stream Cache] Fix for Mac regression when checking for USoundWave::LoadingBehavior being ForceInline.
#rnx

diff --git a/Engine/Source/Runtime/Engine/Private/SoundWave.cpp b/Engine/Source/Runtime/Engine/Private/SoundWave.cpp
--- a/Engine/Source/Runtime/Engine/Private/SoundWave.cpp
+++ b/Engine/Source/Runtime/Engine/Private/SoundWave.cpp
@@ -1914,7 +1914,10 @@ void USoundWave::FlushAudioRenderingCommands() const
 bool USoundWave::HasStreamingChunks()
 {
 	check(SoundWaveDataPtr);
-	const bool bIsForceInline = (GetLoadingBehavior() == ESoundWaveLoadingBehavior::ForceInline);
+
+	// Only checking to see if it is set to ForceInline. ForceInline is not supported on 
+	// USoundClasses, so it is safe to ignore USoundClasses when calling `GetLoadingBehavior(...)`
+	const bool bIsForceInline = (GetLoadingBehavior(false /* bCheckSoundClasses */) == ESoundWaveLoadingBehavior::ForceInline);
 	const bool bHasChunks = (SoundWaveDataPtr->RunningPlatformData.GetNumChunks() > 0);
 	return !bIsForceInline && bHasChunks;
 }
@@ -4583,8 +4586,7 @@ FSoundWaveProxy::FSoundWaveProxy(USoundWave* InWave)
 	// non-streaming sources need resource data initialized before the FSoundWaveProxy
 	// can be used. 
 	InWave->InitAudioResource(SoundWaveDataPtr->GetRuntimeFormat());
-	check((InWave->IsStreaming(nullptr)) || (SoundWaveDataPtr->GetResourceSize() > 0));
-}
+	check((InWave->IsStreaming(nullptr)) || (SoundWaveDataPtr->GetResourceSize() > 0));}
 
 FSoundWaveProxy::~FSoundWaveProxy()
 {
