PR #10749: Optimization for USoundWaveProcedural::QueueAudio (MoveTemp)
#rnx

diff --git a/Engine/Source/Runtime/Engine/Private/SoundWaveProcedural.cpp b/Engine/Source/Runtime/Engine/Private/SoundWaveProcedural.cpp
--- a/Engine/Source/Runtime/Engine/Private/SoundWaveProcedural.cpp
+++ b/Engine/Source/Runtime/Engine/Private/SoundWaveProcedural.cpp
@@ -40,7 +40,7 @@ void USoundWaveProcedural::QueueAudio(const uint8* AudioData, const int32 Buffer
 	TArray<uint8> NewAudioBuffer;
 	NewAudioBuffer.AddUninitialized(BufferSize);
 	FMemory::Memcpy(NewAudioBuffer.GetData(), AudioData, BufferSize);
-	QueuedAudio.Enqueue(NewAudioBuffer);
+	QueuedAudio.Enqueue(MoveTemp(NewAudioBuffer));
 
 	AvailableByteCount.Add(BufferSize);
 }
