fix crash on undo/redo of collision update when render data has also been updated and needs rebuild by triggering the render rebuild slightly earlier in the process

diff --git a/Engine/Source/Runtime/Engine/Private/PhysicsEngine/BodySetup.cpp b/Engine/Source/Runtime/Engine/Private/PhysicsEngine/BodySetup.cpp
--- a/Engine/Source/Runtime/Engine/Private/PhysicsEngine/BodySetup.cpp
+++ b/Engine/Source/Runtime/Engine/Private/PhysicsEngine/BodySetup.cpp
@@ -1391,7 +1391,7 @@ FByteBulkData* UBodySetup::GetCookedData(FName Format)
 #endif
 
 	bool bContainedData = UseCookedData->Contains(Format);
-	FByteBulkData* Result = &UseCookedData->GetFormat(Format);
+	FByteBulkData* Result = nullptr;
 	bool bIsRuntime = IsRuntime(this);
 
 #if WITH_EDITOR
@@ -1399,7 +1399,9 @@ FByteBulkData* UBodySetup::GetCookedData(FName Format)
 	{
 		SCOPE_CYCLE_COUNTER(STAT_PhysXCooking);
 
-		if (AggGeom.ConvexElems.Num() == 0 && (CDP == nullptr || CDP->ContainsPhysicsTriMeshData(bMeshCollideAll) == false))
+		// Note: Check ContainsPhysicsTriMeshData before looking at the number of convex elems, to ensure the side effects of ContainsPhysicsTriMeshData happen
+		// (specifically, for static mesh this will ensure the mesh render data is already built)
+		if ((CDP == nullptr || CDP->ContainsPhysicsTriMeshData(bMeshCollideAll) == false) && AggGeom.ConvexElems.Num() == 0)
 		{
 			return nullptr;
 		}
@@ -1410,9 +1412,14 @@ FByteBulkData* UBodySetup::GetCookedData(FName Format)
 		const bool bUseRefHolder = false;
 		FChaosDerivedDataCooker* PhysicsDerivedCooker = new FChaosDerivedDataCooker(this, Format, bUseRefHolder);
 
+		Result = &UseCookedData->GetFormat(Format);
 		GetDDCBuiltData(Result, *PhysicsDerivedCooker, this, bIsRuntime);
 	}
+	else
 #endif // #if WITH_EDITOR
+	{
+		Result = &UseCookedData->GetFormat(Format);
+	}
 
 	check(Result);
 	return Result->GetBulkDataSize() > 0 ? Result : nullptr; // we don't return empty bulk data...but we save it to avoid thrashing the DDC
