Shader Compiler:  Added some more verbose logging and assert messages in an effort to track down a bug (cache record not found).  The goal is to see if the crash follows memory cleanup or job cancellation, and whether the crash is for items found in the DDC versus from a job.

#rnx

diff --git a/Engine/Source/Runtime/Engine/Private/ShaderCompiler/ShaderCompiler.cpp b/Engine/Source/Runtime/Engine/Private/ShaderCompiler/ShaderCompiler.cpp
--- a/Engine/Source/Runtime/Engine/Private/ShaderCompiler/ShaderCompiler.cpp
+++ b/Engine/Source/Runtime/Engine/Private/ShaderCompiler/ShaderCompiler.cpp
@@ -397,6 +397,9 @@ struct FShaderJobData
 	/** Output hash will be zero if output data has not been written yet, or can be cleared if output data has been removed */
 	FJobOutputHash OutputHash;
 
+	/** Track which code path wrote this output, for tracking down a bug */
+	bool bOutputFromDDC;
+
 	/**
 	 * In-flight job with the given input hash.  Needs to be a reference pointer to handle cancelling of jobs, where an async DDC query
 	 * (which receives a pointer to FShaderJobData) may be in-flight that still references a job that has otherwise been deleted.
@@ -1606,7 +1609,7 @@ void FShaderJobCache::SubmitJob(FShaderCommonCompileJob* Job)
 		// call to RemoveAllPendingJobsWithId, so we can ignore it and just return.
 		if (!Job->PrevLink)
 		{
-			UE_LOG(LogShaderCompilers, Log, TEXT("Cancelled job 0x%p with pending SubmitJob call."), Job);
+			UE_LOG(LogShaderCompilers, Display, TEXT("Cancelled job 0x%p with pending SubmitJob call."), Job);
 
 			JobLock.WriteUnlock();
 			return;
@@ -10715,7 +10718,8 @@ FShaderJobCacheRef FShaderJobCache::FindOrAdd(const FJobInputHash& Hash, EShader
 
 		FStoredOutput** CannedOutput = Outputs.Find(JobData.OutputHash);
 		// we should not allow a dangling input to output mapping to exist
-		checkf(CannedOutput != nullptr, TEXT("Inconsistency in FShaderJobCache - cache record for ihash %s exists, but output cannot be found."), *LexToString(Hash));
+		checkf(CannedOutput != nullptr, TEXT("Inconsistency in FShaderJobCache - cache record for ihash %s (data 0x%p) exists, but output %s (%s) cannot be found."),
+			*LexToString(Hash), &JobData, *LexToString(JobData.OutputHash), JobData.bOutputFromDDC ? TEXT("DDC") : TEXT("Job"));
 		// update the output hit count
 		(*CannedOutput)->NumHits++;
 
@@ -10806,7 +10810,7 @@ FShaderJobCacheRef FShaderJobCache::FindOrAdd(const FJobInputHash& Hash, EShader
 							// If job was cancelled, it will have been unlinked from PendingSubmitJobTaskJobs, and we can ignore the results.
 							if (!JobDataPtr->JobInFlight->PrevLink)
 							{
-								UE_LOG(LogShaderCompilers, Log, TEXT("Cancelled job 0x%p with pending DDC hit."), JobDataPtr->JobInFlight.GetReference());
+								UE_LOG(LogShaderCompilers, Display, TEXT("Cancelled job 0x%p (data 0x%p) with pending DDC hit."), JobDataPtr->JobInFlight.GetReference(), JobDataPtr);
 
 								delete NewStoredOutput;
 								JobDataPtr->JobInFlight = nullptr;
@@ -10826,6 +10830,7 @@ FShaderJobCacheRef FShaderJobCache::FindOrAdd(const FJobInputHash& Hash, EShader
 						NewStoredOutput->AddRef();
 						Outputs.Add(NewOutputHash, NewStoredOutput);
 						JobDataPtr->OutputHash = NewOutputHash;
+						JobDataPtr->bOutputFromDDC = true;
 
 						CurrentlyAllocatedMemory += NewStoredOutput->GetAllocatedSize();
 
@@ -10887,7 +10892,7 @@ FShaderJobCacheRef FShaderJobCache::FindOrAdd(const FJobInputHash& Hash, EShader
 							// If job was cancelled, it will have been unlinked from PendingSubmitJobTaskJobs, and we can ignore it.
 							if (!Job->PrevLink)
 							{
-								UE_LOG(LogShaderCompilers, Log, TEXT("Cancelled job 0x%p with pending DDC miss."), Job);
+								UE_LOG(LogShaderCompilers, Display, TEXT("Cancelled job 0x%p (data 0x%p) with pending DDC miss."), Job, JobDataPtr);
 
 								JobDataPtr->JobInFlight = nullptr;
 								return;
@@ -11049,6 +11054,7 @@ void FShaderJobCache::AddJobOutput(FShaderJobData& JobData, const FShaderCommonC
 	if (UNLIKELY(bDiscardCacheOutputs == false))
 	{
 		JobData.OutputHash = OutputHash;
+		JobData.bOutputFromDDC = false;
 	}
 
 	FStoredOutput** CannedOutput = Outputs.Find(OutputHash);
@@ -11125,7 +11131,7 @@ void FShaderJobCache::AddJobOutput(FShaderJobData& JobData, const FShaderCommonC
 				// Cull outputs to reach the budget target
 				CullOutputsToMemoryBudget(TargetBudgetBytes);
 
-				UE_LOG(LogShaderCompilers, Log, TEXT("Memory overflow, reduced from %.1lf to %.1lf MB."), (double)MemoryBefore / (1024 * 1024), (double)CurrentlyAllocatedMemory / (1024 * 1024));
+				UE_LOG(LogShaderCompilers, Display, TEXT("Memory overflow, reduced from %.1lf to %.1lf MB."), (double)MemoryBefore / (1024 * 1024), (double)CurrentlyAllocatedMemory / (1024 * 1024));
 			}
 		}
 	}
