PR #9046: Expose audio data compressor to other modules
#rnx

diff --git a/Engine/Source/Runtime/Engine/Public/AudioDerivedData.h b/Engine/Source/Runtime/Engine/Public/AudioDerivedData.h
--- a/Engine/Source/Runtime/Engine/Public/AudioDerivedData.h
+++ b/Engine/Source/Runtime/Engine/Public/AudioDerivedData.h
@@ -16,7 +16,7 @@ private:
 	TUniquePtr<FAudioCookInputs> CookInputs;
 
 public:
-	FDerivedAudioDataCompressor(USoundWave* InSoundNode, FName InBaseFormat, FName InHashedFormat, const FPlatformAudioCookOverrides* InCompressionOverrides, const ITargetPlatform* InTargetPlatform=nullptr);
+	ENGINE_API FDerivedAudioDataCompressor(USoundWave* InSoundNode, FName InBaseFormat, FName InHashedFormat, const FPlatformAudioCookOverrides* InCompressionOverrides, const ITargetPlatform* InTargetPlatform=nullptr);
 
 	virtual const TCHAR* GetPluginName() const override
 	{
