Always create SRV view for FPositionVertexBuffer.

diff --git a/Engine/Source/Runtime/Engine/Private/PositionVertexBuffer.cpp b/Engine/Source/Runtime/Engine/Private/PositionVertexBuffer.cpp
--- a/Engine/Source/Runtime/Engine/Private/PositionVertexBuffer.cpp
+++ b/Engine/Source/Runtime/Engine/Private/PositionVertexBuffer.cpp
@@ -241,23 +241,11 @@ void FPositionVertexBuffer::InitRHI(FRHICommandListBase& RHICmdList)
 	TRACE_CPUPROFILER_EVENT_SCOPE(FPositionVertexBuffer::InitRHI);
 
 	VertexBufferRHI = CreateRHIBuffer(RHICmdList);
-	// we have decide to create the SRV based on GMaxRHIShaderPlatform because this is created once and shared between feature levels for editor preview.
-	// Also check to see whether cpu access has been activated on the vertex data
+
+	// Always create SRV if VertexBufferRHI is valid, as several systems rely on having position SRV. Memory impact has been measured to be minimal.
 	if (VertexBufferRHI)
 	{
-		// we have decide to create the SRV based on GMaxRHIShaderPlatform because this is created once and shared between feature levels for editor preview.
-		bool bSRV = RHISupportsManualVertexFetch(GMaxRHIShaderPlatform) || IsGPUSkinPassThroughSupported(GMaxRHIShaderPlatform);
-
-		// When bAllowCPUAccess is true, the meshes is likely going to be used for Niagara to spawn particles on mesh surface.
-		// And it can be the case for CPU *and* GPU access: no differenciation today. That is why we create a SRV in this case.
-		// This also avoid setting lots of states on all the members of all the different buffers used by meshes. Follow up: https://jira.it.epicgames.net/browse/UE-69376.
-		bSRV |= (VertexData && VertexData->GetAllowCPUAccess());
-		if(bSRV)
-		{
-			// When VertexData is null, this buffer hasn't been streamed in yet. We still need to create a FRHIShaderResourceView which will be
-			// cached in a vertex factory uniform buffer later. The nullptr tells the RHI that the SRV doesn't view on anything yet.
-			PositionComponentSRV = RHICmdList.CreateShaderResourceView(VertexBufferRHI, 4, PF_R32_FLOAT);
-		}
+		PositionComponentSRV = RHICmdList.CreateShaderResourceView(VertexBufferRHI, 4, PF_R32_FLOAT);
 	}
 }
 
