[NetPhysics] Resimulation Improvements

- Runtime correction for rotation improved, faulty calculation fixed and correction not done with pre-resim data.
- Don't apply targets during resimulation, only on rewind.
- Feature to be able to compare predicted states and inputs with the received server states and inputs to trigger resimulation from discrepancies.

diff --git a/Engine/Source/Runtime/Engine/Private/PhysicsEngine/NetworkPhysicsComponent.cpp b/Engine/Source/Runtime/Engine/Private/PhysicsEngine/NetworkPhysicsComponent.cpp
--- a/Engine/Source/Runtime/Engine/Private/PhysicsEngine/NetworkPhysicsComponent.cpp
+++ b/Engine/Source/Runtime/Engine/Private/PhysicsEngine/NetworkPhysicsComponent.cpp
@@ -11,6 +11,7 @@
 #include "Net/UnrealNetwork.h"
 #include "PhysicsReplication.h"
 #include "PhysicsProxy/SingleParticlePhysicsProxy.h"
+#include "Chaos/PhysicsObjectInternalInterface.h"
 
 #if UE_WITH_IRIS
 #include "Iris/ReplicationState/PropertyNetSerializerInfoRegistry.h"
@@ -24,6 +25,10 @@ namespace PhysicsReplicationCVars
 	{
 		static bool bAllowRewindToClosestState = true;
 		static FAutoConsoleVariableRef CVarResimAllowRewindToClosestState(TEXT("np2.Resim.AllowRewindToClosestState"), bAllowRewindToClosestState, TEXT("When rewinding to a specific frame, if the client doens't have state data for that frame, use closest data available. Only affects the first rewind frame, when FPBDRigidsEvolution is set to Reset."));
+		static bool bCompareStateToTriggerRewind = false;
+		static FAutoConsoleVariableRef CVarResimCompareStateToTriggerRewind(TEXT("np2.Resim.CompareStateToTriggerRewind"), bCompareStateToTriggerRewind, TEXT("If we should cache local players custom state struct in rewind history and compare the predicted state with incoming server state to trigger resimulations if they differ, comparison done through FNetworkPhysicsData::CompareData."));
+		static bool bCompareInputToTriggerRewind = false;
+		static FAutoConsoleVariableRef CVarResimCompareInputToTriggerRewind(TEXT("np2.Resim.CompareInputToTriggerRewind"), bCompareInputToTriggerRewind, TEXT("If we should compare local players predicted inputs with incoming server inputs to trigger resimulations if they differ, comparison done through FNetworkPhysicsData::CompareData."));
 	}
 }
 
@@ -387,6 +392,9 @@ UNetworkPhysicsComponent::UNetworkPhysicsComponent() : Super()
 
 void UNetworkPhysicsComponent::InitPhysics()
 {
+	bCompareStateToTriggerRewind = PhysicsReplicationCVars::ResimulationCVars::bCompareStateToTriggerRewind;
+	bCompareInputToTriggerRewind = PhysicsReplicationCVars::ResimulationCVars::bCompareInputToTriggerRewind;
+
 	AActor* Owner = GetOwner();
 	if (Owner)
 	{
@@ -394,6 +402,8 @@ void UNetworkPhysicsComponent::InitPhysics()
 		{
 			InputRedundancy = PhysicsSettings->ResimulationSettings.bOverrideRedundantInputs ? PhysicsSettings->ResimulationSettings.RedundantInputs : InputRedundancy;
 			StateRedundancy = PhysicsSettings->ResimulationSettings.bOverrideRedundantStates ? PhysicsSettings->ResimulationSettings.RedundantStates : StateRedundancy;
+			bCompareStateToTriggerRewind = PhysicsSettings->ResimulationSettings.GetCompareStateToTriggerRewind(PhysicsReplicationCVars::ResimulationCVars::bCompareStateToTriggerRewind);
+			bCompareInputToTriggerRewind = PhysicsSettings->ResimulationSettings.GetCompareInputToTriggerRewind(PhysicsReplicationCVars::ResimulationCVars::bCompareInputToTriggerRewind);
 		}
 
 		if (APawn* Pawn = Cast<APawn>(Owner))
@@ -403,6 +413,11 @@ void UNetworkPhysicsComponent::InitPhysics()
 			RepMovement.RotationQuantizationLevel = ERotatorQuantization::ShortComponents;
 			RepMovement.VelocityQuantizationLevel = EVectorQuantization::RoundTwoDecimals;
 		}
+
+		if (UPrimitiveComponent* RootPrimComp = Cast<UPrimitiveComponent>(Owner->GetRootComponent()))
+		{
+			RootPhysicsObject = RootPrimComp->GetPhysicsObjectByName(NAME_None);
+		}
 	}
 
 	bAutoActivate = true;
@@ -633,9 +648,36 @@ void UNetworkPhysicsComponent::OnRep_SetReplicatedStates()
 				TSharedPtr<Chaos::FBaseRewindHistory> ReceivedStates = MakeShareable(ReplicatedStates.History->Clone().Release());
 				PhysScene->EnqueueAsyncPhysicsCommand(0, this, [this, PhysScene, ReceivedStates, LocalOffset]()
 				{
-					PRAGMA_DISABLE_DEPRECATION_WARNINGS // TODO: Change to ReceiveNewData() in UE 5.6 and remove deprecation pragma
-					StateHistory->ReceiveNewDatas(*ReceivedStates, LocalOffset);
-					PRAGMA_ENABLE_DEPRECATION_WARNINGS
+					if (bCompareStateToTriggerRewind)
+					{
+						int32 ResimFrame = StateHistory->ReceiveNewData(*ReceivedStates, LocalOffset, /*CompareDataForRewind*/ bCompareStateToTriggerRewind);
+						if (ResimFrame != INDEX_NONE)
+						{
+							if (Chaos::FPhysicsSolver* Solver = PhysScene->GetSolver())
+							{
+								if (Chaos::FRewindData* RewindData = Solver->GetRewindData())
+								{
+									// Mark particle/island as resim
+									Chaos::FReadPhysicsObjectInterface_Internal Interface = Chaos::FPhysicsObjectInternalInterface::GetRead();
+									if (Chaos::FPBDRigidParticleHandle* POHandle = Interface.GetRigidParticle(RootPhysicsObject))
+									{
+										Solver->GetEvolution()->GetIslandManager().SetParticleResimFrame(POHandle, ResimFrame);
+									}
+
+									// Set resim frame in rewind data
+									ResimFrame = (RewindData->GetResimFrame() == INDEX_NONE) ? ResimFrame : FMath::Min(ResimFrame, RewindData->GetResimFrame());
+									RewindData->SetResimFrame(ResimFrame);
+								}
+							}
+						}
+					}
+					else 
+					{
+						PRAGMA_DISABLE_DEPRECATION_WARNINGS // TODO: Change to ReceiveNewData() in UE 5.6 and remove deprecation pragma
+						StateHistory->ReceiveNewDatas(*ReceivedStates, LocalOffset);
+						PRAGMA_ENABLE_DEPRECATION_WARNINGS
+					}
+
 #if DEBUG_NETWORK_PHYSICS
 					{
 						TArray<int32> LocalFrames, ServerFrames, InputFrames;
@@ -682,9 +724,36 @@ void UNetworkPhysicsComponent::OnRep_SetReplicatedInputs()
 				TSharedPtr<Chaos::FBaseRewindHistory> ReceivedInputs = MakeShareable(ReplicatedInputs.History->Clone().Release());
 				PhysScene->EnqueueAsyncPhysicsCommand(0, this, [this, PhysScene, ReceivedInputs, LocalOffset]()
 				{
-					PRAGMA_DISABLE_DEPRECATION_WARNINGS // TODO: Change to ReceiveNewData() in UE 5.6 and remove deprecation pragma
-					InputHistory->ReceiveNewDatas(*ReceivedInputs, LocalOffset);
-					PRAGMA_ENABLE_DEPRECATION_WARNINGS
+						if (bCompareInputToTriggerRewind)
+						{
+							int32 ResimFrame = StateHistory->ReceiveNewData(*ReceivedInputs, LocalOffset, /*CompareDataForRewind*/ bCompareInputToTriggerRewind);
+							if (ResimFrame != INDEX_NONE)
+							{
+								if (Chaos::FPhysicsSolver* Solver = PhysScene->GetSolver())
+								{
+									if (Chaos::FRewindData* RewindData = Solver->GetRewindData())
+									{
+										// Mark particle/island as resim
+										Chaos::FReadPhysicsObjectInterface_Internal Interface = Chaos::FPhysicsObjectInternalInterface::GetRead();
+										if (Chaos::FPBDRigidParticleHandle* POHandle = Interface.GetRigidParticle(RootPhysicsObject))
+										{
+											Solver->GetEvolution()->GetIslandManager().SetParticleResimFrame(POHandle, ResimFrame);
+										}
+
+										// Set resim frame in rewind data
+										ResimFrame = (RewindData->GetResimFrame() == INDEX_NONE) ? ResimFrame : FMath::Min(ResimFrame, RewindData->GetResimFrame());
+										RewindData->SetResimFrame(ResimFrame);
+									}
+								}
+							}
+						}
+						else
+						{
+							PRAGMA_DISABLE_DEPRECATION_WARNINGS // TODO: Change to ReceiveNewData() in UE 5.6 and remove deprecation pragma
+							InputHistory->ReceiveNewDatas(*ReceivedInputs, LocalOffset);
+							PRAGMA_ENABLE_DEPRECATION_WARNINGS
+						}
+
 #if DEBUG_NETWORK_PHYSICS
 					{
 						TArray<int32> LocalFrames, ServerFrames, InputFrames;
@@ -830,7 +899,7 @@ void UNetworkPhysicsComponent::OnPreProcessInputsInternal(const int32 PhysicsSte
 	#endif
 			const bool bExactFrame = PhysicsReplicationCVars::ResimulationCVars::bAllowRewindToClosestState ? !bIsSolverReset : true;
 			PRAGMA_DISABLE_DEPRECATION_WARNINGS // TODO: Change to ExtractData() in UE 5.6 and remove deprecation pragma
-			if (StateHistory->ExtractDatas(PhysicsStep, bIsSolverReset, PhysicsData, bExactFrame))
+			if (StateHistory->ExtractDatas(PhysicsStep, bIsSolverReset, PhysicsData, bExactFrame) && PhysicsData->bReceivedData)
 			PRAGMA_ENABLE_DEPRECATION_WARNINGS
 			{
 				PhysicsData->ApplyData(ActorComponent);
@@ -871,13 +940,15 @@ void UNetworkPhysicsComponent::OnPostProcessInputsInternal(const int32 PhysicsSt
 			PlayerController = GetWorld()->GetFirstPlayerController();
 		}
 
+		const bool bShouldCacheInputHistory = PlayerController && IsLocallyControlled() && !bIsSolverResim;
 		// for the inputs client local ones are ground truth otherwise use the replicated ones coming from the server
-		if (PlayerController && IsLocallyControlled() && !bIsSolverResim && (InputData != nullptr))
+		if (bShouldCacheInputHistory && (InputData != nullptr))
 		{
 			FNetworkPhysicsData* PhysicsData = InputData.Get();
 			PhysicsData->LocalFrame = PhysicsStep;
 			PhysicsData->ServerFrame = HasServerWorld() ? PhysicsStep : PhysicsStep + PlayerController->GetNetworkPhysicsTickOffset();
 			PhysicsData->InputFrame = PhysicsStep;
+			PhysicsData->bReceivedData = false;
 
 			PhysicsData->BuildData(ActorComponent);
 
@@ -891,7 +962,9 @@ void UNetworkPhysicsComponent::OnPostProcessInputsInternal(const int32 PhysicsSt
 #endif
 		}
 
-		if (HasServerWorld())
+		const bool bIsServer = HasServerWorld();
+		const bool bShouldCacheStateHistory = bIsServer || (bCompareStateToTriggerRewind && bShouldCacheInputHistory);
+		if (bShouldCacheStateHistory)
 		{
 			// Compute of the local frame coming from the client that was used to generate this state
 			int32 InputFrame = INDEX_NONE;
@@ -909,6 +982,7 @@ void UNetworkPhysicsComponent::OnPostProcessInputsInternal(const int32 PhysicsSt
 			PhysicsData->LocalFrame = PhysicsStep;
 			PhysicsData->ServerFrame = PhysicsStep;
 			PhysicsData->InputFrame = InputFrame;
+			PhysicsData->bReceivedData = false;
 
 			PhysicsData->BuildData(ActorComponent);
 
diff --git a/Engine/Source/Runtime/Engine/Private/PhysicsEngine/PhysicsReplication.cpp b/Engine/Source/Runtime/Engine/Private/PhysicsEngine/PhysicsReplication.cpp
--- a/Engine/Source/Runtime/Engine/Private/PhysicsEngine/PhysicsReplication.cpp
+++ b/Engine/Source/Runtime/Engine/Private/PhysicsEngine/PhysicsReplication.cpp
@@ -1786,47 +1786,49 @@ bool FPhysicsReplicationAsync::ResimulationReplication(Chaos::FPBDRigidParticleH
 	}
 #endif
 
-	if (ShouldTriggerResim && Target.TickCount == 0 && LocalFrame > RewindData->GetBlockedResimFrame())
+	if (LocalFrame > RewindData->GetBlockedResimFrame())
 	{
-		// Trigger resimulation
-		RigidsSolver->GetEvolution()->GetIslandManager().SetParticleResimFrame(Handle, LocalFrame);
-
-		int32 ResimFrame = RewindData->GetResimFrame();
-		ResimFrame = (ResimFrame == INDEX_NONE) ? LocalFrame : FMath::Min(ResimFrame, LocalFrame);
-		RewindData->SetResimFrame(ResimFrame);
-	}
-	else if (SettingsCurrent.ResimulationSettings.GetRuntimeCorrectionEnabled())
-	{
-		const int32 NumPredictedFrames = RigidsSolver->GetCurrentFrame() - LocalFrame - Target.TickCount;
-
-		if (Target.TickCount <= NumPredictedFrames && NumPredictedFrames > 0)
+		if (ShouldTriggerResim && Target.TickCount == 0)
 		{
-			// Calculate correction to position
-			const float CorrectionAmountX = SettingsCurrent.ResimulationSettings.GetPosStabilityMultiplier() / NumPredictedFrames; // Same result as (ErrorOffset / NumPredictedFrames) * PosStabilityMultiplier
-			const FVector CorrectedX = Handle->X() + (ErrorOffset * CorrectionAmountX);
+			// Trigger resimulation
+			RigidsSolver->GetEvolution()->GetIslandManager().SetParticleResimFrame(Handle, LocalFrame);
 
-			// Calculate correction to rotation
-			const float CorrectionAmountR = (1.f / NumPredictedFrames) * SettingsCurrent.ResimulationSettings.GetRotStabilityMultiplier();
-			const FQuat InvCurrentQuat = PastState.GetR().Inverse();
-			const FQuat DeltaQuat = Target.TargetState.Quaternion * InvCurrentQuat;
-			const FQuat TargetCorrectionR = Handle->GetR() * DeltaQuat;
-			const FQuat CorrectedR = FQuat::Slerp(Handle->GetR(), TargetCorrectionR, CorrectionAmountR);
+			int32 ResimFrame = RewindData->GetResimFrame();
+			ResimFrame = (ResimFrame == INDEX_NONE) ? LocalFrame : FMath::Min(ResimFrame, LocalFrame);
+			RewindData->SetResimFrame(ResimFrame);
+		}
+		else if (SettingsCurrent.ResimulationSettings.GetRuntimeCorrectionEnabled())
+		{
+			const int32 NumPredictedFrames = RigidsSolver->GetCurrentFrame() - LocalFrame - Target.TickCount;
 
-#if !(UE_BUILD_SHIPPING || UE_BUILD_TEST)
-			if (Chaos::FPhysicsSolverBase::CanDebugNetworkPhysicsPrediction())
+			if (Target.TickCount <= NumPredictedFrames && NumPredictedFrames > 0)
 			{
-				Chaos::FDebugDrawQueue::GetInstance().DrawDebugDirectionalArrow(Handle->X(), CorrectedX, 5.0f, FColor::MakeRandomSeededColor(LocalFrame), true, CharacterMovementCVars::NetCorrectionLifetime, 0, 0.5f);
+				// Calculate correction to position
+				const float CorrectionAmountX = SettingsCurrent.ResimulationSettings.GetPosStabilityMultiplier() / NumPredictedFrames; // Same result as (ErrorOffset / NumPredictedFrames) * PosStabilityMultiplier
+				const FVector CorrectedX = Handle->X() + (ErrorOffset * CorrectionAmountX);
+
+				// Calculate correction to rotation
+				const float CorrectionAmountR = (1.f / NumPredictedFrames) * SettingsCurrent.ResimulationSettings.GetRotStabilityMultiplier();
+				const FQuat InvTargetQuat = Target.TargetState.Quaternion;
+				const FQuat DeltaQuat = PastState.GetR().Inverse() * InvTargetQuat;
+				const FQuat TargetCorrectionR = Handle->GetR() * DeltaQuat;
+				const FQuat CorrectedR = FQuat::Slerp(Handle->GetR(), TargetCorrectionR, CorrectionAmountR);
+
+	#if !(UE_BUILD_SHIPPING || UE_BUILD_TEST)
+				if (Chaos::FPhysicsSolverBase::CanDebugNetworkPhysicsPrediction())
+				{
+					Chaos::FDebugDrawQueue::GetInstance().DrawDebugDirectionalArrow(Handle->X(), CorrectedX, 5.0f, FColor::MakeRandomSeededColor(LocalFrame), true, CharacterMovementCVars::NetCorrectionLifetime, 0, 0.5f);
+				}
+	#endif
+				// Apply correction
+				Handle->SetX(CorrectedX);
+				Handle->SetR(CorrectedR);
 			}
-#endif
-			// Apply correction
-			Handle->SetX(CorrectedX);
-			Handle->SetR(CorrectedR);
-		}
 
-		// Keep target for NumPredictedFrames time to perform runtime corrections with until a new target is received
-		bClearTarget = Target.TickCount >= NumPredictedFrames;
+			// Keep target for NumPredictedFrames time to perform runtime corrections with until a new target is received
+			bClearTarget = Target.TickCount >= NumPredictedFrames;
+		}
 	}
-
 	return bClearTarget;
 }
 
diff --git a/Engine/Source/Runtime/Engine/Public/Physics/NetworkPhysicsComponent.h b/Engine/Source/Runtime/Engine/Public/Physics/NetworkPhysicsComponent.h
--- a/Engine/Source/Runtime/Engine/Public/Physics/NetworkPhysicsComponent.h
+++ b/Engine/Source/Runtime/Engine/Public/Physics/NetworkPhysicsComponent.h
@@ -10,6 +10,7 @@
 #include "Physics/Experimental/PhysScene_Chaos.h"
 #include "Subsystems/WorldSubsystem.h"
 #include "PhysicsEngine/PhysicsSettings.h"
+#include "Chaos/PhysicsObject.h"
 
 #include "NetworkPhysicsComponent.generated.h"
 
@@ -75,10 +76,11 @@ struct TNetRewindHistory : public Chaos::TDataRewindHistory<DataType>
 		return Copy;
 	}
 
-	virtual void ReceiveNewData(Chaos::FBaseRewindHistory& NewData, const int32 FrameOffset) override
+	virtual int32 ReceiveNewData(Chaos::FBaseRewindHistory& NewData, const int32 FrameOffset, bool CompareDataForRewind = false) override
 	{
 		TNetRewindHistory& NetNewData = static_cast<TNetRewindHistory&>(NewData);
 
+		int32 RewindFrame = INDEX_NONE;
 		if (NetNewData.NumFrames > 0)
 		{
 			for (int32 FrameIndex = 0; FrameIndex < NetNewData.NumFrames; ++FrameIndex)
@@ -88,12 +90,20 @@ struct TNetRewindHistory : public Chaos::TDataRewindHistory<DataType>
 				FrameData.LocalFrame = FrameData.ServerFrame - FrameOffset;
 				if (FrameData.LocalFrame >= 0)
 				{
+					FrameData.bReceivedData = true; // Received data is marked to differentiate from locally predicted data
+
+					if (CompareDataForRewind && FrameData.LocalFrame > RewindFrame && Super::TriggerRewindFromNewData(&FrameData))
+					{
+						RewindFrame = FrameData.LocalFrame;
+					}
+
 					PRAGMA_DISABLE_DEPRECATION_WARNINGS // TODO: Change to RecordData() in UE 5.6 and remove deprecation pragma
 					Super::RecordDatas(FrameData.LocalFrame, &FrameData);
 					PRAGMA_ENABLE_DEPRECATION_WARNINGS
 				}
 			}
 		}
+		return RewindFrame;
 	}
 
 	virtual void NetSerialize(FArchive& Ar, UPackageMap* InPackageMap) override
@@ -360,6 +370,9 @@ struct FNetworkPhysicsData
 	UPROPERTY()
 	int32 InputFrame = INDEX_NONE;
 
+	// If this data was received over the network or locally predicted
+	bool bReceivedData = false;
+
 	// Serialize the data into/from the archive
 	void SerializeFrames(FArchive& Ar)
 	{
@@ -411,6 +424,11 @@ struct FNetworkPhysicsData
 	*/
 	virtual void ValidateData(const UActorComponent* NetworkComponent) { }
 
+	/** Define how to compare client and server data for the same frame, returning false means the data differ enough to trigger a resimulation.
+	* @param PredictedData is data predicted on the client to compare with the current data received from the server.
+	*/
+	virtual bool CompareData(const FNetworkPhysicsData& PredictedData) { return true; }
+
 	/** DEPRECATED */
 	UE_DEPRECATED(5.4, "Deprecated, use ApplyData instead")
 	virtual void ApplyDatas(UActorComponent* NetworkComponent) const { }
@@ -616,8 +634,17 @@ private:
 	// Actor component that will be used to fill the histories
 	TObjectPtr<UActorComponent> ActorComponent;
 
+	// Root components physics object
+	Chaos::FPhysicsObjectHandle RootPhysicsObject;
+
 	// Locally relayed inputs makes this component act as if it's a locally controlled pawn.
 	bool bIsRelayingLocalInputs = false;
+
+	// Cache locally predicted states and then compare then via FNetworkPhysicsData::CompareData to trigger rewind if comparison differ
+	bool bCompareStateToTriggerRewind = false;
+
+	// Compare locally predicted inputs via FNetworkPhysicsData::CompareData to trigger rewind if comparison differ
+	bool bCompareInputToTriggerRewind = false;
 };
 
 /** DEPRECATED UE 5.4 */
diff --git a/Engine/Source/Runtime/Engine/Public/Physics/NetworkPhysicsSettingsComponent.h b/Engine/Source/Runtime/Engine/Public/Physics/NetworkPhysicsSettingsComponent.h
--- a/Engine/Source/Runtime/Engine/Public/Physics/NetworkPhysicsSettingsComponent.h
+++ b/Engine/Source/Runtime/Engine/Public/Physics/NetworkPhysicsSettingsComponent.h
@@ -213,6 +213,21 @@ struct FNetworkPhysicsSettingsResimulation
 	UPROPERTY(config, EditDefaultsOnly, Category = "Overrides", Meta = (EditCondition = "bOverrideResimulationErrorThreshold"))
 	uint32 ResimulationErrorThreshold = 10;
 	uint32 GetResimulationErrorThreshold(uint32 DefaultValue) { return bOverrideRotStabilityMultiplier ? ResimulationErrorThreshold : DefaultValue; }
+	
+	// np2.Resim.CompareStateToTriggerRewind
+	UPROPERTY(config, EditDefaultsOnly, Category = "Overrides", Meta = (InlineEditConditionToggle))
+	uint32 bOverrideCompareStateToTriggerRewind : 1;
+	UPROPERTY(config, EditDefaultsOnly, Category = "Overrides", Meta = (EditCondition = "bOverrideCompareStateToTriggerRewind"))
+	bool bCompareStateToTriggerRewind = false;
+	bool GetCompareStateToTriggerRewind(bool DefaultValue) { return bOverrideCompareStateToTriggerRewind ? bCompareStateToTriggerRewind : DefaultValue; }
+
+	// np2.Resim.CompareInputToTriggerRewind
+	UPROPERTY(config, EditDefaultsOnly, Category = "Overrides", Meta = (InlineEditConditionToggle))
+	uint32 bOverrideCompareInputToTriggerRewind : 1;
+	UPROPERTY(config, EditDefaultsOnly, Category = "Overrides", Meta = (EditCondition = "bOverrideCompareInputToTriggerRewind"))
+	bool bCompareInputToTriggerRewind = false;
+	bool GetCompareInputToTriggerRewind(bool DefaultValue) { return bOverrideCompareInputToTriggerRewind ? bCompareInputToTriggerRewind : DefaultValue; }
+
 };
 
 /*
diff --git a/Engine/Source/Runtime/Experimental/Chaos/Private/PBDRigidsSolver.cpp b/Engine/Source/Runtime/Experimental/Chaos/Private/PBDRigidsSolver.cpp
--- a/Engine/Source/Runtime/Experimental/Chaos/Private/PBDRigidsSolver.cpp
+++ b/Engine/Source/Runtime/Experimental/Chaos/Private/PBDRigidsSolver.cpp
@@ -380,6 +380,15 @@ namespace Chaos
 	}
 }
 
+namespace PhysicsReplicationCVars
+{
+	namespace ResimulationCVars
+	{
+		bool bApplyTargetsWhileResimulating = false;
+		static FAutoConsoleVariableRef CVarResimApplyTargetsWhileResimulating(TEXT("np2.Resim.ApplyTargetsWhileResimulating"), bApplyTargetsWhileResimulating, TEXT("If false, target states from the server are only applied on rewind. If true, target states from the server are applied during resimulation if there are any available."));
+	}
+}
+
 namespace Chaos
 {
 	using namespace CVars;
@@ -1929,8 +1938,11 @@ namespace Chaos
 					{
 						if ((LastStep - Step) < RecordedPushData.Num())
 						{
-							// We need to reset all the particles having received a target from the server even if not at first resim frame
-							MRewindData->ApplyTargets(Step, bFirst);
+							if (PhysicsReplicationCVars::ResimulationCVars::bApplyTargetsWhileResimulating || bFirst)
+							{
+								// Update all the particles having received a target from the server
+								MRewindData->ApplyTargets(Step, bFirst);
+							}
 
 							FPushPhysicsData* PushData = RecordedPushData[LastStep - Step];	//push data is sorted as latest first
 							if (bFirst)
diff --git a/Engine/Source/Runtime/Experimental/Chaos/Private/RewindData.cpp b/Engine/Source/Runtime/Experimental/Chaos/Private/RewindData.cpp
--- a/Engine/Source/Runtime/Experimental/Chaos/Private/RewindData.cpp
+++ b/Engine/Source/Runtime/Experimental/Chaos/Private/RewindData.cpp
@@ -603,7 +603,7 @@ void FRewindData::AdvanceFrameImp(IResimCacheBase* ResimCache)
 	//TODO: if joint is desynced we should desync particles as well
 	//If particle of joint is desynced, we need to make sure the joint is reconsidered too for optimization, though maybe not "desynced"
 
-	if(IsResim() && ResimCache)
+	if (IsResim() && ResimCache)
 	{
 		ResimCache->SetDesyncedParticles(MoveTemp(DesyncedParticles));
 	}
@@ -706,7 +706,7 @@ void FRewindData::PushGTDirtyData(const FDirtyPropertiesManager& SrcManager,cons
 			}
 		});
 
-		if(bKeepRecording)
+		if (bKeepRecording)
 		{
 			//Dynamics are not available at head (sim zeroes them out), so we have to record them as PostPushData (since they're applied as part of PushData)
 			if (auto NewData = Dirty.PropertyData.FindDynamics(SrcManager, SrcDataIdx))
@@ -811,7 +811,7 @@ void FRewindData::MarkDirtyFromPT(FGeometryParticleHandle& Handle)
 
 	if(bRecordingHistory || Latest.ParticlePositionRotation.IsClean(FrameAndPhase))
 	{
-		if (auto Data = Latest.ParticlePositionRotation.WriteAccessNonDecreasing(FFrameAndPhase{ CurFrame, FFrameAndPhase::PostPushData }, PropertiesPool))
+		if (auto Data = Latest.ParticlePositionRotation.WriteAccessNonDecreasing(FrameAndPhase, PropertiesPool))
 		{
 			Data->CopyFrom(Handle);
 		}
@@ -855,7 +855,7 @@ void FRewindData::MarkDirtyJointFromPT(FPBDJointConstraintHandle& Handle)
 
 	if (bRecordingHistory || Latest.JointSettings.IsClean(FrameAndPhase))
 	{
-		if (auto Data = Latest.JointSettings.WriteAccessNonDecreasing(FFrameAndPhase{ CurFrame, FFrameAndPhase::PostPushData }, PropertiesPool))
+		if (auto Data = Latest.JointSettings.WriteAccessNonDecreasing(FrameAndPhase, PropertiesPool))
 		{
 			CopyDataFromObject(*Data, Handle);
 		}
@@ -1009,7 +1009,6 @@ int32 FRewindData::FindValidResimFrame(const int32 RequestedFrame)
 		{
 			for (TWeakPtr<FBaseRewindHistory>& StateHistory : StateHistories)
 			{
-
 				PRAGMA_DISABLE_DEPRECATION_WARNINGS // TODO: Change to HasValidData() in UE 5.6 and remove deprecation pragma
 				if (StateHistory.IsValid() && !StateHistory.Pin().Get()->HasValidDatas(ValidFrame))
 				PRAGMA_ENABLE_DEPRECATION_WARNINGS
diff --git a/Engine/Source/Runtime/Experimental/Chaos/Public/RewindData.h b/Engine/Source/Runtime/Experimental/Chaos/Public/RewindData.h
--- a/Engine/Source/Runtime/Experimental/Chaos/Public/RewindData.h
+++ b/Engine/Source/Runtime/Experimental/Chaos/Public/RewindData.h
@@ -59,11 +59,14 @@ struct FBaseRewindHistory
 	/** Create a polymorphic copy of only a range of frames, applying the frame offset to the copies */
 	virtual TUniquePtr<FBaseRewindHistory> CopyFramesWithOffset(const uint32 StartFrame, const uint32 EndFrame, const int32 FrameOffset) = 0;
 
-	/** Copy new data (received from the network) into this history */
-	virtual void ReceiveNewData(FBaseRewindHistory& NewData, const int32 FrameOffset) { }
+	/** Copy new data (received from the network) into this history, returns frame to resimulate from if @param CompareDataForRewind is set to true and compared data differ enough */
+	virtual int32 ReceiveNewData(FBaseRewindHistory& NewData, const int32 FrameOffset, bool CompareDataForRewind = false) { return INDEX_NONE; }
 	UE_DEPRECATED(5.4, "Deprecated, use ReceiveNewData() instead")
 	virtual void ReceiveNewDatas(FBaseRewindHistory& NewDatas, const int32 FrameOffset) { ReceiveNewData(NewDatas, FrameOffset); }
 
+	/** Compares new received data with local predicted data and returns true if they differ enough to trigger a resimulation  */
+	FORCEINLINE virtual bool TriggerRewindFromNewData(void* NewData) { return false; }
+
 	/** Serialize the data to or from a network archive */
 	virtual void NetSerialize(FArchive& Ar, UPackageMap* PackageMap) {}
 	
@@ -200,6 +203,16 @@ public :
 		}
 	}
 
+	FORCEINLINE virtual bool TriggerRewindFromNewData(void* NewData) override
+	{
+		if (EvalData(static_cast<DataType*>(NewData)->LocalFrame))
+		{
+			return !static_cast<DataType*>(NewData)->CompareData(DataHistory[CurrentIndex]);
+		}
+
+		return false;
+	}
+
 	/** Load the data from the buffer at a specific frame */
 	FORCEINLINE bool LoadData(const int32 LoadFrame)
 	{
@@ -1442,21 +1455,22 @@ public:
 		Managers[CurFrame].FrameCreatedFor = CurFrame;
 		TUniquePtr<IResimCacheBase>& ResimCache = Managers[CurFrame].ExternalResimCache;
 
-		if(bResimOptimization)
+		if (bResimOptimization)
 		{
-			if(IsResim())
+			if (IsResim())
 			{
-				if(ResimCache)
+				if (ResimCache)
 				{
 					ResimCache->SetResimming(true);
 				}
 			}
 			else
 			{
-				if(ResimCache)
+				if (ResimCache)
 				{
 					ResimCache->ResetCache();
-				} else
+				}
+				else
 				{
 					ResimCache = CreateCacheFunc();
 				}
