Bug fix for replication being turned off when geometry collections are part of cluster union actor (normally a seperate actor)

diff --git a/Engine/Source/Runtime/Engine/Private/PhysicsEngine/ClusterUnionComponent.cpp b/Engine/Source/Runtime/Engine/Private/PhysicsEngine/ClusterUnionComponent.cpp
--- a/Engine/Source/Runtime/Engine/Private/PhysicsEngine/ClusterUnionComponent.cpp
+++ b/Engine/Source/Runtime/Engine/Private/PhysicsEngine/ClusterUnionComponent.cpp
@@ -1085,7 +1085,7 @@ void UClusterUnionComponent::HandleAddOrModifiedClusteredComponent(const FMapped
 				NewData.bWasReplicatingMovement = Owner->IsReplicatingMovement();
 				ActorToComponents.Add(Owner, NewData);
 
-				if (IsAuthority())
+				if (IsAuthority() && Owner != GetOwner())
 				{
 					Owner->SetReplicatingMovement(false);
 				}
