Use bulk data batch API for texture streaming

#rnx

diff --git a/Engine/Source/Runtime/Engine/Private/Streaming/Texture2DStreamIn_IO.cpp b/Engine/Source/Runtime/Engine/Private/Streaming/Texture2DStreamIn_IO.cpp
--- a/Engine/Source/Runtime/Engine/Private/Streaming/Texture2DStreamIn_IO.cpp
+++ b/Engine/Source/Runtime/Engine/Private/Streaming/Texture2DStreamIn_IO.cpp
@@ -8,6 +8,7 @@ Texture2DStreamIn.cpp: Stream in helper for 2D textures using texture streaming
 #include "HAL/PlatformFile.h"
 #include "Streaming/TextureStreamingHelpers.h"
 #include "ContentStreaming.h"
+#include "IO/IoBuffer.h"
 #include "Rendering/Texture2DResource.h"
 #include "Streaming/Texture2DStreamIn.h"
 #include "Streaming/Texture2DUpdate.h"
@@ -17,12 +18,23 @@ Texture2DStreamIn.cpp: Stream in helper for 2D textures using texture streaming
 #include "RenderUtils.h"
 #endif
 
+// Will be removed once the batch API has proven to work
+bool GBatchTextureMipIoRequests = true;
+static FAutoConsoleVariableRef CVarBatchTextureMipIoRequests(
+	TEXT("r.BatchTextureMipIoRequests"),
+	GBatchTextureMipIoRequests,
+	TEXT("Whether to batch texture mip I/O requests"));
+
 FTexture2DStreamIn_IO::FTexture2DStreamIn_IO(UTexture2D* InTexture, bool InPrioritizedIORequest)
 	: FTexture2DStreamIn(InTexture)
 	, bPrioritizedIORequest(InPrioritizedIORequest)
+	, bBatchIORequest(GBatchTextureMipIoRequests)
 
 {
-	IORequests.AddZeroed(ResourceState.MaxNumLODs);
+	if (!bBatchIORequest)
+	{
+		IORequests.AddZeroed(ResourceState.MaxNumLODs);
+	}
 }
 
 FTexture2DStreamIn_IO::~FTexture2DStreamIn_IO()
@@ -54,6 +66,91 @@ static void ValidateMipBulkDataSize(const UTexture2D& Texture, int32 MipSizeX, i
 
 void FTexture2DStreamIn_IO::SetIORequests(const FContext& Context)
 {
+	if (bBatchIORequest)
+	{
+		const int32 BatchCount = CurrentFirstLODIdx - PendingFirstLODIdx;
+		FBulkDataBatchRequest::FBatchBuilder Batch = FBulkDataBatchRequest::NewBatch(BatchCount);
+		for (int32 MipIndex = PendingFirstLODIdx; MipIndex < CurrentFirstLODIdx && !IsCancelled(); ++MipIndex)
+		{
+			const FTexture2DMipMap& MipMap = *Context.MipsView[MipIndex];
+			check(MipData[MipIndex].Data != nullptr);
+
+			int64 BulkDataSize = MipMap.BulkData.GetBulkDataSize();
+			if (BulkDataSize > 0)
+			{
+				// Validate buffer size for the mip, so we don't overrun it on streaming
+				// note: MipData[] should have size
+				// ValidateMipBulkDataSize only does anything on Android
+				ValidateMipBulkDataSize(*Context.Texture, MipMap.SizeX, MipMap.SizeY, MipIndex, BulkDataSize);
+				
+				// reads directly into MipData[] , doesn't respect Pitch
+				// we do get a completion callback at AsyncFileCallBack
+				// so in theory could fix Pitch there
+				uint32 DestPitch = MipData[MipIndex].Pitch;
+				FTexture2DResource::WarnRequiresTightPackedMip(MipMap.SizeX, MipMap.SizeY, Context.Resource->GetPixelFormat(), DestPitch);
+
+				EAsyncIOPriorityAndFlags Priority = AIOP_Low;
+				if (bPrioritizedIORequest)
+				{
+					static IConsoleVariable* CVarAsyncLoadingPrecachePriority = IConsoleManager::Get().FindConsoleVariable(TEXT("s.AsyncLoadingPrecachePriority"));
+					const bool bLoadBeforeAsyncPrecache = CVarStreamingLowResHandlingMode.GetValueOnAnyThread() == (int32)FRenderAssetStreamingSettings::LRHM_LoadBeforeAsyncPrecache;
+
+					if (CVarAsyncLoadingPrecachePriority && bLoadBeforeAsyncPrecache)
+					{
+						const int32 AsyncIOPriority = CVarAsyncLoadingPrecachePriority->GetInt();
+						// Higher priority than regular requests but don't go over max
+						Priority = (EAsyncIOPriorityAndFlags)FMath::Clamp<int32>(AsyncIOPriority + 1, AIOP_BelowNormal, AIOP_MAX);
+					}
+					else
+					{
+						Priority = AIOP_BelowNormal;
+					}
+				}
+
+				FIoBuffer Dst(FIoBuffer::Wrap, MipData[MipIndex].Data, BulkDataSize);
+				Batch.Read(MipMap.BulkData, 0, BulkDataSize, Priority | AIOP_FLAG_DONTCACHE, Dst);
+			}
+			else // Bulk data size can only be 0 when not available, in which case, we need to recache the file state.
+			{
+				bFailedOnIOError = true;
+				MarkAsCancelled();
+				break;
+			}
+		}
+
+		if (bFailedOnIOError)
+		{
+			return;
+		}
+
+		TaskSynchronization.Increment();
+
+		return Batch.Issue(
+			[this](FBulkDataBatchRequest::EStatus Status)
+			{
+				TaskSynchronization.Decrement();
+
+				if (Status != FBulkDataBatchRequest::EStatus::Ok && !bIsCancelled)
+				{
+					// If IO requests was cancelled but the streaming request wasn't, this is an IO error.
+					bFailedOnIOError = true;
+					MarkAsCancelled();
+				}
+#if !UE_BUILD_SHIPPING
+				// On some platforms the IO is too fast to test cancelation requests timing issues.
+				if (FRenderAssetStreamingSettings::ExtraIOLatency > 0 && TaskSynchronization.GetValue() == 0)
+				{
+					FPlatformProcess::Sleep(FRenderAssetStreamingSettings::ExtraIOLatency * .001f); // Slow down the streaming.
+				}
+#endif
+				// The tick here is intended to schedule the success or cancel callback.
+				// Using TT_None ensure gets which could create a dead lock.
+				Tick(FTexture2DUpdate::TT_None);
+			},
+			BatchRequest);
+	}
+
+	check(BatchRequest.IsNone());
 	SetAsyncFileCallback();
 	
 	for (int32 MipIndex = PendingFirstLODIdx; MipIndex < CurrentFirstLODIdx && !IsCancelled(); ++MipIndex)
@@ -115,6 +212,13 @@ void FTexture2DStreamIn_IO::SetIORequests(const FContext& Context)
 
 void FTexture2DStreamIn_IO::CancelIORequests()
 {
+	if (bBatchIORequest)
+	{
+		check(IORequests.IsEmpty());
+		return BatchRequest.Cancel();
+	}
+
+	check(BatchRequest.IsNone());
 	for (int32 MipIndex = 0; MipIndex < IORequests.Num(); ++MipIndex)
 	{
 		IBulkDataIORequest* IORequest = IORequests[MipIndex];
@@ -128,6 +232,13 @@ void FTexture2DStreamIn_IO::CancelIORequests()
 
 void FTexture2DStreamIn_IO::ClearIORequests(const FContext& Context)
 {
+	if (bBatchIORequest)
+	{
+		check(IORequests.IsEmpty());
+		return BatchRequest.Reset();
+	}
+
+	check(BatchRequest.IsNone());
 	for (int32 MipIndex = PendingFirstLODIdx; MipIndex < CurrentFirstLODIdx; ++MipIndex)
 	{
 		IBulkDataIORequest* IORequest = IORequests[MipIndex];
@@ -201,6 +312,13 @@ void FTexture2DStreamIn_IO::Abort()
 
 		if (HasPendingIORequests())
 		{
+			if (bBatchIORequest)
+			{
+				check(IORequests.IsEmpty());
+				return BatchRequest.Cancel();
+			}
+
+			check(BatchRequest.IsNone());
 			// Prevent the update from being considered done before this is finished.
 			// By checking that it was not already canceled, we make sure this doesn't get called twice.
 			(new FAsyncCancelIORequestsTask(this))->StartBackgroundTask();
@@ -210,6 +328,13 @@ void FTexture2DStreamIn_IO::Abort()
 
 bool FTexture2DStreamIn_IO::HasPendingIORequests()
 {
+	if (bBatchIORequest)
+	{
+		check(IORequests.IsEmpty());
+		return BatchRequest.IsPending();
+	}
+
+	check(BatchRequest.IsNone());
 	for (IBulkDataIORequest* IORequest : IORequests)
 	{
 		if (IORequest != nullptr)
diff --git a/Engine/Source/Runtime/Engine/Private/Streaming/Texture2DStreamIn_IO.h b/Engine/Source/Runtime/Engine/Private/Streaming/Texture2DStreamIn_IO.h
--- a/Engine/Source/Runtime/Engine/Private/Streaming/Texture2DStreamIn_IO.h
+++ b/Engine/Source/Runtime/Engine/Private/Streaming/Texture2DStreamIn_IO.h
@@ -10,6 +10,7 @@ Texture2DStreamIn.h: Stream in helper for 2D textures using texture streaming fi
 #include "Texture2DStreamIn.h"
 
 class IBulkDataIORequest;
+class FBulkDataRequest;
 
 // Base StreamIn framework exposing MipData
 class FTexture2DStreamIn_IO : public FTexture2DStreamIn
@@ -62,7 +63,7 @@ private:
 
 
 	// Request for loading into each mip.
-	TArray<IBulkDataIORequest*, TInlineAllocator<MAX_TEXTURE_MIP_COUNT> > IORequests;
+	TArray<IBulkDataIORequest*> IORequests;
 
 	// Whether an IO error was detected (when files do not exists).
 	bool bFailedOnIOError = false;
@@ -70,6 +71,10 @@ private:
 	// Whether the IO request should be created with an higher priority for quicker response time.
 	bool bPrioritizedIORequest = false;
 
+	// Whether to batch IO request(s) 
+	bool bBatchIORequest;
+
 	FBulkDataIORequestCallBack AsyncFileCallBack;
+	FBulkDataBatchRequest BatchRequest;
 };
 
