Fix bug when FStaticMeshRenderData::ReleaseResources() is called when FStaticMeshRenderData is not initialized.
- caused crash when changing materials inside the Static Mesh Editor.
- Fix:
  - clear RayTracingGeometryGroupHandle after release.
  - check whether FStaticMeshRenderData was initialized.

diff --git a/Engine/Source/Runtime/Engine/Private/StaticMesh.cpp b/Engine/Source/Runtime/Engine/Private/StaticMesh.cpp
--- a/Engine/Source/Runtime/Engine/Private/StaticMesh.cpp
+++ b/Engine/Source/Runtime/Engine/Private/StaticMesh.cpp
@@ -1957,6 +1957,8 @@ void FStaticMeshRenderData::InitResources(ERHIFeatureLevel::Type InFeatureLevel,
 
 void FStaticMeshRenderData::ReleaseResources()
 {
+	const bool bWasInitialized = bIsInitialized;
+
 	bIsInitialized = false;
 
 	for (int32 LODIndex = 0; LODIndex < LODResources.Num(); ++LODIndex)
@@ -1969,12 +1971,13 @@ void FStaticMeshRenderData::ReleaseResources()
 	}
 
 #if RHI_RAYTRACING
-	if (IsRayTracingAllowed())
+	if (bWasInitialized && IsRayTracingAllowed())
 	{
 		ENQUEUE_RENDER_COMMAND(CmdReleaseRayTracingGeometryGroup)(
 			[this](FRHICommandListImmediate&)
 			{
 				GRayTracingGeometryManager->ReleaseRayTracingGeometryGroup(RayTracingGeometryGroupHandle);
+				RayTracingGeometryGroupHandle = INDEX_NONE;
 			});
 	}
 #endif
