World Partition: handle the case where actors can be destroyed as part of deprecation by removing them from their actor descriptor container instance.

diff --git a/Engine/Source/Runtime/Engine/Private/WorldPartition/WorldPartition.cpp b/Engine/Source/Runtime/Engine/Private/WorldPartition/WorldPartition.cpp
--- a/Engine/Source/Runtime/Engine/Private/WorldPartition/WorldPartition.cpp
+++ b/Engine/Source/Runtime/Engine/Private/WorldPartition/WorldPartition.cpp
@@ -867,6 +867,17 @@ bool UWorldPartition::IsMainWorldPartition() const
 	return World == GetTypedOuter<UWorld>();
 }
 
+void UWorldPartition::OnLevelActorDeleted(AActor* Actor)
+{
+	if (GIsEditorLoadingPackage)
+	{
+		if (UActorDescContainerInstance* DescContainerInstance = GetActorDescContainerInstance())
+		{
+			DescContainerInstance->RemoveActor(Actor->GetActorGuid());
+		}
+	}
+}
+
 void UWorldPartition::OnPostBugItGoCalled(const FVector& Loc, const FRotator& Rot)
 {
 #if WITH_EDITOR
@@ -904,9 +915,10 @@ void UWorldPartition::RegisterDelegates()
 			FEditorDelegates::CancelPIE.AddUObject(this, &UWorldPartition::OnCancelPIE);
 			FGameDelegates::Get().GetEndPlayMapDelegate().AddUObject(this, &UWorldPartition::OnEndPlay);
 			FCoreUObjectDelegates::PostReachabilityAnalysis.AddUObject(this, &UWorldPartition::OnGCPostReachabilityAnalysis);
+			GEditor->OnLevelActorDeleted().AddUObject(this, &UWorldPartition::OnLevelActorDeleted);
 			GEditor->OnPostBugItGoCalled().AddUObject(this, &UWorldPartition::OnPostBugItGoCalled);
 			GEditor->OnEditorClose().AddUObject(this, &UWorldPartition::SavePerUserSettings);
-			FWorldDelegates::OnPostWorldRename.AddUObject(this, &UWorldPartition::OnWorldRenamed);
+			FWorldDelegates::OnPostWorldRename.AddUObject(this, &UWorldPartition::OnWorldRenamed);			
 		}
 
 		if (!IsRunningCommandlet())
@@ -949,6 +961,7 @@ void UWorldPartition::UnregisterDelegates()
 				FCoreUObjectDelegates::PostReachabilityAnalysis.RemoveAll(this);
 			}
 
+			GEditor->OnLevelActorDeleted().RemoveAll(this);
 			GEditor->OnPostBugItGoCalled().RemoveAll(this);
 			GEditor->OnEditorClose().RemoveAll(this);
 		}
diff --git a/Engine/Source/Runtime/Engine/Public/WorldPartition/WorldPartition.h b/Engine/Source/Runtime/Engine/Public/WorldPartition/WorldPartition.h
--- a/Engine/Source/Runtime/Engine/Public/WorldPartition/WorldPartition.h
+++ b/Engine/Source/Runtime/Engine/Public/WorldPartition/WorldPartition.h
@@ -650,6 +650,8 @@ private:
 	static ENGINE_API FAutoConsoleVariableRef CVarUseMakingInvisibleTransactionRequests;
 
 	ENGINE_API void OnWorldMatchStarting();
+	
+	ENGINE_API void OnLevelActorDeleted(AActor* Actor);
 	ENGINE_API void OnPostBugItGoCalled(const FVector& Loc, const FRotator& Rot);
 
 	// Delegates registration
