Path Tracing: Avoid path tracer invalidation when a camera is selected in the editor

This was creating an editor widget, which was in turn setting some colors on a dynamic material instance which led to the path traced output being invalidated. Simply exclude any materials with bUsedWithEditorCompositing set from invalidating the path tracer. The assumption is that this flag will only be set for specific engine materials used in the UI.

diff --git a/Engine/Source/Runtime/Engine/Private/Materials/MaterialInstance.cpp b/Engine/Source/Runtime/Engine/Private/Materials/MaterialInstance.cpp
--- a/Engine/Source/Runtime/Engine/Private/Materials/MaterialInstance.cpp
+++ b/Engine/Source/Runtime/Engine/Private/Materials/MaterialInstance.cpp
@@ -435,7 +435,7 @@ void GameThread_UpdateMIParameter(const UMaterialInstance* Instance, const Param
 		{
 			EMaterialDomain Domain = Material->MaterialDomain;
 			// check if this material has any relevance to path tracing
-			if (Domain != MD_PostProcess && Domain != MD_UI)
+			if (Domain != MD_PostProcess && Domain != MD_UI && !Material->bUsedWithEditorCompositing)
 			{
 				GetRendererModule().InvalidatePathTracedOutput();
 			}
