Avoid recreating the light component on selection

Marking the renderer state dirty causes the path tracer to invalidate because the light is removed and added. Instead, simply update the proxy's relevant flag from the render thread.

diff --git a/Engine/Source/Runtime/Engine/Private/Components/LightComponent.cpp b/Engine/Source/Runtime/Engine/Private/Components/LightComponent.cpp
--- a/Engine/Source/Runtime/Engine/Private/Components/LightComponent.cpp
+++ b/Engine/Source/Runtime/Engine/Private/Components/LightComponent.cpp
@@ -14,6 +14,7 @@
 #include "MaterialDomain.h"
 #include "UObject/ObjectSaveContext.h"
 #include "SceneInterface.h"
+#include "LightSceneProxy.h"
 #include "UObject/RenderingObjectVersion.h"
 #include "UObject/UE5MainStreamObjectVersion.h"
 #include "UObject/UObjectAnnotation.h"
@@ -1437,7 +1438,18 @@ void ULightComponent::SetMaterial(int32 ElementIndex, UMaterialInterface* InMate
 
 void ULightComponent::PushSelectionToProxy()
 {
-	MarkRenderStateDirty();
+	if (SceneProxy)
+	{
+		const bool bIsSelected = IsSelected() || IsOwnerSelected();
+		FLightSceneProxy* LocalSceneProxy = SceneProxy;
+		ENQUEUE_RENDER_COMMAND(SetLightSelection)(
+			[LocalSceneProxy, bIsSelected](FRHICommandListImmediate& RHICmdList)
+			{
+				// NOTE: The selection flag is currently used on the C++ side for debug features, so simply updating this flag is enough.
+				LocalSceneProxy->SetSelected(bIsSelected);
+			});
+	}
+
 }
 
 /** Stores a light and a channel it has been assigned to. */
diff --git a/Engine/Source/Runtime/Engine/Private/LightSceneProxy.cpp b/Engine/Source/Runtime/Engine/Private/LightSceneProxy.cpp
--- a/Engine/Source/Runtime/Engine/Private/LightSceneProxy.cpp
+++ b/Engine/Source/Runtime/Engine/Private/LightSceneProxy.cpp
@@ -53,7 +53,7 @@ FLightSceneProxy::FLightSceneProxy(const ULightComponent* InLightComponent)
 	, bUseVirtualShadowMaps(false)	// See below
 	, bCastModulatedShadows(false)
 	, bUseWholeSceneCSMForMovableObjects(false)
-	, bSelected(InLightComponent->GetOwner() ? InLightComponent->GetOwner()->IsActorOrSelectionParentSelected() : false)
+	, bSelected(InLightComponent->IsSelected() || InLightComponent->IsOwnerSelected())
 	, AtmosphereSunLightIndex(InLightComponent->GetAtmosphereSunLightIndex())
 	, AtmosphereSunDiskColorScale(InLightComponent->GetAtmosphereSunDiskColorScale())
 	, LightType(InLightComponent->GetLightType())
diff --git a/Engine/Source/Runtime/Engine/Public/LightSceneProxy.h b/Engine/Source/Runtime/Engine/Public/LightSceneProxy.h
--- a/Engine/Source/Runtime/Engine/Public/LightSceneProxy.h
+++ b/Engine/Source/Runtime/Engine/Public/LightSceneProxy.h
@@ -213,6 +213,7 @@ public:
 	inline FName GetComponentFName() const { return ComponentName; }
 
 	inline bool IsSelected() const { return bSelected; }
+	inline void SetSelected(bool bNewSelected) { bSelected = bNewSelected; }
 
 	/**
 	 * Use to get the owning actor label (or component name as fallback, if the owner is null or ENABLE_DEBUG_LABELS is off) for diagnostic messages, debug or profiling.
@@ -434,7 +435,8 @@ protected:
 	/** Whether to render csm shadows for movable objects only (mobile). */
 	uint8 bUseWholeSceneCSMForMovableObjects : 1;
 
-	const uint8 bSelected : 1;
+    /** Is the light selected in the editor? */
+	uint8 bSelected : 1;
 
 	/** The index of the atmospheric light. Multiple lights can be considered when computing the sky/atmospheric scattering. */
 	const uint8 AtmosphereSunLightIndex;
diff --git a/Engine/Source/Runtime/Renderer/Private/VirtualShadowMaps/VirtualShadowMapArray.cpp b/Engine/Source/Runtime/Renderer/Private/VirtualShadowMaps/VirtualShadowMapArray.cpp
--- a/Engine/Source/Runtime/Renderer/Private/VirtualShadowMaps/VirtualShadowMapArray.cpp
+++ b/Engine/Source/Runtime/Renderer/Private/VirtualShadowMaps/VirtualShadowMapArray.cpp
@@ -1075,16 +1075,12 @@ void FVirtualShadowMapVisualizeLightSearch::CheckLight(const FLightSceneProxy* C
 		UE_LOG(LogRenderer, Display, TEXT("%s"), *CheckLightName);
 	}
 
-	const ULightComponent* Component = CheckProxy->GetLightComponent();
-	check(Component);
-
 	// Fill out new sort key and compare to our best found so far
 	SortKey CheckKey;
 	CheckKey.Packed = 0;
 	CheckKey.Fields.bExactNameMatch = (CheckLightName == GVirtualShadowMapVisualizeLightName);
 	CheckKey.Fields.bPartialNameMatch = CheckKey.Fields.bExactNameMatch || CheckLightName.Contains(GVirtualShadowMapVisualizeLightName);
-	CheckKey.Fields.bSelected = Component->IsSelected();
-	CheckKey.Fields.bOwnerSelected = Component->IsOwnerSelected();
+	CheckKey.Fields.bSelected = CheckProxy->IsSelected();
 	CheckKey.Fields.bDirectionalLight = CheckProxy->GetLightType() == LightType_Directional;
 	CheckKey.Fields.bExists = 1;
 
diff --git a/Engine/Source/Runtime/Renderer/Private/VirtualShadowMaps/VirtualShadowMapArray.h b/Engine/Source/Runtime/Renderer/Private/VirtualShadowMaps/VirtualShadowMapArray.h
--- a/Engine/Source/Runtime/Renderer/Private/VirtualShadowMaps/VirtualShadowMapArray.h
+++ b/Engine/Source/Runtime/Renderer/Private/VirtualShadowMaps/VirtualShadowMapArray.h
@@ -204,7 +204,6 @@ private:
 			// NOTE: Lowest to highest priority
 			uint32 bExists : 1;				// Catch-all
 			uint32 bDirectionalLight : 1;
-			uint32 bOwnerSelected : 1;		// In editor
 			uint32 bSelected : 1;			// In editor
 			uint32 bPartialNameMatch : 1;
 			uint32 bExactNameMatch : 1;
