fix linux build

diff --git a/Engine/Source/Runtime/Engine/Classes/VisualLogger/VisualLoggerTraceDevice.h b/Engine/Source/Runtime/Engine/Classes/VisualLogger/VisualLoggerTraceDevice.h
--- a/Engine/Source/Runtime/Engine/Classes/VisualLogger/VisualLoggerTraceDevice.h
+++ b/Engine/Source/Runtime/Engine/Classes/VisualLogger/VisualLoggerTraceDevice.h
@@ -18,7 +18,7 @@ public:
 		return GDevice;
 	}
 
-	FVisualLoggerTraceDevice();
+	ENGINE_API FVisualLoggerTraceDevice();
 	virtual void Cleanup(bool bReleaseMemory = false) override;
 	virtual void StartRecordingToFile(double TimeStamp) override;
 	virtual void StopRecordingToFile(double TimeStamp) override;
