Improved fix for mid-tick changes in movement modes continuing movement with the wrong Phys**** function, where we now verify a mode change since entry.

This version accounts for the possibility of a project's custom mode explicitly calling the Phys*** function.

diff --git a/Engine/Source/Runtime/Engine/Private/Components/CharacterMovementComponent.cpp b/Engine/Source/Runtime/Engine/Private/Components/CharacterMovementComponent.cpp
--- a/Engine/Source/Runtime/Engine/Private/Components/CharacterMovementComponent.cpp
+++ b/Engine/Source/Runtime/Engine/Private/Components/CharacterMovementComponent.cpp
@@ -5396,6 +5396,9 @@ void UCharacterMovementComponent::PhysWalking(float deltaTime, int32 Iterations)
 	bool bTriedLedgeMove = false;
 	float remainingTime = deltaTime;
 
+	const EMovementMode StartingMovementMode = MovementMode;
+	const uint8 StartingCustomMovementMode = CustomMovementMode;
+
 	// Perform the move
 	while ( (remainingTime >= MIN_TICK_TIME) && (Iterations < MaxSimulationIterations) && CharacterOwner && (CharacterOwner->Controller || bRunPhysicsWithNoController || HasAnimRootMotion() || CurrentRootMotion.HasOverrideVelocity() || (CharacterOwner->GetLocalRole() == ROLE_SimulatedProxy)) )
 	{
@@ -5427,9 +5430,9 @@ void UCharacterMovementComponent::PhysWalking(float deltaTime, int32 Iterations)
 		ApplyRootMotionToVelocity(timeTick);
 		devCode(ensureMsgf(!Velocity.ContainsNaN(), TEXT("PhysWalking: Velocity contains NaN after Root Motion application (%s)\n%s"), *GetPathNameSafe(this), *Velocity.ToString()));
 
-		if (MovementMode != MOVE_Walking)
+		if (MovementMode != StartingMovementMode || CustomMovementMode != StartingCustomMovementMode)
 		{
-			// Root motion could have taken us out of walking mode
+			// Root motion could have taken us out of our current mode
 			// No movement has taken place this movement tick so we pass on full time/past iteration count
 			StartNewPhysics(remainingTime+timeTick, Iterations-1);
 			return;
@@ -5455,7 +5458,7 @@ void UCharacterMovementComponent::PhysWalking(float deltaTime, int32 Iterations)
 				StartSwimming(OldLocation, OldVelocity, timeTick, remainingTime, Iterations);
 				return;
 			}
-			else if (MovementMode != MOVE_Walking)
+			else if (MovementMode != StartingMovementMode || CustomMovementMode != StartingCustomMovementMode)
 			{
 				// pawn ended up in a different mode, probably due to the step-up-and-over flow
 				// let's refund the estimated unused time (if any) and keep moving in the new mode
@@ -5610,6 +5613,9 @@ void UCharacterMovementComponent::PhysNavWalking(float deltaTime, int32 Iteratio
 		return;
 	}
 
+	const EMovementMode StartingMovementMode = MovementMode;
+	const uint8 StartingCustomMovementMode = CustomMovementMode;
+
 	RestorePreAdditiveRootMotionVelocity();
 
 	// Ensure velocity is horizontal.
@@ -5626,9 +5632,9 @@ void UCharacterMovementComponent::PhysNavWalking(float deltaTime, int32 Iteratio
 
 	ApplyRootMotionToVelocity(deltaTime);
 
-	if (MovementMode != MOVE_NavWalking)
+	if (MovementMode != StartingMovementMode || CustomMovementMode != StartingCustomMovementMode)
 	{
-		// Root motion could have taken us out of walking mode
+		// Root motion could have taken us out of our current mode
 		StartNewPhysics(deltaTime, Iterations);
 		return;
 	}
