Skipping the material translation DDC when a material is not flagged as Persistent.

diff --git a/Engine/Source/Runtime/Engine/Private/Materials/HLSLMaterialTranslator.cpp b/Engine/Source/Runtime/Engine/Private/Materials/HLSLMaterialTranslator.cpp
--- a/Engine/Source/Runtime/Engine/Private/Materials/HLSLMaterialTranslator.cpp
+++ b/Engine/Source/Runtime/Engine/Private/Materials/HLSLMaterialTranslator.cpp
@@ -1162,7 +1162,7 @@ bool FHLSLMaterialTranslator::Translate()
 #if FORCE_DISABLE_MATERIAL_TRANSLATION_DDC
 	const bool bDisableTranslationDDC = true;
 #else
-	const bool bDisableTranslationDDC = bNoMaterialTranslationDDC || Material->IsPreview();
+	const bool bDisableTranslationDDC = bNoMaterialTranslationDDC || Material->IsPreview() || !Material->IsPersistent();
 #endif
 
 	// DDC query local data
