Substrate - Fixed support for dithered lod transitions.

diff --git a/Engine/Shaders/Shared/SubstrateDefinitions.h b/Engine/Shaders/Shared/SubstrateDefinitions.h
--- a/Engine/Shaders/Shared/SubstrateDefinitions.h
+++ b/Engine/Shaders/Shared/SubstrateDefinitions.h
@@ -9,7 +9,7 @@
 
 #ifndef __cplusplus
 // Change this to force recompilation of all Substrate dependent shaders (for instance https://guidgenerator.com/online-guid-generator.aspx)
-#pragma message("UESHADERMETADATA_VERSION B8E54812-515D-45F4-8A68-AE2FB6553231")
+#pragma message("UESHADERMETADATA_VERSION B013B52C-494E-4FC3-9EB5-B16B754E5CA4")
 #endif
 
 // Closure offsets are packed into 32bits, each entry using SUBSTRATE_CLOSURE_OFFSET_BIT_COUNT bits
diff --git a/Engine/Source/Runtime/Engine/Private/Materials/MaterialShared.cpp b/Engine/Source/Runtime/Engine/Private/Materials/MaterialShared.cpp
--- a/Engine/Source/Runtime/Engine/Private/Materials/MaterialShared.cpp
+++ b/Engine/Source/Runtime/Engine/Private/Materials/MaterialShared.cpp
@@ -2541,17 +2541,12 @@ void FMaterial::SetupMaterialEnvironment(
 		switch (GetBlendMode())
 		{
 		case BLEND_Opaque:
-		{
-			SET_SHADER_DEFINE(OutEnvironment, MATERIALBLENDING_SOLID, 1);
-			SET_SHADER_DEFINE(OutEnvironment, SUBSTRATE_BLENDING_OPAQUE, 1);
-			break;
-		}
 		case BLEND_Masked:
 		{
 			// Only set MATERIALBLENDING_MASKED if the material is truly masked
 			//@todo - this may cause mismatches with what the shader compiles and what the renderer thinks the shader needs
 			// For example IsTranslucentBlendMode doesn't check IsMasked
-			if(!WritesEveryPixel())
+			if (!WritesEveryPixel())
 			{
 				SET_SHADER_DEFINE(OutEnvironment, MATERIALBLENDING_MASKED, 1);
 			}
@@ -2559,7 +2554,6 @@ void FMaterial::SetupMaterialEnvironment(
 			{
 				SET_SHADER_DEFINE(OutEnvironment, MATERIALBLENDING_SOLID, 1);
 			}
-			SET_SHADER_DEFINE(OutEnvironment, SUBSTRATE_BLENDING_MASKED, 1);
 			break;
 		}
 		case BLEND_Additive:
