Fix issues with MorphTarget build by disabling async build for skeletal mesh containing morph targets.

#rnx

diff --git a/Engine/Source/Runtime/Engine/Private/SkeletalMesh.cpp b/Engine/Source/Runtime/Engine/Private/SkeletalMesh.cpp
--- a/Engine/Source/Runtime/Engine/Private/SkeletalMesh.cpp
+++ b/Engine/Source/Runtime/Engine/Private/SkeletalMesh.cpp
@@ -1975,8 +1975,10 @@ void USkeletalMesh::Build()
 	FSkinnedAssetBuildContext Context;
 	BeginBuildInternal(Context);
 	
-	// Inline reduction is not thread-safe, prevent async build until this is fixed.
-	if (FSkinnedAssetCompilingManager::Get().IsAsyncCompilationAllowed(this))
+	//5.4 hack fix: if this skeletal mesh has some morph targets we disable asynchronous build.
+	//TODO: To remove this constraint we must get rid of the UMorphTarget sub object and directly create the render data from the FSkeletalMeshImportData.
+	const bool bBuildAsynchronous = FSkinnedAssetCompilingManager::Get().IsAsyncCompilationAllowed(this) && GetMorphTargets().IsEmpty();
+	if (bBuildAsynchronous)
 	{
 		PrepareForAsyncCompilation();
 
