World Partition PIE startup: don't parse external actor folders for cell packages.

diff --git a/Engine/Source/Runtime/Engine/Private/Level.cpp b/Engine/Source/Runtime/Engine/Private/Level.cpp
--- a/Engine/Source/Runtime/Engine/Private/Level.cpp
+++ b/Engine/Source/Runtime/Engine/Private/Level.cpp
@@ -1133,6 +1133,8 @@ void ULevel::PreSave(FObjectPreSaveContext ObjectSaveContext)
 
 void ULevel::PostLoad()
 {
+	TRACE_CPUPROFILER_EVENT_SCOPE(ULevel::PostLoad);
+
 	Super::PostLoad();
 
 	// Ensure that the level is pointed to the owning world.  For streamed levels, this will be the world of the P map
@@ -1175,7 +1177,7 @@ void ULevel::PostLoad()
 
 	if (IsUsingActorFolders() && IsUsingExternalObjects() && IsActorFolderObjectsFeatureAvailable())
 	{
-		if (!bWasDuplicated)
+		if (!bWasDuplicated && !FPackageName::IsTempPackage(GetPackage()->GetName()))
 		{
 			// Load all folders for this level
 			FExternalPackageHelper::LoadObjectsFromExternalPackages<UActorFolder>(this, [this](UActorFolder* LoadedFolder)
diff --git a/Engine/Source/Runtime/Engine/Public/ExternalPackageHelper.h b/Engine/Source/Runtime/Engine/Public/ExternalPackageHelper.h
--- a/Engine/Source/Runtime/Engine/Public/ExternalPackageHelper.h
+++ b/Engine/Source/Runtime/Engine/Public/ExternalPackageHelper.h
@@ -109,6 +109,8 @@ private:
 template<typename T>
 void FExternalPackageHelper::LoadObjectsFromExternalPackages(UObject* InOuter, TFunctionRef<void(T*)> Operation)
 {
+	TRACE_CPUPROFILER_EVENT_SCOPE(FExternalPackageHelper::LoadObjectsFromExternalPackages);
+
 	check(InOuter);
 	UPackage* OutermostPackage = InOuter->IsA<UPackage>() ? CastChecked<UPackage>(InOuter) : InOuter->GetOutermostObject()->GetPackage();
 	const FString OutermostPackageName = !OutermostPackage->GetLoadedPath().IsEmpty() ? OutermostPackage->GetLoadedPath().GetPackageName() : OutermostPackage->GetName();
