World Partition - HLOD: Removed leftover code that should have been removed in 25205034

diff --git a/Engine/Source/Runtime/Engine/Private/WorldPartition/HLOD/HLODActor.cpp b/Engine/Source/Runtime/Engine/Private/WorldPartition/HLOD/HLODActor.cpp
--- a/Engine/Source/Runtime/Engine/Private/WorldPartition/HLOD/HLODActor.cpp
+++ b/Engine/Source/Runtime/Engine/Private/WorldPartition/HLOD/HLODActor.cpp
@@ -156,13 +156,6 @@ void AWorldPartitionHLOD::PostLoad()
 {
 	Super::PostLoad();
 
-	// If world is instanced, we need to recompute our bounds since they are in the instanced-world space
-	ForEachComponent<USceneComponent>(false, [](USceneComponent* SceneComponent)
-	{
-		// Clear bComputedBoundsOnceForGame so that the bounds are recomputed once
-		SceneComponent->bComputedBoundsOnceForGame = false;
-	});
-
 #if WITH_EDITOR
 	if (GetLinkerCustomVersion(FFortniteMainBranchObjectVersion::GUID) < FFortniteMainBranchObjectVersion::WorldPartitionStreamingCellsNamingShortened)
 	{
