World Partition PIE startup: minimize calls to StaticFindObject when generating streaming cells.

diff --git a/Engine/Source/Runtime/Engine/Private/WorldPartition/WorldPartitionStreamingGeneration.cpp b/Engine/Source/Runtime/Engine/Private/WorldPartition/WorldPartitionStreamingGeneration.cpp
--- a/Engine/Source/Runtime/Engine/Private/WorldPartition/WorldPartitionStreamingGeneration.cpp
+++ b/Engine/Source/Runtime/Engine/Private/WorldPartition/WorldPartitionStreamingGeneration.cpp
@@ -724,7 +724,7 @@ class FWorldPartitionStreamingGenerator
 		const bool bIsTempContainerPackage = FPackageName::IsTempPackage((InActorDescCollection.GetBaseContainerInstancePackageName().ToString()));
 
 		// Test whether an actor descriptor instance should be included in the ActorDescViewMap.
-		auto ShouldRegisterActorDesc = [this](const FWorldPartitionActorDescInstance* InActorDescInstance)
+		auto ShouldRegisterActorDesc = [this](const FWorldPartitionActorDescInstance* InActorDescInstance, AActor** OutActor = nullptr)
 		{
 			for (UClass* FilteredClass : FilteredClasses)
 			{
@@ -739,7 +739,16 @@ class FWorldPartitionStreamingGenerator
 				return false;
 			}
 
-			return InActorDescInstance->IsLoaded() ? !InActorDescInstance->GetActor()->IsEditorOnly() : !InActorDescInstance->GetActorIsEditorOnly();
+			if (AActor* Actor = InActorDescInstance->GetActor())
+			{
+				if (OutActor)
+				{
+					*OutActor = Actor;
+				}
+				return !Actor->IsEditorOnly();
+			}
+
+			return !InActorDescInstance->GetActorIsEditorOnly();
 		};
 
 		// Register the actor descriptor view
@@ -760,10 +769,12 @@ class FWorldPartitionStreamingGenerator
 		{
 			// @todo_ow: this is to validate that new parenting of container instance code is equivalent
 			check(Iterator->GetContainerInstance()->GetContainerID() == InContainerID);
-			if (ShouldRegisterActorDesc(*Iterator))
+
+			AActor* Actor = nullptr;
+			if (ShouldRegisterActorDesc(*Iterator, &Actor))
 			{
 				// Handle unsaved actors
-				if (AActor* Actor = Iterator->GetActor())
+				if (Actor)
 				{
 					// Deleted actors
 					if (!IsValid(Actor))
