World Partition PIE startup: don't compute HLOD file size on disk when duplicating the world for PIE.

diff --git a/Engine/Source/Runtime/Engine/Private/WorldPartition/HLOD/HLODActor.cpp b/Engine/Source/Runtime/Engine/Private/WorldPartition/HLOD/HLODActor.cpp
--- a/Engine/Source/Runtime/Engine/Private/WorldPartition/HLOD/HLODActor.cpp
+++ b/Engine/Source/Runtime/Engine/Private/WorldPartition/HLOD/HLODActor.cpp
@@ -239,8 +239,11 @@ void AWorldPartitionHLOD::PostLoad()
 		}
 	}
 
-	// Update the disk size stat on load, as we can't really know it when saving
-	HLODStats.Add(FWorldPartitionHLODStats::MemoryDiskSizeBytes, FHLODActorDesc::GetPackageSize(this));
+	if (GetPackage()->GetPIEInstanceID() == INDEX_NONE)
+	{
+		// Update the disk size stat on load, as we can't really know it when saving
+		HLODStats.Add(FWorldPartitionHLODStats::MemoryDiskSizeBytes, FHLODActorDesc::GetPackageSize(this));
+	}
 #endif
 }
 
