Added cvar r.SkyLight.RealTimeReflectionCapture to set a max resolution for skylight cubemap if needed

#rnx

#localization none
#tests tested in live

diff --git a/Engine/Source/Runtime/Engine/Private/Components/SkyLightComponent.cpp b/Engine/Source/Runtime/Engine/Private/Components/SkyLightComponent.cpp
--- a/Engine/Source/Runtime/Engine/Private/Components/SkyLightComponent.cpp
+++ b/Engine/Source/Runtime/Engine/Private/Components/SkyLightComponent.cpp
@@ -108,6 +108,13 @@ FAutoConsoleVariableRef CVarSkylightRealTimeReflectionCapture(
 	ECVF_Scalability
 	);
 
+int32 GSkylightCubemapMaxResolution = -1;
+FAutoConsoleVariableRef CVarSkylightCubemapMaxResolution(
+	TEXT("r.SkyLight.CubemapMaxResolution"),
+	GSkylightCubemapMaxResolution,
+	TEXT("Force max resolution of skylight cubemap (default to -1: takes default property value of USkyLightComponent::CubeMapResolution)")
+);
+
 constexpr EPixelFormat SKYLIGHT_CUBEMAP_FORMAT = PF_FloatRGBA;
 
 void FSkyTextureCubeResource::InitRHI(FRHICommandListBase&)
@@ -345,6 +352,11 @@ void USkyLightComponent::SanitizeCubemapSize()
 	const int32 MaxCubemapResolution = GetMaxCubeTextureDimension();
 	const int32 MinCubemapResolution = 8;
 
+	if (GSkylightCubemapMaxResolution > 0)
+	{
+		CubemapResolution = GSkylightCubemapMaxResolution;
+	}
+
 	CubemapResolution = FMath::Clamp(int32(FMath::RoundUpToPowerOfTwo(CubemapResolution)), MinCubemapResolution, MaxCubemapResolution);
 
 #if WITH_EDITOR
