Fixed crash when creating SVT Sample node with Substrate enabled. This was actually an issue with how the SVT Sample node behaves and is technically unrelated to Substrate. Substrate just happened to call into the buggy function.

#rnx

diff --git a/Engine/Source/Runtime/Engine/Private/Materials/MaterialExpressions.cpp b/Engine/Source/Runtime/Engine/Private/Materials/MaterialExpressions.cpp
--- a/Engine/Source/Runtime/Engine/Private/Materials/MaterialExpressions.cpp
+++ b/Engine/Source/Runtime/Engine/Private/Materials/MaterialExpressions.cpp
@@ -28069,11 +28069,8 @@ uint32 UMaterialExpressionSparseVolumeTextureSample::GetInputType(int32 InputInd
 	case 2:
 		return MCT_Float1;
 	default:
-		break;
+		return MCT_Unknown;
 	}
-
-	check(false);
-	return MCT_Float1;
 }
 
 int32 UMaterialExpressionSparseVolumeTextureSample::Compile(class FMaterialCompiler* Compiler, int32 OutputIndex)
