World Partition Actor Descriptors: make sure that actors still async loading report as not loaded.

diff --git a/Engine/Source/Runtime/Engine/Private/WorldPartition/WorldPartitionActorDescInstance.cpp b/Engine/Source/Runtime/Engine/Private/WorldPartition/WorldPartitionActorDescInstance.cpp
--- a/Engine/Source/Runtime/Engine/Private/WorldPartition/WorldPartitionActorDescInstance.cpp
+++ b/Engine/Source/Runtime/Engine/Private/WorldPartition/WorldPartitionActorDescInstance.cpp
@@ -40,6 +40,11 @@ void FWorldPartitionActorDescInstance::UpdateActorDesc(FWorldPartitionActorDesc*
 
 bool FWorldPartitionActorDescInstance::IsLoaded(bool bEvenIfPendingKill) const
 {
+	if (AsyncLoadID != INDEX_NONE)
+	{
+		return false;
+	}
+
 	if (ActorPtr.IsExplicitlyNull() || ActorPtr.IsStale())
 	{
 		ActorPtr = FindObject<AActor>(nullptr, *GetActorSoftPath().ToString());
