AggGeom:

Fix GetElementIndexByName to return the correct index.
#tests PIE

diff --git a/Engine/Source/Runtime/Engine/Classes/PhysicsEngine/AggregateGeom.h b/Engine/Source/Runtime/Engine/Classes/PhysicsEngine/AggregateGeom.h
--- a/Engine/Source/Runtime/Engine/Classes/PhysicsEngine/AggregateGeom.h
+++ b/Engine/Source/Runtime/Engine/Classes/PhysicsEngine/AggregateGeom.h
@@ -171,45 +171,52 @@ struct FKAggregateGeom
 	int32 GetElementIndexByName(const FName InName) const
 	{
 		int32 FoundIndex = GetElementIndexByName<FKSphereElem>(MakeArrayView(SphereElems), InName);
+		int32 StartIndex = 0;
 		if (FoundIndex != INDEX_NONE)
 		{
-			return FoundIndex;
+			return FoundIndex + StartIndex;
 		}
+		StartIndex += SphereElems.Num();
 
 		FoundIndex = GetElementIndexByName<FKBoxElem>(MakeArrayView(BoxElems), InName);
 		if (FoundIndex != INDEX_NONE)
 		{
-			return FoundIndex;
+			return FoundIndex + StartIndex;
 		}
+		StartIndex += BoxElems.Num();
 
 		FoundIndex = GetElementIndexByName<FKSphylElem>(MakeArrayView(SphylElems), InName);
 		if (FoundIndex != INDEX_NONE)
 		{
-			return FoundIndex;
+			return FoundIndex + StartIndex;
 		}
+		StartIndex += SphylElems.Num();
 
 		FoundIndex = GetElementIndexByName<FKConvexElem>(MakeArrayView(ConvexElems), InName);
 		if (FoundIndex != INDEX_NONE)
 		{
-			return FoundIndex;
+			return FoundIndex + StartIndex;
 		}
+		StartIndex += ConvexElems.Num();
 
 		FoundIndex = GetElementIndexByName<FKTaperedCapsuleElem>(MakeArrayView(TaperedCapsuleElems), InName);
 		if (FoundIndex != INDEX_NONE)
 		{
-			return FoundIndex;
+			return FoundIndex + StartIndex;
 		}
+		StartIndex += TaperedCapsuleElems.Num();
 
 		FoundIndex = GetElementIndexByName<FKLevelSetElem>(MakeArrayView(LevelSetElems), InName);
 		if (FoundIndex != INDEX_NONE)
 		{
-			return FoundIndex;
+			return FoundIndex + StartIndex;
 		}
+		StartIndex += LevelSetElems.Num();
 
 		FoundIndex = GetElementIndexByName<FKSkinnedLevelSetElem>(MakeArrayView(SkinnedLevelSetElems), InName);
 		if (FoundIndex != INDEX_NONE)
 		{
-			return FoundIndex;
+			return FoundIndex + StartIndex;
 		}
 
 		return INDEX_NONE;
