World Partition: handle the possibility that streaming sources velocity computation code receives a zero delta time.

diff --git a/Engine/Source/Runtime/Engine/Private/WorldPartition/WorldPartitionSubsystem.cpp b/Engine/Source/Runtime/Engine/Private/WorldPartition/WorldPartitionSubsystem.cpp
--- a/Engine/Source/Runtime/Engine/Private/WorldPartition/WorldPartitionSubsystem.cpp
+++ b/Engine/Source/Runtime/Engine/Private/WorldPartition/WorldPartitionSubsystem.cpp
@@ -1534,36 +1534,45 @@ FVector FStreamingSourceVelocity::GetAverageVelocity(const FVector& NewPosition,
 
 	const bool bNewSource = (LastUpdateTime <= 0.0);
 	const double DeltaSeconds = bNewSource ? 1.0 : (CurrentTime - LastUpdateTime);
-	LastUpdateTime = CurrentTime;
 
-	const FVector AbsMovement = (NewPosition - LastPosition);
-	const FVector AbsVelocity = AbsMovement / DeltaSeconds;
-	LastPosition = NewPosition;
-
-	const double TeleportDistance = 10000;
-	const double MaxDeltaSeconds = 5.0;
-	const double Distance = AbsMovement.Size();
-
-	if (bNewSource)
-	{
-		UE_LOG(LogWorldPartition, Verbose, TEXT("New Streaming Source: %s -> Position: %s"), *SourceName.ToString(), *NewPosition.ToString());
-		AvgVelocity = FVector::Zero();
-	}
-	else if (Distance > TeleportDistance)
+	if (DeltaSeconds <= 0.0)
 	{
-		UE_LOG(LogWorldPartition, Verbose, TEXT("Detected Streaming Source Teleport: %s -> Last Position: %s -> New Position: %s"), *SourceName.ToString(), *LastPosition.ToString(), *NewPosition.ToString());
-		AvgVelocity = FVector::Zero();
-	}
-	else if  (DeltaSeconds > MaxDeltaSeconds)
-	{
-		UE_LOG(LogWorldPartition, Verbose, TEXT("Detected Inactive Streaming Source: %s -> Last Position: %s -> New Position: %s"), *SourceName.ToString(), *LastPosition.ToString(), *NewPosition.ToString());
+		UE_LOG(LogWorldPartition, Verbose, TEXT("Detected Invalid Delta Time: %s"), *SourceName.ToString());
 		AvgVelocity = FVector::Zero();
 	}
 	else
 	{
-		// Compute the new value in a weighted moving average series
-		const double AvgWeight = FMath::Clamp(DeltaSeconds * 100, 0, 1);
-		AvgVelocity = AvgVelocity * (1.0 - AvgWeight) + AbsVelocity * AvgWeight;
+		LastUpdateTime = CurrentTime;
+
+		const FVector AbsMovement = (NewPosition - LastPosition);
+		const FVector AbsVelocity = AbsMovement / DeltaSeconds;
+		LastPosition = NewPosition;
+
+		const double TeleportDistance = 10000;
+		const double MaxDeltaSeconds = 5.0;
+		const double Distance = AbsMovement.Size();
+
+		if (bNewSource)
+		{
+			UE_LOG(LogWorldPartition, Verbose, TEXT("New Streaming Source: %s -> Position: %s"), *SourceName.ToString(), *NewPosition.ToString());
+			AvgVelocity = FVector::Zero();
+		}
+		else if (Distance > TeleportDistance)
+		{
+			UE_LOG(LogWorldPartition, Verbose, TEXT("Detected Streaming Source Teleport: %s -> Last Position: %s -> New Position: %s"), *SourceName.ToString(), *LastPosition.ToString(), *NewPosition.ToString());
+			AvgVelocity = FVector::Zero();
+		}
+		else if  (DeltaSeconds > MaxDeltaSeconds)
+		{
+			UE_LOG(LogWorldPartition, Verbose, TEXT("Detected Inactive Streaming Source: %s -> Last Position: %s -> New Position: %s"), *SourceName.ToString(), *LastPosition.ToString(), *NewPosition.ToString());
+			AvgVelocity = FVector::Zero();
+		}
+		else
+		{
+			// Compute the new value in a weighted moving average series
+			const double AvgWeight = FMath::Clamp(DeltaSeconds * 100, 0, 1);
+			AvgVelocity = AvgVelocity * (1.0 - AvgWeight) + AbsVelocity * AvgWeight;
+		}
 	}
 
 	return AvgVelocity;
