Added more consistent CPU/GPU profiling events to SVT streaming.

#rnx

diff --git a/Engine/Source/Runtime/Engine/Private/SparseVolumeTexture/SparseVolumeTextureStreamingManager.cpp b/Engine/Source/Runtime/Engine/Private/SparseVolumeTexture/SparseVolumeTextureStreamingManager.cpp
--- a/Engine/Source/Runtime/Engine/Private/SparseVolumeTexture/SparseVolumeTextureStreamingManager.cpp
+++ b/Engine/Source/Runtime/Engine/Private/SparseVolumeTexture/SparseVolumeTextureStreamingManager.cpp
@@ -20,6 +20,8 @@
 
 DEFINE_LOG_CATEGORY(LogSparseVolumeTextureStreamingManager);
 
+DECLARE_GPU_STAT(SVTStreaming);
+
 #ifndef SVT_STREAMING_LOG_VERBOSE
 #define SVT_STREAMING_LOG_VERBOSE 0
 #endif
@@ -342,7 +344,10 @@ void FStreamingManager::BeginAsyncUpdate(FRDGBuilder& GraphBuilder, bool bBlocki
 		return;
 	}
 
-	TRACE_CPUPROFILER_EVENT_SCOPE(SVT::FStreamingManager::BeginAsyncUpdate);
+	RDG_EVENT_SCOPE(GraphBuilder, "SVT::StreamingBeginAsyncUpdate");
+	RDG_GPU_STAT_SCOPE(GraphBuilder, SVTStreaming);
+	RDG_CSV_STAT_EXCLUSIVE_SCOPE(GraphBuilder, SVTStreaming);
+	SCOPED_NAMED_EVENT_TEXT("SVT::StreamingBeginAsyncUpdate", FColor::Green);
 
 #if SVT_STREAMING_LOG_VERBOSE
 	UE_LOG(LogSparseVolumeTextureStreamingManager, Display, TEXT("SVT Streaming Update %i"), NextUpdateIndex);
@@ -486,7 +491,10 @@ void FStreamingManager::EndAsyncUpdate(FRDGBuilder& GraphBuilder)
 	}
 	check(AsyncState.bUpdateActive);
 
-	TRACE_CPUPROFILER_EVENT_SCOPE(SVT::FStreamingManager::EndAsyncUpdate);
+	RDG_EVENT_SCOPE(GraphBuilder, "SVT::StreamingEndAsyncUpdate");
+	RDG_GPU_STAT_SCOPE(GraphBuilder, SVTStreaming);
+	RDG_CSV_STAT_EXCLUSIVE_SCOPE(GraphBuilder, SVTStreaming);
+	SCOPED_NAMED_EVENT_TEXT("SVT::StreamingEndAsyncUpdate", FColor::Green);
 
 	// Wait for async processing to finish
 	if (AsyncState.bUpdateIsAsync)
@@ -545,7 +553,10 @@ void FStreamingManager::AddInternal(FRDGBuilder& GraphBuilder, FNewSparseVolumeT
 		return;
 	}
 
-	TRACE_CPUPROFILER_EVENT_SCOPE(SVT::FStreamingManager::AddInternal);
+	RDG_EVENT_SCOPE(GraphBuilder, "SVT::StreamingAddInternal");
+	RDG_GPU_STAT_SCOPE(GraphBuilder, SVTStreaming);
+	RDG_CSV_STAT_EXCLUSIVE_SCOPE(GraphBuilder, SVTStreaming);
+	SCOPED_NAMED_EVENT_TEXT("SVT::StreamingAddInternal", FColor::Green);
 
 	const int32 NumFrames = NewSVTInfo.FrameInfo.Num();
 
@@ -815,7 +826,7 @@ bool FStreamingManager::AddRequest(const FStreamingRequest& Request)
 
 void FStreamingManager::AddParentRequests()
 {
-	TRACE_CPUPROFILER_EVENT_SCOPE(SVT::AddParentRequests);
+	TRACE_CPUPROFILER_EVENT_SCOPE(SVT::StreamingAddParentRequests);
 
 	ParentRequestsToAdd.Reset();
 	for (const auto& Request : RequestsHashTable)
@@ -859,7 +870,7 @@ void FStreamingManager::SelectHighestPriorityRequestsAndUpdateLRU(int32 MaxSelec
 
 	if (!RequestsHashTable.IsEmpty())
 	{
-		TRACE_CPUPROFILER_EVENT_SCOPE(SVT::SelectHighestPriorityRequestsAndUpdateLRU);
+		TRACE_CPUPROFILER_EVENT_SCOPE(SVT::StreamingFilterRequests);
 
 		for (const auto& Request : RequestsHashTable)
 		{
@@ -942,7 +953,7 @@ void FStreamingManager::IssueRequests(int32 MaxSelectedRequests)
 		return;
 	}
 
-	TRACE_CPUPROFILER_EVENT_SCOPE(SVT::IssueRequests);
+	TRACE_CPUPROFILER_EVENT_SCOPE(SVT::StreamingIssueRequests);
 
 #if WITH_EDITORONLY_DATA
 	TArray<FCacheGetChunkRequest> DDCRequests;
@@ -1202,7 +1213,7 @@ int32 FStreamingManager::DetermineReadyMipLevels()
 {
 	using namespace UE::DerivedData;
 
-	TRACE_CPUPROFILER_EVENT_SCOPE(SVT::DetermineReadyMipLevels);
+	TRACE_CPUPROFILER_EVENT_SCOPE(SVT::StreamingDetermineReadyRequests);
 
 	const int32 StartPendingMipLevelIndex = (NextPendingMipLevelIndex + MaxPendingMipLevels - NumPendingMipLevels) % MaxPendingMipLevels;
 	int32 NumReadyMipLevels = 0;
@@ -1302,7 +1313,7 @@ void FStreamingManager::InstallReadyMipLevels()
 		return;
 	}
 
-	TRACE_CPUPROFILER_EVENT_SCOPE(SVT::InstallReadyMipLevels);
+	TRACE_CPUPROFILER_EVENT_SCOPE(SVT::StreamingInstallReadyRequests);
 
 	UploadTasks.Reset();
 	UploadTasks.Reserve(AsyncState.NumReadyMipLevels * 2 /*slack for splitting large uploads*/);
@@ -1440,13 +1451,11 @@ void FStreamingManager::InstallReadyMipLevels()
 	// Do all the memcpy's in parallel
 	ParallelFor(UploadTasks.Num(), [&](int32 TaskIndex)
 		{
-			TRACE_CPUPROFILER_EVENT_SCOPE(SVT::FUploadTask);
-
 			FUploadTask& Task = UploadTasks[TaskIndex];
 
 			if (Task.Union.HasSubtype<FUploadTask::FPageTableTask>())
 			{
-				TRACE_CPUPROFILER_EVENT_SCOPE(SVT::PageTableUpload);
+				TRACE_CPUPROFILER_EVENT_SCOPE(SVT::StreamingPageTableUpload);
 				FUploadTask::FPageTableTask& PageTableTask = Task.Union.GetSubtype<FUploadTask::FPageTableTask>();
 				if (PageTableTask.NumPageTableUpdates > 0)
 				{
@@ -1467,7 +1476,7 @@ void FStreamingManager::InstallReadyMipLevels()
 			}
 			else if(Task.Union.HasSubtype<FUploadTask::FTileDataTask>())
 			{
-				TRACE_CPUPROFILER_EVENT_SCOPE(SVT::TileDataUpload);
+				TRACE_CPUPROFILER_EVENT_SCOPE(SVT::StreamingTileDataUpload);
 				FUploadTask::FTileDataTask& TileDataTask = Task.Union.GetSubtype<FUploadTask::FTileDataTask>();
 				for (int32 i = 0; i < 2; ++i)
 				{
@@ -1499,7 +1508,7 @@ void FStreamingManager::InstallReadyMipLevels()
 
 	ParallelFor(UploadCleanupTasks.Num(), [&](int32 TaskIndex)
 		{
-			TRACE_CPUPROFILER_EVENT_SCOPE(SVT::FUploadCleanupTask);
+			TRACE_CPUPROFILER_EVENT_SCOPE(SVT::StreamingUploadCleanupTask);
 
 			FPendingMipLevel* PendingMipLevel = UploadCleanupTasks[TaskIndex];
 #if WITH_EDITORONLY_DATA
@@ -1531,6 +1540,8 @@ void FStreamingManager::InstallReadyMipLevels()
 
 void FStreamingManager::PatchPageTable(FRDGBuilder& GraphBuilder)
 {
+	TRACE_CPUPROFILER_EVENT_SCOPE(SVT::StreamingPatchPageTable);
+
 	int32 NumUpdates = 0;
 
 	// Generate bitsets of invalidated pages for every frame.
