Fix intermittent hang at end of cook caused by shader compiler's completion callback destroying its own RequestOwner leading later to a hang in DDC's WaitForQuiescence method.

diff --git a/Engine/Source/Runtime/Engine/Private/ShaderCompiler/ShaderCompiler.cpp b/Engine/Source/Runtime/Engine/Private/ShaderCompiler/ShaderCompiler.cpp
--- a/Engine/Source/Runtime/Engine/Private/ShaderCompiler/ShaderCompiler.cpp
+++ b/Engine/Source/Runtime/Engine/Private/ShaderCompiler/ShaderCompiler.cpp
@@ -1886,7 +1886,11 @@ void FShaderJobCache::AddToCacheAndProcessPending(FShaderCommonCompileJob* Finis
 		}
 
 		// remove ourselves from the jobs in flight
-		JobData.JobInFlight = nullptr;
+		if (JobData.JobInFlight)
+		{
+			JobData.JobInFlight->RequestOwner->KeepAlive();
+			JobData.JobInFlight = nullptr;
+		}
 		FinishedJob->JobCacheRef.Clear();
 	}
 
@@ -10813,7 +10817,11 @@ FShaderJobCacheRef FShaderJobCache::FindOrAdd(const FJobInputHash& Hash, EShader
 								UE_LOG(LogShaderCompilers, Display, TEXT("Cancelled job 0x%p (data 0x%p) with pending DDC hit."), JobDataPtr->JobInFlight.GetReference(), JobDataPtr);
 
 								delete NewStoredOutput;
-								JobDataPtr->JobInFlight = nullptr;
+								if (JobDataPtr->JobInFlight)
+								{
+									JobDataPtr->JobInFlight->RequestOwner->KeepAlive();
+									JobDataPtr->JobInFlight = nullptr;
+								}
 								JobLock.WriteUnlock();
 								return;
 							}
@@ -10860,7 +10868,11 @@ FShaderJobCacheRef FShaderJobCache::FindOrAdd(const FJobInputHash& Hash, EShader
 								CurHead = CurHead->NextLink;
 							}
 							JobDataPtr->DuplicateJobsWaitList = nullptr;
-							JobDataPtr->JobInFlight = nullptr;
+							if (JobDataPtr->JobInFlight)
+							{
+								JobDataPtr->JobInFlight->RequestOwner->KeepAlive();
+								JobDataPtr->JobInFlight = nullptr;
+							}
 							Job->JobCacheRef.Clear();
 
 							// Need to release the lock before calling ProcessFinishedJobs
@@ -10894,7 +10906,11 @@ FShaderJobCacheRef FShaderJobCache::FindOrAdd(const FJobInputHash& Hash, EShader
 							{
 								UE_LOG(LogShaderCompilers, Display, TEXT("Cancelled job 0x%p (data 0x%p) with pending DDC miss."), Job, JobDataPtr);
 
-								JobDataPtr->JobInFlight = nullptr;
+								if (JobDataPtr->JobInFlight)
+								{
+									JobDataPtr->JobInFlight->RequestOwner->KeepAlive();
+									JobDataPtr->JobInFlight = nullptr;
+								}
 								return;
 							}
 							else
