Changes ConformSparseClassData to iterate and copy each class property instead of relying on CopyScriptStruct. This avoids losing property data during a conform due whee a native copy operator implementation does not actually copy all properties.

diff --git a/Engine/Source/Runtime/Engine/Private/BlueprintGeneratedClass.cpp b/Engine/Source/Runtime/Engine/Private/BlueprintGeneratedClass.cpp
--- a/Engine/Source/Runtime/Engine/Private/BlueprintGeneratedClass.cpp
+++ b/Engine/Source/Runtime/Engine/Private/BlueprintGeneratedClass.cpp
@@ -978,7 +978,13 @@ void UBlueprintGeneratedClass::ConformSparseClassData(UObject* Object)
 			}
 			if (SparseClassDataStructToCopy)
 			{
-				SparseClassDataStructToCopy->CopyScriptStruct(GetOrCreateSparseClassData(), SparseClassDataPendingConform);
+				// Copy all properties from SparseClassDataPendingConform into the current sparse class data struct.
+				// NOTE: Avoids the copy assignment operators of the class properties, as they may not preserve all data.
+				void* const ThisSparseClassData = GetOrCreateSparseClassData();
+				for (TFieldIterator<FProperty> It(SparseClassDataStructToCopy); It; ++It)
+				{
+					It->CopyCompleteValue_InContainer(ThisSparseClassData, SparseClassDataPendingConform);
+				}
 			}
 		}
 
