#ue5 - Reject expired moves in UCharacterMovementComponent::VerifyClientTimeStamp() when we are doing forced position updates on the server. Should be rare but networking could deliver packets very late and then cause more corrections since the older timestamps were used on the server. Also fix handling of timestamp resets (very old timestamps sent on purpose) that arrive during ForcePositionUpdate, so they correctly trigger timestamp reset calbacks and sync the time correctly. Fixed APlayerController::TickActor() not resetting forced update state once good new moves had been received, it was waiting on the max timeout period first.

#tests PIE player movement with network emulation inlcluding up to Bad and Extreme profiles, as well as with "p.NetForceClientServerMoveLossPercent" and "p.NetForceClientServerMoveLossDuration" dropping client moves varying between 1%-10% of moves and 0.05sec - 0.75sec duration. Forced MinTimeBetweenTimeStampResets artificially low to every 20 seconds to stress the reset during forced corrections. Iris enabled.

diff --git a/Engine/Source/Runtime/Engine/Private/Components/CharacterMovementComponent.cpp b/Engine/Source/Runtime/Engine/Private/Components/CharacterMovementComponent.cpp
--- a/Engine/Source/Runtime/Engine/Private/Components/CharacterMovementComponent.cpp
+++ b/Engine/Source/Runtime/Engine/Private/Components/CharacterMovementComponent.cpp
@@ -8993,22 +8993,10 @@ bool UCharacterMovementComponent::VerifyClientTimeStamp(float TimeStamp, FNetwor
 	UWorld* World = GetWorld();
 
 	const bool bFirstMoveAfterForcedUpdates = ServerData.bTriggeringForcedUpdates;
-	if (bFirstMoveAfterForcedUpdates)
-	{
-		// We have been performing ForcedUpdates because we hadn't received any moves from this connection in a while but we've now received a new move!
-		// Let's sync up to this TimeStamp in order to resolve movement desyncs ASAP
-		// This will result in this move having a zero DeltaTime, so it will perform no movement but it will send a correction, and we should be able to process the next move that arrives.
-		UE_LOG(LogNetPlayerMovement, Log, TEXT("Received a new move after performing ForcedUpdates.  Updating CurrentClientTimeStamp from %f to %f"), ServerData.CurrentClientTimeStamp, TimeStamp);
-		ServerData.CurrentClientTimeStamp = TimeStamp;
-		if (World != nullptr)
-		{
-			ServerData.ServerTimeStamp = World->GetTimeSeconds();
-		}
-	}
-
+	
 	bool bTimeStampResetDetected = false;
 	bool bNeedsForcedUpdate = false;
-	const bool bIsValid = bFirstMoveAfterForcedUpdates || IsClientTimeStampValid(TimeStamp, ServerData, bTimeStampResetDetected);
+	const bool bIsValid = IsClientTimeStampValid(TimeStamp, ServerData, bTimeStampResetDetected);
 	if (bIsValid)
 	{
 		if (bTimeStampResetDetected)
@@ -9029,6 +9017,20 @@ bool UCharacterMovementComponent::VerifyClientTimeStamp(float TimeStamp, FNetwor
 			UE_LOG(LogNetPlayerMovement, VeryVerbose, TEXT("TimeStamp %f Accepted! CurrentTimeStamp: %f"), TimeStamp, ServerData.CurrentClientTimeStamp);
 			ProcessClientTimeStampForTimeDiscrepancy(TimeStamp, ServerData);
 		}
+
+		if (bFirstMoveAfterForcedUpdates)
+		{
+			// We have been performing ForcedUpdates because we hadn't received any moves from this connection in a while but we've now received a new move!
+			// Let's sync up to this TimeStamp in order to resolve movement desyncs ASAP
+			// This will result in this move having a zero DeltaTime (TimeStamp - CurrentClientTimeStamp), so it will perform no movement but it will send a correction,
+			// and we should be able to process the next move that arrives, which will advance ServerTimeStamp and unlock the forced updates (see 'APlayerController::TickActor()').
+			UE_LOG(LogNetPlayerMovement, Log, TEXT("Received a new move after performing ForcedUpdates. Updating CurrentClientTimeStamp from %f to %f"), ServerData.CurrentClientTimeStamp, TimeStamp);
+			ServerData.CurrentClientTimeStamp = TimeStamp;
+			if (World != nullptr)
+			{
+				ServerData.ServerTimeStamp = World->GetTimeSeconds();
+			}
+		}
 	}
 	else
 	{
@@ -9038,7 +9040,8 @@ bool UCharacterMovementComponent::VerifyClientTimeStamp(float TimeStamp, FNetwor
 		}
 		else
 		{
-			bNeedsForcedUpdate = (TimeStamp <= ServerData.LastReceivedClientTimeStamp);
+			// NOTE: Disabled: can be dangerous, keeping server advancing and client can't catch up. Has been effectively disabled since CL 16194374, but other changes could enable it again.
+			//bNeedsForcedUpdate = (TimeStamp <= ServerData.LastReceivedClientTimeStamp);
 		}
 	}
 
diff --git a/Engine/Source/Runtime/Engine/Private/PlayerController.cpp b/Engine/Source/Runtime/Engine/Private/PlayerController.cpp
--- a/Engine/Source/Runtime/Engine/Private/PlayerController.cpp
+++ b/Engine/Source/Runtime/Engine/Private/PlayerController.cpp
@@ -5097,17 +5097,18 @@ void APlayerController::TickActor( float DeltaSeconds, ELevelTick TickType, FAct
 						ServerData->bForcedUpdateDurationExceeded = false;
 						if (ServerData->bTriggeringForcedUpdates)
 						{
-							const float PawnTimeSinceForcingUpdates = (WorldTimeStamp - ServerData->ServerTimeBeginningForcedUpdates) * GetPawn()->CustomTimeDilation;
-							const float PawnTimeForcedUpdateMaxDuration = ForcedUpdateMaxDuration * GetPawn()->GetActorTimeDilation();
-							
-							if (PawnTimeSinceForcingUpdates > PawnTimeForcedUpdateMaxDuration)
+							if (ServerData->ServerTimeStamp > ServerData->ServerTimeLastForcedUpdate)
 							{
-								if (ServerData->ServerTimeStamp > ServerData->ServerTimeLastForcedUpdate)
-								{
-									// An update came in that was not a forced update (ie a real move), since ServerTimeStamp advanced outside this code.
-									ServerData->ResetForcedUpdateState();
-								}
-								else
+								// An update came in that was not a forced update (ie a real move), since ServerTimeStamp advanced outside this code.
+								UE_LOG(LogNetPlayerMovement, Log, TEXT("Movement detected, resetting forced update state (ServerTimeStamp %.6f > ServerTimeLastForcedUpdate %.6f)"), ServerData->ServerTimeStamp, ServerData->ServerTimeLastForcedUpdate);
+								ServerData->ResetForcedUpdateState();
+							}
+							else
+							{
+								const float PawnTimeSinceForcingUpdates = (WorldTimeStamp - ServerData->ServerTimeBeginningForcedUpdates) * GetPawn()->CustomTimeDilation;
+								const float PawnTimeForcedUpdateMaxDuration = ForcedUpdateMaxDuration * GetPawn()->GetActorTimeDilation();
+
+								if (PawnTimeSinceForcingUpdates > PawnTimeForcedUpdateMaxDuration)
 								{
 									if (ServerData->bLastRequestNeedsForcedUpdates)
 									{
@@ -5117,6 +5118,7 @@ void APlayerController::TickActor( float DeltaSeconds, ELevelTick TickType, FAct
 									else
 									{
 										// Waiting for ServerTimeStamp to advance from a client move.
+										UE_LOG(LogNetPlayerMovement, Log, TEXT("Setting bForcedUpdateDurationExceeded=true (PawnTimeSinceForcingUpdates %.6f > PawnTimeForcedUpdateMaxDuration %.6f)"), PawnTimeSinceForcingUpdates, PawnTimeForcedUpdateMaxDuration);
 										ServerData->bForcedUpdateDurationExceeded = true;
 									}
 								}
