Export the new functions added for converting multidimensional array indices, this facilitates native code calling these functions
#rnx

diff --git a/Engine/Source/Runtime/Engine/Classes/Kismet/KismetMathLibrary.h b/Engine/Source/Runtime/Engine/Classes/Kismet/KismetMathLibrary.h
--- a/Engine/Source/Runtime/Engine/Classes/Kismet/KismetMathLibrary.h
+++ b/Engine/Source/Runtime/Engine/Classes/Kismet/KismetMathLibrary.h
@@ -4568,7 +4568,7 @@ class UKismetMathLibrary : public UBlueprintFunctionLibrary
 	 * @return The equivalent 2D index of the array
 	 */
 	UFUNCTION(BlueprintPure, Category="Math|Conversions|Indices", meta=(DisplayName="Convert a 1D Index to a 2D Index"))
-	static FIntPoint Convert1DTo2D(int32 Index1D, int32 XSize);
+	static ENGINE_API FIntPoint Convert1DTo2D(int32 Index1D, int32 XSize);
 
 	/**
 	 *
@@ -4581,7 +4581,7 @@ class UKismetMathLibrary : public UBlueprintFunctionLibrary
 	 * @return The equivalent 3D index of the array
 	 */
 	UFUNCTION(BlueprintPure, Category="Math|Conversions|Indices", meta=(DisplayName="Convert a 1D Index to a 3D Index"))
-	static FIntVector Convert1DTo3D(int32 Index1D, int32 XSize, int32 YSize);
+	static ENGINE_API FIntVector Convert1DTo3D(int32 Index1D, int32 XSize, int32 YSize);
 
 	/**
 	 *
@@ -4593,7 +4593,7 @@ class UKismetMathLibrary : public UBlueprintFunctionLibrary
 	 * @return The equivalent 1D index of the array
 	 */
 	UFUNCTION(BlueprintPure, Category="Math|Conversions|Indices", meta=(DisplayName="Convert a 2D Index to a 1D Index"))
-	static int32 Convert2DTo1D(const FIntPoint& Index2D, int32 XSize);
+	static ENGINE_API int32 Convert2DTo1D(const FIntPoint& Index2D, int32 XSize);
 
 	/**
 	 *
@@ -4606,7 +4606,7 @@ class UKismetMathLibrary : public UBlueprintFunctionLibrary
 	 * @return The equivalent 1D index of the array
 	 */
 	UFUNCTION(BlueprintPure, Category="Math|Conversions|Indices", meta=(DisplayName="Convert a 3D Index to a 1D Index"))
-	static int32 Convert3DTo1D(const FIntVector& Index3D, int32 XSize, int32 YSize);
+	static ENGINE_API int32 Convert3DTo1D(const FIntVector& Index3D, int32 XSize, int32 YSize);
 
 	// NetQuantized vector make/breaks
 	UFUNCTION(BlueprintPure, Category = "Math|Vector", meta = (NativeMakeFunc))
