[ChaosVD] Fix CIS issue in ChaosVDSQTraceHelper.cpp.

FCollisionQueryParams::GetIgnoredActors was deprecated in main, so after the merge this resulted in a warning.

diff --git a/Engine/Source/Runtime/Engine/Private/Collision/ChaosVDSQTraceHelper.cpp b/Engine/Source/Runtime/Engine/Private/Collision/ChaosVDSQTraceHelper.cpp
--- a/Engine/Source/Runtime/Engine/Private/Collision/ChaosVDSQTraceHelper.cpp
+++ b/Engine/Source/Runtime/Engine/Private/Collision/ChaosVDSQTraceHelper.cpp
@@ -83,9 +83,9 @@ int32 GetWorldSolverID(const UWorld* World)
 		FChaosVDCollisionQueryParams CVDCollisionParams;
 		CVDCollisionParams.CopyFrom(Params);
 
-		CVDCollisionParams.IgnoredActorsNames.Reserve(Params.GetIgnoredActors().Num());
+		CVDCollisionParams.IgnoredActorsNames.Reserve(Params.GetIgnoredSourceObjects().Num());
 		CVDCollisionParams.IgnoredComponentsNames.Reserve(Params.GetIgnoredComponents().Num());
-		Algo::Transform(Params.GetIgnoredActors(), CVDCollisionParams.IgnoredActorsNames, [](uint32 ID){ return GetFNameSafe(GetObjetPointerFromID<AActor>(ID)); } );
+		Algo::Transform(Params.GetIgnoredSourceObjects(), CVDCollisionParams.IgnoredActorsNames, [](uint32 ID){ return GetFNameSafe(GetObjetPointerFromID<AActor>(ID)); } );
 		Algo::Transform(Params.GetIgnoredComponents(), CVDCollisionParams.IgnoredComponentsNames, [](uint32 ID){ return GetFNameSafe(GetObjetPointerFromID<UActorComponent>(ID)); } );
 
 		FChaosVDCollisionResponseParams CVDCollisionResponseParams;
