Fix crashes in static and skeletal mesh streaming systems
 - Removed FRHIAsyncCommandList. Use of this helper was dangerous as the resulting RHICmdList is enqueued for execution via ENQUEUE_RENDER_COMMAND, which doesn't guarantee correct ordering vs other RHI commands.
 - Removed deprecated functions that made use of this async helper type.
 - Reimplemented mesh streaming jobs to create a local RHICmdList to use for resource creation / initialization on the async thread, which is then submitted on the render thread in DoFinishUpdate. This guarantees command ordering in all cases (e.g. with bypass / no RHI thread etc).
 - Most raytracing operations have been kept on the render thread, as they are not yet threadsafe to run on the async streaming threads.
 - Added a breadcrumb in CreateBuffers to show which static / skeletal mesh is getting streamed in Insights etc traces.

Renamed "update RHI resource" APIs used by mesh streaming to "replace RHI resource" (e.g. RHICmdList.UpdateRHIResources() -> RHICmdList.ReplaceResources(), TRHIResourceUpdateBatcher -> FRHIResourceReplaceBatcher)
 - Old naming was generic and didn't describe what was actually happening.
 - Also removed the reference counting done by the API now that the mesh streamer holds references until after DoFinishUpdate is called.

diff --git a/Engine/Source/Runtime/Engine/Private/Animation/SkinWeightProfile.cpp b/Engine/Source/Runtime/Engine/Private/Animation/SkinWeightProfile.cpp
--- a/Engine/Source/Runtime/Engine/Private/Animation/SkinWeightProfile.cpp
+++ b/Engine/Source/Runtime/Engine/Private/Animation/SkinWeightProfile.cpp
@@ -748,18 +748,7 @@ void FSkinWeightProfilesData::ApplyOverrideProfile(FSkinWeightVertexBuffer* Over
 	}	
 }
 
-void FSkinWeightProfilesData::CreateRHIBuffers_RenderThread(TArray<TPair<FName, FSkinWeightRHIInfo>>& OutBuffers)
-{
-	CreateRHIBuffers(FRHICommandListImmediate::Get(), OutBuffers);
-}
-
-void FSkinWeightProfilesData::CreateRHIBuffers_Async(TArray<TPair<FName, FSkinWeightRHIInfo>>& OutBuffers)
-{
-	FRHIAsyncCommandList CommandList;
-	CreateRHIBuffers(*CommandList, OutBuffers);
-}
-
-void FSkinWeightProfilesData::InitRHIForStreaming(const TArray<TPair<FName, FSkinWeightRHIInfo>>& IntermediateBuffers, FRHIResourceUpdateBatcher& Batcher)
+void FSkinWeightProfilesData::InitRHIForStreaming(const TArray<TPair<FName, FSkinWeightRHIInfo>>& IntermediateBuffers, FRHIResourceReplaceBatcher& Batcher)
 {
 	for (int32 Idx = 0; Idx < IntermediateBuffers.Num(); ++Idx)
 	{
@@ -769,7 +758,7 @@ void FSkinWeightProfilesData::InitRHIForStreaming(const TArray<TPair<FName, FSki
 	}
 }
 
-void FSkinWeightProfilesData::ReleaseRHIForStreaming(FRHIResourceUpdateBatcher& Batcher)
+void FSkinWeightProfilesData::ReleaseRHIForStreaming(FRHIResourceReplaceBatcher& Batcher)
 {
 	for (TMap<FName, FSkinWeightVertexBuffer*>::TIterator It(ProfileNameToBuffer); It; ++It)
 	{
diff --git a/Engine/Source/Runtime/Engine/Private/PositionVertexBuffer.cpp b/Engine/Source/Runtime/Engine/Private/PositionVertexBuffer.cpp
--- a/Engine/Source/Runtime/Engine/Private/PositionVertexBuffer.cpp
+++ b/Engine/Source/Runtime/Engine/Private/PositionVertexBuffer.cpp
@@ -6,7 +6,7 @@
 #include "DataDrivenShaderPlatformInfo.h"
 #include "RenderingThread.h"
 #include "RenderUtils.h"
-#include "RHIResourceUpdates.h"
+#include "RHIResourceReplace.h"
 #include "StaticMeshVertexData.h"
 
 /*-----------------------------------------------------------------------------
@@ -210,30 +210,19 @@ FBufferRHIRef FPositionVertexBuffer::CreateRHIBuffer(FRHICommandListBase& RHICmd
 	return FRenderResource::CreateRHIBuffer(RHICmdList, VertexData, NumVertices, BUF_Static | BUF_ShaderResource, TEXT("FPositionVertexBuffer"));
 }
 
-FBufferRHIRef FPositionVertexBuffer::CreateRHIBuffer_RenderThread()
-{
-	return CreateRHIBuffer(FRHICommandListExecutor::GetImmediateCommandList());
-}
-
-FBufferRHIRef FPositionVertexBuffer::CreateRHIBuffer_Async()
-{
-	FRHIAsyncCommandList CommandList;
-	return CreateRHIBuffer(*CommandList);
-}
-
-void FPositionVertexBuffer::InitRHIForStreaming(FRHIBuffer* IntermediateBuffer, FRHIResourceUpdateBatcher& Batcher)
+void FPositionVertexBuffer::InitRHIForStreaming(FRHIBuffer* IntermediateBuffer, FRHIResourceReplaceBatcher& Batcher)
 {
 	check(VertexBufferRHI);
 	if (IntermediateBuffer)
 	{
-		Batcher.QueueUpdateRequest(VertexBufferRHI, IntermediateBuffer);
+		Batcher.EnqueueReplace(VertexBufferRHI, IntermediateBuffer);
 	}
 }
 
-void FPositionVertexBuffer::ReleaseRHIForStreaming(FRHIResourceUpdateBatcher& Batcher)
+void FPositionVertexBuffer::ReleaseRHIForStreaming(FRHIResourceReplaceBatcher& Batcher)
 {
 	check(VertexBufferRHI);
-	Batcher.QueueUpdateRequest(VertexBufferRHI, nullptr);
+	Batcher.EnqueueReplace(VertexBufferRHI, nullptr);
 }
 
 void FPositionVertexBuffer::InitRHI(FRHICommandListBase& RHICmdList)
diff --git a/Engine/Source/Runtime/Engine/Private/RawIndexBuffer.cpp b/Engine/Source/Runtime/Engine/Private/RawIndexBuffer.cpp
--- a/Engine/Source/Runtime/Engine/Private/RawIndexBuffer.cpp
+++ b/Engine/Source/Runtime/Engine/Private/RawIndexBuffer.cpp
@@ -11,7 +11,7 @@
 #include "RenderUtils.h"
 #include "Modules/ModuleManager.h"
 #include "RenderingThread.h"
-#include "RHIResourceUpdates.h"
+#include "RHIResourceReplace.h"
 
 #if WITH_EDITOR
 #include "MeshUtilities.h"
@@ -353,30 +353,19 @@ FBufferRHIRef FRawStaticIndexBuffer::CreateRHIBuffer(FRHICommandListBase& RHICmd
 	return nullptr;
 }
 
-FBufferRHIRef FRawStaticIndexBuffer::CreateRHIBuffer_RenderThread()
-{
-	return CreateRHIBuffer(FRHICommandListExecutor::GetImmediateCommandList());
-}
-
-FBufferRHIRef FRawStaticIndexBuffer::CreateRHIBuffer_Async()
-{
-	FRHIAsyncCommandList CommandList;
-	return CreateRHIBuffer(*CommandList);
-}
-
-void FRawStaticIndexBuffer::InitRHIForStreaming(FRHIBuffer* IntermediateBuffer, FRHIResourceUpdateBatcher& Batcher)
+void FRawStaticIndexBuffer::InitRHIForStreaming(FRHIBuffer* IntermediateBuffer, FRHIResourceReplaceBatcher& Batcher)
 {
 	if (IndexBufferRHI && IntermediateBuffer)
 	{
-		Batcher.QueueUpdateRequest(IndexBufferRHI, IntermediateBuffer);
+		Batcher.EnqueueReplace(IndexBufferRHI, IntermediateBuffer);
 	}
 }
 
-void FRawStaticIndexBuffer::ReleaseRHIForStreaming(FRHIResourceUpdateBatcher& Batcher)
+void FRawStaticIndexBuffer::ReleaseRHIForStreaming(FRHIResourceReplaceBatcher& Batcher)
 {
 	if (IndexBufferRHI)
 	{
-		Batcher.QueueUpdateRequest(IndexBufferRHI, nullptr);
+		Batcher.EnqueueReplace(IndexBufferRHI, nullptr);
 	}
 }
 
@@ -454,19 +443,19 @@ bool FRawStaticIndexBuffer16or32Interface::IsSRVNeeded(bool bAllowCPUAccess) con
 	return bSRV;
 }
 
-void FRawStaticIndexBuffer16or32Interface::InitRHIForStreaming(FRHIBuffer* IntermediateBuffer, size_t IndexSize, FRHIResourceUpdateBatcher& Batcher)
+void FRawStaticIndexBuffer16or32Interface::InitRHIForStreaming(FRHIBuffer* IntermediateBuffer, size_t IndexSize, FRHIResourceReplaceBatcher& Batcher)
 {
 	if (IndexBufferRHI && IntermediateBuffer)
 	{
-		Batcher.QueueUpdateRequest(IndexBufferRHI, IntermediateBuffer);
+		Batcher.EnqueueReplace(IndexBufferRHI, IntermediateBuffer);
 	}
 }
 
-void FRawStaticIndexBuffer16or32Interface::ReleaseRHIForStreaming(FRHIResourceUpdateBatcher& Batcher)
+void FRawStaticIndexBuffer16or32Interface::ReleaseRHIForStreaming(FRHIResourceReplaceBatcher& Batcher)
 {
 	if (IndexBufferRHI)
 	{
-		Batcher.QueueUpdateRequest(IndexBufferRHI, nullptr);
+		Batcher.EnqueueReplace(IndexBufferRHI, nullptr);
 	}
 }
 
diff --git a/Engine/Source/Runtime/Engine/Private/Rendering/ColorVertexBuffer.cpp b/Engine/Source/Runtime/Engine/Private/Rendering/ColorVertexBuffer.cpp
--- a/Engine/Source/Runtime/Engine/Private/Rendering/ColorVertexBuffer.cpp
+++ b/Engine/Source/Runtime/Engine/Private/Rendering/ColorVertexBuffer.cpp
@@ -9,7 +9,7 @@
 #include "VertexFactory.h"
 #include "DataDrivenShaderPlatformInfo.h"
 #include "GlobalRenderResources.h"
-#include "RHIResourceUpdates.h"
+#include "RHIResourceReplace.h"
 
 /*-----------------------------------------------------------------------------
 FColorVertexBuffer
@@ -396,30 +396,19 @@ FBufferRHIRef FColorVertexBuffer::CreateRHIBuffer(FRHICommandListBase& RHICmdLis
 	return FRenderResource::CreateRHIBuffer(RHICmdList, VertexData, NumVertices, BUF_Static | BUF_ShaderResource, TEXT("FColorVertexBuffer"));
 }
 
-FBufferRHIRef FColorVertexBuffer::CreateRHIBuffer_RenderThread()
-{
-	return CreateRHIBuffer(FRHICommandListImmediate::Get());
-}
-
-FBufferRHIRef FColorVertexBuffer::CreateRHIBuffer_Async()
-{
-	FRHIAsyncCommandList CommandList;
-	return CreateRHIBuffer(*CommandList);
-}
-
-void FColorVertexBuffer::InitRHIForStreaming(FRHIBuffer* IntermediateBuffer, FRHIResourceUpdateBatcher& Batcher)
+void FColorVertexBuffer::InitRHIForStreaming(FRHIBuffer* IntermediateBuffer, FRHIResourceReplaceBatcher& Batcher)
 {
 	if (VertexBufferRHI && IntermediateBuffer)
 	{
-		Batcher.QueueUpdateRequest(VertexBufferRHI, IntermediateBuffer);
+		Batcher.EnqueueReplace(VertexBufferRHI, IntermediateBuffer);
 	}
 }
 
-void FColorVertexBuffer::ReleaseRHIForStreaming(FRHIResourceUpdateBatcher& Batcher)
+void FColorVertexBuffer::ReleaseRHIForStreaming(FRHIResourceReplaceBatcher& Batcher)
 {
 	if (VertexBufferRHI)
 	{
-		Batcher.QueueUpdateRequest(VertexBufferRHI, nullptr);
+		Batcher.EnqueueReplace(VertexBufferRHI, nullptr);
 	}
 }
 
diff --git a/Engine/Source/Runtime/Engine/Private/Rendering/MultiSizeIndexContainer.cpp b/Engine/Source/Runtime/Engine/Private/Rendering/MultiSizeIndexContainer.cpp
--- a/Engine/Source/Runtime/Engine/Private/Rendering/MultiSizeIndexContainer.cpp
+++ b/Engine/Source/Runtime/Engine/Private/Rendering/MultiSizeIndexContainer.cpp
@@ -210,18 +210,7 @@ FBufferRHIRef FMultiSizeIndexContainer::CreateRHIBuffer(FRHICommandListBase& RHI
 	return nullptr;
 }
 
-FBufferRHIRef FMultiSizeIndexContainer::CreateRHIBuffer_RenderThread()
-{
-	return CreateRHIBuffer(FRHICommandListImmediate::Get());
-}
-
-FBufferRHIRef FMultiSizeIndexContainer::CreateRHIBuffer_Async()
-{
-	FRHIAsyncCommandList RHICmdList;
-	return CreateRHIBuffer(*RHICmdList);
-}
-
-void FMultiSizeIndexContainer::InitRHIForStreaming(FRHIBuffer* IntermediateBuffer, FRHIResourceUpdateBatcher& Batcher)
+void FMultiSizeIndexContainer::InitRHIForStreaming(FRHIBuffer* IntermediateBuffer, FRHIResourceReplaceBatcher& Batcher)
 {
 	check(!((uint32)!!IntermediateBuffer ^ (uint32)!!IndexBuffer));
 	if (IntermediateBuffer)
@@ -237,7 +226,7 @@ void FMultiSizeIndexContainer::InitRHIForStreaming(FRHIBuffer* IntermediateBuffe
 	}
 }
 
-void FMultiSizeIndexContainer::ReleaseRHIForStreaming(FRHIResourceUpdateBatcher& Batcher)
+void FMultiSizeIndexContainer::ReleaseRHIForStreaming(FRHIResourceReplaceBatcher& Batcher)
 {
 	if (IndexBuffer)
 	{
diff --git a/Engine/Source/Runtime/Engine/Private/Rendering/SkeletalMeshVertexClothBuffer.cpp b/Engine/Source/Runtime/Engine/Private/Rendering/SkeletalMeshVertexClothBuffer.cpp
--- a/Engine/Source/Runtime/Engine/Private/Rendering/SkeletalMeshVertexClothBuffer.cpp
+++ b/Engine/Source/Runtime/Engine/Private/Rendering/SkeletalMeshVertexClothBuffer.cpp
@@ -6,7 +6,7 @@
 #include "EngineUtils.h"
 #include "Experimental/Containers/HazardPointer.h"
 #include "ProfilingDebugging/LoadTimeTracker.h"
-#include "RHIResourceUpdates.h"
+#include "RHIResourceReplace.h"
 #include "RHI.h"
 
 /**
@@ -80,30 +80,19 @@ FBufferRHIRef FSkeletalMeshVertexClothBuffer::CreateRHIBuffer(FRHICommandListBas
 	return FRenderResource::CreateRHIBuffer(RHICmdList, VertexData, NumVertices, BUF_Static | BUF_ShaderResource, TEXT("FSkeletalMeshVertexClothBuffer"));
 }
 
-FBufferRHIRef FSkeletalMeshVertexClothBuffer::CreateRHIBuffer_RenderThread()
-{
-	return CreateRHIBuffer(FRHICommandListImmediate::Get());
-}
-
-FBufferRHIRef FSkeletalMeshVertexClothBuffer::CreateRHIBuffer_Async()
-{
-	FRHIAsyncCommandList CommandList;
-	return CreateRHIBuffer(*CommandList);
-}
-
-void FSkeletalMeshVertexClothBuffer::InitRHIForStreaming(FRHIBuffer* IntermediateBuffer, FRHIResourceUpdateBatcher& Batcher)
+void FSkeletalMeshVertexClothBuffer::InitRHIForStreaming(FRHIBuffer* IntermediateBuffer, FRHIResourceReplaceBatcher& Batcher)
 {
 	if (VertexBufferRHI && IntermediateBuffer)
 	{
-		Batcher.QueueUpdateRequest(VertexBufferRHI, IntermediateBuffer);
+		Batcher.EnqueueReplace(VertexBufferRHI, IntermediateBuffer);
 	}
 }
 
-void FSkeletalMeshVertexClothBuffer::ReleaseRHIForStreaming(FRHIResourceUpdateBatcher& Batcher)
+void FSkeletalMeshVertexClothBuffer::ReleaseRHIForStreaming(FRHIResourceReplaceBatcher& Batcher)
 {
 	if (VertexBufferRHI)
 	{
-		Batcher.QueueUpdateRequest(VertexBufferRHI, nullptr);
+		Batcher.EnqueueReplace(VertexBufferRHI, nullptr);
 	}
 }
 
diff --git a/Engine/Source/Runtime/Engine/Private/Rendering/SkinWeightVertexBuffer.cpp b/Engine/Source/Runtime/Engine/Private/Rendering/SkinWeightVertexBuffer.cpp
--- a/Engine/Source/Runtime/Engine/Private/Rendering/SkinWeightVertexBuffer.cpp
+++ b/Engine/Source/Runtime/Engine/Private/Rendering/SkinWeightVertexBuffer.cpp
@@ -5,7 +5,7 @@
 #include "EngineUtils.h"
 #include "Rendering/SkeletalMeshLODModel.h"
 #include "ProfilingDebugging/LoadTimeTracker.h"
-#include "RHIResourceUpdates.h"
+#include "RHIResourceReplace.h"
 #include "SkeletalMeshLegacyCustomVersions.h"
 #include "UObject/UE5MainStreamObjectVersion.h"
 #include "Rendering/RenderCommandPipes.h"
@@ -152,30 +152,19 @@ void FSkinWeightLookupVertexBuffer::ReleaseRHI()
 	FVertexBuffer::ReleaseRHI();
 }
 
-FBufferRHIRef FSkinWeightLookupVertexBuffer::CreateRHIBuffer_RenderThread()
-{
-	return CreateRHIBuffer(FRHICommandListImmediate::Get());
-}
-
-FBufferRHIRef FSkinWeightLookupVertexBuffer::CreateRHIBuffer_Async()
-{
-	FRHIAsyncCommandList CommandList;
-	return CreateRHIBuffer(*CommandList);
-}
-
-void FSkinWeightLookupVertexBuffer::InitRHIForStreaming(FRHIBuffer* IntermediateBuffer, FRHIResourceUpdateBatcher& Batcher)
+void FSkinWeightLookupVertexBuffer::InitRHIForStreaming(FRHIBuffer* IntermediateBuffer, FRHIResourceReplaceBatcher& Batcher)
 {
 	if (VertexBufferRHI && IntermediateBuffer)
 	{
-		Batcher.QueueUpdateRequest(VertexBufferRHI, IntermediateBuffer);
+		Batcher.EnqueueReplace(VertexBufferRHI, IntermediateBuffer);
 	}
 }
 
-void FSkinWeightLookupVertexBuffer::ReleaseRHIForStreaming(FRHIResourceUpdateBatcher& Batcher)
+void FSkinWeightLookupVertexBuffer::ReleaseRHIForStreaming(FRHIResourceReplaceBatcher& Batcher)
 {
 	if (VertexBufferRHI)
 	{
-		Batcher.QueueUpdateRequest(VertexBufferRHI, nullptr);
+		Batcher.EnqueueReplace(VertexBufferRHI, nullptr);
 	}
 }
 
@@ -400,30 +389,19 @@ FBufferRHIRef FSkinWeightDataVertexBuffer::CreateRHIBuffer(FRHICommandListBase&
 	return FRenderResource::CreateRHIBuffer(RHICmdList, WeightData, NumBoneWeights, BUF_Static | BUF_ShaderResource | BUF_SourceCopy, TEXT("FSkinWeightDataVertexBuffer"));
 }
 
-FBufferRHIRef FSkinWeightDataVertexBuffer::CreateRHIBuffer_RenderThread()
-{
-	return CreateRHIBuffer(FRHICommandListImmediate::Get());
-}
-
-FBufferRHIRef FSkinWeightDataVertexBuffer::CreateRHIBuffer_Async()
-{
-	FRHIAsyncCommandList CommandList;
-	return CreateRHIBuffer(*CommandList);
-}
-
-void FSkinWeightDataVertexBuffer::InitRHIForStreaming(FRHIBuffer* IntermediateBuffer, FRHIResourceUpdateBatcher& Batcher)
+void FSkinWeightDataVertexBuffer::InitRHIForStreaming(FRHIBuffer* IntermediateBuffer, FRHIResourceReplaceBatcher& Batcher)
 {
 	if (VertexBufferRHI && IntermediateBuffer)
 	{
-		Batcher.QueueUpdateRequest(VertexBufferRHI, IntermediateBuffer);
+		Batcher.EnqueueReplace(VertexBufferRHI, IntermediateBuffer);
 	}
 }
 
-void FSkinWeightDataVertexBuffer::ReleaseRHIForStreaming(FRHIResourceUpdateBatcher& Batcher)
+void FSkinWeightDataVertexBuffer::ReleaseRHIForStreaming(FRHIResourceReplaceBatcher& Batcher)
 {
 	if (VertexBufferRHI)
 	{
-		Batcher.QueueUpdateRequest(VertexBufferRHI, nullptr);
+		Batcher.EnqueueReplace(VertexBufferRHI, nullptr);
 	}
 }
 
@@ -712,20 +690,6 @@ FSkinWeightRHIInfo FSkinWeightVertexBuffer::CreateRHIBuffer(FRHICommandListBase&
 	return RHIInfo;
 }
 
-FSkinWeightRHIInfo FSkinWeightVertexBuffer::CreateRHIBuffer_RenderThread()
-{
-	return CreateRHIBuffer(FRHICommandListImmediate::Get());
-}
-
-FSkinWeightRHIInfo FSkinWeightVertexBuffer::CreateRHIBuffer_Async()
-{
-	FSkinWeightRHIInfo RHIInfo;
-	FRHIAsyncCommandList CommandList;
-	RHIInfo.DataVertexBufferRHI = DataVertexBuffer.CreateRHIBuffer(*CommandList);
-	RHIInfo.LookupVertexBufferRHI = LookupVertexBuffer.CreateRHIBuffer(*CommandList);
-	return RHIInfo;
-}
-
 GPUSkinBoneInfluenceType FSkinWeightVertexBuffer::GetBoneInfluenceType() const
 {
 	if (GetVariableBonesPerVertex())
diff --git a/Engine/Source/Runtime/Engine/Private/Rendering/StaticMeshVertexBuffer.cpp b/Engine/Source/Runtime/Engine/Private/Rendering/StaticMeshVertexBuffer.cpp
--- a/Engine/Source/Runtime/Engine/Private/Rendering/StaticMeshVertexBuffer.cpp
+++ b/Engine/Source/Runtime/Engine/Private/Rendering/StaticMeshVertexBuffer.cpp
@@ -8,7 +8,7 @@
 #include "LocalVertexFactory.h"
 #include "MeshUVChannelInfo.h"
 #include "ProfilingDebugging/LoadTimeTracker.h"
-#include "RHIResourceUpdates.h"
+#include "RHIResourceReplace.h"
 #include "StaticMeshVertexData.h"
 
 FStaticMeshVertexBuffer::FStaticMeshVertexBuffer() :
@@ -272,54 +272,32 @@ FBufferRHIRef FStaticMeshVertexBuffer::CreateTangentsRHIBuffer(FRHICommandListBa
 	return CreateRHIBuffer(RHICmdList, TangentsData, GetNumVertices(), BUF_Static | BUF_ShaderResource, TEXT("TangentsRHIBuffer"));
 }
 
-FBufferRHIRef FStaticMeshVertexBuffer::CreateTangentsRHIBuffer_RenderThread()
-{
-	return CreateTangentsRHIBuffer(FRHICommandListImmediate::Get());
-}
-
-FBufferRHIRef FStaticMeshVertexBuffer::CreateTangentsRHIBuffer_Async()
-{
-	FRHIAsyncCommandList CommandList;
-	return CreateTangentsRHIBuffer(*CommandList);
-}
-
 FBufferRHIRef FStaticMeshVertexBuffer::CreateTexCoordRHIBuffer(FRHICommandListBase& RHICmdList)
 {
 	return CreateRHIBuffer(RHICmdList, TexcoordData, GetNumTexCoords(), BUF_Static | BUF_ShaderResource, TEXT("TexCoordRHIBuffer"));
 }
 
-FBufferRHIRef FStaticMeshVertexBuffer::CreateTexCoordRHIBuffer_RenderThread()
-{
-	return CreateTexCoordRHIBuffer(FRHICommandListImmediate::Get());
-}
-
-FBufferRHIRef FStaticMeshVertexBuffer::CreateTexCoordRHIBuffer_Async()
-{
-	FRHIAsyncCommandList CommandList;
-	return CreateTexCoordRHIBuffer(*CommandList);
-}
-
 void FStaticMeshVertexBuffer::InitRHIForStreaming(
 	FRHIBuffer* IntermediateTangentsBuffer,
 	FRHIBuffer* IntermediateTexCoordBuffer,
-	FRHIResourceUpdateBatcher& Batcher)
+	FRHIResourceReplaceBatcher& Batcher)
 {
 	check(TangentsVertexBuffer.VertexBufferRHI && TexCoordVertexBuffer.VertexBufferRHI);
 	if (IntermediateTangentsBuffer)
 	{
-		Batcher.QueueUpdateRequest(TangentsVertexBuffer.VertexBufferRHI, IntermediateTangentsBuffer);
+		Batcher.EnqueueReplace(TangentsVertexBuffer.VertexBufferRHI, IntermediateTangentsBuffer);
 	}
 	if (IntermediateTexCoordBuffer)
 	{
-		Batcher.QueueUpdateRequest(TexCoordVertexBuffer.VertexBufferRHI, IntermediateTexCoordBuffer);
+		Batcher.EnqueueReplace(TexCoordVertexBuffer.VertexBufferRHI, IntermediateTexCoordBuffer);
 	}
 }
 
-void FStaticMeshVertexBuffer::ReleaseRHIForStreaming(FRHIResourceUpdateBatcher& Batcher)
+void FStaticMeshVertexBuffer::ReleaseRHIForStreaming(FRHIResourceReplaceBatcher& Batcher)
 {
 	check(TangentsVertexBuffer.VertexBufferRHI && TexCoordVertexBuffer.VertexBufferRHI);
-	Batcher.QueueUpdateRequest(TangentsVertexBuffer.VertexBufferRHI, nullptr);
-	Batcher.QueueUpdateRequest(TexCoordVertexBuffer.VertexBufferRHI, nullptr);
+	Batcher.EnqueueReplace(TangentsVertexBuffer.VertexBufferRHI, nullptr);
+	Batcher.EnqueueReplace(TexCoordVertexBuffer.VertexBufferRHI, nullptr);
 }
 
 void FStaticMeshVertexBuffer::InitRHI(FRHICommandListBase& RHICmdList)
diff --git a/Engine/Source/Runtime/Engine/Private/SkeletalMesh.cpp b/Engine/Source/Runtime/Engine/Private/SkeletalMesh.cpp
--- a/Engine/Source/Runtime/Engine/Private/SkeletalMesh.cpp
+++ b/Engine/Source/Runtime/Engine/Private/SkeletalMesh.cpp
@@ -1100,30 +1100,20 @@ bool USkeletalMesh::StreamIn(int32 NewMipCount, bool bHighPrio)
 	check(IsInGameThread());
 	if (!HasPendingInitOrStreaming() && CachedSRRState.StreamIn(NewMipCount))
 	{
+		FRenderAssetUpdate::EThreadType CreateResourcesThread = GRHISupportsAsyncTextureCreation
+			? FRenderAssetUpdate::TT_Async
+			: FRenderAssetUpdate::TT_Render;
+
 #if WITH_EDITOR
 		// If editor data is available for the current platform, and the package isn't actually cooked.
 		if (FPlatformProperties::HasEditorOnlyData() && !GetOutermost()->bIsCookedForEditor)
 		{
-			if (GRHISupportsAsyncTextureCreation)
-			{
-				PendingUpdate = new FSkeletalMeshStreamIn_DDC_Async(this);
-			}
-			else
-			{
-				PendingUpdate = new FSkeletalMeshStreamIn_DDC_RenderThread(this);
-			}
+			PendingUpdate = new FSkeletalMeshStreamIn_DDC(this, CreateResourcesThread);
 		}
 		else
 #endif
 		{
-			if (GRHISupportsAsyncTextureCreation)
-			{
-				PendingUpdate = new FSkeletalMeshStreamIn_IO_Async(this, bHighPrio);
-			}
-			else
-			{
-				PendingUpdate = new FSkeletalMeshStreamIn_IO_RenderThread(this, bHighPrio);
-			}
+			PendingUpdate = new FSkeletalMeshStreamIn_IO(this, bHighPrio, CreateResourcesThread);
 		}
 		return !PendingUpdate->IsCancelled();
 	}
diff --git a/Engine/Source/Runtime/Engine/Private/StaticMesh.cpp b/Engine/Source/Runtime/Engine/Private/StaticMesh.cpp
--- a/Engine/Source/Runtime/Engine/Private/StaticMesh.cpp
+++ b/Engine/Source/Runtime/Engine/Private/StaticMesh.cpp
@@ -7303,30 +7303,19 @@ bool UStaticMesh::StreamIn(int32 NewMipCount, bool bHighPrio)
 	check(IsInGameThread());
 	if (!HasPendingInitOrStreaming() && CachedSRRState.StreamIn(NewMipCount))
 	{
+		FRenderAssetUpdate::EThreadType CreateResourcesThread = GRHISupportsAsyncTextureCreation
+			? FRenderAssetUpdate::TT_Async
+			: FRenderAssetUpdate::TT_Render;
+
 #if WITH_EDITOR
 		if (FPlatformProperties::HasEditorOnlyData())
 		{
-			if (GRHISupportsAsyncTextureCreation)
-			{
-				PendingUpdate = new FStaticMeshStreamIn_DDC_Async(this);
-			}
-			else
-			{
-				PendingUpdate = new FStaticMeshStreamIn_DDC_RenderThread(this);
-			}
+			PendingUpdate = new FStaticMeshStreamIn_DDC(this, CreateResourcesThread);
 		}
 		else
 #endif
 		{
-			// When not using threaded rendering, rendercommands get executed on async thread which create issues on some RHI. See EnqueueUniqueRenderCommand() and IsInRenderingThread().
-			if (GRHISupportsAsyncTextureCreation && GIsThreadedRendering)
-			{
-				PendingUpdate = new FStaticMeshStreamIn_IO_Async(this, bHighPrio);
-			}
-			else
-			{
-				PendingUpdate = new FStaticMeshStreamIn_IO_RenderThread(this, bHighPrio);
-			}
+			PendingUpdate = new FStaticMeshStreamIn_IO(this, bHighPrio, CreateResourcesThread);
 		}
 		return !PendingUpdate->IsCancelled();
 	}
diff --git a/Engine/Source/Runtime/Engine/Private/Streaming/SkeletalMeshUpdate.cpp b/Engine/Source/Runtime/Engine/Private/Streaming/SkeletalMeshUpdate.cpp
--- a/Engine/Source/Runtime/Engine/Private/Streaming/SkeletalMeshUpdate.cpp
+++ b/Engine/Source/Runtime/Engine/Private/Streaming/SkeletalMeshUpdate.cpp
@@ -13,7 +13,7 @@ SkeletalMeshUpdate.cpp: Helpers to stream in and out skeletal mesh LODs.
 #include "Rendering/SkeletalMeshRenderData.h"
 #include "Components/SkinnedMeshComponent.h"
 #include "Streaming/RenderAssetUpdate.inl"
-#include "RHIResourceUpdates.h"
+#include "RHIResourceReplace.h"
 
 extern int32 GStreamingMaxReferenceChecks;
 
@@ -49,10 +49,8 @@ FSkeletalMeshUpdate::FSkeletalMeshUpdate(const USkeletalMesh* InMesh)
 {
 }
 
-void FSkeletalMeshStreamIn::FIntermediateBuffers::CreateFromCPUData_RenderThread(FSkeletalMeshLODRenderData& LODResource)
+void FSkeletalMeshStreamIn::FIntermediateBuffers::CreateFromCPUData(FRHICommandListBase& RHICmdList, FSkeletalMeshLODRenderData& LODResource)
 {
-	FRHICommandListBase& RHICmdList = FRHICommandListImmediate::Get();
-
 	FStaticMeshVertexBuffers& VBs = LODResource.StaticVertexBuffers;
 	TangentsVertexBuffer = VBs.StaticMeshVertexBuffer.CreateTangentsRHIBuffer(RHICmdList);
 	TexCoordVertexBuffer = VBs.StaticMeshVertexBuffer.CreateTexCoordRHIBuffer(RHICmdList);
@@ -64,35 +62,7 @@ void FSkeletalMeshStreamIn::FIntermediateBuffers::CreateFromCPUData_RenderThread
 	IndexBuffer = LODResource.MultiSizeIndexContainer.CreateRHIBuffer(RHICmdList);
 }
 
-void FSkeletalMeshStreamIn::FIntermediateBuffers::CreateFromCPUData_Async(FSkeletalMeshLODRenderData& LODResource)
-{
-	FRHIAsyncCommandList RHICmdList;
-
-	FStaticMeshVertexBuffers& VBs = LODResource.StaticVertexBuffers;
-	TangentsVertexBuffer = VBs.StaticMeshVertexBuffer.CreateTangentsRHIBuffer(*RHICmdList);
-	TexCoordVertexBuffer = VBs.StaticMeshVertexBuffer.CreateTexCoordRHIBuffer(*RHICmdList);
-	PositionVertexBuffer = VBs.PositionVertexBuffer.CreateRHIBuffer(*RHICmdList);
-	ColorVertexBuffer = VBs.ColorVertexBuffer.CreateRHIBuffer(*RHICmdList);
-	LODResource.SkinWeightProfilesData.CreateRHIBuffers(*RHICmdList, AltSkinWeightVertexBuffers);
-	SkinWeightVertexBuffer = LODResource.SkinWeightVertexBuffer.CreateRHIBuffer(*RHICmdList);
-	ClothVertexBuffer = LODResource.ClothVertexBuffer.CreateRHIBuffer(*RHICmdList);
-	IndexBuffer = LODResource.MultiSizeIndexContainer.CreateRHIBuffer(*RHICmdList);
-}
-
-void FSkeletalMeshStreamIn::FIntermediateBuffers::SafeRelease()
-{
-	TangentsVertexBuffer.SafeRelease();
-	TexCoordVertexBuffer.SafeRelease();
-	PositionVertexBuffer.SafeRelease();
-	ColorVertexBuffer.SafeRelease();
-	SkinWeightVertexBuffer.DataVertexBufferRHI.SafeRelease();
-	SkinWeightVertexBuffer.LookupVertexBufferRHI.SafeRelease();
-	ClothVertexBuffer.SafeRelease();
-	IndexBuffer.SafeRelease();
-	AltSkinWeightVertexBuffers.Empty();
-}
-
-void FSkeletalMeshStreamIn::FIntermediateBuffers::TransferBuffers(FSkeletalMeshLODRenderData& LODResource, FRHIResourceUpdateBatcher& Batcher)
+void FSkeletalMeshStreamIn::FIntermediateBuffers::TransferBuffers(FSkeletalMeshLODRenderData& LODResource, FRHIResourceReplaceBatcher& Batcher)
 {
 	FStaticMeshVertexBuffers& VBs = LODResource.StaticVertexBuffers;
 	VBs.StaticMeshVertexBuffer.InitRHIForStreaming(TangentsVertexBuffer, TexCoordVertexBuffer, Batcher);
@@ -102,24 +72,11 @@ void FSkeletalMeshStreamIn::FIntermediateBuffers::TransferBuffers(FSkeletalMeshL
 	LODResource.ClothVertexBuffer.InitRHIForStreaming(ClothVertexBuffer, Batcher);
 	LODResource.MultiSizeIndexContainer.InitRHIForStreaming(IndexBuffer, Batcher);
 	LODResource.SkinWeightProfilesData.InitRHIForStreaming(AltSkinWeightVertexBuffers, Batcher);
-	SafeRelease();
-}
-
-void FSkeletalMeshStreamIn::FIntermediateBuffers::CheckIsNull() const
-{
-	check(!TangentsVertexBuffer
-		&& !TexCoordVertexBuffer
-		&& !PositionVertexBuffer
-		&& !ColorVertexBuffer
-		&& !SkinWeightVertexBuffer.DataVertexBufferRHI
-		&& !SkinWeightVertexBuffer.LookupVertexBufferRHI
-		&& !ClothVertexBuffer
-		&& !IndexBuffer
-		&& !AltSkinWeightVertexBuffers.Num());
 }
 
-FSkeletalMeshStreamIn::FSkeletalMeshStreamIn(const USkeletalMesh* InMesh)
+FSkeletalMeshStreamIn::FSkeletalMeshStreamIn(const USkeletalMesh* InMesh, EThreadType CreateResourcesThread)
 	: FSkeletalMeshUpdate(InMesh)
+	, CreateResourcesThread(CreateResourcesThread)
 {
 	if (!ensure(PendingFirstLODIdx < CurrentFirstLODIdx))
 	{
@@ -129,117 +86,89 @@ FSkeletalMeshStreamIn::FSkeletalMeshStreamIn(const USkeletalMesh* InMesh)
 
 FSkeletalMeshStreamIn::~FSkeletalMeshStreamIn()
 {
-#if DO_CHECK
-	for (int32 Idx = 0; Idx < MAX_MESH_LOD_COUNT; ++Idx)
-	{
-		IntermediateBuffersArray[Idx].CheckIsNull();
-	}
-#endif
+	check(!StreamingRHICmdList);
 }
 
-template <bool bRenderThread>
-void FSkeletalMeshStreamIn::CreateBuffers_Internal(const FContext& Context)
+void FSkeletalMeshStreamIn::CreateBuffers(const FContext& Context)
 {
 	LLM_SCOPE(ELLMTag::SkeletalMesh);
 
-	const USkeletalMesh* Mesh = Context.Mesh;
-	FSkeletalMeshRenderData* RenderData = Context.RenderData;
-	if (!IsCancelled() && Mesh && RenderData)
-	{
-		for (int32 LODIndex = PendingFirstLODIdx; LODIndex < CurrentFirstLODIdx; ++LODIndex)
-		{
-			FSkeletalMeshLODRenderData& LODResource = *Context.LODResourcesView[LODIndex];
-
-			if (bRenderThread)
-			{
-				IntermediateBuffersArray[LODIndex].CreateFromCPUData_RenderThread(LODResource);
-			}
-			else
-			{
-				IntermediateBuffersArray[LODIndex].CreateFromCPUData_Async(LODResource);
-			}
-		}
-	}
-}
-
-void FSkeletalMeshStreamIn::CreateBuffers_RenderThread(const FContext& Context)
-{
-	check(Context.CurrentThread == TT_Render);
-	CreateBuffers_Internal<true>(Context);
-}
+	check(Context.Mesh && Context.RenderData);
 
-void FSkeletalMeshStreamIn::CreateBuffers_Async(const FContext& Context)
-{
-	check(Context.CurrentThread == TT_Async);
-	CreateBuffers_Internal<false>(Context);
-}
-
-void FSkeletalMeshStreamIn::DiscardNewLODs(const FContext& Context)
-{
-	FSkeletalMeshRenderData* RenderData = Context.RenderData;
-	if (RenderData)
+	StreamingRHICmdList = new FRHICommandList();
+	StreamingRHICmdList->SwitchPipeline(ERHIPipeline::Graphics);
 	{
+		SCOPED_DRAW_EVENTF(*StreamingRHICmdList, SkeletalMesh_StreamIn, TEXT("SkeletalMesh - StreamIn: %s"), Context.Mesh->GetFName());
+
 		for (int32 LODIndex = PendingFirstLODIdx; LODIndex < CurrentFirstLODIdx; ++LODIndex)
 		{
 			FSkeletalMeshLODRenderData& LODResource = *Context.LODResourcesView[LODIndex];
-			LODResource.ReleaseCPUResources(true);
+			IntermediateBuffersArray[LODIndex].CreateFromCPUData(*StreamingRHICmdList, LODResource);
 		}
-	}
-}
 
-void FSkeletalMeshStreamIn::DoFinishUpdate(const FContext& Context)
-{
-	const USkeletalMesh* Mesh = Context.Mesh;
-	FSkeletalMeshRenderData* RenderData = Context.RenderData;
-	if (!IsCancelled() && Mesh && RenderData)
-	{
-		check(Context.CurrentThread == TT_Render);
 		// Use a scope to flush the batcher before updating CurrentFirstLODIdx
 		{
-			TRHIResourceUpdateBatcher<GSkelMeshMaxNumResourceUpdatesPerBatch> Batcher;
+			FRHIResourceReplaceBatcher Batcher(*StreamingRHICmdList, GSkelMeshMaxNumResourceUpdatesPerBatch);
 
 			for (int32 LODIndex = PendingFirstLODIdx; LODIndex < CurrentFirstLODIdx; ++LODIndex)
 			{
 				FSkeletalMeshLODRenderData& LODResource = *Context.LODResourcesView[LODIndex];
-				LODResource.IncrementMemoryStats(Mesh->GetHasVertexColors());
+				LODResource.IncrementMemoryStats(Context.Mesh->GetHasVertexColors());
 				LODResource.InitMorphResources();
 				IntermediateBuffersArray[LODIndex].TransferBuffers(LODResource, Batcher);
 			}
 		}
 
 #if RHI_RAYTRACING
-		FRHICommandListBase& RHICmdList = FRHICommandListImmediate::Get();
-
 		// Must happen after the batched updates have been flushed
 		if (IsRayTracingAllowed())
 		{
 			for (int32 LODIndex = PendingFirstLODIdx; LODIndex < CurrentFirstLODIdx; ++LODIndex)
 			{
 				// Skip LODs that have their render data stripped
-				if (RenderData->LODRenderData[LODIndex].GetNumVertices() > 0)
+				if (Context.RenderData->LODRenderData[LODIndex].GetNumVertices() > 0)
 				{
-					if (RenderData->LODRenderData[LODIndex].bReferencedByStaticSkeletalMeshObjects_RenderThread)
+					if (Context.RenderData->LODRenderData[LODIndex].bReferencedByStaticSkeletalMeshObjects_RenderThread)
 					{
-						ensure(!RenderData->LODRenderData[LODIndex].StaticRayTracingGeometry.IsInitialized());
-						RenderData->LODRenderData[LODIndex].StaticRayTracingGeometry.InitResource(RHICmdList);
+						ensure(!Context.RenderData->LODRenderData[LODIndex].StaticRayTracingGeometry.IsInitialized());
+						Context.RenderData->LODRenderData[LODIndex].StaticRayTracingGeometry.InitResource(*StreamingRHICmdList);
 					}
 				}
 			}
 		}
 #endif
-
-		RenderData->PendingFirstLODIdx = RenderData->CurrentFirstLODIdx = ResourceState.LODCountToAssetFirstLODIdx(ResourceState.NumRequestedLODs);
-		MarkAsSuccessfullyFinished();
 	}
-	else
+
+	StreamingRHICmdList->FinishRecording();
+}
+
+void FSkeletalMeshStreamIn::DiscardNewLODs(const FContext& Context)
+{
+	if (Context.RenderData)
 	{
 		for (int32 LODIndex = PendingFirstLODIdx; LODIndex < CurrentFirstLODIdx; ++LODIndex)
 		{
-			IntermediateBuffersArray[LODIndex].SafeRelease();
+			FSkeletalMeshLODRenderData& LODResource = *Context.LODResourcesView[LODIndex];
+			LODResource.ReleaseCPUResources(true);
 		}
 	}
 }
 
+void FSkeletalMeshStreamIn::DoFinishUpdate(const FContext& Context)
+{
+	check(Context.CurrentThread == TT_Render);
+	check(IsInRenderingThread());
+
+	if (StreamingRHICmdList)
+	{
+		FRHICommandListImmediate::Get().QueueAsyncCommandListSubmit(StreamingRHICmdList);
+		StreamingRHICmdList = nullptr;
+	}
+
+	Context.RenderData->PendingFirstLODIdx = Context.RenderData->CurrentFirstLODIdx = ResourceState.LODCountToAssetFirstLODIdx(ResourceState.NumRequestedLODs);
+	MarkAsSuccessfullyFinished();
+}
+
 void FSkeletalMeshStreamIn::DoCancel(const FContext& Context)
 {
 	// TODO: support streaming CPU data for editor builds
@@ -247,7 +176,9 @@ void FSkeletalMeshStreamIn::DoCancel(const FContext& Context)
 	{
 		DiscardNewLODs(Context);
 	}
-	DoFinishUpdate(Context);
+
+	check(!StreamingRHICmdList);
+	MarkAsSuccessfullyFinished();
 }
 
 FSkeletalMeshStreamOut::FSkeletalMeshStreamOut(const USkeletalMesh* InMesh)
@@ -343,42 +274,47 @@ void FSkeletalMeshStreamOut::ReleaseBuffers(const FContext& Context)
 {
 	DECLARE_SCOPE_CYCLE_COUNTER(TEXT("FSkeletalMeshStreamOut::ReleaseBuffers"), STAT_SkeletalMeshStreamOut_ReleaseBuffers, STATGROUP_StreamingDetails);
 	check(Context.CurrentThread == TT_Render);
+	check(IsInRenderingThread());
+
 	const USkeletalMesh* Mesh = Context.Mesh;
 	FSkeletalMeshRenderData* RenderData = Context.RenderData;
 	if (!IsCancelled() && Mesh && RenderData)
 	{
 		RenderData->CurrentFirstLODIdx = Context.RenderData->PendingFirstLODIdx;
 
-		TRHIResourceUpdateBatcher<GSkelMeshMaxNumResourceUpdatesPerBatch> Batcher;
-
-		for (int32 LODIndex = CurrentFirstLODIdx; LODIndex < PendingFirstLODIdx; ++LODIndex)
 		{
-			FSkeletalMeshLODRenderData& LODResource = *Context.LODResourcesView[LODIndex];
-			FStaticMeshVertexBuffers& VBs = LODResource.StaticVertexBuffers;
-			LODResource.DecrementMemoryStats();
-			VBs.StaticMeshVertexBuffer.ReleaseRHIForStreaming(Batcher);
-			VBs.PositionVertexBuffer.ReleaseRHIForStreaming(Batcher);
-			VBs.ColorVertexBuffer.ReleaseRHIForStreaming(Batcher);
-			LODResource.SkinWeightVertexBuffer.ReleaseRHIForStreaming(Batcher);
-			LODResource.ClothVertexBuffer.ReleaseRHIForStreaming(Batcher);
-			LODResource.MultiSizeIndexContainer.ReleaseRHIForStreaming(Batcher);
-			LODResource.SkinWeightProfilesData.ReleaseRHIForStreaming(Batcher);
-
-			if (!FPlatformProperties::HasEditorOnlyData())
+			FRHIResourceReplaceBatcher Batcher(FRHICommandListImmediate::Get(), GSkelMeshMaxNumResourceUpdatesPerBatch);
+
+			for (int32 LODIndex = CurrentFirstLODIdx; LODIndex < PendingFirstLODIdx; ++LODIndex)
 			{
-				// TODO requires more testing : LODResource.ReleaseCPUResources(true);
-			}
+				FSkeletalMeshLODRenderData& LODResource = *Context.LODResourcesView[LODIndex];
+				FStaticMeshVertexBuffers& VBs = LODResource.StaticVertexBuffers;
+				LODResource.DecrementMemoryStats();
+				VBs.StaticMeshVertexBuffer.ReleaseRHIForStreaming(Batcher);
+				VBs.PositionVertexBuffer.ReleaseRHIForStreaming(Batcher);
+				VBs.ColorVertexBuffer.ReleaseRHIForStreaming(Batcher);
+				LODResource.SkinWeightVertexBuffer.ReleaseRHIForStreaming(Batcher);
+				LODResource.ClothVertexBuffer.ReleaseRHIForStreaming(Batcher);
+				LODResource.MultiSizeIndexContainer.ReleaseRHIForStreaming(Batcher);
+				LODResource.SkinWeightProfilesData.ReleaseRHIForStreaming(Batcher);
+
+				if (!FPlatformProperties::HasEditorOnlyData())
+				{
+					// TODO requires more testing : LODResource.ReleaseCPUResources(true);
+				}
 
 #if RHI_RAYTRACING
-			if (IsRayTracingAllowed())
-			{
-				if (RenderData->LODRenderData[LODIndex].bReferencedByStaticSkeletalMeshObjects_RenderThread)
+				if (IsRayTracingAllowed())
 				{
-					LODResource.StaticRayTracingGeometry.ReleaseResource();
+					if (RenderData->LODRenderData[LODIndex].bReferencedByStaticSkeletalMeshObjects_RenderThread)
+					{
+						LODResource.StaticRayTracingGeometry.ReleaseResource();
+					}
 				}
-			}
 #endif
+			}
 		}
+
 		MarkAsSuccessfullyFinished();
 	}
 }
@@ -403,11 +339,6 @@ void FSkeletalMeshStreamIn_IO::FCancelIORequestsTask::DoWork()
 	PendingUpdate->DoUnlock(OldState);
 }
 
-FSkeletalMeshStreamIn_IO::FSkeletalMeshStreamIn_IO(const USkeletalMesh* InMesh, bool bHighPrio)
-	: FSkeletalMeshStreamIn(InMesh)
-	, bHighPrioIORequest(bHighPrio)
-{}
-
 void FSkeletalMeshStreamIn_IO::Abort()
 {
 	if (!IsCancelled() && !IsCompleted())
@@ -591,15 +522,14 @@ void FSkeletalMeshStreamIn_IO::CancelIORequest()
 	}
 }
 
-template <bool bRenderThread>
-TSkeletalMeshStreamIn_IO<bRenderThread>::TSkeletalMeshStreamIn_IO(const USkeletalMesh* InMesh, bool bHighPrio)
-	: FSkeletalMeshStreamIn_IO(InMesh, bHighPrio)
+FSkeletalMeshStreamIn_IO::FSkeletalMeshStreamIn_IO(const USkeletalMesh* InMesh, bool bHighPrio, EThreadType CreateResourcesThread)
+	: FSkeletalMeshStreamIn(InMesh, CreateResourcesThread)
+	, bHighPrioIORequest(bHighPrio)
 {
 	PushTask(FContext(InMesh, TT_None), TT_Async, SRA_UPDATE_CALLBACK(DoInitiateIO), TT_None, nullptr);
 }
 
-template <bool bRenderThread>
-void TSkeletalMeshStreamIn_IO<bRenderThread>::DoInitiateIO(const FContext& Context)
+void FSkeletalMeshStreamIn_IO::DoInitiateIO(const FContext& Context)
 {
 	check(Context.CurrentThread == TT_Async);
 
@@ -608,46 +538,37 @@ void TSkeletalMeshStreamIn_IO<bRenderThread>::DoInitiateIO(const FContext& Conte
 	PushTask(Context, TT_Async, SRA_UPDATE_CALLBACK(DoSerializeLODData), TT_Async, SRA_UPDATE_CALLBACK(DoCancelIO));
 }
 
-template <bool bRenderThread>
-void TSkeletalMeshStreamIn_IO<bRenderThread>::DoSerializeLODData(const FContext& Context)
+void FSkeletalMeshStreamIn_IO::DoSerializeLODData(const FContext& Context)
 {
 	check(Context.CurrentThread == TT_Async);
 	SerializeLODData(Context);
 	ClearIORequest(Context);
-	const EThreadType TThread = bRenderThread ? TT_Render : TT_Async;
-	const EThreadType CThread = (EThreadType)Context.CurrentThread;
-	PushTask(Context, TThread, SRA_UPDATE_CALLBACK(DoCreateBuffers), CThread, SRA_UPDATE_CALLBACK(Cancel));
+
+	PushTask(Context
+		, CreateResourcesThread, SRA_UPDATE_CALLBACK(DoCreateBuffers)
+		, (EThreadType)Context.CurrentThread, SRA_UPDATE_CALLBACK(Cancel));
 }
 
-template <bool bRenderThread>
-void TSkeletalMeshStreamIn_IO<bRenderThread>::DoCreateBuffers(const FContext& Context)
+void FSkeletalMeshStreamIn_IO::DoCreateBuffers(const FContext& Context)
 {
-	if (bRenderThread)
-	{
-		CreateBuffers_RenderThread(Context);
-	}
-	else
-	{
-		CreateBuffers_Async(Context);
-	}
+	CreateBuffers(Context);
+
 	check(!TaskSynchronization.GetValue());
-	PushTask(Context, TT_Render, SRA_UPDATE_CALLBACK(DoFinishUpdate), (EThreadType)Context.CurrentThread, SRA_UPDATE_CALLBACK(Cancel));
+	PushTask(Context, TT_Render, SRA_UPDATE_CALLBACK(DoFinishUpdate), TT_None, nullptr);
 }
 
-template <bool bRenderThread>
-void TSkeletalMeshStreamIn_IO<bRenderThread>::DoCancelIO(const FContext& Context)
+void FSkeletalMeshStreamIn_IO::DoCancelIO(const FContext& Context)
 {
 	ClearIORequest(Context);
 	PushTask(Context, TT_None, nullptr, (EThreadType)Context.CurrentThread, SRA_UPDATE_CALLBACK(Cancel));
 }
 
-template class TSkeletalMeshStreamIn_IO<true>;
-template class TSkeletalMeshStreamIn_IO<false>;
-
 #if WITH_EDITOR
-FSkeletalMeshStreamIn_DDC::FSkeletalMeshStreamIn_DDC(const USkeletalMesh* InMesh)
-	: FSkeletalMeshStreamIn(InMesh)
-{}
+FSkeletalMeshStreamIn_DDC::FSkeletalMeshStreamIn_DDC(const USkeletalMesh* InMesh, EThreadType CreateResourcesThread)
+	: FSkeletalMeshStreamIn(InMesh, CreateResourcesThread)
+{
+	PushTask(FContext(InMesh, TT_None), TT_Async, SRA_UPDATE_CALLBACK(DoLoadNewLODsFromDDC), TT_None, nullptr);
+}
 
 void FSkeletalMeshStreamIn_DDC::LoadNewLODsFromDDC(const FContext& Context)
 {
@@ -655,38 +576,21 @@ void FSkeletalMeshStreamIn_DDC::LoadNewLODsFromDDC(const FContext& Context)
 	// TODO: support streaming CPU data for editor builds
 }
 
-template <bool bRenderThread>
-TSkeletalMeshStreamIn_DDC<bRenderThread>::TSkeletalMeshStreamIn_DDC(const USkeletalMesh* InMesh)
-	: FSkeletalMeshStreamIn_DDC(InMesh)
-{
-	PushTask(FContext(InMesh, TT_None), TT_Async, SRA_UPDATE_CALLBACK(DoLoadNewLODsFromDDC), TT_None, nullptr);
-}
-
-template <bool bRenderThread>
-void TSkeletalMeshStreamIn_DDC<bRenderThread>::DoLoadNewLODsFromDDC(const FContext& Context)
+void FSkeletalMeshStreamIn_DDC::DoLoadNewLODsFromDDC(const FContext& Context)
 {
 	LoadNewLODsFromDDC(Context);
 	check(!TaskSynchronization.GetValue());
-	const EThreadType TThread = bRenderThread ? TT_Render : TT_Async;
-	const EThreadType CThread = (EThreadType)Context.CurrentThread;
-	PushTask(Context, TThread, SRA_UPDATE_CALLBACK(DoCreateBuffers), CThread, SRA_UPDATE_CALLBACK(DoCancel));
+
+	PushTask(Context
+		, CreateResourcesThread, SRA_UPDATE_CALLBACK(DoCreateBuffers)
+		, (EThreadType)Context.CurrentThread, SRA_UPDATE_CALLBACK(DoCancel));
 }
 
-template <bool bRenderThread>
-void TSkeletalMeshStreamIn_DDC<bRenderThread>::DoCreateBuffers(const FContext& Context)
+void FSkeletalMeshStreamIn_DDC::DoCreateBuffers(const FContext& Context)
 {
-	if (bRenderThread)
-	{
-		CreateBuffers_RenderThread(Context);
-	}
-	else
-	{
-		CreateBuffers_Async(Context);
-	}
+	CreateBuffers(Context);
+	
 	check(!TaskSynchronization.GetValue());
-	PushTask(Context, TT_Render, SRA_UPDATE_CALLBACK(DoFinishUpdate), (EThreadType)Context.CurrentThread, SRA_UPDATE_CALLBACK(DoCancel));
+	PushTask(Context, TT_Render, SRA_UPDATE_CALLBACK(DoFinishUpdate), TT_None, nullptr);
 }
-
-template class TSkeletalMeshStreamIn_DDC<true>;
-template class TSkeletalMeshStreamIn_DDC<false>;
 #endif
diff --git a/Engine/Source/Runtime/Engine/Private/Streaming/SkeletalMeshUpdate.h b/Engine/Source/Runtime/Engine/Private/Streaming/SkeletalMeshUpdate.h
--- a/Engine/Source/Runtime/Engine/Private/Streaming/SkeletalMeshUpdate.h
+++ b/Engine/Source/Runtime/Engine/Private/Streaming/SkeletalMeshUpdate.h
@@ -67,7 +67,7 @@ public:
 class FSkeletalMeshStreamIn : public FSkeletalMeshUpdate
 {
 public:
-	FSkeletalMeshStreamIn(const USkeletalMesh* InMesh);
+	FSkeletalMeshStreamIn(const USkeletalMesh* InMesh, EThreadType CreateResourcesThread);
 
 	virtual ~FSkeletalMeshStreamIn();
 
@@ -84,20 +84,14 @@ protected:
 		FBufferRHIRef IndexBuffer;
 		TArray<TPair<FName, FSkinWeightRHIInfo>> AltSkinWeightVertexBuffers;
 
-		void CreateFromCPUData_RenderThread(FSkeletalMeshLODRenderData& LODResource);
-		void CreateFromCPUData_Async(FSkeletalMeshLODRenderData& LODResource);
-
-		void SafeRelease();
+		void CreateFromCPUData(FRHICommandListBase& RHICmdList, FSkeletalMeshLODRenderData& LODResource);
 
 		/** Transfer ownership of buffers to a LOD resource */
-		void TransferBuffers(FSkeletalMeshLODRenderData& LODResource, FRHIResourceUpdateBatcher& Batcher);
-
-		void CheckIsNull() const;
+		void TransferBuffers(FSkeletalMeshLODRenderData& LODResource, FRHIResourceReplaceBatcher& Batcher);
 	};
 
-	/** Create buffers with new LOD data on render or pooled thread */
-	void CreateBuffers_RenderThread(const FContext& Context);
-	void CreateBuffers_Async(const FContext& Context);
+	/** Create buffers with new LOD data */
+	void CreateBuffers(const FContext& Context);
 
 	/** Discard newly streamed-in CPU data */
 	void DiscardNewLODs(const FContext& Context);
@@ -111,9 +105,11 @@ protected:
 	/** The intermediate buffers created in the update process. */
 	FIntermediateBuffers IntermediateBuffersArray[MAX_MESH_LOD_COUNT];
 
-private:
-	template <bool bRenderThread>
-	void CreateBuffers_Internal(const FContext& Context);
+	/** RHI command list used for creating buffers and replacing the streaming placeholders. Submitted in DoFinishUpdate */
+	FRHICommandList* StreamingRHICmdList = nullptr;
+
+	/** The thread to use for recording the above command list */
+	const EThreadType CreateResourcesThread;
 };
 
 class FSkeletalMeshStreamOut : public FSkeletalMeshUpdate
@@ -145,7 +141,7 @@ private:
 class FSkeletalMeshStreamIn_IO : public FSkeletalMeshStreamIn
 {
 public:
-	FSkeletalMeshStreamIn_IO(const USkeletalMesh* InMesh, bool bHighPrio);
+	FSkeletalMeshStreamIn_IO(const USkeletalMesh* InMesh, bool bHighPrio, EThreadType CreateResourcesThread);
 
 	virtual ~FSkeletalMeshStreamIn_IO() {}
 
@@ -201,17 +197,7 @@ protected:
 
 	// Whether an IO error was detected (when files do not exists).
 	bool bFailedOnIOError = false;
-};
-
-template <bool bRenderThread>
-class TSkeletalMeshStreamIn_IO : public FSkeletalMeshStreamIn_IO
-{
-public:
-	TSkeletalMeshStreamIn_IO(const USkeletalMesh* InMesh, bool bHighPrio);
-
-	virtual ~TSkeletalMeshStreamIn_IO() {}
 
-protected:
 	void DoInitiateIO(const FContext& Context);
 
 	void DoSerializeLODData(const FContext& Context);
@@ -221,28 +207,16 @@ protected:
 	void DoCancelIO(const FContext& Context);
 };
 
-typedef TSkeletalMeshStreamIn_IO<true> FSkeletalMeshStreamIn_IO_RenderThread;
-typedef TSkeletalMeshStreamIn_IO<false> FSkeletalMeshStreamIn_IO_Async;
-
 #if WITH_EDITOR
 class FSkeletalMeshStreamIn_DDC : public FSkeletalMeshStreamIn
 {
 public:
-	FSkeletalMeshStreamIn_DDC(const USkeletalMesh* InMesh);
+	FSkeletalMeshStreamIn_DDC(const USkeletalMesh* InMesh, EThreadType CreateResourcesThread);
 
 	virtual ~FSkeletalMeshStreamIn_DDC() {}
 
 protected:
 	void LoadNewLODsFromDDC(const FContext& Context);
-};
-
-template <bool bRenderThread>
-class TSkeletalMeshStreamIn_DDC : public FSkeletalMeshStreamIn_DDC
-{
-public:
-	TSkeletalMeshStreamIn_DDC(const USkeletalMesh* InMesh);
-
-	virtual ~TSkeletalMeshStreamIn_DDC() {}
 
 private:
 	/** Load new LOD buffers from DDC and queue a task to create RHI buffers on RT */
@@ -251,7 +225,4 @@ private:
 	/** Create RHI buffers for newly streamed-in LODs and queue a task to rename references on RT */
 	void DoCreateBuffers(const FContext& Context);
 };
-
-typedef TSkeletalMeshStreamIn_DDC<true> FSkeletalMeshStreamIn_DDC_RenderThread;
-typedef TSkeletalMeshStreamIn_DDC<false> FSkeletalMeshStreamIn_DDC_Async;
 #endif
diff --git a/Engine/Source/Runtime/Engine/Private/Streaming/StaticMeshUpdate.cpp b/Engine/Source/Runtime/Engine/Private/Streaming/StaticMeshUpdate.cpp
--- a/Engine/Source/Runtime/Engine/Private/Streaming/StaticMeshUpdate.cpp
+++ b/Engine/Source/Runtime/Engine/Private/Streaming/StaticMeshUpdate.cpp
@@ -12,7 +12,7 @@ StaticMeshUpdate.cpp: Helpers to stream in and out static mesh LODs.
 #include "StaticMeshResources.h"
 #include "Streaming/RenderAssetUpdate.inl"
 #include "ContentStreaming.h"
-#include "RHIResourceUpdates.h"
+#include "RHIResourceReplace.h"
 
 int32 GStreamingMaxReferenceChecks = 2;
 static FAutoConsoleVariableRef CVarStreamingMaxReferenceChecksBeforeStreamOut(
@@ -55,10 +55,8 @@ FStaticMeshUpdate::FStaticMeshUpdate(const UStaticMesh* InMesh)
 {
 }
 
-void FStaticMeshStreamIn::FIntermediateBuffers::CreateFromCPUData_RenderThread(FStaticMeshLODResources& LODResource)
+void FStaticMeshStreamIn::FIntermediateBuffers::CreateFromCPUData(FRHICommandListBase& RHICmdList, FStaticMeshLODResources& LODResource)
 {
-	FRHICommandListBase& RHICmdList = FRHICommandListImmediate::Get();
-
 	FStaticMeshVertexBuffers& VBs = LODResource.VertexBuffers;
 	TangentsVertexBuffer = VBs.StaticMeshVertexBuffer.CreateTangentsRHIBuffer(RHICmdList);
 	TexCoordVertexBuffer = VBs.StaticMeshVertexBuffer.CreateTexCoordRHIBuffer(RHICmdList);
@@ -67,7 +65,6 @@ void FStaticMeshStreamIn::FIntermediateBuffers::CreateFromCPUData_RenderThread(F
 	IndexBuffer = LODResource.IndexBuffer.CreateRHIBuffer(RHICmdList);
 	DepthOnlyIndexBuffer = LODResource.DepthOnlyIndexBuffer.CreateRHIBuffer(RHICmdList);
 
-
 	if (LODResource.AdditionalIndexBuffers)
 	{
 		ReversedIndexBuffer = LODResource.AdditionalIndexBuffers->ReversedIndexBuffer.CreateRHIBuffer(RHICmdList);
@@ -76,40 +73,7 @@ void FStaticMeshStreamIn::FIntermediateBuffers::CreateFromCPUData_RenderThread(F
 	}
 }
 
-void FStaticMeshStreamIn::FIntermediateBuffers::CreateFromCPUData_Async(FStaticMeshLODResources& LODResource)
-{
-	FRHIAsyncCommandList RHICmdList;
-
-	FStaticMeshVertexBuffers& VBs = LODResource.VertexBuffers;
-	TangentsVertexBuffer = VBs.StaticMeshVertexBuffer.CreateTangentsRHIBuffer(*RHICmdList);
-	TexCoordVertexBuffer = VBs.StaticMeshVertexBuffer.CreateTexCoordRHIBuffer(*RHICmdList);
-	PositionVertexBuffer = VBs.PositionVertexBuffer.CreateRHIBuffer(*RHICmdList);
-	ColorVertexBuffer = VBs.ColorVertexBuffer.CreateRHIBuffer(*RHICmdList);
-	IndexBuffer = LODResource.IndexBuffer.CreateRHIBuffer(*RHICmdList);
-	DepthOnlyIndexBuffer = LODResource.DepthOnlyIndexBuffer.CreateRHIBuffer(*RHICmdList);
-
-	if (LODResource.AdditionalIndexBuffers)
-	{
-		ReversedIndexBuffer = LODResource.AdditionalIndexBuffers->ReversedIndexBuffer.CreateRHIBuffer(*RHICmdList);
-		ReversedDepthOnlyIndexBuffer = LODResource.AdditionalIndexBuffers->ReversedDepthOnlyIndexBuffer.CreateRHIBuffer(*RHICmdList);
-		WireframeIndexBuffer = LODResource.AdditionalIndexBuffers->WireframeIndexBuffer.CreateRHIBuffer(*RHICmdList);
-	}
-}
-
-void FStaticMeshStreamIn::FIntermediateBuffers::SafeRelease()
-{
-	TangentsVertexBuffer.SafeRelease();
-	TexCoordVertexBuffer.SafeRelease();
-	PositionVertexBuffer.SafeRelease();
-	ColorVertexBuffer.SafeRelease();
-	IndexBuffer.SafeRelease();
-	ReversedIndexBuffer.SafeRelease();
-	DepthOnlyIndexBuffer.SafeRelease();
-	ReversedDepthOnlyIndexBuffer.SafeRelease();
-	WireframeIndexBuffer.SafeRelease();
-}
-
-void FStaticMeshStreamIn::FIntermediateBuffers::TransferBuffers(FStaticMeshLODResources& LODResource, FRHIResourceUpdateBatcher& Batcher)
+void FStaticMeshStreamIn::FIntermediateBuffers::TransferBuffers(FStaticMeshLODResources& LODResource, FRHIResourceReplaceBatcher& Batcher)
 {
 	FStaticMeshVertexBuffers& VBs = LODResource.VertexBuffers;
 	VBs.StaticMeshVertexBuffer.InitRHIForStreaming(TangentsVertexBuffer, TexCoordVertexBuffer, Batcher);
@@ -124,25 +88,11 @@ void FStaticMeshStreamIn::FIntermediateBuffers::TransferBuffers(FStaticMeshLODRe
 		LODResource.AdditionalIndexBuffers->ReversedDepthOnlyIndexBuffer.InitRHIForStreaming(ReversedDepthOnlyIndexBuffer, Batcher);
 		LODResource.AdditionalIndexBuffers->WireframeIndexBuffer.InitRHIForStreaming(WireframeIndexBuffer, Batcher);
 	}
-	SafeRelease();
-}
-
-void FStaticMeshStreamIn::FIntermediateBuffers::CheckIsNull() const
-{
-	check(!TangentsVertexBuffer
-		&& !TexCoordVertexBuffer
-		&& !PositionVertexBuffer
-		&& !ColorVertexBuffer
-		&& !IndexBuffer
-		&& !ReversedIndexBuffer
-		&& !DepthOnlyIndexBuffer
-		&& !ReversedDepthOnlyIndexBuffer
-		&& !WireframeIndexBuffer);
 }
 
 #if RHI_RAYTRACING
 
-void FStaticMeshStreamIn::FIntermediateRayTracingGeometry::CreateFromCPUData(FRHICommandList& RHICmdList, FRayTracingGeometry& RayTracingGeometry)
+void FStaticMeshStreamIn::FIntermediateRayTracingGeometry::CreateFromCPUData(FRHICommandListBase& RHICmdList, FRayTracingGeometry& RayTracingGeometry)
 {
 	Initializer = RayTracingGeometry.Initializer;
 	Initializer.Type = ERayTracingGeometryInitializerType::StreamingSource;
@@ -172,7 +122,7 @@ void FStaticMeshStreamIn::FIntermediateRayTracingGeometry::SafeRelease()
 	RayTracingGeometryRHI.SafeRelease();
 }
 
-void FStaticMeshStreamIn::FIntermediateRayTracingGeometry::TransferRayTracingGeometry(FRayTracingGeometry& RayTracingGeometry, FRHIResourceUpdateBatcher& Batcher)
+void FStaticMeshStreamIn::FIntermediateRayTracingGeometry::TransferRayTracingGeometry(FRayTracingGeometry& RayTracingGeometry, FRHIResourceReplaceBatcher& Batcher)
 {
 	RayTracingGeometry.InitRHIForStreaming(RayTracingGeometryRHI, Batcher);
 	RayTracingGeometry.SetRequiresBuild(bRequiresBuild);
@@ -233,44 +183,35 @@ static void CheckRayTracingGeometryInitializer(
 
 #endif
 
-FStaticMeshStreamIn::FStaticMeshStreamIn(const UStaticMesh* InMesh)
+FStaticMeshStreamIn::FStaticMeshStreamIn(const UStaticMesh* InMesh, EThreadType CreateResourcesThread)
 	: FStaticMeshUpdate(InMesh)
+	, CreateResourcesThread(CreateResourcesThread)
 {}
 
 FStaticMeshStreamIn::~FStaticMeshStreamIn()
 {
-#if DO_CHECK
-	for (int32 Idx = 0; Idx < MAX_MESH_LOD_COUNT; ++Idx)
-	{
-		IntermediateBuffersArray[Idx].CheckIsNull();
-	}
-#endif
+	check(!StreamingRHICmdList);
 }
 
-template <bool bRenderThread>
-void FStaticMeshStreamIn::CreateBuffers_Internal(const FContext& Context)
+void FStaticMeshStreamIn::CreateBuffers(const FContext& Context)
 {
 	LLM_SCOPE(ELLMTag::StaticMesh);
-	
-	FStaticMeshRenderData* RenderData = Context.RenderData;
-	if (!IsCancelled() && RenderData)
+
+	check(Context.Mesh && Context.RenderData);
+
+	StreamingRHICmdList = new FRHICommandList();
+	StreamingRHICmdList->SwitchPipeline(ERHIPipeline::Graphics);
 	{
+		SCOPED_DRAW_EVENTF(*StreamingRHICmdList, StaticMesh_StreamIn, TEXT("StaticMesh - StreamIn: %s"), Context.Mesh->GetFName());
+
 		for (int32 LODIdx = PendingFirstLODIdx; LODIdx < CurrentFirstLODIdx; ++LODIdx)
 		{
 			FStaticMeshLODResources& LODResource = *Context.LODResourcesView[LODIdx];
 
-			if (bRenderThread)
-			{
-				IntermediateBuffersArray[LODIdx].CreateFromCPUData_RenderThread(LODResource);
-			}
-			else
-			{
-				IntermediateBuffersArray[LODIdx].CreateFromCPUData_Async(LODResource);
-			}
+			IntermediateBuffersArray[LODIdx].CreateFromCPUData(*StreamingRHICmdList, LODResource);
 
 #if RHI_RAYTRACING
-			if (IsRayTracingEnabled() && Context.Mesh->bSupportRayTracing &&
-				LODResource.VertexBuffers.StaticMeshVertexBuffer.GetNumVertices() > 0)
+			if (IsRayTracingEnabled() && Context.Mesh->bSupportRayTracing && LODResource.VertexBuffers.StaticMeshVertexBuffer.GetNumVertices() > 0)
 			{
 #if DO_CHECK
 				CheckRayTracingGeometryInitializer(
@@ -281,32 +222,29 @@ void FStaticMeshStreamIn::CreateBuffers_Internal(const FContext& Context)
 					LODResource.RayTracingGeometry.Initializer);
 #endif
 
-				FRHIAsyncCommandList AsyncCommandList;
-				FRHICommandList& RHICmdList = bRenderThread ? FRHICommandListImmediate::Get() : *AsyncCommandList;
-
-				IntermediateRayTracingGeometry[LODIdx].CreateFromCPUData(RHICmdList, LODResource.RayTracingGeometry);
+				IntermediateRayTracingGeometry[LODIdx].CreateFromCPUData(*StreamingRHICmdList, LODResource.RayTracingGeometry);
 			}
 #endif
 		}
-	}
-}
 
-void FStaticMeshStreamIn::CreateBuffers_RenderThread(const FContext& Context)
-{
-	check(Context.CurrentThread == TT_Render);
-	CreateBuffers_Internal<true>(Context);
-}
+		// Use a scope to flush the batcher before updating CurrentFirstLODIdx
+		{
+			FRHIResourceReplaceBatcher Batcher(*StreamingRHICmdList, GStaticMeshMaxNumResourceUpdatesPerBatch);
+			for (int32 LODIdx = PendingFirstLODIdx; LODIdx < CurrentFirstLODIdx; ++LODIdx)
+			{
+				FStaticMeshLODResources& LODResource = *Context.LODResourcesView[LODIdx];
+				LODResource.IncrementMemoryStats();
+				IntermediateBuffersArray[LODIdx].TransferBuffers(LODResource, Batcher);
+			}
+		}
+	}
 
-void FStaticMeshStreamIn::CreateBuffers_Async(const FContext& Context)
-{
-	check(Context.CurrentThread == TT_Async);
-	CreateBuffers_Internal<false>(Context);
+	StreamingRHICmdList->FinishRecording();
 }
 
 void FStaticMeshStreamIn::DiscardNewLODs(const FContext& Context)
 {
-	FStaticMeshRenderData* RenderData = Context.RenderData;
-	if (RenderData)
+	if (Context.RenderData)
 	{
 		for (int32 LODIdx = PendingFirstLODIdx; LODIdx < CurrentFirstLODIdx; ++LODIdx)
 		{
@@ -318,79 +256,66 @@ void FStaticMeshStreamIn::DiscardNewLODs(const FContext& Context)
 
 void FStaticMeshStreamIn::DoFinishUpdate(const FContext& Context)
 {
-	FStaticMeshRenderData* RenderData = Context.RenderData;
-	if (!IsCancelled() && RenderData)
+	check(Context.CurrentThread == TT_Render);
+	check(IsInRenderingThread());
+	
+	if (StreamingRHICmdList)
+	{
+		FRHICommandListImmediate::Get().QueueAsyncCommandListSubmit(StreamingRHICmdList);
+		StreamingRHICmdList = nullptr;
+	}
+
+#if RHI_RAYTRACING
+	if (IsRayTracingAllowed() && Context.Mesh->bSupportRayTracing)
 	{
-		check(Context.CurrentThread == TT_Render);
 		// Use a scope to flush the batcher before updating CurrentFirstLODIdx
 		{
-			TRHIResourceUpdateBatcher<GStaticMeshMaxNumResourceUpdatesPerBatch> Batcher;
-
+			FRHIResourceReplaceBatcher Batcher(FRHICommandListImmediate::Get(), GStaticMeshMaxNumResourceUpdatesPerBatch);
 			for (int32 LODIdx = PendingFirstLODIdx; LODIdx < CurrentFirstLODIdx; ++LODIdx)
 			{
 				FStaticMeshLODResources& LODResource = *Context.LODResourcesView[LODIdx];
-				LODResource.IncrementMemoryStats();
-				IntermediateBuffersArray[LODIdx].TransferBuffers(LODResource, Batcher);
-
-#if RHI_RAYTRACING
-				if (IsRayTracingAllowed() && Context.Mesh->bSupportRayTracing && LODResource.VertexBuffers.StaticMeshVertexBuffer.GetNumVertices() > 0)
+				if (LODResource.VertexBuffers.StaticMeshVertexBuffer.GetNumVertices() > 0)
 				{
 					IntermediateRayTracingGeometry[LODIdx].TransferRayTracingGeometry(LODResource.RayTracingGeometry, Batcher);
 				}
-#endif
 			}
 		}
 
-#if RHI_RAYTRACING
 		// Must happen after the batched updates have been flushed
-		if (IsRayTracingAllowed() && Context.Mesh->bSupportRayTracing)
+		for (int32 LODIndex = PendingFirstLODIdx; LODIndex < CurrentFirstLODIdx; ++LODIndex)
 		{
-			for (int32 LODIndex = PendingFirstLODIdx; LODIndex < CurrentFirstLODIdx; ++LODIndex)
-			{
-				FStaticMeshLODResources& LODResource = *Context.LODResourcesView[LODIndex];
+			FStaticMeshLODResources& LODResource = *Context.LODResourcesView[LODIndex];
 
-				// Skip LODs that have their render data stripped
-				if (LODResource.VertexBuffers.StaticMeshVertexBuffer.GetNumVertices() > 0)
-				{
+			// Skip LODs that have their render data stripped
+			if (LODResource.VertexBuffers.StaticMeshVertexBuffer.GetNumVertices() > 0)
+			{
 #if DO_CHECK
-					// Streaming LODs in/out shouldn't affect the ray tracing geometry initializer
-					// Here we check that assumption
-					CheckRayTracingGeometryInitializer(
-						Context.Mesh,
-						LODIndex + Context.Mesh->GetStreamableResourceState().AssetLODBias,
-						LODResource,
-						ERayTracingGeometryInitializerType::Rendering,
-						LODResource.RayTracingGeometry.Initializer);
-
-					check(EnumHasAllFlags(LODResource.RayTracingGeometry.GetGeometryState(), FRayTracingGeometry::EGeometryStateFlags::StreamedIn));
+				// Streaming LODs in/out shouldn't affect the ray tracing geometry initializer
+				// Here we check that assumption
+				CheckRayTracingGeometryInitializer(
+					Context.Mesh,
+					LODIndex + Context.Mesh->GetStreamableResourceState().AssetLODBias,
+					LODResource,
+					ERayTracingGeometryInitializerType::Rendering,
+					LODResource.RayTracingGeometry.Initializer);
+
+				check(EnumHasAllFlags(LODResource.RayTracingGeometry.GetGeometryState(), FRayTracingGeometry::EGeometryStateFlags::StreamedIn));
 #endif
 
-					// Under very rare circumstances that we switch ray tracing on/off right in the middle of streaming RayTracingGeometryRHI might not be valid.
-					if (IsRayTracingEnabled() && ensure(LODResource.RayTracingGeometry.RayTracingGeometryRHI.IsValid()))
-					{
-						LODResource.RayTracingGeometry.RequestBuildIfNeeded(ERTAccelerationStructureBuildPriority::Normal);
-					}
+				// Under very rare circumstances that we switch ray tracing on/off right in the middle of streaming RayTracingGeometryRHI might not be valid.
+				if (IsRayTracingEnabled() && ensure(LODResource.RayTracingGeometry.RayTracingGeometryRHI.IsValid()))
+				{
+					LODResource.RayTracingGeometry.RequestBuildIfNeeded(ERTAccelerationStructureBuildPriority::Normal);
 				}
 			}
-
 		}
-#endif
-				
-		Context.Mesh->RequestUpdateCachedRenderState();
-		RenderData->CurrentFirstLODIdx = ResourceState.LODCountToAssetFirstLODIdx(ResourceState.NumRequestedLODs);
-		MarkAsSuccessfullyFinished();
-	}
-	else
-	{
-		for (int32 LODIdx = PendingFirstLODIdx; LODIdx < CurrentFirstLODIdx; ++LODIdx)
-		{
-			IntermediateBuffersArray[LODIdx].SafeRelease();
 
-#if RHI_RAYTRACING
-			IntermediateRayTracingGeometry[LODIdx].SafeRelease();
-#endif
-		}
 	}
+#endif
+
+	Context.Mesh->RequestUpdateCachedRenderState();
+	Context.RenderData->CurrentFirstLODIdx = ResourceState.LODCountToAssetFirstLODIdx(ResourceState.NumRequestedLODs);
+	MarkAsSuccessfullyFinished();
 }
 
 void FStaticMeshStreamIn::DoCancel(const FContext& Context)
@@ -400,7 +325,9 @@ void FStaticMeshStreamIn::DoCancel(const FContext& Context)
 	{
 		DiscardNewLODs(Context);
 	}
-	DoFinishUpdate(Context);
+
+	check(!StreamingRHICmdList);
+	MarkAsSuccessfullyFinished();
 }
 
 FStaticMeshStreamOut::FStaticMeshStreamOut(const UStaticMesh* InMesh, bool InDiscardCPUData)
@@ -482,10 +409,9 @@ void FStaticMeshStreamOut::ReleaseRHIBuffers(const FContext& Context)
 	DECLARE_SCOPE_CYCLE_COUNTER(TEXT("FStaticMeshStreamOut::ReleaseRHIBuffers"), STAT_StaticMeshStreamOut_ReleaseRHIBuffers, STATGROUP_StreamingDetails);
 	check(Context.CurrentThread == TT_Render);
 
-	FStaticMeshRenderData* RenderData = Context.RenderData;
-	if (RenderData)
+	if (Context.RenderData)
 	{
-		TRHIResourceUpdateBatcher<GStaticMeshMaxNumResourceUpdatesPerBatch> Batcher;
+		FRHIResourceReplaceBatcher Batcher(FRHICommandListImmediate::Get(), GStaticMeshMaxNumResourceUpdatesPerBatch);
 		for (int32 LODIdx = CurrentFirstLODIdx; LODIdx < PendingFirstLODIdx; ++LODIdx)
 		{
 			FStaticMeshLODResources& LODResource = *Context.LODResourcesView[LODIdx];
@@ -526,11 +452,6 @@ void FStaticMeshStreamIn_IO::FCancelIORequestsTask::DoWork()
 	PendingUpdate->DoUnlock(PreviousTaskState);
 }
 
-FStaticMeshStreamIn_IO::FStaticMeshStreamIn_IO(const UStaticMesh* InMesh, bool bHighPrio)
-	: FStaticMeshStreamIn(InMesh)
-	, bHighPrioIORequest(bHighPrio)
-{}
-
 void FStaticMeshStreamIn_IO::Abort()
 {
 	if (!IsCancelled() && !IsCompleted())
@@ -555,9 +476,7 @@ void FStaticMeshStreamIn_IO::SetIORequest(const FContext& Context)
 
 	check(BulkDataRequest.IsNone() && PendingFirstLODIdx < CurrentFirstLODIdx);
 
-	const UStaticMesh* Mesh = Context.Mesh;
-	FStaticMeshRenderData* RenderData = Context.RenderData;
-	if (Mesh && RenderData)
+	if (Context.Mesh && Context.RenderData)
 	{
 		const int32 BatchCount = CurrentFirstLODIdx - PendingFirstLODIdx;
 		FBulkDataBatchRequest::FScatterGatherBuilder Batch = FBulkDataBatchRequest::ScatterGather(BatchCount);
@@ -711,17 +630,16 @@ void FStaticMeshStreamIn_IO::CancelIORequest()
 	}
 }
 
-template <bool bRenderThread>
-TStaticMeshStreamIn_IO<bRenderThread>::TStaticMeshStreamIn_IO(const UStaticMesh* InMesh, bool bHighPrio)
-	: FStaticMeshStreamIn_IO(InMesh, bHighPrio)
+FStaticMeshStreamIn_IO::FStaticMeshStreamIn_IO(const UStaticMesh* InMesh, bool bHighPrio, EThreadType CreateResourcesThread)
+	: FStaticMeshStreamIn(InMesh, CreateResourcesThread)
+	, bHighPrioIORequest(bHighPrio)
 {
 	PushTask(FContext(InMesh, TT_None), TT_Async, SRA_UPDATE_CALLBACK(DoInitiateIO), TT_None, nullptr);
 }
 
-template <bool bRenderThread>
-void TStaticMeshStreamIn_IO<bRenderThread>::DoInitiateIO(const FContext& Context)
+void FStaticMeshStreamIn_IO::DoInitiateIO(const FContext& Context)
 {
-	DECLARE_SCOPE_CYCLE_COUNTER(TEXT("TStaticMeshStreamIn_IO::DoInitiateIO"), STAT_StaticMeshStreamInIO_DoInitiateIO, STATGROUP_StreamingDetails);
+	DECLARE_SCOPE_CYCLE_COUNTER(TEXT("FStaticMeshStreamIn_IO::DoInitiateIO"), STAT_StaticMeshStreamInIO_DoInitiateIO, STATGROUP_StreamingDetails);
 	check(Context.CurrentThread == TT_Async);
 
 	SetIORequest(Context);
@@ -729,49 +647,48 @@ void TStaticMeshStreamIn_IO<bRenderThread>::DoInitiateIO(const FContext& Context
 	PushTask(Context, TT_Async, SRA_UPDATE_CALLBACK(DoSerializeLODData), TT_Async, SRA_UPDATE_CALLBACK(DoCancelIO));
 }
 
-template <bool bRenderThread>
-void TStaticMeshStreamIn_IO<bRenderThread>::DoSerializeLODData(const FContext& Context)
+void FStaticMeshStreamIn_IO::DoSerializeLODData(const FContext& Context)
 {
-	DECLARE_SCOPE_CYCLE_COUNTER(TEXT("TStaticMeshStreamIn_IO::DoSerializeLODData"), STAT_StaticMeshStreamInIO_DoSerializeLODData, STATGROUP_StreamingDetails);
+	DECLARE_SCOPE_CYCLE_COUNTER(TEXT("FStaticMeshStreamIn_IO::DoSerializeLODData"), STAT_StaticMeshStreamInIO_DoSerializeLODData, STATGROUP_StreamingDetails);
 	check(Context.CurrentThread == TT_Async);
+
 	SerializeLODData(Context);
 	ClearIORequest(Context);
-	const EThreadType TThread = bRenderThread ? TT_Render : TT_Async;
-	const EThreadType CThread = (EThreadType)Context.CurrentThread;
-	PushTask(Context, TThread, SRA_UPDATE_CALLBACK(DoCreateBuffers), CThread, SRA_UPDATE_CALLBACK(Cancel));
+
+	PushTask(Context
+		, CreateResourcesThread, SRA_UPDATE_CALLBACK(DoCreateBuffers)
+		, (EThreadType)Context.CurrentThread, SRA_UPDATE_CALLBACK(Cancel));
 }
 
-template <bool bRenderThread>
-void TStaticMeshStreamIn_IO<bRenderThread>::DoCreateBuffers(const FContext& Context)
+void FStaticMeshStreamIn_IO::DoCreateBuffers(const FContext& Context)
 {
-	DECLARE_SCOPE_CYCLE_COUNTER(TEXT("TStaticMeshStreamIn_IO::DoCreateBuffers"), STAT_StaticMeshStreamInIO_DoCreateBuffers, STATGROUP_StreamingDetails);
-	if (bRenderThread)
-	{
-		CreateBuffers_RenderThread(Context);
-	}
-	else
-	{
-		CreateBuffers_Async(Context);
-	}
+	DECLARE_SCOPE_CYCLE_COUNTER(TEXT("FStaticMeshStreamIn_IO::DoCreateBuffers"), STAT_StaticMeshStreamInIO_DoCreateBuffers, STATGROUP_StreamingDetails);
+	CreateBuffers(Context);
+
 	check(!TaskSynchronization.GetValue());
-	PushTask(Context, TT_Render, SRA_UPDATE_CALLBACK(DoFinishUpdate), (EThreadType)Context.CurrentThread, SRA_UPDATE_CALLBACK(Cancel));
+
+	// We cannot cancel once DoCreateBuffers has started executing, as there's an RHICmdList that must be submitted.
+	// Pass the same callback for both task and cancel.
+	PushTask(Context
+		, TT_Render, SRA_UPDATE_CALLBACK(DoFinishUpdate)
+		, TT_Render, SRA_UPDATE_CALLBACK(DoFinishUpdate)
+	);
 }
 
-template <bool bRenderThread>
-void TStaticMeshStreamIn_IO<bRenderThread>::DoCancelIO(const FContext& Context)
+void FStaticMeshStreamIn_IO::DoCancelIO(const FContext& Context)
 {
-	DECLARE_SCOPE_CYCLE_COUNTER(TEXT("TStaticMeshStreamIn_IO::DoCancelIO"), STAT_StaticMeshStreamInIO_DoCancelIO, STATGROUP_StreamingDetails);
+	DECLARE_SCOPE_CYCLE_COUNTER(TEXT("FStaticMeshStreamIn_IO::DoCancelIO"), STAT_StaticMeshStreamInIO_DoCancelIO, STATGROUP_StreamingDetails);
 	ClearIORequest(Context);
+
 	PushTask(Context, TT_None, nullptr, (EThreadType)Context.CurrentThread, SRA_UPDATE_CALLBACK(Cancel));
 }
 
-template class TStaticMeshStreamIn_IO<true>;
-template class TStaticMeshStreamIn_IO<false>;
-
 #if WITH_EDITOR
-FStaticMeshStreamIn_DDC::FStaticMeshStreamIn_DDC(const UStaticMesh* InMesh)
-	: FStaticMeshStreamIn(InMesh)
-{}
+FStaticMeshStreamIn_DDC::FStaticMeshStreamIn_DDC(const UStaticMesh* InMesh, EThreadType CreateResourcesThread)
+	: FStaticMeshStreamIn(InMesh, CreateResourcesThread)
+{
+	PushTask(FContext(InMesh, TT_None), TT_Async, SRA_UPDATE_CALLBACK(DoLoadNewLODsFromDDC), TT_None, nullptr);
+}
 
 void FStaticMeshStreamIn_DDC::LoadNewLODsFromDDC(const FContext& Context)
 {
@@ -779,40 +696,23 @@ void FStaticMeshStreamIn_DDC::LoadNewLODsFromDDC(const FContext& Context)
 	// TODO: support streaming CPU data for editor builds
 }
 
-template <bool bRenderThread>
-TStaticMeshStreamIn_DDC<bRenderThread>::TStaticMeshStreamIn_DDC(const UStaticMesh* InMesh)
-	: FStaticMeshStreamIn_DDC(InMesh)
-{
-	PushTask(FContext(InMesh, TT_None), TT_Async, SRA_UPDATE_CALLBACK(DoLoadNewLODsFromDDC), TT_None, nullptr);
-}
-
-template <bool bRenderThread>
-void TStaticMeshStreamIn_DDC<bRenderThread>::DoLoadNewLODsFromDDC(const FContext& Context)
+void FStaticMeshStreamIn_DDC::DoLoadNewLODsFromDDC(const FContext& Context)
 {
-	DECLARE_SCOPE_CYCLE_COUNTER(TEXT("TStaticMeshStreamInDDC::DoLoadNewLODsFromDDC"), STAT_StaticMeshStreamInDDC_DoLoadNewLODsFromDDC, STATGROUP_StreamingDetails);
+	DECLARE_SCOPE_CYCLE_COUNTER(TEXT("FStaticMeshStreamIn_DDC::DoLoadNewLODsFromDDC"), STAT_StaticMeshStreamInDDC_DoLoadNewLODsFromDDC, STATGROUP_StreamingDetails);
 	LoadNewLODsFromDDC(Context);
 	check(!TaskSynchronization.GetValue());
-	const EThreadType TThread = bRenderThread ? TT_Render : TT_Async;
-	const EThreadType CThread = (EThreadType)Context.CurrentThread;
-	PushTask(Context, TThread, SRA_UPDATE_CALLBACK(DoCreateBuffers), CThread, SRA_UPDATE_CALLBACK(DoCancel));
+
+	PushTask(Context
+		, CreateResourcesThread, SRA_UPDATE_CALLBACK(DoCreateBuffers)
+		, (EThreadType)Context.CurrentThread, SRA_UPDATE_CALLBACK(DoCancel));
 }
 
-template <bool bRenderThread>
-void TStaticMeshStreamIn_DDC<bRenderThread>::DoCreateBuffers(const FContext& Context)
+void FStaticMeshStreamIn_DDC::DoCreateBuffers(const FContext& Context)
 {
-	DECLARE_SCOPE_CYCLE_COUNTER(TEXT("TStaticMeshStreamInDDC::DoCreateBuffers"), STAT_StaticMeshStreamInDDC_DoCreateBuffers, STATGROUP_StreamingDetails);
-	if (bRenderThread)
-	{
-		CreateBuffers_RenderThread(Context);
-	}
-	else
-	{
-		CreateBuffers_Async(Context);
-	}
+	DECLARE_SCOPE_CYCLE_COUNTER(TEXT("FStaticMeshStreamIn_DDC::DoCreateBuffers"), STAT_StaticMeshStreamInDDC_DoCreateBuffers, STATGROUP_StreamingDetails);
+	CreateBuffers(Context);
+
 	check(!TaskSynchronization.GetValue());
-	PushTask(Context, TT_Render, SRA_UPDATE_CALLBACK(DoFinishUpdate), (EThreadType)Context.CurrentThread, SRA_UPDATE_CALLBACK(DoCancel));
+	PushTask(Context, TT_Render, SRA_UPDATE_CALLBACK(DoFinishUpdate), TT_None, nullptr);
 }
-
-template class TStaticMeshStreamIn_DDC<true>;
-template class TStaticMeshStreamIn_DDC<false>;
 #endif
diff --git a/Engine/Source/Runtime/Engine/Private/Streaming/StaticMeshUpdate.h b/Engine/Source/Runtime/Engine/Private/Streaming/StaticMeshUpdate.h
--- a/Engine/Source/Runtime/Engine/Private/Streaming/StaticMeshUpdate.h
+++ b/Engine/Source/Runtime/Engine/Private/Streaming/StaticMeshUpdate.h
@@ -68,7 +68,7 @@ protected:
 class FStaticMeshStreamIn : public FStaticMeshUpdate
 {
 public:
-	FStaticMeshStreamIn(const UStaticMesh* InMesh);
+	FStaticMeshStreamIn(const UStaticMesh* InMesh, EThreadType CreateResourcesThread);
 
 	virtual ~FStaticMeshStreamIn();
 
@@ -86,15 +86,10 @@ protected:
 		FBufferRHIRef ReversedDepthOnlyIndexBuffer;
 		FBufferRHIRef WireframeIndexBuffer;
 
-		void CreateFromCPUData_RenderThread(FStaticMeshLODResources& LODResource);
-		void CreateFromCPUData_Async(FStaticMeshLODResources& LODResource);
-
-		void SafeRelease();
+		void CreateFromCPUData(FRHICommandListBase& RHICmdList, FStaticMeshLODResources& LODResource);
 
 		/** Transfer ownership of buffers to a LOD resource */
-		void TransferBuffers(FStaticMeshLODResources& LODResource, FRHIResourceUpdateBatcher& Batcher);
-
-		void CheckIsNull() const;
+		void TransferBuffers(FStaticMeshLODResources& LODResource, FRHIResourceReplaceBatcher& Batcher);
 	};
 
 #if RHI_RAYTRACING
@@ -106,17 +101,16 @@ protected:
 		bool bRequiresBuild = false;
 
 	public:
-		void CreateFromCPUData(FRHICommandList& RHICmdList, FRayTracingGeometry& RayTracingGeometry);
+		void CreateFromCPUData(FRHICommandListBase& RHICmdList, FRayTracingGeometry& RayTracingGeometry);
 
 		void SafeRelease();
 
-		void TransferRayTracingGeometry(FRayTracingGeometry& RayTracingGeometry, FRHIResourceUpdateBatcher& Batcher);
+		void TransferRayTracingGeometry(FRayTracingGeometry& RayTracingGeometry, FRHIResourceReplaceBatcher& Batcher);
 	};
 #endif
 
-	/** Create buffers with new LOD data on render or pooled thread */
-	void CreateBuffers_RenderThread(const FContext& Context);
-	void CreateBuffers_Async(const FContext& Context);
+	/** Create buffers with new LOD data */
+	void CreateBuffers(const FContext& Context);
 
 	/** Discard newly streamed-in CPU data */
 	void DiscardNewLODs(const FContext& Context);
@@ -134,9 +128,11 @@ protected:
 	FIntermediateRayTracingGeometry IntermediateRayTracingGeometry[MAX_MESH_LOD_COUNT];
 #endif
 
-private:
-	template <bool bRenderThread>
-	void CreateBuffers_Internal(const FContext& Context);
+	/** RHI command list used for creating buffers and replacing the streaming placeholders. Submitted in DoFinishUpdate */
+	FRHICommandList* StreamingRHICmdList = nullptr;
+
+	/** The thread to use for recording the above command list */
+	const EThreadType CreateResourcesThread;
 };
 
 /** A streamout that doesn't actually touches the CPU data. Required because DDC stream in doesn't reset. */
@@ -160,7 +156,7 @@ private:
 class FStaticMeshStreamIn_IO : public FStaticMeshStreamIn
 {
 public:
-	FStaticMeshStreamIn_IO(const UStaticMesh* InMesh, bool bHighPrio);
+	FStaticMeshStreamIn_IO(const UStaticMesh* InMesh, bool bHighPrio, EThreadType CreateResourcesThread);
 
 	virtual ~FStaticMeshStreamIn_IO() {}
 
@@ -216,17 +212,7 @@ protected:
 
 	// Whether an IO error was detected (when files do not exists).
 	bool bFailedOnIOError = false;
-};
-
-template <bool bRenderThread>
-class TStaticMeshStreamIn_IO : public FStaticMeshStreamIn_IO
-{
-public:
-	TStaticMeshStreamIn_IO(const UStaticMesh* InMesh, bool bHighPrio);
 
-	virtual ~TStaticMeshStreamIn_IO() {}
-
-protected:
 	void DoInitiateIO(const FContext& Context);
 
 	void DoSerializeLODData(const FContext& Context);
@@ -236,28 +222,16 @@ protected:
 	void DoCancelIO(const FContext& Context);
 };
 
-typedef TStaticMeshStreamIn_IO<true> FStaticMeshStreamIn_IO_RenderThread;
-typedef TStaticMeshStreamIn_IO<false> FStaticMeshStreamIn_IO_Async;
-
 #if WITH_EDITOR
 class FStaticMeshStreamIn_DDC : public FStaticMeshStreamIn
 {
 public:
-	FStaticMeshStreamIn_DDC(const UStaticMesh* InMesh);
+	FStaticMeshStreamIn_DDC(const UStaticMesh* InMesh, EThreadType CreateResourcesThread);
 
 	virtual ~FStaticMeshStreamIn_DDC() {}
 
 protected:
 	void LoadNewLODsFromDDC(const FContext& Context);
-};
-
-template <bool bRenderThread>
-class TStaticMeshStreamIn_DDC : public FStaticMeshStreamIn_DDC
-{
-public:
-	TStaticMeshStreamIn_DDC(const UStaticMesh* InMesh);
-
-	virtual ~TStaticMeshStreamIn_DDC() {}
 
 private:
 	/** Load new LOD buffers from DDC and queue a task to create RHI buffers on RT */
@@ -266,7 +240,4 @@ private:
 	/** Create RHI buffers for newly streamed-in LODs and queue a task to rename references on RT */
 	void DoCreateBuffers(const FContext& Context);
 };
-
-typedef TStaticMeshStreamIn_DDC<true> FStaticMeshStreamIn_DDC_RenderThread;
-typedef TStaticMeshStreamIn_DDC<false> FStaticMeshStreamIn_DDC_Async;
 #endif
diff --git a/Engine/Source/Runtime/Engine/Public/Animation/SkinWeightProfile.h b/Engine/Source/Runtime/Engine/Public/Animation/SkinWeightProfile.h
--- a/Engine/Source/Runtime/Engine/Public/Animation/SkinWeightProfile.h
+++ b/Engine/Source/Runtime/Engine/Public/Animation/SkinWeightProfile.h
@@ -165,13 +165,8 @@ struct FSkinWeightProfilesData
 
 	ENGINE_API void CreateRHIBuffers(FRHICommandListBase& RHICmdList, TArray<TPair<FName, FSkinWeightRHIInfo>>& OutBuffers);
 
-	UE_DEPRECATED(5.4, "Use CreateRHIBuffers instead.")
-	ENGINE_API void CreateRHIBuffers_RenderThread(TArray<TPair<FName, FSkinWeightRHIInfo>>& OutBuffers);
-	UE_DEPRECATED(5.4, "Use CreateRHIBuffers instead.")
-	ENGINE_API void CreateRHIBuffers_Async(TArray<TPair<FName, FSkinWeightRHIInfo>>& OutBuffers);
-
-	ENGINE_API void InitRHIForStreaming(const TArray<TPair<FName, FSkinWeightRHIInfo>>& IntermediateBuffers, FRHIResourceUpdateBatcher& Batcher);
-	ENGINE_API void ReleaseRHIForStreaming(FRHIResourceUpdateBatcher& Batcher);
+	ENGINE_API void InitRHIForStreaming(const TArray<TPair<FName, FSkinWeightRHIInfo>>& IntermediateBuffers, FRHIResourceReplaceBatcher& Batcher);
+	ENGINE_API void ReleaseRHIForStreaming(FRHIResourceReplaceBatcher& Batcher);
 
 	ENGINE_API bool IsPendingReadback() const;
 	ENGINE_API void EnqueueGPUReadback();
diff --git a/Engine/Source/Runtime/Engine/Public/RawIndexBuffer.h b/Engine/Source/Runtime/Engine/Public/RawIndexBuffer.h
--- a/Engine/Source/Runtime/Engine/Public/RawIndexBuffer.h
+++ b/Engine/Source/Runtime/Engine/Public/RawIndexBuffer.h
@@ -278,16 +278,11 @@ public:
 	/** Create an RHI index buffer with CPU data. CPU data may be discarded after creation (see TResourceArray::Discard) */
 	FBufferRHIRef CreateRHIBuffer(FRHICommandListBase& RHICmdList);
 
-	UE_DEPRECATED(5.4, "Use CreateRHIBuffer instead.")
-	FBufferRHIRef CreateRHIBuffer_RenderThread();
-	UE_DEPRECATED(5.4, "Use CreateRHIBuffer instead.")
-	FBufferRHIRef CreateRHIBuffer_Async();
-
 	/** Take over ownership of IntermediateBuffer */
-	void InitRHIForStreaming(FRHIBuffer* IntermediateBuffer, FRHIResourceUpdateBatcher& Batcher);
+	void InitRHIForStreaming(FRHIBuffer* IntermediateBuffer, FRHIResourceReplaceBatcher& Batcher);
 
 	/** Release any GPU resource owned by the RHI object */
-	void ReleaseRHIForStreaming(FRHIResourceUpdateBatcher& Batcher);
+	void ReleaseRHIForStreaming(FRHIResourceReplaceBatcher& Batcher);
 
 	/**
 	 * Serialization.
@@ -362,8 +357,8 @@ protected:
 	ENGINE_API bool IsSRVNeeded(bool bAllowCPUAccess) const;
 
 	/** Similar to Init/ReleaseRHI but only update existing SRV so references to the SRV stays valid */
-	void InitRHIForStreaming(FRHIBuffer* IntermediateBuffer, size_t IndexSize, FRHIResourceUpdateBatcher& Batcher);
-	void ReleaseRHIForStreaming(FRHIResourceUpdateBatcher& Batcher);
+	void InitRHIForStreaming(FRHIBuffer* IntermediateBuffer, size_t IndexSize, FRHIResourceReplaceBatcher& Batcher);
+	void ReleaseRHIForStreaming(FRHIResourceReplaceBatcher& Batcher);
 
 	static ENGINE_API FBufferRHIRef CreateRHIIndexBufferInternal(
 		FRHICommandListBase& RHICmdList,
@@ -519,18 +514,13 @@ public:
 		return nullptr;
 	}
 
-	UE_DEPRECATED(5.4, "Use CreateRHIBuffer instead.")
-	FBufferRHIRef CreateRHIBuffer_RenderThread();
-	UE_DEPRECATED(5.4, "Use CreateRHIBuffer instead.")
-	FBufferRHIRef CreateRHIBuffer_Async();
-
 	/** Similar to Init/ReleaseRHI but only update existing SRV so references to the SRV stays valid */
-	void InitRHIForStreaming(FRHIBuffer* IntermediateBuffer, FRHIResourceUpdateBatcher& Batcher)
+	void InitRHIForStreaming(FRHIBuffer* IntermediateBuffer, FRHIResourceReplaceBatcher& Batcher)
 	{
 		FRawStaticIndexBuffer16or32Interface::InitRHIForStreaming(IntermediateBuffer, sizeof(INDEX_TYPE), Batcher);
 	}
 
-	void ReleaseRHIForStreaming(FRHIResourceUpdateBatcher& Batcher)
+	void ReleaseRHIForStreaming(FRHIResourceReplaceBatcher& Batcher)
 	{
 		FRawStaticIndexBuffer16or32Interface::ReleaseRHIForStreaming(Batcher);
 	}
diff --git a/Engine/Source/Runtime/Engine/Public/Rendering/ColorVertexBuffer.h b/Engine/Source/Runtime/Engine/Public/Rendering/ColorVertexBuffer.h
--- a/Engine/Source/Runtime/Engine/Public/Rendering/ColorVertexBuffer.h
+++ b/Engine/Source/Runtime/Engine/Public/Rendering/ColorVertexBuffer.h
@@ -145,14 +145,9 @@ public:
 	/** Create an RHI vertex buffer with CPU data. CPU data may be discarded after creation (see TResourceArray::Discard) */
 	FBufferRHIRef CreateRHIBuffer(FRHICommandListBase& RHICmdList);
 
-	UE_DEPRECATED(5.4, "Use CreateRHIBuffer instead.")
-	FBufferRHIRef CreateRHIBuffer_RenderThread();
-	UE_DEPRECATED(5.4, "Use CreateRHIBuffer instead.")
-	FBufferRHIRef CreateRHIBuffer_Async();
-
 	/** Similar to Init/ReleaseRHI but only update existing SRV so references to the SRV stays valid */
-	void InitRHIForStreaming(FRHIBuffer* IntermediateBuffer, FRHIResourceUpdateBatcher& Batcher);
-	void ReleaseRHIForStreaming(FRHIResourceUpdateBatcher& Batcher);
+	void InitRHIForStreaming(FRHIBuffer* IntermediateBuffer, FRHIResourceReplaceBatcher& Batcher);
+	void ReleaseRHIForStreaming(FRHIResourceReplaceBatcher& Batcher);
 
 	// FRenderResource interface.
 	ENGINE_API virtual void InitRHI(FRHICommandListBase& RHICmdList) override;
diff --git a/Engine/Source/Runtime/Engine/Public/Rendering/MultiSizeIndexContainer.h b/Engine/Source/Runtime/Engine/Public/Rendering/MultiSizeIndexContainer.h
--- a/Engine/Source/Runtime/Engine/Public/Rendering/MultiSizeIndexContainer.h
+++ b/Engine/Source/Runtime/Engine/Public/Rendering/MultiSizeIndexContainer.h
@@ -101,13 +101,8 @@ public:
 
 	ENGINE_API FBufferRHIRef CreateRHIBuffer(FRHICommandListBase& RHICmdList);
 
-	UE_DEPRECATED(5.4, "Use CreateRHIBuffer instead.")
-	ENGINE_API FBufferRHIRef CreateRHIBuffer_RenderThread();
-	UE_DEPRECATED(5.4, "Use CreateRHIBuffer instead.")
-	ENGINE_API FBufferRHIRef CreateRHIBuffer_Async();
-
-	void InitRHIForStreaming(FRHIBuffer* IntermediateBuffer, FRHIResourceUpdateBatcher& Batcher);
-	void ReleaseRHIForStreaming(FRHIResourceUpdateBatcher& Batcher);
+	void InitRHIForStreaming(FRHIBuffer* IntermediateBuffer, FRHIResourceReplaceBatcher& Batcher);
+	void ReleaseRHIForStreaming(FRHIResourceReplaceBatcher& Batcher);
 
 private:
 	/** Size of the index buffer's index type (should be 2 or 4 bytes) */
diff --git a/Engine/Source/Runtime/Engine/Public/Rendering/PositionVertexBuffer.h b/Engine/Source/Runtime/Engine/Public/Rendering/PositionVertexBuffer.h
--- a/Engine/Source/Runtime/Engine/Public/Rendering/PositionVertexBuffer.h
+++ b/Engine/Source/Runtime/Engine/Public/Rendering/PositionVertexBuffer.h
@@ -106,15 +106,9 @@ public:
 	/** Create an RHI vertex buffer with CPU data. CPU data may be discarded after creation (see TResourceArray::Discard) */
 	FBufferRHIRef CreateRHIBuffer(FRHICommandListBase& RHICmdList);
 
-	UE_DEPRECATED(5.4, "Use CreateRHIBuffer instead.")
-	FBufferRHIRef CreateRHIBuffer_RenderThread();
-
-	UE_DEPRECATED(5.4, "Use CreateRHIBuffer instead.")
-	FBufferRHIRef CreateRHIBuffer_Async();
-
 	/** Similar to Init/ReleaseRHI but only update existing SRV so references to the SRV stays valid */
-	void InitRHIForStreaming(FRHIBuffer* IntermediateBuffer, FRHIResourceUpdateBatcher& Batcher);
-	void ReleaseRHIForStreaming(FRHIResourceUpdateBatcher& Batcher);
+	void InitRHIForStreaming(FRHIBuffer* IntermediateBuffer, FRHIResourceReplaceBatcher& Batcher);
+	void ReleaseRHIForStreaming(FRHIResourceReplaceBatcher& Batcher);
 
 	// FRenderResource interface.
 	ENGINE_API virtual void InitRHI(FRHICommandListBase& RHICmdList) override;
diff --git a/Engine/Source/Runtime/Engine/Public/Rendering/SkeletalMeshVertexClothBuffer.h b/Engine/Source/Runtime/Engine/Public/Rendering/SkeletalMeshVertexClothBuffer.h
--- a/Engine/Source/Runtime/Engine/Public/Rendering/SkeletalMeshVertexClothBuffer.h
+++ b/Engine/Source/Runtime/Engine/Public/Rendering/SkeletalMeshVertexClothBuffer.h
@@ -126,14 +126,9 @@ public:
 	/** Create an RHI vertex buffer with CPU data. CPU data may be discarded after creation (see TResourceArray::Discard) */
 	FBufferRHIRef CreateRHIBuffer(FRHICommandListBase& RHICmdList);
 
-	UE_DEPRECATED(5.4, "Use CreateRHIBuffer instead.")
-	FBufferRHIRef CreateRHIBuffer_RenderThread();
-	UE_DEPRECATED(5.4, "Use CreateRHIBuffer instead.")
-	FBufferRHIRef CreateRHIBuffer_Async();
-
 	/** Similar to Init/ReleaseRHI but only update existing SRV so references to the SRV stays valid */
-	void InitRHIForStreaming(FRHIBuffer* IntermediateBuffer, FRHIResourceUpdateBatcher& Batcher);
-	void ReleaseRHIForStreaming(FRHIResourceUpdateBatcher& Batcher);
+	void InitRHIForStreaming(FRHIBuffer* IntermediateBuffer, FRHIResourceReplaceBatcher& Batcher);
+	void ReleaseRHIForStreaming(FRHIResourceReplaceBatcher& Batcher);
 
 private:
 	FShaderResourceViewRHIRef VertexBufferSRV;
diff --git a/Engine/Source/Runtime/Engine/Public/Rendering/SkinWeightVertexBuffer.h b/Engine/Source/Runtime/Engine/Public/Rendering/SkinWeightVertexBuffer.h
--- a/Engine/Source/Runtime/Engine/Public/Rendering/SkinWeightVertexBuffer.h
+++ b/Engine/Source/Runtime/Engine/Public/Rendering/SkinWeightVertexBuffer.h
@@ -140,14 +140,9 @@ public:
 	/** Create an RHI vertex buffer with CPU data. CPU data may be discarded after creation (see TResourceArray::Discard) */
 	FBufferRHIRef CreateRHIBuffer(FRHICommandListBase& RHICmdList);
 
-	UE_DEPRECATED(5.4, "Use CreateRHIBuffer instead.")
-	FBufferRHIRef CreateRHIBuffer_RenderThread();
-	UE_DEPRECATED(5.4, "Use CreateRHIBuffer instead.")
-	FBufferRHIRef CreateRHIBuffer_Async();
-
 	/** Similar to Init/ReleaseRHI but only update existing SRV so references to the SRV stays valid */
-	void InitRHIForStreaming(FRHIBuffer* IntermediateBuffer, FRHIResourceUpdateBatcher& Batcher);
-	void ReleaseRHIForStreaming(FRHIResourceUpdateBatcher& Batcher);
+	void InitRHIForStreaming(FRHIBuffer* IntermediateBuffer, FRHIResourceReplaceBatcher& Batcher);
+	void ReleaseRHIForStreaming(FRHIResourceReplaceBatcher& Batcher);
 
 	void GetWeightOffsetAndInfluenceCount(uint32 VertexIndex, uint32& OutWeightOffset, uint32& OutInfluenceCount) const;
 	void SetWeightOffsetAndInfluenceCount(uint32 VertexIndex, uint32 WeightOffset, uint32 InfluenceCount);
@@ -303,14 +298,9 @@ public:
 	/** Create an RHI vertex buffer with CPU data. CPU data may be discarded after creation (see TResourceArray::Discard) */
 	FBufferRHIRef CreateRHIBuffer(FRHICommandListBase& RHICmdList);
 
-	UE_DEPRECATED(5.4, "Use CreateRHIBuffer instead.")
-	FBufferRHIRef CreateRHIBuffer_RenderThread();
-	UE_DEPRECATED(5.4, "Use CreateRHIBuffer instead.")
-	FBufferRHIRef CreateRHIBuffer_Async();
-
 	/** Similar to Init/ReleaseRHI but only update existing SRV so references to the SRV stays valid */
-	void InitRHIForStreaming(FRHIBuffer* IntermediateBuffer, FRHIResourceUpdateBatcher& Batcher);
-	void ReleaseRHIForStreaming(FRHIResourceUpdateBatcher& Batcher);
+	void InitRHIForStreaming(FRHIBuffer* IntermediateBuffer, FRHIResourceReplaceBatcher& Batcher);
+	void ReleaseRHIForStreaming(FRHIResourceReplaceBatcher& Batcher);
 
 	bool IsWeightDataValid() const;
 
@@ -467,11 +457,6 @@ public:
 
 	FSkinWeightRHIInfo CreateRHIBuffer(FRHICommandListBase& RHICmdList);
 
-	UE_DEPRECATED(5.4, "Use CreateRHIBuffer instead.")
-	FSkinWeightRHIInfo CreateRHIBuffer_RenderThread();
-	UE_DEPRECATED(5.4, "Use CreateRHIBuffer instead.")
-	FSkinWeightRHIInfo CreateRHIBuffer_Async();
-
 	ENGINE_API GPUSkinBoneInfluenceType GetBoneInfluenceType() const;
 	ENGINE_API void GetVertexInfluenceOffsetCount(uint32 VertexIndex, uint32& VertexWeightOffset, uint32& VertexInfluenceCount) const;
 	ENGINE_API bool GetRigidWeightBone(uint32 VertexIndex, int32& OutBoneIndex) const;
@@ -492,13 +477,13 @@ public:
 	ENGINE_API void ReleaseResources();
 
 	/** Similar to Init/ReleaseRHI but only update existing SRV so references to the SRV stays valid */
-	void InitRHIForStreaming(const FSkinWeightRHIInfo& RHIInfo, FRHIResourceUpdateBatcher& Batcher)
+	void InitRHIForStreaming(const FSkinWeightRHIInfo& RHIInfo, FRHIResourceReplaceBatcher& Batcher)
 	{
 		DataVertexBuffer.InitRHIForStreaming(RHIInfo.DataVertexBufferRHI, Batcher);
 		LookupVertexBuffer.InitRHIForStreaming(RHIInfo.LookupVertexBufferRHI, Batcher);
 	}
 
-	void ReleaseRHIForStreaming(FRHIResourceUpdateBatcher& Batcher)
+	void ReleaseRHIForStreaming(FRHIResourceReplaceBatcher& Batcher)
 	{
 		DataVertexBuffer.ReleaseRHIForStreaming(Batcher);
 		LookupVertexBuffer.ReleaseRHIForStreaming(Batcher);
diff --git a/Engine/Source/Runtime/Engine/Public/Rendering/StaticMeshVertexBuffer.h b/Engine/Source/Runtime/Engine/Public/Rendering/StaticMeshVertexBuffer.h
--- a/Engine/Source/Runtime/Engine/Public/Rendering/StaticMeshVertexBuffer.h
+++ b/Engine/Source/Runtime/Engine/Public/Rendering/StaticMeshVertexBuffer.h
@@ -440,18 +440,9 @@ public:
 	FBufferRHIRef CreateTangentsRHIBuffer(FRHICommandListBase& RHICmdList);
 	FBufferRHIRef CreateTexCoordRHIBuffer(FRHICommandListBase& RHICmdList);
 
-	UE_DEPRECATED(5.4, "Use CreateTangentsRHIBuffer instead.")
-	FBufferRHIRef CreateTangentsRHIBuffer_RenderThread();
-	UE_DEPRECATED(5.4, "Use CreateTangentsRHIBuffer instead.")
-	FBufferRHIRef CreateTangentsRHIBuffer_Async();
-	UE_DEPRECATED(5.4, "Use CreateTexCoordRHIBuffer instead.")
-	FBufferRHIRef CreateTexCoordRHIBuffer_RenderThread();
-	UE_DEPRECATED(5.4, "Use CreateTexCoordRHIBuffer instead.")
-	FBufferRHIRef CreateTexCoordRHIBuffer_Async();
-
 	/** Similar to Init/ReleaseRHI but only update existing SRV so references to the SRV stays valid */
-	void InitRHIForStreaming(FRHIBuffer* IntermediateTangentsBuffer, FRHIBuffer* IntermediateTexCoordBuffer, FRHIResourceUpdateBatcher& Batcher);
-	void ReleaseRHIForStreaming(FRHIResourceUpdateBatcher& Batcher);
+	void InitRHIForStreaming(FRHIBuffer* IntermediateTangentsBuffer, FRHIBuffer* IntermediateTexCoordBuffer, FRHIResourceReplaceBatcher& Batcher);
+	void ReleaseRHIForStreaming(FRHIResourceReplaceBatcher& Batcher);
 
 	// FRenderResource interface.
 	ENGINE_API virtual void InitRHI(FRHICommandListBase& RHICmdList) override;
diff --git a/Engine/Source/Runtime/RHI/Private/RHICommandList.cpp b/Engine/Source/Runtime/RHI/Private/RHICommandList.cpp
--- a/Engine/Source/Runtime/RHI/Private/RHICommandList.cpp
+++ b/Engine/Source/Runtime/RHI/Private/RHICommandList.cpp
@@ -8,7 +8,7 @@
 #include "ProfilingDebugging/CsvProfiler.h"
 #include "GenericPlatform/GenericPlatformCrashContext.h"
 #include "RHIBreadcrumbs.h"
-#include "RHIResourceUpdates.h"
+#include "RHIResourceReplace.h"
 #include "RHIContext.h"
 #include "RHIFwd.h"
 #include "RHITransition.h"
@@ -980,6 +980,13 @@ RHI_API void FRHICommandListExecutor::Submit(TConstArrayView<FRHICommandListBase
 	check(IsInRenderingThread());
 	SCOPED_NAMED_EVENT(RHICmdList_Submit, FColor::White);
 
+	if (Bypass())
+	{
+		// Always submit to the GPU in Bypass mode. This allows us to wait for all translate tasks to 
+		// complete before  returning from this function. ensuring commands are always executed in-order.
+		EnumAddFlags(SubmitFlags, ERHISubmitFlags::SubmitToGPU);
+	}
+
 	// Commands may already be queued on the immediate command list. These need to be executed
 	// first before any parallel commands can be inserted, otherwise commands will run out-of-order.
 	FRHICommandListBase* ImmCmdList;
@@ -2110,26 +2117,40 @@ void FRHICommandListBase::UpdateTextureReference(FRHITextureReference* TextureRe
 	RHIThreadFence(true);
 }
 
-void FRHICommandListImmediate::UpdateRHIResources(FRHIResourceUpdateInfo* UpdateInfos, int32 Num, bool bNeedReleaseRefs)
+void FRHICommandListBase::ReplaceResources(TArray<FRHIResourceReplaceInfo>&& ReplaceInfos)
 {
-	if (this->Bypass())
+	EnqueueLambda(TEXT("ReplaceResources"), [Infos = MoveTemp(ReplaceInfos)](FRHICommandListBase& ExecutingCmdList)
 	{
-		FRHICommandUpdateRHIResources Cmd(UpdateInfos, Num, bNeedReleaseRefs);
-		Cmd.Execute(*this);
-	}
-	else
-	{
-		const SIZE_T NumBytes = sizeof(FRHIResourceUpdateInfo) * Num;
-		FRHIResourceUpdateInfo* LocalUpdateInfos = reinterpret_cast<FRHIResourceUpdateInfo*>(this->Alloc(NumBytes, alignof(FRHIResourceUpdateInfo)));
-		FMemory::Memcpy(LocalUpdateInfos, UpdateInfos, NumBytes);
-		new (AllocCommand<FRHICommandUpdateRHIResources>()) FRHICommandUpdateRHIResources(LocalUpdateInfos, Num, bNeedReleaseRefs);
-		RHIThreadFence(true);
-		if (GetUsedMemory() > 256 * 1024)
+		RHISTAT(ReplaceResources);
+		for (FRHIResourceReplaceInfo const& Info : Infos)
 		{
-			// we could be loading a level or something, lets get this stuff going
-			ImmediateFlush(EImmediateFlushType::DispatchToRHIThread);
+			switch (Info.GetType())
+			{
+
+			case FRHIResourceReplaceInfo::EType::Buffer:
+				GDynamicRHI->RHITransferBufferUnderlyingResource(
+					ExecutingCmdList,
+					Info.GetBuffer().Dst,
+					Info.GetBuffer().Src);
+				break;
+
+#if RHI_RAYTRACING
+			case FRHIResourceReplaceInfo::EType::RTGeometry:
+				GDynamicRHI->RHITransferRayTracingGeometryUnderlyingResource(
+					ExecutingCmdList,
+					Info.GetRTGeometry().Dst,
+					Info.GetRTGeometry().Src);
+				break;
+#endif // RHI_RAYTRACING
+
+			default:
+				checkNoEntry();
+				break;
+			}
 		}
-	}
+	});
+
+	RHIThreadFence(true);
 }
 
 void FRHICommandListExecutor::CleanupGraphEvents()
diff --git a/Engine/Source/Runtime/RHI/Private/RHIResourceUpdates.cpp b/Engine/Source/Runtime/RHI/Private/RHIResourceUpdates.cpp
deleted file mode 100644
diff --git a/Engine/Source/Runtime/RHI/Public/RHICommandList.h b/Engine/Source/Runtime/RHI/Public/RHICommandList.h
--- a/Engine/Source/Runtime/RHI/Public/RHICommandList.h
+++ b/Engine/Source/Runtime/RHI/Public/RHICommandList.h
@@ -32,6 +32,7 @@
 #include "RHIGlobals.h"
 #include "RHIShaderParameters.h"
 #include "RHITextureReference.h"
+#include "RHIResourceReplace.h"
 #include "Trace/Trace.h"
 
 #include "DynamicRHI.h"
@@ -898,6 +899,8 @@ public:
 	}
 #endif
 
+	RHI_API void ReplaceResources(TArray<FRHIResourceReplaceInfo>&& ReplaceInfos);
+
 	FORCEINLINE void BindDebugLabelName(FRHITexture* Texture, const TCHAR* Name)
 	{
 		GDynamicRHI->RHIBindDebugLabelName(*this, Texture, Name);
@@ -2121,23 +2124,6 @@ FRHICOMMAND_MACRO(FRHICommandDiscardRenderTargets)
 	RHI_API void Execute(FRHICommandListBase& CmdList);
 };
 
-FRHICOMMAND_MACRO(FRHICommandUpdateRHIResources)
-{
-	FRHIResourceUpdateInfo* UpdateInfos;
-	int32 Num;
-	bool bNeedReleaseRefs;
-
-	FRHICommandUpdateRHIResources(FRHIResourceUpdateInfo* InUpdateInfos, int32 InNum, bool bInNeedReleaseRefs)
-		: UpdateInfos(InUpdateInfos)
-		, Num(InNum)
-		, bNeedReleaseRefs(bInNeedReleaseRefs)
-	{}
-
-	~FRHICommandUpdateRHIResources();
-
-	RHI_API void Execute(FRHICommandListBase& CmdList);
-};
-
 FRHICOMMAND_MACRO(FRHICommandCopyBufferRegion)
 {
 	FRHIBuffer* DestBuffer;
@@ -4592,13 +4578,6 @@ public:
 	UE_DEPRECATED(5.5, "RHIPollRenderQueryResults is deprecated. Platform RHIs that require query polling now do this automatically as part of RHI command list submission.")
 	FORCEINLINE void PollRenderQueryResults() {}
 
-	/**
-	 * @param UpdateInfos - an array of update infos
-	 * @param Num - number of update infos
-	 * @param bNeedReleaseRefs - whether Release need to be called on RHI resources referenced by update infos
-	 */
-	RHI_API void UpdateRHIResources(FRHIResourceUpdateInfo* UpdateInfos, int32 Num, bool bNeedReleaseRefs);
-
 	//UE_DEPRECATED(5.1, "SubmitCommandsHint is deprecated. Consider calling ImmediateFlush(EImmediateFlushType::DispatchToRHIThread) instead.")
 	FORCEINLINE_DEBUGGABLE void SubmitCommandsHint()
 	{
@@ -5228,12 +5207,6 @@ FORCEINLINE FShaderResourceViewRHIRef RHICreateShaderResourceView(FRHITexture* T
 	return FRHICommandListExecutor::GetImmediateCommandList().CreateShaderResourceView(Texture, CreateInfo);
 }
 
-FORCEINLINE void RHIUpdateRHIResources(FRHIResourceUpdateInfo* UpdateInfos, int32 Num, bool bNeedReleaseRefs)
-{
-	check(IsInRenderingThread());
-	return FRHICommandListExecutor::GetImmediateCommandList().UpdateRHIResources(UpdateInfos, Num, bNeedReleaseRefs);
-}
-
 FORCEINLINE FTextureReferenceRHIRef RHICreateTextureReference(FRHITexture* InReferencedTexture = nullptr)
 {
 	return FRHICommandListImmediate::Get().CreateTextureReference(InReferencedTexture);
diff --git a/Engine/Source/Runtime/RHI/Public/RHICommandList.inl b/Engine/Source/Runtime/RHI/Public/RHICommandList.inl
--- a/Engine/Source/Runtime/RHI/Public/RHICommandList.inl
+++ b/Engine/Source/Runtime/RHI/Public/RHICommandList.inl
@@ -129,6 +129,14 @@ inline FRHICommandListScopedPipelineGuard::~FRHICommandListScopedPipelineGuard()
 	}
 }
 
+inline FRHIResourceReplaceBatcher::~FRHIResourceReplaceBatcher()
+{
+	if (Infos.Num() > 0)
+	{
+		RHICmdList.ReplaceResources(MoveTemp(Infos));
+	}
+}
+
 #if WITH_RHI_BREADCRUMBS
 
 	template<size_t N, typename... TArgs>
diff --git a/Engine/Source/Runtime/RHI/Public/RHICommandListCommandExecutes.inl b/Engine/Source/Runtime/RHI/Public/RHICommandListCommandExecutes.inl
--- a/Engine/Source/Runtime/RHI/Public/RHICommandListCommandExecutes.inl
+++ b/Engine/Source/Runtime/RHI/Public/RHICommandListCommandExecutes.inl
@@ -13,8 +13,6 @@
 #define INTERNAL_DECORATOR_COMPUTE(Method) CmdList.GetComputeContext().Method
 #endif
 
-#include "RHIResourceUpdates.h"
-
 class FRHICommandListBase;
 class IRHIComputeContext;
 struct FComputedBSS;
@@ -548,69 +546,6 @@ void FRHICommandSetRayTracingBindings::Execute(FRHICommandListBase& CmdList)
 
 #endif // RHI_RAYTRACING
 
-void FRHIResourceUpdateInfo::ReleaseRefs()
-{
-	switch (Type)
-	{
-	case UT_Buffer:
-		Buffer.DestBuffer->Release();
-		if (Buffer.SrcBuffer)
-		{
-			Buffer.SrcBuffer->Release();
-		}
-		break;
-	case UT_RayTracingGeometry:
-		RayTracingGeometry.DestGeometry->Release();
-		if (RayTracingGeometry.SrcGeometry)
-		{
-			RayTracingGeometry.SrcGeometry->Release();
-		}
-	default:
-		// Unrecognized type, do nothing
-		break;
-	}
-}
-
-FRHICommandUpdateRHIResources::~FRHICommandUpdateRHIResources()
-{
-	if (bNeedReleaseRefs)
-	{
-		for (int32 Idx = 0; Idx < Num; ++Idx)
-		{
-			UpdateInfos[Idx].ReleaseRefs();
-		}
-	}
-}
-
-void FRHICommandUpdateRHIResources::Execute(FRHICommandListBase& CmdList)
-{
-	RHISTAT(UpdateRHIResources);
-	for (int32 Idx = 0; Idx < Num; ++Idx)
-	{
-		FRHIResourceUpdateInfo& Info = UpdateInfos[Idx];
-		switch (Info.Type)
-		{
-		case FRHIResourceUpdateInfo::UT_Buffer:
-			GDynamicRHI->RHITransferBufferUnderlyingResource(
-				CmdList,
-				Info.Buffer.DestBuffer,
-				Info.Buffer.SrcBuffer);
-			break;
-#if RHI_RAYTRACING
-		case FRHIResourceUpdateInfo::UT_RayTracingGeometry:
-			GDynamicRHI->RHITransferRayTracingGeometryUnderlyingResource(
-				CmdList,
-				Info.RayTracingGeometry.DestGeometry,
-				Info.RayTracingGeometry.SrcGeometry);
-			break;
-#endif // RHI_RAYTRACING
-		default:
-			// Unrecognized type, do nothing
-			break;
-		}
-	}
-}
-
 void FRHICommandBeginScene::Execute(FRHICommandListBase& CmdList)
 {
 	RHISTAT(BeginScene);
diff --git a/Engine/Source/Runtime/RHI/Public/RHIFwd.h b/Engine/Source/Runtime/RHI/Public/RHIFwd.h
--- a/Engine/Source/Runtime/RHI/Public/RHIFwd.h
+++ b/Engine/Source/Runtime/RHI/Public/RHIFwd.h
@@ -38,8 +38,8 @@ class FRHICommandListImmediate;
 class IRHIComputeContext;
 class IRHICommandContext;
 
-struct FRHIResourceUpdateInfo;
-struct FRHIResourceUpdateBatcher;
+class FRHIResourceReplaceInfo;
+class FRHIResourceReplaceBatcher;
 
 struct FSamplerStateInitializerRHI;
 struct FRasterizerStateInitializerRHI;
diff --git a/Engine/Source/Runtime/RHI/Public/RHIResourceReplace.h b/Engine/Source/Runtime/RHI/Public/RHIResourceReplace.h
new file mode 100644
--- /dev/null
+++ b/Engine/Source/Runtime/RHI/Public/RHIResourceReplace.h
@@ -0,0 +1,76 @@
+// Copyright Epic Games, Inc. All Rights Reserved.
+
+#pragma once
+
+#include "RHIFwd.h"
+#include "RHIResources.h"
+
+#include "Containers/Array.h"
+#include "Misc/TVariant.h"
+
+class FRHIResourceReplaceInfo
+{
+public:
+	template <typename TType>
+	struct TPair
+	{
+		TType* Dst;
+		TType* Src;
+
+		TPair(TType* Dst, TType* Src)
+			: Dst(Dst)
+			, Src(Src)
+		{}
+	};
+
+	typedef TPair<FRHIBuffer            > FBuffer;
+	typedef TPair<FRHIRayTracingGeometry> FRTGeometry;
+
+	typedef TVariant<FBuffer, FRTGeometry> TStorage;
+
+	enum class EType : uint8
+	{
+		Buffer     = TStorage::IndexOfType<FBuffer    >(),
+		RTGeometry = TStorage::IndexOfType<FRTGeometry>(),
+	};
+
+	EType GetType() const { return EType(Storage.GetIndex()); }
+
+	FBuffer     const& GetBuffer    () const { return Storage.Get<FBuffer    >(); }
+	FRTGeometry const& GetRTGeometry() const { return Storage.Get<FRTGeometry>(); }
+
+	FRHIResourceReplaceInfo(FRHIBuffer* Dst, FRHIBuffer* Src)
+		: Storage(TInPlaceType<FBuffer>(), Dst, Src)
+	{}
+
+	FRHIResourceReplaceInfo(FRHIRayTracingGeometry* Dst, FRHIRayTracingGeometry* Src)
+		: Storage(TInPlaceType<FRTGeometry>(), Dst, Src)
+	{}
+
+private:
+	TStorage Storage;
+};
+
+class FRHIResourceReplaceBatcher
+{
+	FRHICommandListBase& RHICmdList;
+	TArray<FRHIResourceReplaceInfo> Infos;
+
+public:
+	FRHIResourceReplaceBatcher(FRHICommandListBase& RHICmdList, uint32 InitialCapacity = 0)
+		: RHICmdList(RHICmdList)
+	{
+		if (InitialCapacity > 0)
+		{
+			Infos.Reserve(InitialCapacity);
+		}
+	}
+
+	inline ~FRHIResourceReplaceBatcher();
+
+	template <typename... TArgs>
+	void EnqueueReplace(TArgs&&... Args)
+	{
+		Infos.Emplace(Forward<TArgs>(Args)...);
+	}
+};
diff --git a/Engine/Source/Runtime/RHI/Public/RHIResourceUpdates.h b/Engine/Source/Runtime/RHI/Public/RHIResourceUpdates.h
deleted file mode 100644
diff --git a/Engine/Source/Runtime/RenderCore/Private/RayTracingGeometry.cpp b/Engine/Source/Runtime/RenderCore/Private/RayTracingGeometry.cpp
--- a/Engine/Source/Runtime/RenderCore/Private/RayTracingGeometry.cpp
+++ b/Engine/Source/Runtime/RenderCore/Private/RayTracingGeometry.cpp
@@ -5,7 +5,7 @@
 #include "HAL/IConsoleManager.h"
 #include "RayTracingGeometryManagerInterface.h"
 #include "RenderUtils.h"
-#include "RHIResourceUpdates.h"
+#include "RHIResourceReplace.h"
 #include "RHITextureReference.h" // IWYU pragma: keep
 
 #if RHI_RAYTRACING
@@ -27,7 +27,7 @@ FRayTracingGeometry::~FRayTracingGeometry() = default;
 
 #if RHI_RAYTRACING
 
-void FRayTracingGeometry::InitRHIForStreaming(FRHIRayTracingGeometry* IntermediateGeometry, FRHIResourceUpdateBatcher& Batcher)
+void FRayTracingGeometry::InitRHIForStreaming(FRHIRayTracingGeometry* IntermediateGeometry, FRHIResourceReplaceBatcher& Batcher)
 {
 	ensureMsgf(RayTracingGeometryRHI || !IsRayTracingEnabled(),
 		TEXT("RayTracingGeometryRHI should be valid when ray tracing is enabled.\n")
@@ -40,7 +40,7 @@ void FRayTracingGeometry::InitRHIForStreaming(FRHIRayTracingGeometry* Intermedia
 
 	if (RayTracingGeometryRHI && IntermediateGeometry)
 	{
-		Batcher.QueueUpdateRequest(RayTracingGeometryRHI, IntermediateGeometry);
+		Batcher.EnqueueReplace(RayTracingGeometryRHI, IntermediateGeometry);
 		EnumAddFlags(GeometryState, EGeometryStateFlags::Valid);
 	}
 	else
@@ -49,7 +49,7 @@ void FRayTracingGeometry::InitRHIForStreaming(FRHIRayTracingGeometry* Intermedia
 	}
 }
 
-void FRayTracingGeometry::ReleaseRHIForStreaming(FRHIResourceUpdateBatcher& Batcher)
+void FRayTracingGeometry::ReleaseRHIForStreaming(FRHIResourceReplaceBatcher& Batcher)
 {
 	RemoveBuildRequest();
 
@@ -60,7 +60,7 @@ void FRayTracingGeometry::ReleaseRHIForStreaming(FRHIResourceUpdateBatcher& Batc
 
 	if (RayTracingGeometryRHI)
 	{
-		Batcher.QueueUpdateRequest(RayTracingGeometryRHI, nullptr);
+		Batcher.EnqueueReplace(RayTracingGeometryRHI, nullptr);
 	}
 }
 
@@ -117,7 +117,7 @@ void FRayTracingGeometry::MakeResident(FRHICommandList& RHICmdList)
 		// Follows the same pattern as: (1) FStaticMeshStreamIn::CreateBuffers_* (2) FStaticMeshStreamIn::DoFinishUpdate
 		// There is no other way to initialize BLAS buffers for the geometry that has a StreamingDestination type.
 		{
-			TRHIResourceUpdateBatcher<1> Batcher;
+			FRHIResourceReplaceBatcher Batcher(RHICmdList, 1);
 			FRayTracingGeometryInitializer IntermediateInitializer = Initializer;
 			IntermediateInitializer.Type = ERayTracingGeometryInitializerType::StreamingSource;
 
diff --git a/Engine/Source/Runtime/RenderCore/Public/RayTracingGeometry.h b/Engine/Source/Runtime/RenderCore/Public/RayTracingGeometry.h
--- a/Engine/Source/Runtime/RenderCore/Public/RayTracingGeometry.h
+++ b/Engine/Source/Runtime/RenderCore/Public/RayTracingGeometry.h
@@ -110,8 +110,8 @@ public:
 		return GeometryState;
 	}
 
-	RENDERCORE_API void InitRHIForStreaming(FRHIRayTracingGeometry* IntermediateGeometry, FRHIResourceUpdateBatcher& Batcher);
-	RENDERCORE_API void ReleaseRHIForStreaming(FRHIResourceUpdateBatcher& Batcher);
+	RENDERCORE_API void InitRHIForStreaming(FRHIRayTracingGeometry* IntermediateGeometry, FRHIResourceReplaceBatcher& Batcher);
+	RENDERCORE_API void ReleaseRHIForStreaming(FRHIResourceReplaceBatcher& Batcher);
 
 	UE_DEPRECATED(5.4, "Use FStaticMeshStreamIn::FIntermediateRayTracingGeometry instead.")
 	RENDERCORE_API void CreateRayTracingGeometryFromCPUData(TResourceArray<uint8>& OfflineData);
diff --git a/Engine/Source/Runtime/RenderCore/Public/RenderingThread.h b/Engine/Source/Runtime/RenderCore/Public/RenderingThread.h
--- a/Engine/Source/Runtime/RenderCore/Public/RenderingThread.h
+++ b/Engine/Source/Runtime/RenderCore/Public/RenderingThread.h
@@ -579,46 +579,6 @@ private:
 // RenderThread scoped work
 ////////////////////////////////////
 
-/** A utility to record RHI commands asynchronously and then enqueue the resulting commands to the render thread. */
-class FRHIAsyncCommandList
-{
-public:
-	FRHIAsyncCommandList(FRHIGPUMask InGPUMask = FRHIGPUMask::All())
-		: RHICmdListStack(InGPUMask)
-	{}
-
-	FRHICommandList& GetCommandList()
-	{
-		return RHICmdListStack;
-	}
-
-	FRHICommandList& operator*()
-	{
-		return RHICmdListStack;
-	}
-
-	FRHICommandList* operator->()
-	{
-		return &RHICmdListStack;
-	}
-
-	~FRHIAsyncCommandList()
-	{
-		RHICmdListStack.FinishRecording();
-		if (RHICmdListStack.HasCommands())
-		{
-			ENQUEUE_RENDER_COMMAND(AsyncCommandListScope)(
-				[RHICmdList = new FRHICommandList(MoveTemp(RHICmdListStack))](FRHICommandListImmediate& RHICmdListImmediate)
-			{
-				RHICmdListImmediate.QueueAsyncCommandListSubmit(RHICmdList);
-			});
-		}
-	}
-
-private:
-	FRHICommandList RHICmdListStack;
-};
-
 class FRenderThreadScope
 {
 	typedef TFunction<void(FRHICommandListImmediate&)> RenderCommandFunction;
