[ActorInstanceHandle] added missing exports

diff --git a/Engine/Source/Runtime/Engine/Classes/Engine/ActorInstanceHandle.h b/Engine/Source/Runtime/Engine/Classes/Engine/ActorInstanceHandle.h
--- a/Engine/Source/Runtime/Engine/Classes/Engine/ActorInstanceHandle.h
+++ b/Engine/Source/Runtime/Engine/Classes/Engine/ActorInstanceHandle.h
@@ -45,13 +45,13 @@ struct FActorInstanceHandle
 	 * A path dedicated to creation of handles while converting actor to a dehydrated representation. This path ensures
 	 * an actor won't be spawned as a side effect of looking for the actor given Manager/Index represents
 	 */
-	static FActorInstanceHandle MakeDehydratedActorHandle(UObject& Manager, int32 InInstanceIndex);
+	static ENGINE_API FActorInstanceHandle MakeDehydratedActorHandle(UObject& Manager, int32 InInstanceIndex);
 
 	/** 
 	 * A path dedicated to creation of handles from any threads.
 	 * This path marks the handle as need resolving once it gets accessed from the game thread.
 	 */
-	static FActorInstanceHandle MakeActorHandleToResolve(const TWeakObjectPtr<UPrimitiveComponent>& WeakComponent, int32 CollisionInstanceIndex);
+	static ENGINE_API FActorInstanceHandle MakeActorHandleToResolve(const TWeakObjectPtr<UPrimitiveComponent>& WeakComponent, int32 CollisionInstanceIndex);
 
 	ENGINE_API bool IsValid() const;
 
