Fix race condition in texture release.
Need to call ReleaseResource() to enqueue resource delete and pointer clear to the render thread.
Fixes some licensee reported crashes when reading texture resource from render thread after game thread release.
Specifically when calling FMaterialRenderProxy::AllocateVTStack() with released resource pointer.

diff --git a/Engine/Source/Runtime/Engine/Private/Texture.cpp b/Engine/Source/Runtime/Engine/Private/Texture.cpp
--- a/Engine/Source/Runtime/Engine/Private/Texture.cpp
+++ b/Engine/Source/Runtime/Engine/Private/Texture.cpp
@@ -1249,7 +1249,7 @@ void UTexture::BeginFinalReleaseResource()
 	// Send the rendering thread a release message for the texture's resource.
 	if (GetResource())
 	{
-		BeginReleaseResource(GetResource());
+		ReleaseResource();
 	}
 	if (TextureReference.IsInitialized_GameThread())
 	{
@@ -1304,12 +1304,8 @@ void UTexture::FinishDestroy()
 	check(!bAsyncResourceReleaseHasBeenStarted || ReleaseFence.IsFenceComplete());
 	check(TextureReference.IsInitialized_GameThread() == false);
 
-	if(PrivateResource)
-	{
-		// Free the resource.
-		delete PrivateResource;
-		PrivateResource = NULL;
-	}
+	check(PrivateResource == nullptr);
+	check(PrivateResourceRenderThread == nullptr);
 
 	CleanupCachedRunningPlatformData();
 #if WITH_EDITOR
