Reworked the LOD Sync Component and Interface to handle LOD streaming of synced components.

A Skinned Mesh Component's desired LOD is no longer clamped to the highest loaded LOD, which allows code monitoring the desired LOD to determine when the component wants to render a LOD that is not currently loaded.

diff --git a/Engine/Plugins/Runtime/HairStrands/Source/HairStrandsCore/Private/GroomComponent.cpp b/Engine/Plugins/Runtime/HairStrands/Source/HairStrandsCore/Private/GroomComponent.cpp
--- a/Engine/Plugins/Runtime/HairStrands/Source/HairStrandsCore/Private/GroomComponent.cpp
+++ b/Engine/Plugins/Runtime/HairStrands/Source/HairStrandsCore/Private/GroomComponent.cpp
@@ -1773,7 +1773,18 @@ int32 UGroomComponent::GetDesiredSyncLOD() const
 	return LODPredictedIndex;
 }
 
-void  UGroomComponent::SetSyncLOD(int32 InLODIndex)
+int32 UGroomComponent::GetBestAvailableLOD() const
+{
+	// For now we assume all LODs are available. This could be made more accurate in future.
+	return 0;
+}
+
+void UGroomComponent::SetForceStreamedLOD(int32 LODIndex)
+{
+	// Force streaming is not supported yet
+}
+
+void UGroomComponent::SetForceRenderedLOD(int32 InLODIndex)
 {
 	SetForcedLOD(InLODIndex);
 }
@@ -1783,7 +1794,13 @@ int32 UGroomComponent::GetNumSyncLODs() const
 	return GetNumLODs();
 }
 
-int32 UGroomComponent::GetCurrentSyncLOD() const
+int32 UGroomComponent::GetForceStreamedLOD() const
+{
+	// Force streaming is not supported yet
+	return INDEX_NONE;
+}
+	
+int32 UGroomComponent::GetForceRenderedLOD() const
 {
 	return GetForcedLOD();
 }
diff --git a/Engine/Plugins/Runtime/HairStrands/Source/HairStrandsCore/Public/GroomComponent.h b/Engine/Plugins/Runtime/HairStrands/Source/HairStrandsCore/Public/GroomComponent.h
--- a/Engine/Plugins/Runtime/HairStrands/Source/HairStrandsCore/Public/GroomComponent.h
+++ b/Engine/Plugins/Runtime/HairStrands/Source/HairStrandsCore/Public/GroomComponent.h
@@ -240,9 +240,12 @@ public:
 
 	///~ Begin ILODSyncInterface Interface.
 	virtual int32 GetDesiredSyncLOD() const override;
-	virtual void SetSyncLOD(int32 LODIndex) override;
+	virtual int32 GetBestAvailableLOD() const override;
+	virtual void SetForceStreamedLOD(int32 LODIndex) override;
+	virtual void SetForceRenderedLOD(int32 LODIndex) override;
 	virtual int32 GetNumSyncLODs() const override;
-	virtual int32 GetCurrentSyncLOD() const override;
+	virtual int32 GetForceStreamedLOD() const override;
+	virtual int32 GetForceRenderedLOD() const override;
 	//~ End ILODSyncInterface
 
 	int32 GetNumLODs() const;
diff --git a/Engine/Source/Runtime/Engine/Classes/Components/LODSyncComponent.h b/Engine/Source/Runtime/Engine/Classes/Components/LODSyncComponent.h
--- a/Engine/Source/Runtime/Engine/Classes/Components/LODSyncComponent.h
+++ b/Engine/Source/Runtime/Engine/Classes/Components/LODSyncComponent.h
@@ -86,7 +86,7 @@ class ULODSyncComponent : public UActorComponent
 	TArray<FComponentSync> ComponentsToSync;
 
 	// by default, the mapping will be one to one
-// but if you want custom, add here. 
+	// but if you want custom, add here. 
 	UPROPERTY(EditAnywhere, BlueprintReadWrite, Category = LOD)
 	TMap<FName, FLODMappingData> CustomLODMapping;
 
@@ -118,6 +118,14 @@ private:
 
 public: 
 	ENGINE_API void RefreshSyncComponents();
+	
+	/**
+	 * Set the LOD of each synced component.
+	 *
+	 * This is called from TickComponent, so there's no need to call it manually. It's exposed here
+	 * for testing purposes.
+	 */
+	ENGINE_API void UpdateLOD();
 
 private:
 	int32 GetCustomMappingLOD(const FName& ComponentName, int32 CurrentWorkingLOD) const;
@@ -125,6 +133,5 @@ private:
 	void InitializeSyncComponents();
 	void UninitializeSyncComponents();
 	const FComponentSync* GetComponentSync(const FName& InName) const;
-	void UpdateLOD();
 };
 
diff --git a/Engine/Source/Runtime/Engine/Classes/Components/SkinnedMeshComponent.h b/Engine/Source/Runtime/Engine/Classes/Components/SkinnedMeshComponent.h
--- a/Engine/Source/Runtime/Engine/Classes/Components/SkinnedMeshComponent.h
+++ b/Engine/Source/Runtime/Engine/Classes/Components/SkinnedMeshComponent.h
@@ -526,6 +526,26 @@ public:
 	UPROPERTY(EditAnywhere, AdvancedDisplay, BlueprintReadOnly, Category=LOD)
 	int32 ForcedLodModel;
 
+private:
+	/**
+	 * Forces the specified LOD to stream in when LOD streaming is enabled.
+	 * 
+	 * If 0, streaming behavior is based on screen size, regardless of ForcedLodModel.
+	 * 
+	 * If > 0, (ForceStreamedLodModel - 1) will be streamed in.
+	 * 
+	 * This value will be copied from ForcedLodModel on initialization and will be set along with
+	 * ForcedLodModel by the SetForcedLOD function. It will only differ from ForcedLodModel if the
+	 * ILODSyncInterface functions are called, as these allow ForcedLodModel and
+	 * ForceStreamedLodModel to be set individually.
+	 * 
+	 * When ForcedLodModel and ForceStreamedLodModel are different, ForcedLodModel determines which
+	 * LOD to use for rendering and ForceStreamedLodModel determines which LOD to request for
+	 * streaming.
+	 */
+	int32 ForceStreamedLodModel;
+public:
+
 	/**
 	 * This is the min LOD that this component will use.  (e.g. if set to 2 then only 2+ LOD Models will be used.) This is useful to set on
 	 * meshes which are known to be a certain distance away and still want to have better LODs when zoomed in on them.
@@ -1855,9 +1875,12 @@ private:
 
 	// BEGIN ILODSyncComponent
 	ENGINE_API virtual int32 GetDesiredSyncLOD() const override;
-	ENGINE_API virtual void SetSyncLOD(int32 LODIndex) override;
+	ENGINE_API virtual int32 GetBestAvailableLOD() const override;
+	ENGINE_API virtual void SetForceStreamedLOD(int32 LODIndex) override;
+	ENGINE_API virtual void SetForceRenderedLOD(int32 LODIndex) override;
 	ENGINE_API virtual int32 GetNumSyncLODs() const override;
-	ENGINE_API virtual int32 GetCurrentSyncLOD() const override;
+	ENGINE_API virtual int32 GetForceStreamedLOD() const override;
+	ENGINE_API virtual int32 GetForceRenderedLOD() const override;
 	// END ILODSyncComponent
 
 	// Animation update rate control.
diff --git a/Engine/Source/Runtime/Engine/Private/Components/LODSyncComponent.cpp b/Engine/Source/Runtime/Engine/Private/Components/LODSyncComponent.cpp
--- a/Engine/Source/Runtime/Engine/Private/Components/LODSyncComponent.cpp
+++ b/Engine/Source/Runtime/Engine/Private/Components/LODSyncComponent.cpp
@@ -165,22 +165,16 @@ void ULODSyncComponent::RefreshSyncComponents()
 
 void ULODSyncComponent::UninitializeSyncComponents()
 {
-	auto ResetComponentList = [&](TArray<UPrimitiveComponent*>& ComponentList)
+	for (UPrimitiveComponent* Component : SubComponents)
 	{
-		for (UPrimitiveComponent* Component : ComponentList)
+		if (Component)
 		{
-			if (Component)
-			{
-				Component->PrimaryComponentTick.RemovePrerequisite(this, PrimaryComponentTick);
-			}
+			Component->PrimaryComponentTick.RemovePrerequisite(this, PrimaryComponentTick);
 		}
+	}
 
-		ComponentList.Reset();
-
-	};
-
-	ResetComponentList(MutableView(DriveComponents));
-	ResetComponentList(MutableView(SubComponents));
+	SubComponents.Reset();
+	DriveComponents.Reset();
 }
 
 void ULODSyncComponent::UpdateLOD()
@@ -204,7 +198,6 @@ void ULODSyncComponent::UpdateLOD()
 			bHaveValidSetting = true;
 			UE_LOG(LogLODSync, Verbose, TEXT("LOD Sync : Using Cvar r.ForceLod Value to set lodsync [%d]"), HighestPriLOD);
 		}
-
 		else if (ForcedLOD >= 0 && ForcedLOD < CurrentNumLODs)
 		{
 			HighestPriLOD = ForcedLOD;
@@ -213,57 +206,126 @@ void ULODSyncComponent::UpdateLOD()
 		}
 		else
 		{
-			bool bHaveVisibleComponent = false;
-			int32 HighestPriVisibleLOD = 0xff;
-			for (UPrimitiveComponent* Component : DriveComponents)
+			// Array is in priority order (last one is highest priority)
+			for (int32 DriveComponentIndex = DriveComponents.Num() - 1; DriveComponentIndex >= 0; DriveComponentIndex--)
 			{
-				if (Component)
+				const UPrimitiveComponent* Component = DriveComponents[DriveComponentIndex];
+				if (!Component)
 				{
-					ILODSyncInterface* LODInterface = CastChecked<ILODSyncInterface>(Component);
-					const int32 DesiredSyncLOD = LODInterface->GetDesiredSyncLOD();
+					continue;
+				}
 
-					if (DesiredSyncLOD >= 0)
-					{
-						const int32 DesiredLOD = GetSyncMappingLOD(Component->GetFName(), DesiredSyncLOD);
-						// Array is in priority order (last one is highest priority)
-						HighestPriLOD = DesiredLOD;
-						bHaveValidSetting = true;
+				const ILODSyncInterface* LODInterface = CastChecked<ILODSyncInterface>(Component);
+				const int32 DesiredSyncLOD = LODInterface->GetDesiredSyncLOD();
 
-						if (Component->WasRecentlyRendered())
-						{
-							bHaveVisibleComponent = true;
-							HighestPriVisibleLOD = DesiredLOD;
-							UE_LOG(LogLODSync, Verbose, TEXT("LOD Sync Drivers : %s (VISIBLE) - Source LOD [%d] RemappedLOD[%d]"), *GetNameSafe(Component), DesiredSyncLOD, DesiredLOD);
-						}
-						else
-						{
-							UE_LOG(LogLODSync, Verbose, TEXT("LOD Sync Drivers : %s - Source LOD [%d] RemappedLOD[%d]"), *GetNameSafe(Component), DesiredSyncLOD, DesiredLOD);
-						}
-					}
+				if (DesiredSyncLOD < 0 || LODInterface->GetBestAvailableLOD() < 0)
+				{
+					// Can't drive from a component that doesn't know what LOD it wants or what
+					// LODs it can render.
+					continue;
 				}
-			}
 
-			// If we found a visible component, use that, if not, use highest priority that wasn't visible
-			if (bHaveVisibleComponent)
-			{
-				HighestPriLOD = HighestPriVisibleLOD;
+				const int32 DesiredLOD = GetSyncMappingLOD(Component->GetFName(), DesiredSyncLOD);
+
+				if (Component->WasRecentlyRendered())
+				{
+					HighestPriLOD = DesiredLOD;
+					bHaveValidSetting = true;
+
+					UE_LOG(LogLODSync, Verbose, TEXT("LOD Sync Drivers : %s (VISIBLE) - Source LOD [%d] RemappedLOD[%d]"), *GetNameSafe(Component), DesiredSyncLOD, DesiredLOD);
+
+					// This is the best possible match, so no need to keep looking
+					break;
+				}
+
+				if (!bHaveValidSetting)
+				{
+					// This is the best match so far, but keep looking in case there's a visible component
+					HighestPriLOD = DesiredLOD;
+					bHaveValidSetting = true;
+				}
+
+				UE_LOG(LogLODSync, Verbose, TEXT("LOD Sync Drivers : %s - Source LOD [%d] RemappedLOD[%d]"), *GetNameSafe(Component), DesiredSyncLOD, DesiredLOD);
 			}
 		}
 
 		if (bHaveValidSetting)
 		{
-			// ensure current WorkingLOD is with in the range
+			// Ensure HighestPriLOD is with in the range
 			HighestPriLOD = FMath::Clamp(HighestPriLOD, MinLOD, CurrentNumLODs - 1);
 			UE_LOG(LogLODSync, Verbose, TEXT("LOD Sync : Current LOD (%d)"), HighestPriLOD);
-			for (UPrimitiveComponent* Component : SubComponents)
+
+			TArray<int32, TInlineAllocator<16>> SubComponentLODs;
+			// The code below ensures that all elements corresponding to non-null components are
+			// written before any elements are read, so it's safe to leave them uninitialized here.
+			SubComponentLODs.SetNumUninitialized(SubComponents.Num());
+
+			bool bFirstIteration = true;
+			while (HighestPriLOD < CurrentNumLODs)
 			{
-				if (Component)
+				// HighestPriLOD can only be used if the corresponding LOD on each component is 
+				// available.
+				//
+				// This loop checks to make sure all the necessary LODs are available, and requests
+				// that the best one be streamed in.
+				bool bAllLODsAvailable = true;
+				for (int32 SubComponentIndex = 0; SubComponentIndex < SubComponents.Num(); SubComponentIndex++)
+				{
+					UPrimitiveComponent* Component = SubComponents[SubComponentIndex];
+					if (Component)
+					{
+						ILODSyncInterface* LODInterface = CastChecked<ILODSyncInterface>(Component);
+						const int32 NewLOD = GetCustomMappingLOD(Component->GetFName(), HighestPriLOD);
+
+						// The first iteration of the enclosing loop tests the most desired LOD, so
+						// this is the one that should be requested for streaming, regardless of
+						// which LOD is ultimately used for rendering.
+						if (bFirstIteration)
+						{
+							LODInterface->SetForceStreamedLOD(NewLOD);
+						}
+
+						// If this LOD isn't available, break out and try the next one.
+						// 
+						// Note that if GetBestAvailableLOD returns INDEX_NONE, this condition will
+						// always be false. We ignore the invalid value rather than let it prevent
+						// us from setting LODs on the other components.
+						if (NewLOD < LODInterface->GetBestAvailableLOD())
+						{
+							bAllLODsAvailable = false;
+
+							// Don't break out if this is the first iteration, because we need to
+							// request streaming on all the most desired LODs during this iteration.
+							if (!bFirstIteration)
+							{
+								break;
+							}
+						}
+
+						SubComponentLODs[SubComponentIndex] = NewLOD;
+					}
+				}
+
+				bFirstIteration = false;
+
+				if (!bAllLODsAvailable)
 				{
-					ILODSyncInterface* LODInterface = Cast<ILODSyncInterface>(Component);
-					const int32 NewLOD = GetCustomMappingLOD(Component->GetFName(), HighestPriLOD);
-					UE_LOG(LogLODSync, Verbose, TEXT("LOD Sync Setter : %s - New LOD [%d]"), *GetNameSafe(Component), NewLOD);
-					LODInterface->SetSyncLOD(NewLOD);
+					HighestPriLOD++;
+					continue;
 				}
+
+				// Found a good set of LODs, now set them.
+				for (int32 SubComponentIndex = 0; SubComponentIndex < SubComponents.Num(); SubComponentIndex++)
+				{
+					UPrimitiveComponent* Component = SubComponents[SubComponentIndex];
+					if (Component)
+					{
+						ILODSyncInterface* LODInterface = CastChecked<ILODSyncInterface>(Component);
+						LODInterface->SetForceRenderedLOD(SubComponentLODs[SubComponentIndex]);
+					}
+				}
+
+				break;
 			}
 		}
 	}
@@ -307,16 +369,22 @@ FString ULODSyncComponent::GetLODSyncDebugText() const
 		if (Component)
 		{
 			ILODSyncInterface* LODInterface = CastChecked<ILODSyncInterface>(Component);
-			const int32 CurrentSyncLOD = LODInterface->GetCurrentSyncLOD();
+			const int32 StreamedLOD = LODInterface->GetForceStreamedLOD();
+			const int32 RenderedLOD = LODInterface->GetForceRenderedLOD();
 			const int32 DesiredSyncLOD = LODInterface->GetDesiredSyncLOD();
 
-			if (DesiredSyncLOD >= 0)
+			if (StreamedLOD >= 0 && StreamedLOD < LODInterface->GetBestAvailableLOD())
+			{
+				// We're waiting for a better LOD to stream in
+				OutString += FString::Printf(TEXT("%s : %d (%d, %d)\n"), *(Component->GetFName().ToString()), RenderedLOD, DesiredSyncLOD, StreamedLOD);
+			}
+			else if (DesiredSyncLOD >= 0)
 			{
-				OutString += FString::Printf(TEXT("%s : %d (%d)\n"), *(Component->GetFName().ToString()), CurrentSyncLOD, DesiredSyncLOD);
+				OutString += FString::Printf(TEXT("%s : %d (%d)\n"), *(Component->GetFName().ToString()), RenderedLOD, DesiredSyncLOD);
 			}
 			else
 			{
-				OutString += FString::Printf(TEXT("%s : %d\n"), *(Component->GetFName().ToString()), CurrentSyncLOD);
+				OutString += FString::Printf(TEXT("%s : %d\n"), *(Component->GetFName().ToString()), RenderedLOD);
 			}
 		}
 	}
diff --git a/Engine/Source/Runtime/Engine/Private/Components/SkinnedMeshComponent.cpp b/Engine/Source/Runtime/Engine/Private/Components/SkinnedMeshComponent.cpp
--- a/Engine/Source/Runtime/Engine/Private/Components/SkinnedMeshComponent.cpp
+++ b/Engine/Source/Runtime/Engine/Private/Components/SkinnedMeshComponent.cpp
@@ -674,6 +674,8 @@ void USkinnedMeshComponent::OnRegister()
 		CachedSceneFeatureLevel = ERHIFeatureLevel::Num;
 	}
 
+	ForceStreamedLodModel = GetForcedLOD();
+
 #if WITH_EDITOR
 	// When we are in editor, ensure that the initial setup is done at LOD 0 so that non-ticking components dont
 	// get unexpected behavior when transitioning LODs post-construction script
@@ -1560,8 +1562,13 @@ void USkinnedMeshComponent::GetStreamingRenderAssetInfo(FStreamingTextureLevelCo
 
 	if (GetSkinnedAsset() && GetSkinnedAsset()->IsStreamable())
 	{
-		const int32 LocalForcedLodModel = GetForcedLOD();
-		const float TexelFactor = LocalForcedLodModel > 0 ? -(GetSkinnedAsset()->GetLODNum() - LocalForcedLodModel + 1) : Bounds.SphereRadius * 2.f;
+		const int32 LocalForcedLodModel = GetForceStreamedLOD();
+
+		// A positive texel factor tells the streaming system the size of the component.
+		// 
+		// A negative texel factor selects a specific LOD index, starting at -1 for the worst LOD
+		const float TexelFactor = LocalForcedLodModel >= 0 ? LocalForcedLodModel - GetSkinnedAsset()->GetLODNum() : Bounds.SphereRadius * 2.f;
+
 		new (OutStreamingRenderAssets) FStreamingRenderAssetPrimitiveInfo(GetSkinnedAsset(), Bounds, TexelFactor, PackedRelativeBox_Identity, false, false);
 	}
 }
@@ -3657,14 +3664,8 @@ void USkinnedMeshComponent::UnHideBoneByName( FName BoneName )
 
 void USkinnedMeshComponent::SetForcedLOD(int32 InNewForcedLOD)
 {
-PRAGMA_DISABLE_DEPRECATION_WARNINGS
-	const int32 OldValue = ForcedLodModel;
-	ForcedLodModel = FMath::Clamp(InNewForcedLOD, 0, GetNumLODs());
-	if (OldValue != ForcedLodModel)
-	{
-		IStreamingManager::Get().NotifyPrimitiveUpdated(this);
-	}
-PRAGMA_ENABLE_DEPRECATION_WARNINGS
+	SetForceStreamedLOD(InNewForcedLOD - 1);
+	SetForceRenderedLOD(InNewForcedLOD - 1);
 }
 
 int32 USkinnedMeshComponent::GetForcedLOD() const
@@ -3759,14 +3760,47 @@ int32 USkinnedMeshComponent::GetDesiredSyncLOD() const
 	return INDEX_NONE;
 }
 
-void USkinnedMeshComponent::SetSyncLOD(int32 LODIndex)
+int32 USkinnedMeshComponent::GetBestAvailableLOD() const
 {
-	SetForcedLOD(LODIndex + 1);
+	if (GetSkinnedAsset() && MeshObject)
+	{
+		const FSkeletalMeshRenderData& RenderData = MeshObject->GetSkeletalMeshRenderData();
+		
+		// If PendingFirstLODIdx is higher (worse) than CurrentFirstLODIdx, it means that
+		// CurrentFirstLODIdx is about to be streamed out, so we can't rely on it to be available.
+		const int32 BestAvailableLOD = FMath::Max<int32>(RenderData.PendingFirstLODIdx, RenderData.CurrentFirstLODIdx);
+		check(BestAvailableLOD >= 0);
+
+		return BestAvailableLOD;
+	}
+
+	return INDEX_NONE;
 }
 
-int32 USkinnedMeshComponent::GetCurrentSyncLOD() const
+void USkinnedMeshComponent::SetForceStreamedLOD(int32 LODIndex)
 {
-	return GetForcedLOD() - 1; // Weird API for forced LOD where 0 means auto, 1 means force to 0 etc
+	const int32 OldValue = ForceStreamedLodModel;
+	// Need to add 1 here, because LODIndex is 0-based but ForceStreamedLodModel is 1-based
+	ForceStreamedLodModel = FMath::Clamp(LODIndex + 1, 0, GetNumLODs());
+	if (OldValue != ForceStreamedLodModel)
+	{
+		IStreamingManager::Get().NotifyPrimitiveUpdated(this);
+	}
+}
+
+void USkinnedMeshComponent::SetForceRenderedLOD(int32 LODIndex)
+{
+PRAGMA_DISABLE_DEPRECATION_WARNINGS
+	// Need to add 1 here, because LODIndex is 0-based but ForcedLodModel is 1-based
+	ForcedLodModel = FMath::Clamp(LODIndex + 1, 0, GetNumLODs());
+PRAGMA_ENABLE_DEPRECATION_WARNINGS
+}
+
+int32 USkinnedMeshComponent::GetForceRenderedLOD() const
+{
+	// Subtract 1 because ILODSyncInterface::GetForceRenderedLOD requires the return value to be 
+	// 0-based, but the return value of GetForcedLOD() is 1-based for legacy reasons.
+	return GetForcedLOD() - 1;
 }
 
 int32 USkinnedMeshComponent::GetNumSyncLODs() const
@@ -3774,6 +3808,13 @@ int32 USkinnedMeshComponent::GetNumSyncLODs() const
 	return GetNumLODs();
 }
 
+int32 USkinnedMeshComponent::GetForceStreamedLOD() const
+{
+	// Subtract 1 because ILODSyncInterface::GetForceStreamedLOD requires the return value to be 
+	// 0-based, but ForceStreamedLodModel is 1-based for legacy reasons.
+	return ForceStreamedLodModel - 1;
+}
+
 bool USkinnedMeshComponent::UpdateLODStatus()
 {
 	return UpdateLODStatus_Internal(INDEX_NONE);
@@ -3854,7 +3895,12 @@ bool USkinnedMeshComponent::UpdateLODStatus_Internal(int32 InLeaderPoseComponent
 
 		if (GetSkinnedAsset()->IsStreamable() && MeshObject)
 		{
-			NewPredictedLODLevel = FMath::Max<int32>(NewPredictedLODLevel, MeshObject->GetSkeletalMeshRenderData().PendingFirstLODIdx);
+			// LOD streaming is enabled for this mesh, so clamp the predicted LOD to the best LOD
+			// that is loaded and is not waiting to be streamed out.
+
+			const FSkeletalMeshRenderData& RenderData = MeshObject->GetSkeletalMeshRenderData();
+			NewPredictedLODLevel = FMath::Max<int32>(NewPredictedLODLevel, RenderData.PendingFirstLODIdx);
+			NewPredictedLODLevel = FMath::Max<int32>(NewPredictedLODLevel, RenderData.CurrentFirstLODIdx);
 		}
 
 #if !(UE_BUILD_SHIPPING || UE_BUILD_TEST)
diff --git a/Engine/Source/Runtime/Engine/Private/SkeletalMesh.cpp b/Engine/Source/Runtime/Engine/Private/SkeletalMesh.cpp
--- a/Engine/Source/Runtime/Engine/Private/SkeletalMesh.cpp
+++ b/Engine/Source/Runtime/Engine/Private/SkeletalMesh.cpp
@@ -6475,7 +6475,7 @@ void FSkeletalMeshSceneProxy::GetMeshElementsConditionallySelectable(const TArra
 
 	const FEngineShowFlags& EngineShowFlags = ViewFamily.EngineShowFlags;
 
-	int32 FirstLODIdx = SkeletalMeshRenderData->GetFirstValidLODIdx(FMath::Max(SkeletalMeshRenderData->PendingFirstLODIdx, SkeletalMeshRenderData->CurrentFirstLODIdx));
+	const int32 FirstLODIdx = SkeletalMeshRenderData->GetFirstValidLODIdx(FMath::Max(SkeletalMeshRenderData->PendingFirstLODIdx, SkeletalMeshRenderData->CurrentFirstLODIdx));
 	if (FirstLODIdx == INDEX_NONE)
 	{
 #if DO_CHECK
@@ -6487,7 +6487,7 @@ void FSkeletalMeshSceneProxy::GetMeshElementsConditionallySelectable(const TArra
 		if (UNLIKELY(!Views.IsEmpty() && IStereoRendering::IsStereoEyeView(*Views[0])))
 		{
 			const FSceneView& View = GetLODView(*Views[0]);
-			MeshObject->UpdateMinDesiredLODLevel(&View, GetBounds(), ViewFamily.FrameNumber, FirstLODIdx);
+			MeshObject->UpdateMinDesiredLODLevel(&View, GetBounds(), ViewFamily.FrameNumber);
 		}
 		else
 		{
@@ -6496,7 +6496,7 @@ void FSkeletalMeshSceneProxy::GetMeshElementsConditionallySelectable(const TArra
 				if (VisibilityMap & (1 << ViewIndex))
 				{
 					const FSceneView* View = Views[ViewIndex];
-					MeshObject->UpdateMinDesiredLODLevel(View, GetBounds(), ViewFamily.FrameNumber, FirstLODIdx);
+					MeshObject->UpdateMinDesiredLODLevel(View, GetBounds(), ViewFamily.FrameNumber);
 				}
 			}
 		}
diff --git a/Engine/Source/Runtime/Engine/Private/SkeletalRender.cpp b/Engine/Source/Runtime/Engine/Private/SkeletalRender.cpp
--- a/Engine/Source/Runtime/Engine/Private/SkeletalRender.cpp
+++ b/Engine/Source/Runtime/Engine/Private/SkeletalRender.cpp
@@ -59,7 +59,7 @@ FSkeletalMeshObject
 -----------------------------------------------------------------------------*/
 
 FSkeletalMeshObject::FSkeletalMeshObject(USkinnedMeshComponent* InMeshComponent, FSkeletalMeshRenderData* InSkelMeshRenderData, ERHIFeatureLevel::Type InFeatureLevel)
-:	MinDesiredLODLevel(FMath::Max<int32>(InMeshComponent->GetPredictedLODLevel(), InSkelMeshRenderData->CurrentFirstLODIdx))
+:	MinDesiredLODLevel(InMeshComponent->GetPredictedLODLevel())
 ,	MaxDistanceFactor(0.f)
 ,	WorkingMinDesiredLODLevel(MinDesiredLODLevel)
 ,	WorkingMaxDistanceFactor(0.f)
@@ -116,7 +116,7 @@ FSkeletalMeshObject::~FSkeletalMeshObject()
 {
 }
 
-void FSkeletalMeshObject::UpdateMinDesiredLODLevel(const FSceneView* View, const FBoxSphereBounds& Bounds, int32 FrameNumber, uint8 CurFirstLODIdx)
+void FSkeletalMeshObject::UpdateMinDesiredLODLevel(const FSceneView* View, const FBoxSphereBounds& Bounds, int32 FrameNumber)
 {
 	// Thumbnail rendering doesn't contribute to MinDesiredLODLevel calculation
 	if (View->Family && (View->Family->bThumbnailRendering || !View->Family->GetIsInFocus()))
@@ -161,7 +161,6 @@ void FSkeletalMeshObject::UpdateMinDesiredLODLevel(const FSceneView* View, const
 		}
 	}
 
-	NewLODLevel = FMath::Max<int32>(NewLODLevel, CurFirstLODIdx);
 	if (!LastFrameNumber)
 	{
 		// We don't have last frame value on the first call to FSkeletalMeshObject::UpdateMinDesiredLODLevel so
diff --git a/Engine/Source/Runtime/Engine/Private/SkeletalRenderStatic.cpp b/Engine/Source/Runtime/Engine/Private/SkeletalRenderStatic.cpp
--- a/Engine/Source/Runtime/Engine/Private/SkeletalRenderStatic.cpp
+++ b/Engine/Source/Runtime/Engine/Private/SkeletalRenderStatic.cpp
@@ -178,6 +178,12 @@ const TArray<FMatrix44f>& FSkeletalMeshObjectStatic::GetReferenceToLocalMatrices
 	return ReferenceToLocalMatrices;
 }
 
+int32 FSkeletalMeshObjectStatic::GetLOD() const
+{
+	// WorkingMinDesiredLODLevel can be a LOD that's not loaded, so need to clamp it to the first loaded LOD
+	return FMath::Max<int32>(WorkingMinDesiredLODLevel, SkeletalMeshRenderData->CurrentFirstLODIdx);
+}
+
 void FSkeletalMeshObjectStatic::FSkeletalMeshObjectLOD::InitResources(FSkelMeshComponentLODInfo* CompLODInfo)
 {
 	check(SkelMeshRenderData);
diff --git a/Engine/Source/Runtime/Engine/Private/SkeletalRenderStatic.h b/Engine/Source/Runtime/Engine/Private/SkeletalRenderStatic.h
--- a/Engine/Source/Runtime/Engine/Private/SkeletalRenderStatic.h
+++ b/Engine/Source/Runtime/Engine/Private/SkeletalRenderStatic.h
@@ -34,11 +34,7 @@ public:
 	ENGINE_API virtual const FVertexFactory* GetSkinVertexFactory(const FSceneView* View, int32 LODIndex, int32 ChunkIdx, ESkinVertexFactoryMode VFMode = ESkinVertexFactoryMode::Default) const override;
 	ENGINE_API virtual TArray<FTransform>* GetComponentSpaceTransforms() const override;
 	ENGINE_API virtual const TArray<FMatrix44f>& GetReferenceToLocalMatrices() const override;
-
-	virtual int32 GetLOD() const override
-	{
-		return WorkingMinDesiredLODLevel;
-	}
+	virtual int32 GetLOD() const override;
 	//virtual const FTwoVectors& GetCustomLeftRightVectors(int32 SectionIndex) const override;
 	virtual void DrawVertexElements(FPrimitiveDrawInterface* PDI, const FMatrix& ToWorldSpace, bool bDrawNormals, bool bDrawTangents, bool bDrawBinormals) const override {};
 	virtual bool HaveValidDynamicData() const override
diff --git a/Engine/Source/Runtime/Engine/Public/LODSyncInterface.h b/Engine/Source/Runtime/Engine/Public/LODSyncInterface.h
--- a/Engine/Source/Runtime/Engine/Public/LODSyncInterface.h
+++ b/Engine/Source/Runtime/Engine/Public/LODSyncInterface.h
@@ -24,15 +24,94 @@ class ILODSyncInterface
 {
 	GENERATED_IINTERFACE_BODY()
 
-	/** returns what is your desired LOD by 0 based index */
+	/**
+	 * Returns the renderable's desired LOD as a 0-based index.
+	 * 
+	 * If LOD streaming is enabled, the desired LOD returned from this function is not guaranteed
+	 * to be available (i.e. loaded). Call GetBestAvailableLOD to find out which LODs are available.
+	 * 
+	 * Returns INDEX_NONE if the desired LOD is not known, e.g. if there's no mesh set on the
+	 * component. In this case, the LOD Sync Component won't take this component's desired LOD into
+	 * account, but will still set a forced LOD on it in case it does get rendered.
+	 */
 	virtual int32 GetDesiredSyncLOD() const = 0;
 
-	/** sets what is synced LOD by 0 based index */
-	virtual void SetSyncLOD(int32 LODIndex) = 0;
+	/**
+	 * Returns the best (i.e. lowest numbered) LOD that is streamed in and can be used for
+	 * rendering immediately.
+	 * 
+	 * Callers can assume that any LOD worse (higher numbered) than this is also available.
+	 *
+	 * Returns INDEX_NONE if no LODs are available, e.g. if there's no mesh set on the component.
+	 * In this case, the LOD Sync Component won't take this component's desired LOD into account,
+	 * but will still set a forced LOD on it in case it does get rendered.
+	 */
+	virtual int32 GetBestAvailableLOD() const = 0;
 
-	/** returns number of LODs */
+	/**
+	 * Returns the total number of LODs, including LODs that are streamed out and unavailable.
+	 *
+	 * Returns 0 if there are no LODs, e.g. if there's no mesh set on the component.
+	 */
 	virtual int32 GetNumSyncLODs() const = 0;
 
-	/** Returns what the current sync LOD has bee set to */
-	virtual int32 GetCurrentSyncLOD() const = 0;
+	/**
+	 * Sets the LOD to force to stream in, as a 0-based index.
+	 * 
+	 * If LODIndex is non-negative and LOD streaming is enabled, this call requests that the
+	 * streaming system stream in LODs up to and including the specified LOD index, ignoring
+	 * any other heuristics such as size on screen.
+	 * 
+	 * Note that there is no guarantee of how long it will take the requested LOD to stream in,
+	 * and if memory is low it may never stream in.
+	 *
+	 * If LODIndex is negative, the streaming behavior will be returned to normal.
+	 *
+	 * This call has no effect while LOD streaming is disabled.
+	 */
+	virtual void SetForceStreamedLOD(int32 LODIndex) = 0;
+
+	/**
+	 * Sets the LOD to use for rendering, as a 0-based index.
+	 *
+	 * If LODIndex is negative, the rendering LOD selection behavior will be returned to normal.
+	 *
+	 * This call has no effect on streaming behavior. If a LOD is requested that is not streamed
+	 * in, the best available LOD will be used instead. If the requested LOD is streamed in later,
+	 * it will be used without the caller needing to call this function again.
+	 */
+	virtual void SetForceRenderedLOD(int32 LODIndex) = 0;
+
+	UE_DEPRECATED(5.4, "SetSyncLOD is deprecated and will be removed. For the same effect, call both SetForceStreamedLOD and SetForceRenderedLOD.")
+	virtual void SetSyncLOD(int32 LODIndex)
+	{
+		SetForceStreamedLOD(LODIndex);
+		SetForceRenderedLOD(LODIndex);
+	}
+
+	/**
+	 * Returns the LOD that will be forced to stream in, or INDEX_NONE if there is no forced LOD
+	 * for streaming.
+	 * 
+	 * Note that this is not guaranteed to return the value that was set by the last call to
+	 * SetForceStreamedLOD, as the implementation is permitted to change the force streamed LOD
+	 * outside of calls to that function.
+	 */
+	virtual int32 GetForceStreamedLOD() const = 0;
+
+	/**
+	 * Returns the LOD that will be forced to use for rendering, or INDEX_NONE if there is no
+	 * forced LOD for rendering.
+	 * 
+	 * Note that this is not guaranteed to return the value that was set by the last call to
+	 * SetForceRenderedLOD, as the implementation is permitted to change the force rendered LOD
+	 * outside of calls to that function.
+	 */
+	virtual int32 GetForceRenderedLOD() const = 0;
+
+	UE_DEPRECATED(5.4, "GetCurrentSyncLOD has been renamed to GetForceRenderedLOD. GetCurrentSyncLOD will be removed.")
+	virtual int32 GetCurrentSyncLOD() const
+	{
+		return GetForceRenderedLOD();
+	}
 };
diff --git a/Engine/Source/Runtime/Engine/Public/SkeletalRenderPublic.h b/Engine/Source/Runtime/Engine/Public/SkeletalRenderPublic.h
--- a/Engine/Source/Runtime/Engine/Public/SkeletalRenderPublic.h
+++ b/Engine/Source/Runtime/Engine/Public/SkeletalRenderPublic.h
@@ -173,7 +173,7 @@ public:
 	 *	This is called from the rendering thread (PreRender) so be very careful what you read/write to.
 	 * @param FrameNumber from ViewFamily.FrameNumber
 	 */
-	void UpdateMinDesiredLODLevel(const FSceneView* View, const FBoxSphereBounds& Bounds, int32 FrameNumber, uint8 CurFirstLODIdx);
+	void UpdateMinDesiredLODLevel(const FSceneView* View, const FBoxSphereBounds& Bounds, int32 FrameNumber);
 
 	/**
 	 *	Return true if this does have valid dynamic data to render
@@ -263,7 +263,10 @@ public:
 	TArray<FCapsuleShape3f> ShadowCapsuleShapes;
 
 	/** 
-	 *	Lowest (best) LOD that was desired for rendering this SkeletalMesh last frame. 
+	 *	Lowest (best) LOD that was desired for rendering this SkeletalMesh last frame.
+	 *
+	 *	Note that if LOD streaming is enabled, the desired LOD is not guaranteed to be currently loaded.
+	 * 
 	 *	This should only ever be WRITTEN by the RENDER thread (in FSkeletalMeshProxy::PreRenderView) and READ by the GAME thread (in USkeletalMeshComponent::UpdateSkelPose).
 	 */
 	int32 MinDesiredLODLevel;
