Fix WP intersection queries being nondeterministic. UWorldPartitionRuntimeCellData::SortCompare was never returning 0 even if comparisons were equal, which meant we wouldn't fall back to a lexical sort on the level name.

diff --git a/Engine/Source/Runtime/Engine/Private/WorldPartition/WorldPartitionRuntimeCellData.cpp b/Engine/Source/Runtime/Engine/Private/WorldPartition/WorldPartitionRuntimeCellData.cpp
--- a/Engine/Source/Runtime/Engine/Private/WorldPartition/WorldPartitionRuntimeCellData.cpp
+++ b/Engine/Source/Runtime/Engine/Private/WorldPartition/WorldPartitionRuntimeCellData.cpp
@@ -90,11 +90,14 @@ int32 UWorldPartitionRuntimeCellData::SortCompare(const UWorldPartitionRuntimeCe
 		if (!Result)
 		{
 			// Cell priority (lower value is higher prio)
-			Result = Priority - InOther->Priority;			
-
+			Result = Priority - InOther->Priority;		
 			if (!Result)
 			{
-				Result = (CachedMinSpatialSortingPriority < InOther->CachedMinSpatialSortingPriority) ? -1 : 1;
+				double Diff = CachedMinSpatialSortingPriority - InOther->CachedMinSpatialSortingPriority;
+				if (!FMath::IsNearlyZero(Diff))
+				{
+					Result = (CachedMinSpatialSortingPriority < InOther->CachedMinSpatialSortingPriority) ? -1 : 1;
+				}
 			}
 		}
 	}
@@ -115,4 +118,4 @@ FBox UWorldPartitionRuntimeCellData::GetCellBounds() const
 FString UWorldPartitionRuntimeCellData::GetDebugName() const
 {
 	return DebugName.GetString();
-}
\ No newline at end of file
+}
