- Improve HandleGetIniCommand to handle section names surrounded with quotes, useful for section names with spaces in them.

diff --git a/Engine/Source/Runtime/Engine/Private/UnrealEngine.cpp b/Engine/Source/Runtime/Engine/Private/UnrealEngine.cpp
--- a/Engine/Source/Runtime/Engine/Private/UnrealEngine.cpp
+++ b/Engine/Source/Runtime/Engine/Private/UnrealEngine.cpp
@@ -9847,6 +9847,7 @@ bool UEngine::HandleGetIniCommand(const TCHAR* Cmd, FOutputDevice& Ar)
 		int32 IniDelim = IniPlusSection.Find(TEXT(":"));
 		FString IniName;
 		FString SectionName = (IniDelim != INDEX_NONE ? IniPlusSection.Mid(IniDelim+1) : IniPlusSection);
+		SectionName.TrimQuotesInline();
 
 		if (IniDelim != INDEX_NONE)
 		{
