PIE fixup - Improved check() to print package & object names in case of failure

diff --git a/Engine/Source/Runtime/Engine/Private/PieFixupSerializer.cpp b/Engine/Source/Runtime/Engine/Private/PieFixupSerializer.cpp
--- a/Engine/Source/Runtime/Engine/Private/PieFixupSerializer.cpp
+++ b/Engine/Source/Runtime/Engine/Private/PieFixupSerializer.cpp
@@ -46,8 +46,8 @@ FArchive& FPIEFixupSerializer::operator<<(UObject*& Object)
 #if WITH_EDITOR
 		if (UPackage* ExternalPackage = Object->GetExternalPackage())
 		{
-			check(Object->IsPackageExternal());
-			check(ExternalPackage->HasAnyPackageFlags(PKG_PlayInEditor));
+			checkf(Object->IsPackageExternal(), TEXT("Expected an external package. Package: '%s'. Object: '%s'."), *ExternalPackage->GetFullName(), *Object->GetFullName());
+			checkf(ExternalPackage->HasAnyPackageFlags(PKG_PlayInEditor), TEXT("Package missing the PKG_PlayInEditor flag. Package: '%s'. Object: '%s'."), *ExternalPackage->GetFullName(), *Object->GetFullName());
 			ExternalPackage->SetPIEInstanceID(PIEInstanceID);
 		}
 #endif
