Orthographic Cameras: Cameras do not obey bAutoCalculateOrthoPlanes=false

The default value of this setting is now true, so code that assumed it defaulted to false could have the wrong behavior. In addition, remove some duplicated logic from the player camera manager.

Fixes PT_OrthoCamera test

diff --git a/Engine/Source/Runtime/Engine/Private/Camera/CameraComponent.cpp b/Engine/Source/Runtime/Engine/Private/Camera/CameraComponent.cpp
--- a/Engine/Source/Runtime/Engine/Private/Camera/CameraComponent.cpp
+++ b/Engine/Source/Runtime/Engine/Private/Camera/CameraComponent.cpp
@@ -419,10 +419,9 @@ void UCameraComponent::GetCameraView(float DeltaTime, FMinimalViewInfo& DesiredV
 	DesiredView.OrthoWidth = OrthoWidth;
 	DesiredView.OrthoNearClipPlane = OrthoNearClipPlane;
 	DesiredView.OrthoFarClipPlane = OrthoFarClipPlane;
-
+	DesiredView.bAutoCalculateOrthoPlanes = bAutoCalculateOrthoPlanes;
 	if (bAutoCalculateOrthoPlanes)
 	{
-		DesiredView.bAutoCalculateOrthoPlanes = true;
 		if (const AActor* ViewTarget = GetOwner())
 		{
 			DesiredView.SetOrthoCameraArmLengthFromOwnerLocation(ViewTarget->GetActorLocation());
diff --git a/Engine/Source/Runtime/Engine/Private/Components/SceneCaptureComponent.cpp b/Engine/Source/Runtime/Engine/Private/Components/SceneCaptureComponent.cpp
--- a/Engine/Source/Runtime/Engine/Private/Components/SceneCaptureComponent.cpp
+++ b/Engine/Source/Runtime/Engine/Private/Components/SceneCaptureComponent.cpp
@@ -701,10 +701,9 @@ void USceneCaptureComponent2D::GetCameraView(float DeltaTime, FMinimalViewInfo&
 	OutMinimalViewInfo.bConstrainAspectRatio = false;
 	OutMinimalViewInfo.ProjectionMode = ProjectionType;
 	OutMinimalViewInfo.OrthoWidth = OrthoWidth;
-		
+	OutMinimalViewInfo.bAutoCalculateOrthoPlanes = bAutoCalculateOrthoPlanes;		
 	if (bAutoCalculateOrthoPlanes)
 	{
-		OutMinimalViewInfo.bAutoCalculateOrthoPlanes = true;
 		if(const AActor* ViewTarget = GetOwner())
 		{
 			OutMinimalViewInfo.SetOrthoCameraArmLengthFromOwnerLocation(ViewTarget->GetActorLocation());
diff --git a/Engine/Source/Runtime/Engine/Private/LocalPlayer.cpp b/Engine/Source/Runtime/Engine/Private/LocalPlayer.cpp
--- a/Engine/Source/Runtime/Engine/Private/LocalPlayer.cpp
+++ b/Engine/Source/Runtime/Engine/Private/LocalPlayer.cpp
@@ -687,14 +687,6 @@ void ULocalPlayer::GetViewPoint(FMinimalViewInfo& OutViewInfo) const
 			OutViewInfo = PlayerController->PlayerCameraManager->GetCameraCacheView();
 			OutViewInfo.FOV = PlayerController->PlayerCameraManager->GetFOVAngle();
 			PlayerController->GetPlayerViewPoint(/*out*/ OutViewInfo.Location, /*out*/ OutViewInfo.Rotation);
-			if (PlayerController->PlayerCameraManager->bAutoCalculateOrthoPlanes)
-			{
-				OutViewInfo.bAutoCalculateOrthoPlanes = true;
-				if (const APawn* Owner = PlayerController->GetPawn())
-				{
-					OutViewInfo.SetOrthoCameraArmLengthFromOwnerLocation(Owner->GetActorLocation());
-				}
-			}
 		}
 		else
 		{
diff --git a/Engine/Source/Runtime/Engine/Private/PlayerCameraManager.cpp b/Engine/Source/Runtime/Engine/Private/PlayerCameraManager.cpp
--- a/Engine/Source/Runtime/Engine/Private/PlayerCameraManager.cpp
+++ b/Engine/Source/Runtime/Engine/Private/PlayerCameraManager.cpp
@@ -365,10 +365,9 @@ void APlayerCameraManager::UpdateViewTarget(FTViewTarget& OutVT, float DeltaTime
 	OutVT.POV.bConstrainAspectRatio = bDefaultConstrainAspectRatio;
 	OutVT.POV.ProjectionMode = bIsOrthographic ? ECameraProjectionMode::Orthographic : ECameraProjectionMode::Perspective;
 	OutVT.POV.PostProcessBlendWeight = 1.0f;
-
+	OutVT.POV.bAutoCalculateOrthoPlanes = bAutoCalculateOrthoPlanes;
 	if (bAutoCalculateOrthoPlanes)
 	{
-		OutVT.POV.bAutoCalculateOrthoPlanes = true;
 		if (OutVT.Target)
 		{
 			OutVT.POV.SetOrthoCameraArmLengthFromOwnerLocation(OutVT.Target->GetActorLocation());
