Fixed case where UWorldPartition::UnregisterActorDescContainerInstance could unhash actordescs that were never initially hashed (which was crashing).

diff --git a/Engine/Source/Runtime/Engine/Private/WorldPartition/WorldPartition.cpp b/Engine/Source/Runtime/Engine/Private/WorldPartition/WorldPartition.cpp
--- a/Engine/Source/Runtime/Engine/Private/WorldPartition/WorldPartition.cpp
+++ b/Engine/Source/Runtime/Engine/Private/WorldPartition/WorldPartition.cpp
@@ -1323,7 +1323,7 @@ bool UWorldPartition::ShouldHashUnhashActorDescInstances() const
 {
 	const bool bIsEditor = !GetWorld()->IsGameWorld();
 	const bool bIsCooking = IsRunningCookCommandlet();
-	const bool bHashActorDescs = bIsEditor && !bIsCooking;
+	const bool bHashActorDescs = EditorHash && bIsEditor && !bIsCooking;
 	return bHashActorDescs;
 }
 
@@ -2028,8 +2028,8 @@ bool UWorldPartition::UnregisterActorDescContainerInstance(UActorDescContainerIn
 
 		OnActorDescContainerInstanceUnregistered.Broadcast(InActorDescContainerInstance);
 
-		// Un-hashing needs to be done for an initialized container instance that was previously hashed
-		if (EditorHash && (IsInitialized() || (InActorDescContainerInstance->IsInitialized() && ShouldHashUnhashActorDescInstances())))
+		// Un-hashing needs to be done for an initialized container instance that was previously hashed (even if WorldPartition is being uninitialized)
+		if (ShouldHashUnhashActorDescInstances() && (IsInitialized() || InActorDescContainerInstance->IsInitialized()))
 		{
 			for (UActorDescContainerInstance::TIterator<> It(InActorDescContainerInstance); It; ++It)
 			{
