Add support of EDL in instanced partitioned worlds.

diff --git a/Engine/Source/Runtime/Engine/Private/WorldPartition/ContentBundle/ContentBundleEditor.cpp b/Engine/Source/Runtime/Engine/Private/WorldPartition/ContentBundle/ContentBundleEditor.cpp
--- a/Engine/Source/Runtime/Engine/Private/WorldPartition/ContentBundle/ContentBundleEditor.cpp
+++ b/Engine/Source/Runtime/Engine/Private/WorldPartition/ContentBundle/ContentBundleEditor.cpp
@@ -526,7 +526,7 @@ bool FContentBundleEditor::PopulateGeneratedPackageForCook(class IWorldPartition
 	}
 	else
 	{
-		if (!ExternalStreamingObject->Rename(nullptr, PackageToCook.GetPackage(), REN_DontCreateRedirectors))
+		if (!ExternalStreamingObject->OnPopulateGeneratedPackageForCook(PackageToCook.GetPackage(), OutModifiedPackages))
 		{
 			UE_LOG(LogContentBundle, Error, TEXT("%s[Cook] Failed to rename streaming object package."), *ContentBundle::Log::MakeDebugInfoString(*this));
 			bIsSuccess = false;
diff --git a/Engine/Source/Runtime/Engine/Private/WorldPartition/Cook/WorldPartitionCook.cpp b/Engine/Source/Runtime/Engine/Private/WorldPartition/Cook/WorldPartitionCook.cpp
--- a/Engine/Source/Runtime/Engine/Private/WorldPartition/Cook/WorldPartitionCook.cpp
+++ b/Engine/Source/Runtime/Engine/Private/WorldPartition/Cook/WorldPartitionCook.cpp
@@ -100,14 +100,38 @@ bool UWorldPartition::PrepareGeneratorPackageForCook(IWorldPartitionCookPackageC
 
 bool UWorldPartition::PopulateGeneratorPackageForCook(IWorldPartitionCookPackageContext& CookContext, const TArray<FWorldPartitionCookPackage*>& InPackagesToCook, TArray<UPackage*>& OutModifiedPackages)
 {
+	auto OnPopulateGeneratorPackageForCook = [this](const FWorldPartitionCookPackage* InPackageToCook)
+	{
+		IWorldPartitionCookPackageObject* CookPackageObject = GetCookPackageObject(*InPackageToCook);
+		return CookPackageObject && CookPackageObject->OnPopulateGeneratorPackageForCook(InPackageToCook->GetPackage());
+	};
+
+	TArray<const FWorldPartitionCookPackage*> GenericPackagesToCook;
 	for (const FWorldPartitionCookPackage* PackageToCook : InPackagesToCook)
 	{
-		IWorldPartitionCookPackageObject* CookPackageObject = GetCookPackageObject(*PackageToCook);
-		if (!CookPackageObject || !CookPackageObject->OnPopulateGeneratorPackageForCook(PackageToCook->GetPackage()))
+		// Defer generic packages as they depend on level packages
+		if (PackageToCook->Type == FWorldPartitionCookPackage::EType::Level)
+		{
+			if (!OnPopulateGeneratorPackageForCook(PackageToCook))
+			{
+				return false;
+			}
+		}
+		else
+		{
+			check(PackageToCook->Type == FWorldPartitionCookPackage::EType::Generic);
+			GenericPackagesToCook.Add(PackageToCook);
+		}
+	}
+
+	for (const FWorldPartitionCookPackage* GenericPackageToCook : GenericPackagesToCook)
+	{
+		if (!OnPopulateGeneratorPackageForCook(GenericPackageToCook))
 		{
 			return false;
 		}
 	}
+
 	return true;
 }
 
diff --git a/Engine/Source/Runtime/Engine/Private/WorldPartition/DataLayer/ExternalDataLayerManager.cpp b/Engine/Source/Runtime/Engine/Private/WorldPartition/DataLayer/ExternalDataLayerManager.cpp
--- a/Engine/Source/Runtime/Engine/Private/WorldPartition/DataLayer/ExternalDataLayerManager.cpp
+++ b/Engine/Source/Runtime/Engine/Private/WorldPartition/DataLayer/ExternalDataLayerManager.cpp
@@ -5,6 +5,8 @@
 #include "Misc/PackageName.h"
 #include "Misc/PathViews.h"
 #include "UObject/Package.h"
+#include "LevelInstance/LevelInstanceInterface.h"
+#include "LevelInstance/LevelInstanceSubsystem.h"
 #include "WorldPartition/ContentBundle/ContentBundleEngineSubsystem.h"
 #include "WorldPartition/DataLayer/DataLayerManager.h"
 #include "WorldPartition/DataLayer/ExternalDataLayerAsset.h"
@@ -47,26 +49,20 @@ void UExternalDataLayerManager::Serialize(FArchive& Ar)
 
 void UExternalDataLayerManager::Initialize()
 {
+	UWorld* OuterWorld = GetTypedOuter<UWorld>();
+	UWorld* OwningWorld = GetOuterUWorldPartition()->GetWorld();
+
 	// EDL in LevelInstance is not currently supported
 	// In this case, don't initialize to make sure it will do nothing.
-	UWorldPartition* OuterWorldPartition = GetOuterUWorldPartition();
-	UWorld* OuterWorldPartitionWorld = OuterWorldPartition->GetWorld();
-	if (OuterWorldPartitionWorld != GetTypedOuter<UWorld>())
+	const ULevelInstanceSubsystem* LevelInstanceSubsystem = OwningWorld->GetSubsystem<ULevelInstanceSubsystem>();
+	if (ILevelInstanceInterface* LevelInstance = LevelInstanceSubsystem ? LevelInstanceSubsystem->GetOwningLevelInstance(OuterWorld->PersistentLevel) : nullptr)
 	{
 		return;
 	}
-
+	
+	const bool bIsInstanced = (OuterWorld && OwningWorld && (OuterWorld != OwningWorld));
+	bIsRunningGameOrInstancedWorldPartition = IsRunningGame() || bIsInstanced;
 	bIsInitialized = true;
-	bIsRunningGameOrInstancedWorldPartition = false;
-	if (IsRunningGame())
-	{
-		bIsRunningGameOrInstancedWorldPartition = true;
-	}
-	//@todo_ow: Add support EDL for instanced partitioned worlds
-	/*else
-	{
-		bIsRunningGameOrInstancedWorldPartition = OuterWorldPartition->GetWorld() && OuterWorldPartitionWorld && OuterWorldPartition->GetWorld() != OuterWorldPartitionWorld;
-	}*/
 
 	UExternalDataLayerEngineSubsystem& ExternalDataLayerEngineSubsystem = UExternalDataLayerEngineSubsystem::Get();
 	ExternalDataLayerEngineSubsystem.OnExternalDataLayerAssetRegistrationStateChanged.AddUObject(this, &UExternalDataLayerManager::OnExternalDataLayerAssetRegistrationStateChanged);
@@ -174,7 +170,26 @@ bool UExternalDataLayerManager::RegisterExternalStreamingObjectForGameWorld(cons
 		return false;
 	}
 
-	UPackage* ExternalStreamingObjectPackage = LoadPackage(nullptr, *ExternalStreamingObjectPackagePath, LOAD_Quiet | LOAD_NoWarn);
+	// Find outer world's instancing suffix and use it to create the package for the ExternalStreamingObject
+	UPackage* DestPackage = nullptr;
+	FLinkerInstancingContext InstancingContext;
+	FLinkerInstancingContext InstancingContext;
+	FString SourceWorldPath, RemappedWorldPath;
+	if (GetTypedOuter<UWorld>()->GetSoftObjectPathMapping(SourceWorldPath, RemappedWorldPath))
+	{
+		FString Source = FTopLevelAssetPath(SourceWorldPath).GetPackageName().ToString();
+		FString Remapped = FTopLevelAssetPath(RemappedWorldPath).GetPackageName().ToString();
+		InstancingContext.AddPackageMapping(FName(Source), FName(Remapped));
+		int32 Index = UE::String::FindFirst(Remapped, Source, ESearchCase::IgnoreCase);
+		if (Index != INDEX_NONE)
+		{
+			const FString Suffix = Remapped.RightChop(Index + Source.Len());
+			const FString RemappedExternalStreamingObjectPackagePath = ExternalStreamingObjectPackagePath + Suffix;
+			DestPackage = CreatePackage(*RemappedExternalStreamingObjectPackagePath);
+		}
+	}
+
+	UPackage* ExternalStreamingObjectPackage = LoadPackage(DestPackage, *ExternalStreamingObjectPackagePath, LOAD_Quiet | LOAD_NoWarn, nullptr, &InstancingContext);
 	if (!ExternalStreamingObjectPackage)
 	{
 		UE_LOG(LogWorldPartition, Error, TEXT("[EDL: %s] No external streaming object found."), *InExternalDataLayerAsset->GetName());
@@ -189,6 +204,10 @@ bool UExternalDataLayerManager::RegisterExternalStreamingObjectForGameWorld(cons
 		return false;
 	}
 
+	// Do some validation on ExternalStreamingObject's OuterWorld and OwningWorld
+	check(ExternalStreamingObject->GetOuterWorld() == GetTypedOuter<UWorld>());
+	check(ExternalStreamingObject->GetOwningWorld() == GetOuterUWorldPartition()->GetWorld());
+
 	ExternalStreamingObjects.Emplace(InExternalDataLayerAsset, ExternalStreamingObject);
 #endif
 	return true;
diff --git a/Engine/Source/Runtime/Engine/Private/WorldPartition/WorldPartitionRuntimeHash.cpp b/Engine/Source/Runtime/Engine/Private/WorldPartition/WorldPartitionRuntimeHash.cpp
--- a/Engine/Source/Runtime/Engine/Private/WorldPartition/WorldPartitionRuntimeHash.cpp
+++ b/Engine/Source/Runtime/Engine/Private/WorldPartition/WorldPartitionRuntimeHash.cpp
@@ -58,20 +58,32 @@ TSet<TObjectPtr<UDataLayerInstance>>& URuntimeHashExternalStreamingObjectBase::G
 	return DataLayerInstances;
 }
 
+UWorld* URuntimeHashExternalStreamingObjectBase::GetOwningWorld() const
+{
+	// Once OnStreamingObjectLoaded is called and OwningWorld is set, use this cached value
+	return OwningWorld.IsSet() ? OwningWorld.GetValue().Get() : GetOuterWorld()->GetWorldPartition()->GetWorld();
+}
+
 void URuntimeHashExternalStreamingObjectBase::OnStreamingObjectLoaded(UWorld* InjectedWorld)
 {
-	bool bIsACookedObject = !CellToLevelStreamingPackage.IsEmpty();
-	if (bIsACookedObject)
+#if !WITH_EDITOR
+	if (!CellToStreamingData.IsEmpty())
 	{
 		// Cooked streaming object's Cells do not have LevelStreaming.
 		ForEachStreamingCells([this](UWorldPartitionRuntimeCell& Cell)
 		{
 			UWorldPartitionRuntimeLevelStreamingCell* RuntimeCell = CastChecked<UWorldPartitionRuntimeLevelStreamingCell>(&Cell);
 
-			FName LevelStreamingPackage = CellToLevelStreamingPackage.FindChecked(RuntimeCell->GetFName());
-			RuntimeCell->CreateAndSetLevelStreaming(*LevelStreamingPackage.ToString());
+			const FWorldPartitionRuntimeCellStreamingData& CellStreamingData = CellToStreamingData.FindChecked(RuntimeCell->GetFName());
+			RuntimeCell->CreateAndSetLevelStreaming(CellStreamingData.PackageName, CellStreamingData.WorldAsset);
 		});
 	}
+#endif
+
+	check(GetOuterWorld());
+	check(GetOuterWorld()->GetWorldPartition());
+	check(GetOuterWorld()->GetWorldPartition()->GetWorld());
+	OwningWorld = GetOuterWorld()->GetWorldPartition()->GetWorld();
 }
 
 #if WITH_EDITOR
@@ -103,7 +115,10 @@ bool URuntimeHashExternalStreamingObjectBase::OnPopulateGeneratorPackageForCook(
 	{
 		UWorldPartitionRuntimeLevelStreamingCell* RuntimeCell = CastChecked<UWorldPartitionRuntimeLevelStreamingCell>(&Cell);
 		UWorldPartitionLevelStreamingDynamic* LevelStreamingDynamic = RuntimeCell->GetLevelStreaming();
-		CellToLevelStreamingPackage.Add(RuntimeCell->GetFName(), LevelStreamingDynamic->PackageNameToLoad);
+		FWorldPartitionRuntimeCellStreamingData& CellStreamingData = CellToStreamingData.Add(RuntimeCell->GetFName());
+		CellStreamingData.PackageName = LevelStreamingDynamic->GetWorldAsset().GetLongPackageName();
+		// SoftObjectPath will be automatically remapped when ExternalStreamingObject will be instanced/loaded at runtime
+		CellStreamingData.WorldAsset = LevelStreamingDynamic->GetWorldAsset().ToSoftObjectPath();
 
 		// Level streaming are outered to the world and would not be saved within the ExternalStreamingObject.
 		// Do not save them, instead they will be created once the external streaming object is loaded at runtime. 
@@ -175,7 +190,7 @@ EWorldPartitionStreamingPerformance UWorldPartitionRuntimeHash::GetStreamingPerf
 	return EWorldPartitionStreamingPerformance::Good;
 }
 
-URuntimeHashExternalStreamingObjectBase* UWorldPartitionRuntimeHash::CreateExternalStreamingObject(TSubclassOf<URuntimeHashExternalStreamingObjectBase> InClass, UObject* InOuter, FName InName, UWorld* InOwningWorld, UWorld* InOuterWorld)
+URuntimeHashExternalStreamingObjectBase* UWorldPartitionRuntimeHash::CreateExternalStreamingObject(TSubclassOf<URuntimeHashExternalStreamingObjectBase> InClass, UObject* InOuter, FName InName, UWorld* InOuterWorld)
 {
 	if (FindObject<URuntimeHashExternalStreamingObjectBase>(InOuter, *InName.ToString()))
 	{
@@ -184,8 +199,7 @@ URuntimeHashExternalStreamingObjectBase* UWorldPartitionRuntimeHash::CreateExter
 	}
 
 	URuntimeHashExternalStreamingObjectBase* StreamingObject = NewObject<URuntimeHashExternalStreamingObjectBase>(InOuter, InClass, InName, RF_Public);
-	StreamingObject->OwningWorld = InOwningWorld;
-	StreamingObject->OuterWorld = InOuterWorld;
+	StreamingObject->OuterWorld = InOuterWorld;	
 	return StreamingObject;
 }
 
@@ -355,7 +369,7 @@ UWorldPartitionRuntimeCell* UWorldPartitionRuntimeHash::GetCellForCookPackage(co
 
 URuntimeHashExternalStreamingObjectBase* UWorldPartitionRuntimeHash::StoreStreamingContentToExternalStreamingObject(FName InStreamingObjectName)
 {
-	URuntimeHashExternalStreamingObjectBase* NewExternalStreamingObject = CreateExternalStreamingObject(GetExternalStreamingObjectClass(), GetOuterUWorldPartition(), InStreamingObjectName, GetWorld(), GetTypedOuter<UWorld>());
+	URuntimeHashExternalStreamingObjectBase* NewExternalStreamingObject = CreateExternalStreamingObject(GetExternalStreamingObjectClass(), GetOuterUWorldPartition(), InStreamingObjectName, GetTypedOuter<UWorld>());
 	StoreStreamingContentToExternalStreamingObject(NewExternalStreamingObject);
 	return NewExternalStreamingObject;
 }
diff --git a/Engine/Source/Runtime/Engine/Private/WorldPartition/WorldPartitionRuntimeLevelStreamingCell.cpp b/Engine/Source/Runtime/Engine/Private/WorldPartition/WorldPartitionRuntimeLevelStreamingCell.cpp
--- a/Engine/Source/Runtime/Engine/Private/WorldPartition/WorldPartitionRuntimeLevelStreamingCell.cpp
+++ b/Engine/Source/Runtime/Engine/Private/WorldPartition/WorldPartitionRuntimeLevelStreamingCell.cpp
@@ -98,9 +98,9 @@ TArray<FName> UWorldPartitionRuntimeLevelStreamingCell::GetActors() const
 	return Actors;
 }
 
-void UWorldPartitionRuntimeLevelStreamingCell::CreateAndSetLevelStreaming(const FString& InPackageName)
+void UWorldPartitionRuntimeLevelStreamingCell::CreateAndSetLevelStreaming(const FString& InPackageName, const FSoftObjectPath& InWorldAsset)
 {
-	LevelStreaming = CreateLevelStreaming(InPackageName);
+	LevelStreaming = CreateLevelStreaming(InPackageName, InWorldAsset);
 }
 
 bool UWorldPartitionRuntimeLevelStreamingCell::CreateAndSetLevelStreaming(const TSoftObjectPtr<UWorld>& InWorldAsset, const FTransform& InInstanceTransform) const
@@ -136,28 +136,45 @@ bool UWorldPartitionRuntimeLevelStreamingCell::CreateAndSetLevelStreaming(const
 
 #if WITH_EDITOR
 	LevelStreaming->SetShouldPerformStandardLevelLoading(true);
-#endif
 
 	if (OwningWorld->IsPlayInEditor() && OwningWorld->GetPackage()->HasAnyPackageFlags(PKG_PlayInEditor) && OwningWorld->GetPackage()->GetPIEInstanceID() != INDEX_NONE)
 	{
 		// When renaming for PIE, make sure to keep World's name so that linker can properly remap with Package's instancing context
 		LevelStreaming->RenameForPIE(OwningWorld->GetPackage()->GetPIEInstanceID(), /*bKeepWorldAssetName*/true);
 	}
+#endif
+
 	return true;
 }
 
-UWorldPartitionLevelStreamingDynamic* UWorldPartitionRuntimeLevelStreamingCell::CreateLevelStreaming(const FString& InPackageName) const
+UWorldPartitionLevelStreamingDynamic* UWorldPartitionRuntimeLevelStreamingCell::CreateLevelStreaming(const FString& InPackageName, const FSoftObjectPath& InWorldAsset) const
 {
+	auto GetUniqueLevelStreamingName = [this]()
+	{
+		const UWorld* OuterWorld = GetOuterWorld();
+		TStringBuilder<128> LevelStreamingNameBuilder;
+		LevelStreamingNameBuilder.Appendf(TEXT("WorldPartitionLevelStreaming_%s"), *GetName());
+		if (OuterWorld->IsGameWorld())
+		{
+			FString OuterWorldPackageShortName = FPackageName::GetShortName(OuterWorld->GetPackage());
+#if WITH_EDITOR
+			OuterWorldPackageShortName = UWorld::RemovePIEPrefix(OuterWorldPackageShortName);
+#endif
+			// Include outer world package name to make sure we generate a unique name since we it is 
+			LevelStreamingNameBuilder.Appendf(TEXT("_%s"), *OuterWorldPackageShortName);
+		}
+		return FName(*LevelStreamingNameBuilder);
+	};
+
 	if (HasActors())
 	{
 		UWorld* OuterWorld = GetOuterWorld();
 		UWorld* OwningWorld = GetOwningWorld();
-		const UWorldPartition* WorldPartition = OuterWorld->GetWorldPartition();
-		const FName LevelStreamingName = FName(*FString::Printf(TEXT("WorldPartitionLevelStreaming_%s"), *GetName()));
-
+		
 		// When called by Commandlet (PopulateGeneratedPackageForCook), LevelStreaming's outer is set to Cell/WorldPartition's outer to prevent warnings when saving Cell Levels (Warning: Obj in another map). 
 		// At runtime, LevelStreaming's outer will be properly set to the main world (see UWorldPartitionRuntimeLevelStreamingCell::Activate).
 		UWorld* LevelStreamingOuterWorld = IsRunningCommandlet() ? OuterWorld : OwningWorld;
+		const FName LevelStreamingName = GetUniqueLevelStreamingName();
 		UWorldPartitionLevelStreamingDynamic* NewLevelStreaming = NewObject<UWorldPartitionLevelStreamingDynamic>(LevelStreamingOuterWorld, UWorldPartitionLevelStreamingDynamic::StaticClass(), LevelStreamingName, RF_NoFlags, NULL);
 
 		FName WorldName = OuterWorld->GetFName();
@@ -174,18 +191,31 @@ UWorldPartitionLevelStreamingDynamic* UWorldPartitionRuntimeLevelStreamingCell::
 		FString PackageName = InPackageName;
 #endif
 
-		TSoftObjectPtr<UWorld> WorldAsset(FSoftObjectPath(FString::Printf(TEXT("%s.%s"), *PackageName, *WorldName.ToString())));
-		NewLevelStreaming->SetWorldAsset(WorldAsset);
+		// Set both PackageNameToLoad and WorldAsset (necessary to properly support instancing)
+		NewLevelStreaming->PackageNameToLoad = *PackageName;
+		if (InWorldAsset.IsValid())
+		{
+			NewLevelStreaming->SetWorldAsset(TSoftObjectPtr<UWorld>(InWorldAsset));
+		}
+		else
+		{
+			TSoftObjectPtr<UWorld> WorldAsset(FSoftObjectPath(FString::Printf(TEXT("%s.%s"), *PackageName, *WorldName.ToString())));
+			NewLevelStreaming->SetWorldAsset(WorldAsset);
+		}
+
 		// Transfer WorldPartition's transform to Level
-		NewLevelStreaming->LevelTransform = WorldPartition->GetInstanceTransform();
+		const UWorldPartition* OuterWorldPartition = OuterWorld->GetWorldPartition();
+		NewLevelStreaming->LevelTransform = OuterWorldPartition->GetInstanceTransform();
 		NewLevelStreaming->bClientOnlyVisible = GetClientOnlyVisible();
 		NewLevelStreaming->Initialize(*this);
 
+#if WITH_EDITOR
 		if (OwningWorld->IsPlayInEditor() && OwningWorld->GetPackage()->HasAnyPackageFlags(PKG_PlayInEditor) && OwningWorld->GetPackage()->GetPIEInstanceID() != INDEX_NONE)
 		{
 			// When renaming for PIE, make sure to keep World's name so that linker can properly remap with Package's instancing context
 			NewLevelStreaming->RenameForPIE(OwningWorld->GetPackage()->GetPIEInstanceID(), /*bKeepWorldAssetName*/true);
 		}
+#endif
 
 		return NewLevelStreaming;
 	}
@@ -490,7 +520,7 @@ UWorldPartitionLevelStreamingDynamic* UWorldPartitionRuntimeLevelStreamingCell::
 		{
 			LevelStreaming->Rename(nullptr, OwningWorld, REN_ForceNoResetLoaders);
 		}
-
+		
 		// Transfer WorldPartition's transform to LevelStreaming
 		LevelStreaming->SetLevelTransform(WorldPartition->GetInstanceTransform());
 
diff --git a/Engine/Source/Runtime/Engine/Public/WorldPartition/WorldPartitionLevelStreamingDynamic.h b/Engine/Source/Runtime/Engine/Public/WorldPartition/WorldPartitionLevelStreamingDynamic.h
--- a/Engine/Source/Runtime/Engine/Public/WorldPartition/WorldPartitionLevelStreamingDynamic.h
+++ b/Engine/Source/Runtime/Engine/Public/WorldPartition/WorldPartitionLevelStreamingDynamic.h
@@ -30,7 +30,7 @@ class UWorldPartitionLevelStreamingDynamic : public ULevelStreamingDynamic
 	ENGINE_API void Unload();
 	ENGINE_API void Activate();
 	ENGINE_API void Deactivate();
-	ENGINE_API UWorld* GetStreamingWorld() const override;
+	ENGINE_API virtual UWorld* GetStreamingWorld() const override;
 	void SetShouldBeAlwaysLoaded(bool bInShouldBeAlwaysLoaded) { bShouldBeAlwaysLoaded = bInShouldBeAlwaysLoaded; }
 	const UWorldPartitionRuntimeCell* GetWorldPartitionRuntimeCell() const { return StreamingCell.Get(); }
 
diff --git a/Engine/Source/Runtime/Engine/Public/WorldPartition/WorldPartitionRuntimeHash.h b/Engine/Source/Runtime/Engine/Public/WorldPartition/WorldPartitionRuntimeHash.h
--- a/Engine/Source/Runtime/Engine/Public/WorldPartition/WorldPartitionRuntimeHash.h
+++ b/Engine/Source/Runtime/Engine/Public/WorldPartition/WorldPartitionRuntimeHash.h
@@ -5,6 +5,7 @@
 #include "UObject/ObjectMacros.h"
 #include "UObject/Object.h"
 #include "Templates/SubclassOf.h"
+#include "Misc/Optional.h"
 #include "Engine/World.h"
 #include "WorldPartition.h"
 #include "WorldPartition/ActorDescList.h"
@@ -36,6 +37,18 @@ enum class EWorldPartitionStreamingPerformance : uint8
 	Critical
 };
 
+USTRUCT()
+struct FWorldPartitionRuntimeCellStreamingData
+{
+	GENERATED_BODY()
+
+	UPROPERTY()
+	FString PackageName;
+
+	UPROPERTY()
+	FSoftObjectPath WorldAsset;
+};
+
 UCLASS(Abstract, MinimalAPI)
 class URuntimeHashExternalStreamingObjectBase : public UObject, public IWorldPartitionCookPackageObject, public IDataLayerInstanceProvider
 {
@@ -51,7 +64,7 @@ public:
 #endif
 	//~ End UObject Interface
 
-	UWorld* GetOwningWorld() const { return OwningWorld.Get(); }
+	UWorld* GetOwningWorld() const;
 	UWorld* GetOuterWorld() const { return OuterWorld.Get(); }
 
 	ENGINE_API void ForEachStreamingCells(TFunctionRef<void(UWorldPartitionRuntimeCell&)> Func);
@@ -91,14 +104,13 @@ public:
 	FWorldPartitionRuntimeContainerResolver ContainerResolver;
 
 protected:
-	UPROPERTY();
-	TSoftObjectPtr<UWorld> OwningWorld;
+	TOptional<TWeakObjectPtr<UWorld>> OwningWorld;
 
 	UPROPERTY();
 	TSoftObjectPtr<UWorld> OuterWorld;
 
 	UPROPERTY();
-	TMap<FName, FName> CellToLevelStreamingPackage;
+	TMap<FName, FWorldPartitionRuntimeCellStreamingData> CellToStreamingData;
 
 	UPROPERTY()
 	TSet<TObjectPtr<UDataLayerInstance>> DataLayerInstances;
@@ -214,18 +226,10 @@ public:
 	virtual bool GetShouldMergeStreamingSourceInfo() const { return false; }
 
 protected:
-	static ENGINE_API URuntimeHashExternalStreamingObjectBase* CreateExternalStreamingObject(TSubclassOf<URuntimeHashExternalStreamingObjectBase> InClass, UObject* InOuter, FName InName, UWorld* InOwningWorld, UWorld* InOuterWorld);
+	static ENGINE_API URuntimeHashExternalStreamingObjectBase* CreateExternalStreamingObject(TSubclassOf<URuntimeHashExternalStreamingObjectBase> InClass, UObject* InOuter, FName InName, UWorld* InOuterWorld);
 	ENGINE_API UWorldPartitionRuntimeCell* CreateRuntimeCell(UClass* CellClass, UClass* CellDataClass, const FString& CellName, const FString& CellInstanceSuffix, UObject* InOuter = nullptr);
 	virtual EWorldPartitionStreamingPerformance GetStreamingPerformanceForCell(const UWorldPartitionRuntimeCell* Cell) const;
 
-#if WITH_EDITOR
-	template <class T>
-	T* CreateExternalStreamingObject(UObject* InOuter, FName InName)
-	{
-		return static_cast<T*>(CreateExternalStreamingObject(T::StaticClass(), InOuter, InName, GetWorld(), GetTypedOuter<UWorld>()));
-	}
-#endif
-
 #if WITH_EDITORONLY_DATA
 	struct FAlwaysLoadedActorForPIE
 	{
diff --git a/Engine/Source/Runtime/Engine/Public/WorldPartition/WorldPartitionRuntimeLevelStreamingCell.h b/Engine/Source/Runtime/Engine/Public/WorldPartition/WorldPartitionRuntimeLevelStreamingCell.h
--- a/Engine/Source/Runtime/Engine/Public/WorldPartition/WorldPartitionRuntimeLevelStreamingCell.h
+++ b/Engine/Source/Runtime/Engine/Public/WorldPartition/WorldPartitionRuntimeLevelStreamingCell.h
@@ -40,9 +40,9 @@ class UWorldPartitionRuntimeLevelStreamingCell : public UWorldPartitionRuntimeCe
 	ENGINE_API bool HasActors() const;
 	ENGINE_API virtual TArray<FName> GetActors() const override;
 
-	ENGINE_API void CreateAndSetLevelStreaming(const FString& InPackageName);
+	ENGINE_API void CreateAndSetLevelStreaming(const FString& InPackageName, const FSoftObjectPath& InWorldAsset = FSoftObjectPath());
 	ENGINE_API bool CreateAndSetLevelStreaming(const TSoftObjectPtr<UWorld>& InWorldAsset, const FTransform& InInstanceTransform) const;
-	ENGINE_API class UWorldPartitionLevelStreamingDynamic* CreateLevelStreaming(const FString& InPackageName = FString()) const;
+	ENGINE_API class UWorldPartitionLevelStreamingDynamic* CreateLevelStreaming(const FString& InPackageName = FString(), const FSoftObjectPath& InWorldAsset = FSoftObjectPath()) const;
 	
 
 #if WITH_EDITOR
