Fixed build

diff --git a/Engine/Source/Runtime/Engine/Private/WorldPartition/DataLayer/ExternalDataLayerManager.cpp b/Engine/Source/Runtime/Engine/Private/WorldPartition/DataLayer/ExternalDataLayerManager.cpp
--- a/Engine/Source/Runtime/Engine/Private/WorldPartition/DataLayer/ExternalDataLayerManager.cpp
+++ b/Engine/Source/Runtime/Engine/Private/WorldPartition/DataLayer/ExternalDataLayerManager.cpp
@@ -173,7 +173,6 @@ bool UExternalDataLayerManager::RegisterExternalStreamingObjectForGameWorld(cons
 	// Find outer world's instancing suffix and use it to create the package for the ExternalStreamingObject
 	UPackage* DestPackage = nullptr;
 	FLinkerInstancingContext InstancingContext;
-	FLinkerInstancingContext InstancingContext;
 	FString SourceWorldPath, RemappedWorldPath;
 	if (GetTypedOuter<UWorld>()->GetSoftObjectPathMapping(SourceWorldPath, RemappedWorldPath))
 	{
