EditorDomain: UMaterialInterface: Fix crash when saving EditorDomain packages during load caused by attempting to garbage collect during load. The disposal code is only necessary for editor saves from the user and can be skipped in the editordomain case.
#rnx

diff --git a/Engine/Source/Runtime/Engine/Private/Materials/MaterialInterface.cpp b/Engine/Source/Runtime/Engine/Private/Materials/MaterialInterface.cpp
--- a/Engine/Source/Runtime/Engine/Private/Materials/MaterialInterface.cpp
+++ b/Engine/Source/Runtime/Engine/Private/Materials/MaterialInterface.cpp
@@ -1687,7 +1687,7 @@ void UMaterialInterface::PreSave(FObjectPreSaveContext ObjectSaveContext)
 	// the editor running on cooked data.
 	UMaterialInterfaceEditorOnlyData* EditorOnly = GetEditorOnlyData();
 	FString EditorOnlyDataName = MaterialInterface::GetEditorOnlyDataName(*GetName());
-	if (!ObjectSaveContext.IsCooking() && EditorOnly && EditorOnly->GetName() != EditorOnlyDataName)
+	if (!ObjectSaveContext.IsProceduralSave() && EditorOnly && EditorOnly->GetName() != EditorOnlyDataName)
 	{
 		UE_LOG(LogMaterial, Display, TEXT("MaterialInterface %s has a incorrectly name EditorOnlyData '%s'. This may cause issues when running the editor on cooked data. Trying to rename it to the correct name '%s'."), *GetName(), *EditorOnly->GetName(), *EditorOnlyDataName);
 		DisposeInvalidEditorOnlyData(this, EditorOnlyDataName);
