World Partition: when unpartitioning a level, make sure to keep all actors loaded when uninitializing the world partition object.

diff --git a/Engine/Source/Runtime/Engine/Private/WorldPartition/WorldPartition.cpp b/Engine/Source/Runtime/Engine/Private/WorldPartition/WorldPartition.cpp
--- a/Engine/Source/Runtime/Engine/Private/WorldPartition/WorldPartition.cpp
+++ b/Engine/Source/Runtime/Engine/Private/WorldPartition/WorldPartition.cpp
@@ -46,6 +46,7 @@
 #include "WorldPartition/LoaderAdapter/LoaderAdapterShape.h"
 #include "WorldPartition/LoaderAdapter/LoaderAdapterPinnedActors.h"
 #include "WorldPartition/HLOD/HLODLayer.h"
+#include "WorldPartition/HLOD/HLODActor.h"
 #include "WorldPartition/Cook/WorldPartitionCookPackageContextInterface.h"
 #include "WorldPartition/WorldPartitionActorDescInstance.h"
 #include "WorldPartition/ActorDescContainerInstance.h"
@@ -1075,16 +1076,25 @@ bool UWorldPartition::RemoveWorldPartition(AWorldSettings* WorldSettings)
 	{
 		if (!WorldPartition->IsStreamingEnabled())
 		{
-			FWorldPartitionLoadingContext::FNull LoadingContext;
+			ULevel* PersistentLevel = WorldSettings->GetLevel();
+
+			TArray<FWorldPartitionReference> ActorReferences;
+			ActorReferences.Reserve(PersistentLevel->Actors.Num());
 
 			WorldSettings->Modify();
 			
-			ULevel* PersistentLevel = WorldSettings->GetLevel();
 			for (AActor* Actor : PersistentLevel->Actors)
 			{
-				if (Actor && (Cast<AWorldDataLayers>(Actor) || Cast<AWorldPartitionMiniMap>(Actor)))
+				if (Actor)
 				{
-					Actor->Destroy();
+					if (Cast<AWorldDataLayers>(Actor) || Cast<AWorldPartitionMiniMap>(Actor) || Cast<AWorldPartitionHLOD>(Actor))
+					{
+						Actor->Destroy();
+					}
+					else if(Actor->GetExternalPackage())
+					{
+						ActorReferences.Emplace(WorldPartition, Actor->GetActorGuid());
+					}
 				}
 			}
 
