Added support for blend spaces with sync markers for motion matching

New:
- Added ResetBlendSamples() method to UBlendspace. Works with marker based syncing and length based syncing

Fixes:
- Fixed advancing backwards instead of forwards when DeltaTime is zero and using sync markers
- Fixed MarkerTickRecords not being reset when setting accumulated time or start time of an asset player
- Fixed sequences and blendspaces not handling their makers becoming invalid when deltatime is zero.

#rnx minor

diff --git a/Engine/Plugins/Animation/PoseSearch/Source/Editor/Private/PoseSearchDatabaseViewModel.cpp b/Engine/Plugins/Animation/PoseSearch/Source/Editor/Private/PoseSearchDatabaseViewModel.cpp
--- a/Engine/Plugins/Animation/PoseSearch/Source/Editor/Private/PoseSearchDatabaseViewModel.cpp
+++ b/Engine/Plugins/Animation/PoseSearch/Source/Editor/Private/PoseSearchDatabaseViewModel.cpp
@@ -87,7 +87,7 @@ bool FDatabasePreviewActor::SpawnPreviewActor(UWorld* World, const UPoseSearchDa
 	Mesh->SetSkeletalMesh(DatabasePreviewMesh ? DatabasePreviewMesh : PoseSearchDatabase->Schema->GetSkeleton(Role)->GetPreviewMesh(true));
 	Mesh->EnablePreview(true, PreviewAsset);
 		
-	AnimInstance->SetAnimationAsset(PreviewAsset, false, 0.0f);
+	AnimInstance->SetAnimationAsset(PreviewAsset, IndexAsset.IsLooping(), 0.0f);
 	AnimInstance->SetBlendSpacePosition(IndexAsset.GetBlendParameters());
 		
 	if (IndexAsset.IsMirrored() && PoseSearchDatabase->Schema)
@@ -108,8 +108,7 @@ bool FDatabasePreviewActor::SpawnPreviewActor(UWorld* World, const UPoseSearchDa
 		RootTransformOrigin = MirrorDataCache.MirrorTransform(Sampler.ExtractRootTransform(PlayTimeOffset));
 	}
 
-	AnimInstance->PlayAnim(false, 0.0f);
-
+	AnimInstance->PlayAnim(IndexAsset.IsLooping(), 0.0f);
 	if (!ActorPtr->GetRootComponent())
 	{
 		ActorPtr->SetRootComponent(Mesh);
diff --git a/Engine/Plugins/Animation/PoseSearch/Source/Runtime/Private/PoseSearchAssetSampler.cpp b/Engine/Plugins/Animation/PoseSearch/Source/Runtime/Private/PoseSearchAssetSampler.cpp
--- a/Engine/Plugins/Animation/PoseSearch/Source/Runtime/Private/PoseSearchAssetSampler.cpp
+++ b/Engine/Plugins/Animation/PoseSearch/Source/Runtime/Private/PoseSearchAssetSampler.cpp
@@ -149,24 +149,49 @@ static void ProcessRootTransform(const UBlendSpace* BlendSpace, const FVector& B
 	// Pre-compute root motion
 	const int32 NumRootSamples = FMath::Max(CachedPlayLength * RootTransformSamplingRate + 1, 1);
 	AccumulatedRootTransform.Init(FTransform::Identity, NumRootSamples);
+	
+	TArray<FBlendSampleData> BlendSamplesData;
 
-	TArray<FBlendSampleData> BlendSamples;
 	int32 TriangulationIndex = 0;
-	if (BlendSpace->GetSamplesFromBlendInput(BlendParameters, BlendSamples, TriangulationIndex, true))
+	if (BlendSpace->GetSamplesFromBlendInput(BlendParameters, BlendSamplesData, TriangulationIndex, true))
 	{
+		TArray<float, TInlineAllocator<16, TMemStackAllocator<>>> PrevSampleTimes;
+		PrevSampleTimes.AddDefaulted(BlendSamplesData.Num());
+		
+		// Get starting time for all samples.
+		BlendSpace->ResetBlendSamples(BlendSamplesData, 0.0f, bIsLoopable, true);
+		
 		for (int32 SampleIdx = 1; SampleIdx < NumRootSamples; ++SampleIdx)
 		{
+			// Keep track of previous samples
+			for (int32 BlendSampleIndex = 0; BlendSampleIndex < BlendSamplesData.Num(); BlendSampleIndex++)
+			{
+				PrevSampleTimes[BlendSampleIndex] = BlendSamplesData[BlendSampleIndex].Time;
+			}
+
+			// Compute samples with new data.
+			const float SampleTime = static_cast<float>(SampleIdx) / (NumRootSamples - 1);
+			BlendSpace->ResetBlendSamples(BlendSamplesData, SampleTime, bIsLoopable, true);
+
+			// Accumulate root motion after samples have been updated.
 			FRootMotionMovementParams RootMotionMovementParams;
-			for (int32 BlendSampleIdex = 0; BlendSampleIdex < BlendSamples.Num(); BlendSampleIdex++)
+			for (int32 BlendSampleIndex = 0; BlendSampleIndex < BlendSamplesData.Num(); BlendSampleIndex++)
 			{
-				FBlendSampleData& BlendSample = BlendSamples[BlendSampleIdex];
-				// @todo: add support for synch marker
-				const float PlayLength = BlendSample.Animation->GetPlayLength();
-				const float DeltaTime = PlayLength / (NumRootSamples - 1);
-				const float PreviousTime = (SampleIdx - 1) * DeltaTime;
-
-				const FTransform BlendSampleRootMotion = BlendSample.Animation->ExtractRootMotion(PreviousTime, DeltaTime, BlendSpace->bLoop);
-				RootMotionMovementParams.AccumulateWithBlend(BlendSampleRootMotion, BlendSample.GetClampedWeight());
+				FBlendSampleData& BlendSample = BlendSamplesData[BlendSampleIndex];
+
+				if (BlendSample.TotalWeight > ZERO_ANIMWEIGHT_THRESH)
+				{
+					float DeltaTime = BlendSample.Time - PrevSampleTimes[BlendSampleIndex];
+
+					// Account for looping.
+					if (DeltaTime < 0.0f)
+					{
+						DeltaTime += BlendSample.Animation->GetPlayLength();
+					}
+					
+					const FTransform BlendSampleRootMotion = BlendSample.Animation->ExtractRootMotion(PrevSampleTimes[BlendSampleIndex], DeltaTime, bIsLoopable);
+					RootMotionMovementParams.AccumulateWithBlend(BlendSampleRootMotion, BlendSample.GetClampedWeight());
+				}
 			}
 
 			AccumulatedRootTransform[SampleIdx] = RootMotionMovementParams.GetRootMotionTransform() * AccumulatedRootTransform[SampleIdx - 1];
@@ -224,14 +249,6 @@ float FAnimationAssetSampler::GetPlayLength(const UAnimationAsset* AnimAsset, co
 			{
 				PlayLength = BlendSpace->GetAnimationLengthFromSampleData(BlendSamples);
 			}
-
-#if !NO_LOGGING
-			const TArray<FName>* UniqueMarkerNames = const_cast<UBlendSpace*>(BlendSpace)->GetUniqueMarkerNames();
-			if (UniqueMarkerNames && !UniqueMarkerNames->IsEmpty())
-			{
-				UE_LOG(LogPoseSearch, Warning, TEXT("FAnimationAssetSampler::Init: sampling blend space (%s) with synch markers is currently not supported"), *BlendSpace->GetName());
-			}
-#endif // !NO_LOGGING
 		}
 		else
 		{
@@ -339,19 +356,7 @@ void FAnimationAssetSampler::ExtractPose(const FAnimExtractContext& ExtractionCt
 		int32 TriangulationIndex = 0;
 		if (BlendSpace->GetSamplesFromBlendInput(BlendParameters, BlendSamples, TriangulationIndex, true))
 		{
-			for (int32 BlendSampleIdex = 0; BlendSampleIdex < BlendSamples.Num(); BlendSampleIdex++)
-			{
-				FBlendSampleData& BlendSample = BlendSamples[BlendSampleIdex];
-				const float Scale = BlendSample.Animation && CachedPlayLength > UE_KINDA_SMALL_NUMBER ? BlendSample.Animation->GetPlayLength() / CachedPlayLength : 1.f;
-
-				FDeltaTimeRecord BlendSampleDeltaTimeRecord;
-				BlendSampleDeltaTimeRecord.Set(ExtractionCtx.DeltaTimeRecord.GetPrevious() * Scale, ExtractionCtx.DeltaTimeRecord.Delta * Scale);
-
-				BlendSample.DeltaTimeRecord = BlendSampleDeltaTimeRecord;
-				BlendSample.PreviousTime = ExtractionCtx.DeltaTimeRecord.GetPrevious() * Scale;
-				BlendSample.Time = ExtractionCtx.CurrentTime * Scale;
-			}
-
+			BlendSpace->ResetBlendSamples(BlendSamples, ToNormalizedTime(ExtractionCtx.CurrentTime), ExtractionCtx.bLooping, false);
 			BlendSpace->GetAnimationPose(BlendSamples, ExtractionCtx, OutAnimPoseData);
 		}
 	}
@@ -387,7 +392,7 @@ void FAnimationAssetSampler::ExtractPose(float Time, FCompactPose& OutPose) cons
 
 	FDeltaTimeRecord DeltaTimeRecord;
 	DeltaTimeRecord.Set(Time, 0.f);
-	FAnimExtractContext ExtractionCtx(double(Time), false, DeltaTimeRecord, false);
+	FAnimExtractContext ExtractionCtx(double(Time), false, DeltaTimeRecord, IsLoopable());
 
 	ExtractPose(ExtractionCtx, AnimPoseData);
 }
diff --git a/Engine/Source/Runtime/AnimGraphRuntime/Private/AnimNodes/AnimNode_BlendSpacePlayer.cpp b/Engine/Source/Runtime/AnimGraphRuntime/Private/AnimNodes/AnimNode_BlendSpacePlayer.cpp
--- a/Engine/Source/Runtime/AnimGraphRuntime/Private/AnimNodes/AnimNode_BlendSpacePlayer.cpp
+++ b/Engine/Source/Runtime/AnimGraphRuntime/Private/AnimNodes/AnimNode_BlendSpacePlayer.cpp
@@ -207,6 +207,8 @@ void FAnimNode_BlendSpacePlayerBase::Reinitialize(bool bResetTime)
 			// Blend spaces run between 0 and 1
 			InternalTimeAccumulator = 1.0f;
 		}
+
+		MarkerTickRecord.Reset();
 	}
 
 	UBlendSpace* CurrentBlendSpace = GetBlendSpace();
diff --git a/Engine/Source/Runtime/Engine/Classes/Animation/BlendSpace.h b/Engine/Source/Runtime/Engine/Classes/Animation/BlendSpace.h
--- a/Engine/Source/Runtime/Engine/Classes/Animation/BlendSpace.h
+++ b/Engine/Source/Runtime/Engine/Classes/Animation/BlendSpace.h
@@ -573,6 +573,16 @@ public:
 	 */
 	ENGINE_API bool UpdateBlendSamples(const FVector& InBlendSpacePosition, float InDeltaTime, TArray<FBlendSampleData>& InOutSampleDataCache, int32& InOutCachedTriangulationIndex) const;
 	
+	/**
+	 * Resets a cached set of blend samples to match a given input time. All samples will be advanced using sync marker if possible, otherwise, their time will just be set match the input normalized time.
+	 * 
+	 * @param	InOutSampleDataCache			The sample data cache to use.
+	 * @param	InNormalizedCurrentTime			The time to match when advancing samples. 
+	 * @param	bLooping						If true, advance samples as a looping blend space would.
+	 * @param	bMatchSyncPhases				If true, all follower samples will pass the same amount of markers the leader sample has passed to match its sync phase. Otherwise, followers samples will only match their next valid sync position.  
+	 */
+	ENGINE_API void ResetBlendSamples(TArray<FBlendSampleData>& InOutSampleDataCache, float InNormalizedCurrentTime, bool bLooping, bool bMatchSyncPhases = true) const;
+
 	/**
 	 * Allows the user to iterate through all the data samples available in the blend space.
 	 * @param Func The function to run for each blend sample
diff --git a/Engine/Source/Runtime/Engine/Private/Animation/AnimNode_AssetPlayerBase.cpp b/Engine/Source/Runtime/Engine/Private/Animation/AnimNode_AssetPlayerBase.cpp
--- a/Engine/Source/Runtime/Engine/Private/Animation/AnimNode_AssetPlayerBase.cpp
+++ b/Engine/Source/Runtime/Engine/Private/Animation/AnimNode_AssetPlayerBase.cpp
@@ -52,6 +52,7 @@ float FAnimNode_AssetPlayerBase::GetAccumulatedTime() const
 void FAnimNode_AssetPlayerBase::SetAccumulatedTime(float NewTime)
 {
 	InternalTimeAccumulator = NewTime;
+	MarkerTickRecord.Reset();
 }
 
 float FAnimNode_AssetPlayerBase::GetCachedBlendWeight() const
diff --git a/Engine/Source/Runtime/Engine/Private/Animation/AnimSequence.cpp b/Engine/Source/Runtime/Engine/Private/Animation/AnimSequence.cpp
--- a/Engine/Source/Runtime/Engine/Private/Animation/AnimSequence.cpp
+++ b/Engine/Source/Runtime/Engine/Private/Animation/AnimSequence.cpp
@@ -2868,7 +2868,7 @@ bool IsMarkerValid(const FAnimSyncMarker* Marker, bool bLooping, const TArray<FN
 void UAnimSequence::AdvanceMarkerPhaseAsLeader(bool bLooping, float MoveDelta, const TArray<FName>& ValidMarkerNames, float& CurrentTime, FMarkerPair& PrevMarker, FMarkerPair& NextMarker, TArray<FPassedMarker>& MarkersPassed, const UMirrorDataTable* MirrorTable) const
 {
 	check(MoveDelta != 0.f);
-	const bool bPlayingForwards = MoveDelta > 0.f;
+	const bool bPlayingForwards = MoveDelta >= 0.f;
 	float CurrentMoveDelta = MoveDelta;
 
 	// Hard to reproduce issue triggering this, ensure & clamp for now
@@ -3536,7 +3536,7 @@ void UAnimSequence::MoveAttributesToModel()
 
 void UAnimSequence::AdvanceMarkerPhaseAsFollower(const FMarkerTickContext& Context, float DeltaRemaining, bool bLooping, float& CurrentTime, FMarkerPair& PreviousMarker, FMarkerPair& NextMarker, const UMirrorDataTable* MirrorTable) const
 {
-	const bool bPlayingForwards = DeltaRemaining > 0.f;
+	const bool bPlayingForwards = DeltaRemaining >= 0.f;
 
 	// Ensures the sequence's markers match the sync start position.
 	ValidateCurrentPosition(Context.GetMarkerSyncStartPosition(), bPlayingForwards, bLooping, CurrentTime, PreviousMarker, NextMarker, MirrorTable);
diff --git a/Engine/Source/Runtime/Engine/Private/Animation/AnimSequenceBase.cpp b/Engine/Source/Runtime/Engine/Private/Animation/AnimSequenceBase.cpp
--- a/Engine/Source/Runtime/Engine/Private/Animation/AnimSequenceBase.cpp
+++ b/Engine/Source/Runtime/Engine/Private/Animation/AnimSequenceBase.cpp
@@ -584,6 +584,11 @@ void UAnimSequenceBase::TickAssetPlayer(FAnimTickRecord& Instance, struct FAnimN
 				UE_LOG(LogAnimMarkerSync, Log, TEXT("Leader (%s) (normal advance)  - PreviousTime (%0.2f), CurrentTime (%0.2f), MoveDelta (%0.2f), Looping (%d) "), *GetName(), PreviousTime, CurrentTime, DeltaTime, Instance.bLooping ? 1 : 0);
 			}
 		}
+		else if (Instance.bCanUseMarkerSync && Context.CanUseMarkerPosition() && !Instance.MarkerTickRecord->IsValid(Instance.bLooping))
+		{
+			// Re-compute marker indices since the asset's tick record is invalid. Get previous and next markers.
+			GetMarkerIndicesForTime(CurrentTime, Instance.bLooping, Context.MarkerTickContext.GetValidMarkerNames(), Instance.MarkerTickRecord->PreviousMarker, Instance.MarkerTickRecord->NextMarker);
+		}
 
 		// Update context's data after ticking.
 		Context.SetAnimationPositionRatio(CurrentTime / GetPlayLength());
diff --git a/Engine/Source/Runtime/Engine/Private/Animation/BlendSpace.cpp b/Engine/Source/Runtime/Engine/Private/Animation/BlendSpace.cpp
--- a/Engine/Source/Runtime/Engine/Private/Animation/BlendSpace.cpp
+++ b/Engine/Source/Runtime/Engine/Private/Animation/BlendSpace.cpp
@@ -297,6 +297,118 @@ bool UBlendSpace::UpdateBlendSamples(const FVector& InBlendSpacePosition, float
 	return bResult;
 }
 
+void UBlendSpace::ResetBlendSamples(TArray<FBlendSampleData>& InOutSampleDataCache, float InNormalizedCurrentTime, bool bLooping, bool bMatchSyncPhases) const
+{
+	const bool bCanDoMarkerSync = SampleIndexWithMarkers != INDEX_NONE && SampleData.Num() > SampleIndexWithMarkers;
+
+	// Ensure we have a valid normalized time.
+	InNormalizedCurrentTime = bLooping ? FMath::Wrap(InNormalizedCurrentTime, 0.0f, 1.0f) : FMath::Clamp(InNormalizedCurrentTime, 0.0f, 1.0f);
+	
+	if (bCanDoMarkerSync)
+	{
+		// Query highest weighted sample with marker information. This will become the leader for all other samples to follow.
+		const int32 HighestMarkerSyncWeightIndex = FBlendSpaceUtilities::GetHighestWeightMarkerSyncSample(InOutSampleDataCache, SampleData);
+		
+		// Query leader sample information.
+		FBlendSampleData& LeaderSampleData = InOutSampleDataCache[HighestMarkerSyncWeightIndex];
+		const FBlendSample& LeaderSample = SampleData[LeaderSampleData.SampleDataIndex];
+
+		ensure(LeaderSample.Animation != nullptr);
+
+		// Leader is known at this point, build it's tick context.
+		FAnimAssetTickContext Context = { 0.0f, ERootMotionMode::NoRootMotionExtraction, true, *LeaderSampleData.Animation->GetUniqueMarkerNames()};
+		
+		// Reset leader sample to match requested time.
+		LeaderSampleData.MarkerTickRecord.Reset();
+		LeaderSampleData.PreviousTime = InNormalizedCurrentTime * LeaderSampleData.Animation->GetPlayLength();
+		LeaderSampleData.Time = InNormalizedCurrentTime * LeaderSampleData.Animation->GetPlayLength();
+
+		// Query valid marker position for leader.
+		LeaderSample.Animation->GetMarkerIndicesForTime(LeaderSampleData.Time, bLooping, Context.MarkerTickContext.GetValidMarkerNames(), LeaderSampleData.MarkerTickRecord.PreviousMarker, LeaderSampleData.MarkerTickRecord.NextMarker);
+
+		// Get leader sync start position.
+		if (bMatchSyncPhases)
+		{
+			FMarkerTickRecord StartMarkerTickRecord;
+
+			// Get sync start position.
+			LeaderSample.Animation->GetMarkerIndicesForTime(0, bLooping, Context.MarkerTickContext.GetValidMarkerNames(), StartMarkerTickRecord.PreviousMarker, StartMarkerTickRecord.NextMarker);
+			Context.MarkerTickContext.SetMarkerSyncStartPosition(LeaderSample.Animation->GetMarkerSyncPositionFromMarkerIndicies(StartMarkerTickRecord.PreviousMarker.MarkerIndex, StartMarkerTickRecord.NextMarker.MarkerIndex, 0, nullptr));
+
+			// We need to account for an extra passed marker when the LeaderSample is looping and its last marker is placed at PlayLength, since GetMarkerIndicesForTime() will give us Prev = LastMarkerIndex - 1 and Next = LastMarkerIndex when CurrentTime is 0.
+			int LeaderLastMarkerIndex = LeaderSampleData.Animation->AuthoredSyncMarkers.Num() - 1;
+			if (bLooping && StartMarkerTickRecord.NextMarker.MarkerIndex == LeaderLastMarkerIndex)
+			{
+				int32 PassedMarker = Context.MarkerTickContext.MarkersPassedThisTick.Add(FPassedMarker());
+				Context.MarkerTickContext.MarkersPassedThisTick[PassedMarker].PassedMarkerName = LeaderSampleData.Animation->AuthoredSyncMarkers[LeaderLastMarkerIndex].MarkerName;
+				Context.MarkerTickContext.MarkersPassedThisTick[PassedMarker].DeltaTimeWhenPassed = LeaderSampleData.Time;
+			}
+		}
+		else
+		{
+			Context.MarkerTickContext.SetMarkerSyncStartPosition(LeaderSample.Animation->GetMarkerSyncPositionFromMarkerIndicies(LeaderSampleData.MarkerTickRecord.PreviousMarker.MarkerIndex, LeaderSampleData.MarkerTickRecord.NextMarker.MarkerIndex, LeaderSampleData.Time, nullptr));
+		}
+
+		// Get leader sync end position.
+		Context.MarkerTickContext.SetMarkerSyncEndPosition(LeaderSample.Animation->GetMarkerSyncPositionFromMarkerIndicies(LeaderSampleData.MarkerTickRecord.PreviousMarker.MarkerIndex, LeaderSampleData.MarkerTickRecord.NextMarker.MarkerIndex, LeaderSampleData.Time, nullptr));
+
+		// Determine how many markers where passed to arrive to the sync phase the leader is currently in.
+		if (bMatchSyncPhases)
+		{
+			FMarkerSyncData SyncData;
+			
+			SyncData.AuthoredSyncMarkers = LeaderSample.Animation->AuthoredSyncMarkers;
+			SyncData.CollectMarkersInRange(0, LeaderSampleData.Time, Context.MarkerTickContext.MarkersPassedThisTick, LeaderSampleData.Time);
+		}
+		
+		// Reset follower samples.
+		for (int32 SampleIndex = 0; SampleIndex < InOutSampleDataCache.Num(); ++SampleIndex)
+		{
+			FBlendSampleData& SampleDataItem = InOutSampleDataCache[SampleIndex];
+			const FBlendSample& Sample = SampleData[SampleDataItem.SampleDataIndex];
+			
+			if (HighestMarkerSyncWeightIndex != SampleIndex && (SampleDataItem.TotalWeight > ZERO_ANIMWEIGHT_THRESH))
+			{
+				if (!Sample.Animation->AuthoredSyncMarkers.IsEmpty())
+				{
+					// Reset time.
+					SampleDataItem.PreviousTime = 0.0f;
+					SampleDataItem.Time = 0.0f;
+				
+					// Get next marker indices that matches sync start position.
+					SampleDataItem.MarkerTickRecord.Reset();
+					Sample.Animation->GetMarkerIndicesForPosition(Context.MarkerTickContext.GetMarkerSyncStartPosition(), bLooping, SampleDataItem.MarkerTickRecord.PreviousMarker, SampleDataItem.MarkerTickRecord.NextMarker, SampleDataItem.Time, nullptr);
+
+					// Ensure we advance and pass all the phases the leader passed.
+					if (bMatchSyncPhases)
+					{
+						Sample.Animation->AdvanceMarkerPhaseAsFollower(Context.MarkerTickContext, 0.0f, bLooping, SampleDataItem.Time, SampleDataItem.MarkerTickRecord.PreviousMarker, SampleDataItem.MarkerTickRecord.NextMarker, nullptr);
+					}
+				}
+				else
+				{
+					// Fallback to length based syncing so it matches default behaviour when ticking a blend space.
+					SampleDataItem.MarkerTickRecord.Reset();
+					SampleDataItem.PreviousTime = InNormalizedCurrentTime * Sample.Animation->GetPlayLength();
+					SampleDataItem.Time = InNormalizedCurrentTime * Sample.Animation->GetPlayLength();
+				}
+			}
+		}
+	}
+	else
+	{
+		// Fallback to length based syncing so it matches default behaviour when ticking a blend space.
+		for (int32 SampleIndex = 0; SampleIndex < InOutSampleDataCache.Num(); ++SampleIndex)
+		{
+			FBlendSampleData& SampleDataItem = InOutSampleDataCache[SampleIndex];
+			
+			SampleDataItem.MarkerTickRecord.Reset();
+			SampleDataItem.PreviousTime = InNormalizedCurrentTime * SampleDataItem.Animation->GetPlayLength();
+			SampleDataItem.Time = InNormalizedCurrentTime * SampleDataItem.Animation->GetPlayLength();
+		}
+	}
+}
+
 void UBlendSpace::ForEachImmutableSample(const TFunctionRef<void(const FBlendSample&)> Func) const
 {
 	for (const FBlendSample & Sample : SampleData)
@@ -500,10 +612,10 @@ void UBlendSpace::TickAssetPlayer(FAnimTickRecord& Instance, struct FAnimNotifyQ
 
 						// Only tick samples if leader sample has any delta time to consume.
 						const float NewDeltaTime = Context.GetDeltaTime() * Instance.PlayRateMultiplier * LeaderSample.RateScale * LeaderSample.Animation->RateScale;
+						Context.SetLeaderDelta(NewDeltaTime);
+
 						if (!FMath::IsNearlyZero(NewDeltaTime))
 						{
-							Context.SetLeaderDelta(NewDeltaTime);
-
 							// Tick leader sample
 							LeaderSample.Animation->TickByMarkerAsLeader(LeaderSampleData.MarkerTickRecord, Context.MarkerTickContext, LeaderSampleData.Time, LeaderSampleData.PreviousTime, NewDeltaTime, Instance.bLooping, Instance.MirrorDataTable);
 
@@ -512,6 +624,19 @@ void UBlendSpace::TickAssetPlayer(FAnimTickRecord& Instance, struct FAnimNotifyQ
 							// Tick all the follower samples
 							TickFollowerSamples(SampleDataList, HighestMarkerSyncWeightIndex, Context, bResetMarkerDataOnFollowers, Instance.bLooping, Instance.MirrorDataTable);
 						}
+						else if (!Instance.MarkerTickRecord->IsValid(Instance.bLooping))
+						{
+							// Re-compute marker indices for leader sample's tick record. Get previous and next markers.
+							LeaderSample.Animation->GetMarkerIndicesForTime(LeaderSampleData.Time, Instance.bLooping, Context.MarkerTickContext.GetValidMarkerNames(), LeaderSampleData.MarkerTickRecord.PreviousMarker, LeaderSampleData.MarkerTickRecord.NextMarker);
+
+							// Get sync position for followers to sync up to.
+							const FMarkerSyncAnimPosition SyncPosition = LeaderSample.Animation->GetMarkerSyncPositionFromMarkerIndicies(LeaderSampleData.MarkerTickRecord.PreviousMarker.MarkerIndex, LeaderSampleData.MarkerTickRecord.NextMarker.MarkerIndex, LeaderSampleData.Time, Instance.MirrorDataTable);
+							Context.MarkerTickContext.SetMarkerSyncStartPosition(SyncPosition);
+							Context.MarkerTickContext.SetMarkerSyncEndPosition(SyncPosition);
+							
+							// Make all follower samples match next sync position to equal that of the leader.
+							TickFollowerSamples(SampleDataList, HighestMarkerSyncWeightIndex, Context, true, Instance.bLooping, Instance.MirrorDataTable);
+						}
 						
 						NormalizedCurrentTime = LeaderSampleData.Time / LeaderSample.Animation->GetPlayLength();
 						*Instance.MarkerTickRecord = LeaderSampleData.MarkerTickRecord;
