World Partition - HLOD: Treat HLOD actors which were never built (or failed to build components for various reasons) as editor only.

diff --git a/Engine/Source/Runtime/Engine/Private/WorldPartition/HLOD/HLODActor.cpp b/Engine/Source/Runtime/Engine/Private/WorldPartition/HLOD/HLODActor.cpp
--- a/Engine/Source/Runtime/Engine/Private/WorldPartition/HLOD/HLODActor.cpp
+++ b/Engine/Source/Runtime/Engine/Private/WorldPartition/HLOD/HLODActor.cpp
@@ -146,6 +146,12 @@ void AWorldPartitionHLOD::Serialize(FArchive& Ar)
 #endif
 }
 
+bool AWorldPartitionHLOD::IsEditorOnly() const
+{
+	// Treat HLOD actors which were never built (or failed to build components for various reasons) as editor only.
+	return GetRootComponent() == nullptr || Super::IsEditorOnly();
+}
+
 bool AWorldPartitionHLOD::NeedsLoadForServer() const
 {
 	// Only needed on server if this HLOD actor has anything to replicate to clients
diff --git a/Engine/Source/Runtime/Engine/Public/WorldPartition/HLOD/HLODActor.h b/Engine/Source/Runtime/Engine/Public/WorldPartition/HLOD/HLODActor.h
--- a/Engine/Source/Runtime/Engine/Public/WorldPartition/HLOD/HLODActor.h
+++ b/Engine/Source/Runtime/Engine/Public/WorldPartition/HLOD/HLODActor.h
@@ -67,6 +67,7 @@ public:
 protected:
 	//~ Begin UObject Interface.
 	ENGINE_API virtual void Serialize(FArchive& Ar) override;
+	ENGINE_API virtual bool IsEditorOnly() const override;
 	ENGINE_API virtual bool NeedsLoadForServer() const override;
 	ENGINE_API virtual void PostLoad() override;
 #if WITH_EDITOR
