Fixed CIS - Static Analysis

diff --git a/Engine/Source/Runtime/Engine/Private/WorldPartition/DataLayer/ExternalDataLayerManager.cpp b/Engine/Source/Runtime/Engine/Private/WorldPartition/DataLayer/ExternalDataLayerManager.cpp
--- a/Engine/Source/Runtime/Engine/Private/WorldPartition/DataLayer/ExternalDataLayerManager.cpp
+++ b/Engine/Source/Runtime/Engine/Private/WorldPartition/DataLayer/ExternalDataLayerManager.cpp
@@ -52,6 +52,9 @@ void UExternalDataLayerManager::Initialize()
 	UWorld* OuterWorld = GetTypedOuter<UWorld>();
 	UWorld* OwningWorld = GetOuterUWorldPartition()->GetWorld();
 
+	check(OuterWorld);
+	check(OwningWorld);
+
 	// EDL in LevelInstance is not currently supported
 	// In this case, don't initialize to make sure it will do nothing.
 	const ULevelInstanceSubsystem* LevelInstanceSubsystem = OwningWorld->GetSubsystem<ULevelInstanceSubsystem>();
