Rename component templates that are no longer needed, so that creating new components with the same name does not result in a fatal name collision

diff --git a/Engine/Source/Runtime/Engine/Private/InheritableComponentHandler.cpp b/Engine/Source/Runtime/Engine/Private/InheritableComponentHandler.cpp
--- a/Engine/Source/Runtime/Engine/Private/InheritableComponentHandler.cpp
+++ b/Engine/Source/Runtime/Engine/Private/InheritableComponentHandler.cpp
@@ -255,6 +255,10 @@ void UInheritableComponentHandler::ValidateTemplates()
 					if (Record.ComponentTemplate)
 					{
 						Record.ComponentTemplate->SetFlags(RF_Transient);
+#if WITH_EDITOR
+						// in editor, move the component template aside so its name is free:
+						Record.ComponentTemplate->Rename(nullptr, GetTransientPackage(), REN_DoNotDirty | REN_DontCreateRedirectors | REN_ForceNoResetLoaders | REN_NonTransactional);
+#endif // WITH_EDITOR
 						UnnecessaryComponents.AddUnique(Record.ComponentTemplate);
 					}
 
