Expose a native way to get the 'main' anim instance in a thread safe way

Should prevent the need to nest property accesses to acheive the same effect (Lyra does this, and shouldnt have to)

diff --git a/Engine/Source/Runtime/Engine/Classes/Animation/AnimInstance.h b/Engine/Source/Runtime/Engine/Classes/Animation/AnimInstance.h
--- a/Engine/Source/Runtime/Engine/Classes/Animation/AnimInstance.h
+++ b/Engine/Source/Runtime/Engine/Classes/Animation/AnimInstance.h
@@ -506,6 +506,10 @@ public:
 	UFUNCTION(BlueprintCallable, Category = "Animation")
 	ENGINE_API USkeletalMeshComponent* GetOwningComponent() const;
 
+	/** Get the 'main' anim instance, i.e. the one that is hosted on the skeletal mesh component */
+	UFUNCTION(BlueprintCallable, Category = "Animation", meta = (DisplayName="Get Main Anim Instance", BlueprintThreadSafe))
+	ENGINE_API UAnimInstance* Blueprint_GetMainAnimInstance() const;
+
 public:
 
 	/** Executed when the Animation is initialized */
@@ -531,7 +535,7 @@ public:
 	/** Executed when the Animation Blueprint is updated on a worker thread, just prior to graph update */
 	UFUNCTION(BlueprintImplementableEvent, meta=(BlueprintThreadSafe))
 	ENGINE_API void BlueprintThreadSafeUpdateAnimation(float DeltaTime);
-	
+
 	ENGINE_API bool CanTransitionSignature() const;
 	
 	/*********************************************************************************************
diff --git a/Engine/Source/Runtime/Engine/Private/Animation/AnimInstance.cpp b/Engine/Source/Runtime/Engine/Private/Animation/AnimInstance.cpp
--- a/Engine/Source/Runtime/Engine/Private/Animation/AnimInstance.cpp
+++ b/Engine/Source/Runtime/Engine/Private/Animation/AnimInstance.cpp
@@ -225,6 +225,11 @@ USkeletalMeshComponent* UAnimInstance::GetOwningComponent() const
 	return GetSkelMeshComponent();
 }
 
+UAnimInstance* UAnimInstance::Blueprint_GetMainAnimInstance() const
+{
+	return GetSkelMeshComponent()->GetAnimInstance();
+}
+
 UWorld* UAnimInstance::GetWorld() const
 {
 	// The CDO isn't owned by a SkelMeshComponent (and doesn't have a World)
