Fixed CIS - Static Analysis

diff --git a/Engine/Source/Runtime/Engine/Private/WorldPartition/DataLayer/ExternalDataLayerManager.cpp b/Engine/Source/Runtime/Engine/Private/WorldPartition/DataLayer/ExternalDataLayerManager.cpp
--- a/Engine/Source/Runtime/Engine/Private/WorldPartition/DataLayer/ExternalDataLayerManager.cpp
+++ b/Engine/Source/Runtime/Engine/Private/WorldPartition/DataLayer/ExternalDataLayerManager.cpp
@@ -63,7 +63,7 @@ void UExternalDataLayerManager::Initialize()
 		return;
 	}
 	
-	const bool bIsInstanced = (OuterWorld && OwningWorld && (OuterWorld != OwningWorld));
+	const bool bIsInstanced = (OuterWorld != OwningWorld);
 	bIsRunningGameOrInstancedWorldPartition = IsRunningGame() || bIsInstanced;
 	bIsInitialized = true;
 
