[shaders] fix refcounting and memory leak in handling of job cache entries created from DDC callbacks by checking if an appropriate FStoredOutput object exists already before creating a new one (always addref the stored output regardless)

diff --git a/Engine/Source/Runtime/Engine/Private/ShaderCompiler/ShaderCompiler.cpp b/Engine/Source/Runtime/Engine/Private/ShaderCompiler/ShaderCompiler.cpp
--- a/Engine/Source/Runtime/Engine/Private/ShaderCompiler/ShaderCompiler.cpp
+++ b/Engine/Source/Runtime/Engine/Private/ShaderCompiler/ShaderCompiler.cpp
@@ -10799,16 +10799,14 @@ FShaderJobCacheRef FShaderJobCache::FindOrAdd(const FJobInputHash& Hash, EShader
 
 					if (Response.Status == UE::DerivedData::EStatus::Ok)
 					{
-						// Create a new entry to store in the FShaderJobCache
-						FStoredOutput* NewStoredOutput = new FStoredOutput();
-						NewStoredOutput->JobOutput = Response.Record.GetValue(ShaderJobCacheId).GetData().Decompress();
+						// Retrieve the shared buffer containing the job output and compute the associated output hash for the result retrieved from DDC
+						// If an existing duplicate of this buffer is already registered in the Outputs map, this copy will be freed at end of scope
+						FSharedBuffer JobOutput = Response.Record.GetValue(ShaderJobCacheId).GetData().Decompress();
+						FJobOutputHash OutputHash = FBlake3::HashBuffer(JobOutput.GetData(), JobOutput.GetSize());
 
-						TRACE_COUNTER_ADD(Shaders_JobCacheDDCBytesReceived, NewStoredOutput->JobOutput.GetSize());
+						TRACE_COUNTER_ADD(Shaders_JobCacheDDCBytesReceived, JobOutput.GetSize());
 						TRACE_COUNTER_INCREMENT(Shaders_JobCacheDDCHits);
 
-						// Generate an output hash
-						FJobOutputHash NewOutputHash = FBlake3::HashBuffer(NewStoredOutput->JobOutput.GetData(), NewStoredOutput->JobOutput.GetSize());
-
 						// If we are running the cache logic async (not blocking in the main thread), we need a lock before writing to the job cache.
 						// Otherwise, the lock will already be held by the main thread (and trying to lock here would just deadlock).
 						if (DerivedDataPriority != UE::DerivedData::EPriority::Blocking)
@@ -10820,8 +10818,6 @@ FShaderJobCacheRef FShaderJobCache::FindOrAdd(const FJobInputHash& Hash, EShader
 							if (!JobDataPtr->JobInFlight->PrevLink)
 							{
 								UE_LOG(LogShaderCompilers, Display, TEXT("Cancelled job 0x%p (data 0x%p) with pending DDC hit."), JobDataPtr->JobInFlight.GetReference(), JobDataPtr);
-
-								delete NewStoredOutput;
 								if (JobDataPtr->JobInFlight)
 								{
 #if WITH_EDITOR
@@ -10844,18 +10840,27 @@ FShaderJobCacheRef FShaderJobCache::FindOrAdd(const FJobInputHash& Hash, EShader
 						// Add a DDC hit
 						++TotalCacheDDCHits;
 
-						// Cache the result in the FShaderJobCache
-						NewStoredOutput->AddRef();
-						Outputs.Add(NewOutputHash, NewStoredOutput);
-						JobDataPtr->OutputHash = NewOutputHash;
-						JobDataPtr->bOutputFromDDC = true;
+						FStoredOutput** ExistingStoredOutput = Outputs.Find(OutputHash);
+						FStoredOutput* StoredOutput = ExistingStoredOutput ? *ExistingStoredOutput : nullptr;
+						if (StoredOutput == nullptr)
+						{
+							// Create a new entry to store in the FShaderJobCache if one doesn't already exist for this output hash
+							StoredOutput = new FStoredOutput();
+							StoredOutput->JobOutput = JobOutput;
+							Outputs.Add(OutputHash, StoredOutput);
+							CurrentlyAllocatedMemory += StoredOutput->GetAllocatedSize();
+						}
+
+						// Increment refcount of output whether or not we created it above
+						StoredOutput->AddRef();
 
-						CurrentlyAllocatedMemory += NewStoredOutput->GetAllocatedSize();
+						JobDataPtr->OutputHash = OutputHash;
+						JobDataPtr->bOutputFromDDC = true;
 
 						// Optionally send results back to the main thread
 						if (OutCachedOutputPtr)
 						{
-							*OutCachedOutputPtr = &NewStoredOutput->JobOutput;
+							*OutCachedOutputPtr = &StoredOutput->JobOutput;
 						}
 
 						// If non-blocking, add processed results to output.  For the blocking case, this is handled back in the main thread.
@@ -10896,7 +10901,7 @@ FShaderJobCacheRef FShaderJobCache::FindOrAdd(const FJobInputHash& Hash, EShader
 							// Call ProcessFinishedJob on main job and duplicates
 							for (FShaderCommonCompileJob* FinishedJob : FinishedJobs)
 							{
-								FMemoryReaderView MemReader(NewStoredOutput->JobOutput);
+								FMemoryReaderView MemReader(StoredOutput->JobOutput);
 								FinishedJob->SerializeOutput(MemReader);
 								ProcessFinishedJob(FinishedJob, true);
 							}
