SVT streaming data is now stored as an array of structs. This enables loading in arbitrary tiles with a single read and is another step towards streaming at tile granularity. The streaming manager still streams at mip level granularity, so for now we need to translate from a range of pages in a given mip level to a range of tiles in the streaming data. To simplify the logic in the streaming manager, we only deduplicate tiles within a single mip level, so we currently don't need to do any ref counting on the resident tiles because we effectively stream in/out all the tiles of a given mip level at the same time and all these tiles are only referenced by that mip level.

#rnx

diff --git a/Engine/Source/Runtime/Engine/Classes/SparseVolumeTexture/SparseVolumeTexture.h b/Engine/Source/Runtime/Engine/Classes/SparseVolumeTexture/SparseVolumeTexture.h
--- a/Engine/Source/Runtime/Engine/Classes/SparseVolumeTexture/SparseVolumeTexture.h
+++ b/Engine/Source/Runtime/Engine/Classes/SparseVolumeTexture/SparseVolumeTexture.h
@@ -63,24 +63,16 @@ struct FPageTopology
 		uint32 PageCount; // Number of pages in this mip
 	};
 
-	struct FPage
-	{
-		uint32 PackedPageTableCoord; // 11|11|10
-		uint32 TileIndex; // Index of the tile this page points to. Indexes into compressed tile data and is relative to each mip level (starts at 0 for every mip)
-		uint32 ParentIndex;
-	};
-
-	struct FInteriorPageData
-	{
-		uint32 ChildIndices[8];
-	};
-
 	TArray<FMip> MipInfo;
-	TArray<FPage> Pages; // One entry for every page in the sparse topology. This single array is shared across all mip levels
-	TArray<FInteriorPageData> InteriorPageData; // One entry for every interior page/node (all mips except mip0)
+	TArray<uint32> PackedPageTableCoords; // 11|11|10 packed coords of the page within the dense page table. One entry per (sparse) page.
+	TArray<uint32> TileIndices; // Index of tile data the page points to. One entry per (sparse) page.
+	TArray<uint32> ParentIndices; // Parent index. One entry per (sparse) page.
 
 	void Reset();
 	void Serialize(FArchive& Ar);
+	uint32 NumPages() const { return PackedPageTableCoords.Num(); }
+	bool IsValidPageIndex(uint32 PageIndex) const { return PackedPageTableCoords.IsValidIndex(PageIndex); }
+	void GetTileRange(uint32 PageOffset, uint32 PageCount, uint32& OutTileOffset, uint32& OutTileCount) const;
 };
 
 // Describes a mip level of a SVT frame in terms of the sizes and offsets of the data in the built bulk data.
@@ -116,6 +108,68 @@ struct FMipLevelStreamingInfo
 	int32 NumPhysicalTiles;
 };
 
+// All the tiles are essentially stored as an array of structs with each tile having different memory sections:
+// 
+//    | OccupancyBitsA | OccupancyBitsB | VoxelsA | VoxelsB |
+// 
+// OccupancyBitsA and OccupancyBitsB will only be stored if the respective texture/attributes group exists (Format is != Unknown).
+// Each set of occupancy bits has a fixed size, so it doesn't need to be stored explicitely.
+// VoxelsA and VoxelsB are the compacted non-fallback-value voxels of the tile. The sizes of these sections varies with the number of active
+// voxels in the tile (and in each texture).
+
+struct FTileInfo
+{
+	uint32 Offset;
+	uint32 Size;
+	TStaticArray<uint32, 2> OccupancyBitsOffsets; // Relative to Offset
+	TStaticArray<uint32, 2> OccupancyBitsSizes; // Relative to Offset
+	TStaticArray<uint32, 2> VoxelDataOffsets; // Relative to Offset
+	TStaticArray<uint32, 2> VoxelDataSizes;
+	TStaticArray<uint32, 2> NumVoxels;
+};
+
+// Compactly stores data to construct a FTileInfo for every stored (and compressed) tile.
+struct FTileStreamingMetaData
+{
+	// Array of byte offsets into the tile data. Has N+1 entries, with the last entry effectively being the total size of all tiles. This is a "logical" offset
+	// so we can compute the size of each tile as TileDataOffsets[N+1] - TileDataOffsets[N]. Tiles < FirstStreamingTileIndex are not actually stored with the rest of the
+	// streaming tiles, so to get the actual file offset of a streaming tile, the size of the root/non-streaming tile needs to be subtracted first. GetTileInfo() takes this into account.
+	TArray<uint32> TileDataOffsets;
+	TArray<uint16> NumVoxelsA; // Number of stored voxels for texture A for every tile. We reconstruct the number of voxels for texture B based on NumVoxelsA and the tile size.
+	uint32 FirstStreamingTileIndex; // Index of first tile that does not belong to the root mip level. Root tiles are always in memory/resident and do not stream.
+
+	void Reset()
+	{
+		TileDataOffsets.Reset();
+		NumVoxelsA.Reset();
+		FirstStreamingTileIndex = 0;
+	}
+
+	uint32 GetNumTiles() const
+	{
+		return NumVoxelsA.Num();
+	}
+
+	uint32 GetNumStreamingTiles() const
+	{
+		return GetNumTiles() - FirstStreamingTileIndex;
+	}
+
+	uint32 GetRootTileSize() const
+	{
+		return TileDataOffsets[FirstStreamingTileIndex] - TileDataOffsets[0];
+	}
+
+	uint32 GetTileMemorySize(uint32 TileIndex) const
+	{
+		return TileDataOffsets[TileIndex + 1] - TileDataOffsets[TileIndex];
+	}
+
+	FTileInfo GetTileInfo(uint32 TileIndex, uint32 FormatSizeA, uint32 FormatSizeB) const;
+
+	void GetNumVoxelsInTileRange(uint32 TileOffset, uint32 TileCount, uint32 FormatSizeA, uint32 FormatSizeB, uint32& OutNumVoxelsA, uint32& OutNumVoxelsB) const;
+};
+
 enum EResourceFlag : uint32
 {
 	EResourceFlag_StreamingDataInDDC = 1 << 0u, // FResources was cached, so MipLevelStreamingInfo can be streamed from DDC
@@ -127,8 +181,9 @@ struct FResources
 public:
 	FHeader Header;
 	uint32 ResourceFlags = 0;
-	// Info about sizes and offsets into the streamable mip level data. The last entry refers to the root mip level which is stored in RootData, not StreamableMipLevels.
-	TArray<FMipLevelStreamingInfo> MipLevelStreamingInfo;
+	int32 NumMipLevels = 0;
+	// Info about offsets into the tile data. 
+	FTileStreamingMetaData StreamingMetaData;
 	// Data for the highest/"root" mip level
 	TArray<uint8> RootData;
 	// Data for all streamable mip levels
@@ -175,13 +230,7 @@ private:
 	void EndRebuildBulkDataFromCache();
 #endif
 
-	// Stores page table into BulkData as two consecutive arrays of packed page coordinates and linear indices into the physical tiles array.
-	// Returns number of written/non-zero page table entries
-	static int32 CompressPageTable(const TArray<uint32>& PageTable, const FIntVector3& Resolution, TArray<uint8>& BulkData, TMap<uint32, uint32>& CoordToIndexMap);
-
-	static void CompressTiles(int32 NumTiles, const TArray64<uint8>& PhysicalTileDataA, const TArray64<uint8>& PhysicalTileDataB, const TArrayView<EPixelFormat>& Formats, const TArrayView<FVector4f>& FallbackValues, TArray<uint8>& BulkData, FMipLevelStreamingInfo& MipStreamingInfo);
-	
-	static FPageTopology BuildTopology(const FTextureData& InTextureData, const TArray<uint8>& InRootData, const TArray<uint8>& InStreamableBulkData, const TArray<TMap<uint32, uint32>, TInlineAllocator<16>>& InCoordToPageIndexInMipData, const TArray<FMipLevelStreamingInfo>& InMipLevelStreamingInfo);
+	static FTileStreamingMetaData CompressTiles(const FPageTopology& Topology, const struct FDerivedTextureData& DerivedTextureData, TArray<uint8>& OutRootBulkData, TArray<uint8>& OutStreamingBulkData);
 };
 
 // Encapsulates RHI resources needed to render a SparseVolumeTexture.
@@ -221,8 +270,6 @@ private:
 
 FArchive& operator<<(FArchive& Ar, UE::SVT::FHeader& Header);
 FArchive& operator<<(FArchive& Ar, UE::SVT::FPageTopology::FMip& Mip);
-FArchive& operator<<(FArchive& Ar, UE::SVT::FPageTopology::FPage& Page);
-FArchive& operator<<(FArchive& Ar, UE::SVT::FPageTopology::FInteriorPageData& Children);
 
 enum ESparseVolumeTextureShaderUniform
 {
diff --git a/Engine/Source/Runtime/Engine/Private/SparseVolumeTexture/SparseVolumeTexture.cpp b/Engine/Source/Runtime/Engine/Private/SparseVolumeTexture/SparseVolumeTexture.cpp
--- a/Engine/Source/Runtime/Engine/Private/SparseVolumeTexture/SparseVolumeTexture.cpp
+++ b/Engine/Source/Runtime/Engine/Private/SparseVolumeTexture/SparseVolumeTexture.cpp
@@ -90,23 +90,6 @@ FArchive& operator<<(FArchive& Ar, UE::SVT::FPageTopology::FMip& Mip)
 	return Ar;
 }
 
-FArchive& operator<<(FArchive& Ar, UE::SVT::FPageTopology::FPage& Page)
-{
-	Ar << Page.PackedPageTableCoord;
-	Ar << Page.TileIndex;
-	Ar << Page.ParentIndex;
-	return Ar;
-}
-
-FArchive& operator<<(FArchive& Ar, UE::SVT::FPageTopology::FInteriorPageData& InteriorPageData)
-{
-	for (uint32& ChildIndex : InteriorPageData.ChildIndices)
-	{
-		Ar << ChildIndex;
-	}
-	return Ar;
-}
-
 ////////////////////////////////////////////////////////////////////////////////////////////////
 
 namespace UE
@@ -129,15 +112,10 @@ static int32 ComputeNumMipLevels(const FIntVector3& InVirtualVolumeMin, const FI
 
 static const FString& GetDerivedDataVersion()
 {
-	static FString CachedVersionString = TEXT("4187DF14-A41A-45DD-957A-386EC727980D");	// Bump this if you want to ignore all cached data so far.
+	static FString CachedVersionString = TEXT("6C460857-3C9C-4580-98F5-63811E5EA4D5");	// Bump this if you want to ignore all cached data so far.
 	return CachedVersionString;
 }
 
-uint32 PackX11Y11Z10(uint32 X, uint32 Y, uint32 Z)
-{
-	return (X & 0x7FFu) | ((Y & 0x7FFu) << 11u) | ((Z & 0x3FFu) << 22u);
-}
-
 FHeader::FHeader(const FIntVector3& AABBMin, const FIntVector3& AABBMax, EPixelFormat FormatA, EPixelFormat FormatB, const FVector4f& FallbackValueA, const FVector4f& FallbackValueB)
 {
 	VirtualVolumeAABBMin = AABBMin;
@@ -220,15 +198,40 @@ bool FHeader::Validate(bool bPrintToLog)
 void FPageTopology::Reset()
 {
 	MipInfo.Reset();
-	Pages.Reset();
-	InteriorPageData.Reset();
+	PackedPageTableCoords.Reset();
+	TileIndices.Reset();
+	ParentIndices.Reset();
 }
 
 void FPageTopology::Serialize(FArchive& Ar)
 {
 	Ar << MipInfo;
-	Ar << Pages;
-	Ar << InteriorPageData;
+	Ar << PackedPageTableCoords;
+	Ar << TileIndices;
+	Ar << ParentIndices;
+}
+
+void FPageTopology::GetTileRange(uint32 PageOffset, uint32 PageCount, uint32& OutTileOffset, uint32& OutTileCount) const
+{
+	uint32 TileRangeMin = UINT32_MAX;
+	uint32 TileRangeMax = 0;
+	for (uint32 PageIndex = PageOffset; PageIndex < (PageOffset + PageCount); ++PageIndex)
+	{
+		check(IsValidPageIndex(PageIndex));
+		const uint32 TileIndex = TileIndices[PageIndex];
+		TileRangeMin = FMath::Min(TileIndex, TileRangeMin);
+		TileRangeMax = FMath::Max(TileIndex, TileRangeMax);
+	}
+	if (TileRangeMax >= TileRangeMin)
+	{
+		OutTileOffset = TileRangeMin;
+		OutTileCount = TileRangeMax - TileRangeMin + 1;
+	}
+	else
+	{
+		OutTileOffset = 0;
+		OutTileCount = 0;
+	}
 }
 
 ////////////////////////////////////////////////////////////////////////////////////////////////
@@ -254,7 +257,9 @@ void FResources::Serialize(FArchive& Ar, UObject* Owner, bool bCooked)
 			StoredResourceFlags = ResourceFlags;
 		}
 
-		Ar << MipLevelStreamingInfo;
+		Ar << StreamingMetaData.TileDataOffsets;
+		Ar << StreamingMetaData.NumVoxelsA;
+		Ar << StreamingMetaData.FirstStreamingTileIndex;
 		Ar << RootData;
 
 		Topology.Serialize(Ar);
@@ -280,21 +285,7 @@ void FResources::Serialize(FArchive& Ar, UObject* Owner, bool bCooked)
 
 bool FResources::HasStreamingData() const
 {
-	if (MipLevelStreamingInfo.Num() == 1)
-	{
-		// Root mip level does not stream
-		return false;
-	}
-	else
-	{
-		// It is possible for multiple mip levels to exist but all these levels are empty, so we can't just check the number of mip levels
-		bool bHasStreamingData = false;
-		for (int32 MipLevelIndex = 0; MipLevelIndex < MipLevelStreamingInfo.Num() - 1; ++MipLevelIndex)
-		{
-			bHasStreamingData = bHasStreamingData || (MipLevelStreamingInfo[MipLevelIndex].BulkSize > 0);
-		}
-		return bHasStreamingData;
-	}
+	return StreamingMetaData.GetNumStreamingTiles() > 0;
 }
 
 #if WITH_EDITORONLY_DATA
@@ -361,7 +352,7 @@ bool FResources::Build(USparseVolumeTextureFrame* Owner, UE::Serialization::FEdi
 
 		const int32 NumMipLevelsGlobal = Owner->GetNumMipLevels();
 		const bool bMoveMip0FromSource = true; // we have no need to keep SourceTextureData around
-		FTextureData DerivedTextureData;
+		FDerivedTextureData DerivedTextureData;
 		if (!SourceTextureData.BuildDerivedData(AddressingInfo, NumMipLevelsGlobal, bMoveMip0FromSource, DerivedTextureData))
 		{
 			return false;
@@ -370,11 +361,9 @@ bool FResources::Build(USparseVolumeTextureFrame* Owner, UE::Serialization::FEdi
 		// Now unload the source data
 		SourceData.UnloadData();
 
-		const int32 NumMipLevels = DerivedTextureData.MipMaps.Num();
-
 		Header = DerivedTextureData.Header;
+		NumMipLevels = DerivedTextureData.MipPageRanges.Num();
 		RootData.Reset();
-		MipLevelStreamingInfo.SetNumZeroed(NumMipLevels);
 		Topology.Reset();
 		ResourceFlags = 0;
 		ResourceName.Reset();
@@ -382,31 +371,22 @@ bool FResources::Build(USparseVolumeTextureFrame* Owner, UE::Serialization::FEdi
 		DDCRawHash.Reset();
 		DDCRebuildState.store(EDDCRebuildState::Initial);
 
-		TArray<TMap<uint32, uint32>, TInlineAllocator<16>> CoordToPageIndexInMipData;
-		CoordToPageIndexInMipData.SetNum(NumMipLevels);
-
-		TArray<uint8> StreamableBulkData;
-		for (int32 MipLevelIdx = NumMipLevels - 1; MipLevelIdx >= 0; --MipLevelIdx)
+		// Build page topology
+		Topology.Reset();
+		Topology.MipInfo.Reserve(DerivedTextureData.MipPageRanges.Num());
+		for (const FDerivedTextureData::FMipPageRange& MipRange : DerivedTextureData.MipPageRanges)
 		{
-			const FTextureData::FMipMap& Mip = DerivedTextureData.MipMaps[MipLevelIdx];
-			const bool bIsRootMipLevel = (MipLevelIdx == (NumMipLevels - 1));
-			TArray<uint8>& BulkData = bIsRootMipLevel ? RootData : StreamableBulkData;
-
-			FIntVector3 MipPageTableResolution = DerivedTextureData.Header.PageTableVolumeResolution >> MipLevelIdx;
-			MipPageTableResolution = FIntVector3(FMath::Max(1, MipPageTableResolution.X), FMath::Max(1, MipPageTableResolution.Y), FMath::Max(1, MipPageTableResolution.Z));
-
-			FMipLevelStreamingInfo& MipStreamingInfo = MipLevelStreamingInfo[MipLevelIdx];
-			MipStreamingInfo.BulkOffset = BulkData.Num();
-
-			MipStreamingInfo.PageTableOffset = BulkData.Num() - MipStreamingInfo.BulkOffset;
-			const int32 NumNonZeroPageTableEntries = CompressPageTable(Mip.PageTable, MipPageTableResolution, BulkData, CoordToPageIndexInMipData[MipLevelIdx]);
-			MipStreamingInfo.PageTableSize = NumNonZeroPageTableEntries * 2 * sizeof(uint32);
-			
-			CompressTiles(Mip.NumPhysicalTiles, Mip.PhysicalTileDataA, Mip.PhysicalTileDataB, Header.AttributesFormats, Header.FallbackValues, BulkData, MipStreamingInfo);
-
-			MipStreamingInfo.BulkSize = BulkData.Num() - MipStreamingInfo.BulkOffset;
-			MipStreamingInfo.NumPhysicalTiles = Mip.NumPhysicalTiles;
+			FPageTopology::FMip& Mip = Topology.MipInfo.AddDefaulted_GetRef();
+			Mip.PageOffset = MipRange.PageOffset;
+			Mip.PageCount = MipRange.PageCount;
 		}
+		Topology.PackedPageTableCoords = MoveTemp(DerivedTextureData.PageTableCoords);
+		Topology.TileIndices = MoveTemp(DerivedTextureData.PageTableTileIndices);
+		Topology.ParentIndices = MoveTemp(DerivedTextureData.PageTableParentIndices);
+
+		// Compress tile data
+		TArray<uint8> StreamableBulkData;
+		StreamingMetaData = CompressTiles(Topology, DerivedTextureData, RootData, StreamableBulkData);
 
 		// Store StreamableMipLevels
 		{
@@ -419,9 +399,6 @@ bool FResources::Build(USparseVolumeTextureFrame* Owner, UE::Serialization::FEdi
 			StreamableMipLevels.SetBulkDataFlags(BULKDATA_Force_NOT_InlinePayload);
 		}
 
-		// Build page topology
-		Topology = BuildTopology(DerivedTextureData, RootData, StreamableBulkData, CoordToPageIndexInMipData, MipLevelStreamingInfo);
-
 		return true;
 	}
 	return false;
@@ -509,19 +486,16 @@ void FResources::Cache(USparseVolumeTextureFrame* Owner, UE::Serialization::FEdi
 		check(bBuiltSuccessfully);
 
 		FCacheRecordBuilder RecordBuilder(CacheKey);
-		if (!MipLevelStreamingInfo.IsEmpty())
+		if (HasStreamingData())
 		{
-			if (HasStreamingData())
-			{
-				FByteBulkData& BulkData = StreamableMipLevels;
-
-				FValue Value = FValue::Compress(FSharedBuffer::MakeView(BulkData.LockReadOnly(), BulkData.GetBulkDataSize()));
-				RecordBuilder.AddValue(SVTStreamingDataId, Value);
-				BulkData.Unlock();
-				ResourceFlags |= EResourceFlag_StreamingDataInDDC;
-				DDCKeyHash = CacheKey.Hash;
-				DDCRawHash = Value.GetRawHash();
-			}
+			FByteBulkData& BulkData = StreamableMipLevels;
+
+			FValue Value = FValue::Compress(FSharedBuffer::MakeView(BulkData.LockReadOnly(), BulkData.GetBulkDataSize()));
+			RecordBuilder.AddValue(SVTStreamingDataId, Value);
+			BulkData.Unlock();
+			ResourceFlags |= EResourceFlag_StreamingDataInDDC;
+			DDCKeyHash = CacheKey.Hash;
+			DDCRawHash = Value.GetRawHash();
 		}
 
 		// Serialize to a buffer and store into DDC.
@@ -568,7 +542,8 @@ void FResources::SetDefault(EPixelFormat FormatA, EPixelFormat FormatB, const FV
 {
 	Header = FHeader(FIntVector(0, 0, 0), FIntVector(1, 1, 1), FormatA, FormatB, FallbackValueA, FallbackValueB);
 	ResourceFlags = 0;
-	MipLevelStreamingInfo.SetNumZeroed(1);
+	NumMipLevels = 1;
+	StreamingMetaData.Reset();
 	RootData.Reset();
 	StreamableMipLevels.RemoveBulkData();
 	ResourceName.Reset();
@@ -633,53 +608,11 @@ void FResources::EndRebuildBulkDataFromCache()
 
 #endif // WITH_EDITORONLY_DATA
 
-int32 FResources::CompressPageTable(const TArray<uint32>& PageTable, const FIntVector3& Resolution, TArray<uint8>& BulkData, TMap<uint32, uint32>& CoordToIndexMap)
-{
-	TRACE_CPUPROFILER_EVENT_SCOPE(SVT::FResources::Build::CompressPageTable);
-
-	int32 NumNonZeroEntries = 0;
-	for (uint32 Entry : PageTable)
-	{
-		if (Entry)
-		{
-			++NumNonZeroEntries;
-		}
-	}
-	CoordToIndexMap.Reserve(NumNonZeroEntries);
-
-	const int32 BaseOffset = BulkData.Num();
-	BulkData.SetNum(BulkData.Num() + NumNonZeroEntries * 2 * sizeof(uint32));
-	uint32* PackedCoords = reinterpret_cast<uint32*>(BulkData.GetData() + BaseOffset);
-	uint32* PageEntries = PackedCoords + NumNonZeroEntries;
-
-	int32 NumWrittenEntries = 0;
-	for (int32 Z = 0; Z < Resolution.Z; ++Z)
-	{
-		for (int32 Y = 0; Y < Resolution.Y; ++Y)
-		{
-			for (int32 X = 0; X < Resolution.X; ++X)
-			{
-				const int32 LinearCoord = (Z * Resolution.Y * Resolution.X) + (Y * Resolution.X) + X;
-				const uint32 Entry = PageTable[LinearCoord];
-				if (Entry)
-				{
-					const uint32 Packed = PackX11Y11Z10(X, Y, Z);
-					PackedCoords[NumWrittenEntries] = Packed;
-					PageEntries[NumWrittenEntries] = Entry - 1;
-					CoordToIndexMap.Add(Packed, NumWrittenEntries);
-					++NumWrittenEntries;
-				}
-			}
-		}
-	}
-
-	return NumNonZeroEntries;
-}
-
-void FResources::CompressTiles(int32 NumTiles, const TArray64<uint8>& PhysicalTileDataA, const TArray64<uint8>& PhysicalTileDataB, const TArrayView<EPixelFormat>& Formats, const TArrayView<FVector4f>& FallbackValues, TArray<uint8>& BulkData, FMipLevelStreamingInfo& MipStreamingInfo)
+FTileStreamingMetaData FResources::CompressTiles(const FPageTopology& Topology, const FDerivedTextureData& DerivedTextureData, TArray<uint8>& OutRootBulkData, TArray<uint8>& OutStreamingBulkData)
 {
 	TRACE_CPUPROFILER_EVENT_SCOPE(SVT::FResources::Build::CompressTiles);
 
+	const EPixelFormat Formats[] = { DerivedTextureData.Header.AttributesFormats[0], DerivedTextureData.Header.AttributesFormats[1] };
 	const int64 FormatSize[] = { GPixelFormats[Formats[0]].BlockBytes, GPixelFormats[Formats[1]].BlockBytes };
 	uint8 NullTileValuesU8[2][sizeof(float) * 4] = {};
 	int32 NumTextures = 0;
@@ -688,41 +621,23 @@ void FResources::CompressTiles(int32 NumTiles, const TArray64<uint8>& PhysicalTi
 		if (Formats[i] != PF_Unknown)
 		{
 			++NumTextures;
-			SVT::WriteVoxel(0, NullTileValuesU8[i], Formats[i], FallbackValues[i]);
+			SVT::WriteVoxel(0, NullTileValuesU8[i], DerivedTextureData.Header.AttributesFormats[i], DerivedTextureData.Header.FallbackValues[i]);
 		}
 	}
+	const uint8* PhysicalTileData[] = { DerivedTextureData.PhysicalTileDataA.GetData(), DerivedTextureData.PhysicalTileDataB.GetData() };
 
-	const int32 OccupancySizePerTexture = NumTiles * SVT::NumOccupancyWordsPerPaddedTile * sizeof(uint32);
-	const int32 OccupancySize = NumTextures * OccupancySizePerTexture;
-	const int32 TileDataOffsetsSize = NumTextures * NumTiles * sizeof(uint32);
-	BulkData.Reserve(BulkData.Num() + OccupancySize + TileDataOffsetsSize);
-
-	MipStreamingInfo.OccupancyBitsOffset[0] = BulkData.Num() - MipStreamingInfo.BulkOffset;
-	MipStreamingInfo.OccupancyBitsSize[0] = Formats[0] != PF_Unknown ? OccupancySizePerTexture : 0;
-	BulkData.SetNum(BulkData.Num() + MipStreamingInfo.OccupancyBitsSize[0]);
-
-	MipStreamingInfo.OccupancyBitsOffset[1] = BulkData.Num() - MipStreamingInfo.BulkOffset;
-	MipStreamingInfo.OccupancyBitsSize[1] = Formats[1] != PF_Unknown ? OccupancySizePerTexture : 0;
-	BulkData.SetNum(BulkData.Num() + MipStreamingInfo.OccupancyBitsSize[1]);
-
-	MipStreamingInfo.TileDataOffsetsOffset[0] = BulkData.Num() - MipStreamingInfo.BulkOffset;
-	MipStreamingInfo.TileDataOffsetsSize[0] = Formats[0] != PF_Unknown ? NumTiles * sizeof(uint32) : 0;
-	BulkData.SetNum(BulkData.Num() + MipStreamingInfo.TileDataOffsetsSize[0]);
-
-	MipStreamingInfo.TileDataOffsetsOffset[1] = BulkData.Num() - MipStreamingInfo.BulkOffset;
-	MipStreamingInfo.TileDataOffsetsSize[1] = Formats[1] != PF_Unknown ? NumTiles * sizeof(uint32) : 0;
-	BulkData.SetNum(BulkData.Num() + MipStreamingInfo.TileDataOffsetsSize[1]);
-
-	// Zero the bitmasks
-	FMemory::Memzero(BulkData.GetData() + MipStreamingInfo.BulkOffset + MipStreamingInfo.OccupancyBitsOffset[0], MipStreamingInfo.OccupancyBitsSize[0] + MipStreamingInfo.OccupancyBitsSize[1]);
+	const uint32 NumTiles = DerivedTextureData.NumPhysicalTiles;
+	const uint32 NumOccupancyWordsPerTexture = NumTiles * SVT::NumOccupancyWordsPerPaddedTile;
+	const uint32 NumOccupancyWords = NumTextures * NumOccupancyWordsPerTexture;
+	TArray<uint32> OccupancyBits;
+	OccupancyBits.SetNumZeroed(NumOccupancyWords);
+	uint32* OccupancyBitsPtr[2];
+	OccupancyBitsPtr[0] = FormatSize[0] ? OccupancyBits.GetData() : nullptr;
+	OccupancyBitsPtr[1] = FormatSize[1] ? (OccupancyBits.GetData() + (FormatSize[0] ? NumOccupancyWordsPerTexture : 0)) : nullptr;
 
-	uint32* OccupancyBits[2];
-	OccupancyBits[0] = reinterpret_cast<uint32*>(BulkData.GetData() + MipStreamingInfo.BulkOffset + MipStreamingInfo.OccupancyBitsOffset[0]);
-	OccupancyBits[1] = reinterpret_cast<uint32*>(BulkData.GetData() + MipStreamingInfo.BulkOffset + MipStreamingInfo.OccupancyBitsOffset[1]);
-
-	uint32* PrefixSums[2];
-	PrefixSums[0] = reinterpret_cast<uint32*>(BulkData.GetData() + MipStreamingInfo.BulkOffset + MipStreamingInfo.TileDataOffsetsOffset[0]);
-	PrefixSums[1] = reinterpret_cast<uint32*>(BulkData.GetData() + MipStreamingInfo.BulkOffset + MipStreamingInfo.TileDataOffsetsOffset[1]);
+	TStaticArray<TArray<uint16>, 2> NumVoxels;
+	NumVoxels[0].SetNumZeroed(NumTiles);
+	NumVoxels[1].SetNum((FormatSize[1] > 0) ? NumTiles : 0);
 
 	// Compute occupancy bitmasks and count number of non-fallback voxels per tile
 	ParallelFor(NumTiles, [&](int32 TileIndex)
@@ -731,63 +646,90 @@ void FResources::CompressTiles(int32 NumTiles, const TArray64<uint8>& PhysicalTi
 		{
 			if (Formats[AttributesIdx] != PF_Unknown)
 			{
-				PrefixSums[AttributesIdx][TileIndex] = 0;
+				NumVoxels[AttributesIdx][TileIndex] = 0;
 				for (int64 VoxelIndex = 0; VoxelIndex < SVT::NumVoxelsPerPaddedTile; ++VoxelIndex)
 				{
-					const uint8* Src = (AttributesIdx == 0 ? PhysicalTileDataA.GetData() : PhysicalTileDataB.GetData()) + FormatSize[AttributesIdx] * (TileIndex * SVT::NumVoxelsPerPaddedTile + VoxelIndex);
+					const uint8* Src = PhysicalTileData[AttributesIdx] + FormatSize[AttributesIdx] * (TileIndex * SVT::NumVoxelsPerPaddedTile + VoxelIndex);
 					bool bIsFallbackValue = FMemory::Memcmp(Src, NullTileValuesU8[AttributesIdx], FormatSize[AttributesIdx]) == 0;
 					if (!bIsFallbackValue)
 					{
 						const int64 WordIndex = TileIndex * SVT::NumOccupancyWordsPerPaddedTile + (VoxelIndex / 32);
-						OccupancyBits[AttributesIdx][WordIndex] |= 1u << (static_cast<uint32>(VoxelIndex) % 32u);
-						PrefixSums[AttributesIdx][TileIndex]++;
+						OccupancyBitsPtr[AttributesIdx][WordIndex] |= 1u << (static_cast<uint32>(VoxelIndex) % 32u);
+						NumVoxels[AttributesIdx][TileIndex]++;
 					}
 				}
 			}
 		}
 	});
 
-	// Compute actual per-tile voxel data offsets (prefix sum)
-	uint32 Sums[2] = {};
-	for (int32 AttributesIdx = 0; AttributesIdx < 2; ++AttributesIdx)
+	TArray<uint32> TileDataOffsets;
+	TileDataOffsets.SetNum(NumTiles + 1);
+
+	// Compute actual tile data offsets. Each tile stores data like this: | OccupancyBitsA | OccupancyBitsB | VoxelsA | VoxelsB |
+	// OccupancyBits A and B always have the same size, but VoxelsA and VoxelsB can vary
+	uint32 CurrentOffset = 0;
+	for (uint32 TileIndex = 0; TileIndex < NumTiles; ++TileIndex)
 	{
-		if (Formats[AttributesIdx] != PF_Unknown)
+		TileDataOffsets[TileIndex] = CurrentOffset;
+		CurrentOffset += SVT::NumOccupancyWordsPerPaddedTile * sizeof(uint32) * NumTextures;
+		for (int32 AttributesIdx = 0; AttributesIdx < 2; ++AttributesIdx)
 		{
-			for (int32 TileIndex = 0; TileIndex < NumTiles; ++TileIndex)
+			if (Formats[AttributesIdx] != PF_Unknown)
 			{
-				uint32 Tmp = PrefixSums[AttributesIdx][TileIndex];
-				PrefixSums[AttributesIdx][TileIndex] = Sums[AttributesIdx];
-				Sums[AttributesIdx] += Tmp;
+				CurrentOffset += NumVoxels[AttributesIdx][TileIndex] * FormatSize[AttributesIdx];
 			}
 		}
 	}
 
-	// Allocate compacted tile data memory
-	MipStreamingInfo.TileDataOffset[0] = BulkData.Num() - MipStreamingInfo.BulkOffset;
-	MipStreamingInfo.TileDataSize[0] = Sums[0] * FormatSize[0];
-
-	MipStreamingInfo.TileDataOffset[1] = MipStreamingInfo.TileDataOffset[0] + MipStreamingInfo.TileDataSize[0];
-	MipStreamingInfo.TileDataSize[1] = Sums[1] * FormatSize[1];
+	// Write final size at the end of the array so we can compute individual tile sizes as (Offsets[N+1] - Offsets[N])
+	TileDataOffsets[NumTiles] = CurrentOffset;
 
-	BulkData.SetNum(BulkData.Num() + MipStreamingInfo.TileDataSize[0] + MipStreamingInfo.TileDataSize[1]);
+	// Reuse the allocations
+	FTileStreamingMetaData StreamingMetaData;
+	StreamingMetaData.TileDataOffsets = MoveTemp(TileDataOffsets);
+	StreamingMetaData.NumVoxelsA = MoveTemp(NumVoxels[0]);
+	StreamingMetaData.FirstStreamingTileIndex = Topology.MipInfo.Last().PageOffset + Topology.MipInfo.Last().PageCount;
 
-	// All above pointers into the BulkData are now invalid!
+	// We're assuming that the very first tile belongs to the root mip and that there is only a maximum of one such tile
+	check(Topology.MipInfo.Last().PageOffset == 0);
+	check(Topology.MipInfo.Last().PageCount <= 1);
+	check(Topology.MipInfo.Last().PageCount == 0 || Topology.TileIndices[0] == 0);
 
-	uint8* RelativeBulkDataPtr = BulkData.GetData() + MipStreamingInfo.BulkOffset;
+	// Allocate memory for all tiles
+	const uint32 RootTileSize = StreamingMetaData.GetRootTileSize();
+	OutRootBulkData.SetNum(RootTileSize);
+	OutStreamingBulkData.SetNum(CurrentOffset - RootTileSize);
 
-	uint8* DstTileData[2];
-	DstTileData[0] = RelativeBulkDataPtr + MipStreamingInfo.TileDataOffset[0];
-	DstTileData[1] = RelativeBulkDataPtr + MipStreamingInfo.TileDataOffset[1];
+	uint8* RootTileData = OutRootBulkData.GetData();
+	uint8* StreamingTileData = OutStreamingBulkData.GetData();
 
-	// Copy voxels to compacted locations
-	ParallelFor(NumTiles, [&](int32 TileIndex)
+	// Copy tile data to compacted locations
+	ParallelFor(NumTiles, [&](uint32 TileIndex)
 	{
+		const bool bStreamingTile = TileIndex >= StreamingMetaData.FirstStreamingTileIndex;
+		uint8* WritePtr = bStreamingTile ? (StreamingTileData + StreamingMetaData.TileDataOffsets[TileIndex] - RootTileSize) : RootTileData;
+
+		const FTileInfo TileInfo = StreamingMetaData.GetTileInfo(TileIndex, FormatSize[0], FormatSize[1]);
+		const uint8* BaseWritePtr = WritePtr;
+
+		// Write occupancy bits
+		for (int32 AttributesIdx = 0; AttributesIdx < 2; ++AttributesIdx)
+		{
+			if (Formats[AttributesIdx] != PF_Unknown)
+			{
+				check(WritePtr == (BaseWritePtr + TileInfo.OccupancyBitsOffsets[AttributesIdx]));
+				FMemory::Memcpy(WritePtr, OccupancyBitsPtr[AttributesIdx] + TileIndex * SVT::NumOccupancyWordsPerPaddedTile, SVT::NumOccupancyWordsPerPaddedTile * sizeof(uint32));
+				WritePtr += SVT::NumOccupancyWordsPerPaddedTile * sizeof(uint32);
+			}
+		}
+
+		// Write voxel data
 		for (int32 AttributesIdx = 0; AttributesIdx < 2; ++AttributesIdx)
 		{
 			if (Formats[AttributesIdx] != PF_Unknown)
 			{
-				uint32 VoxelDataWriteOffset = reinterpret_cast<uint32*>(RelativeBulkDataPtr + MipStreamingInfo.TileDataOffsetsOffset[AttributesIdx])[TileIndex];
-				const uint32* TileOccupancyBits = reinterpret_cast<uint32*>(RelativeBulkDataPtr + MipStreamingInfo.OccupancyBitsOffset[AttributesIdx] + TileIndex * SVT::NumOccupancyWordsPerPaddedTile * sizeof(uint32));
+				check(WritePtr == (BaseWritePtr + TileInfo.VoxelDataOffsets[AttributesIdx]));
+				const uint32* TileOccupancyBits = OccupancyBitsPtr[AttributesIdx] + TileIndex * SVT::NumOccupancyWordsPerPaddedTile;
 
 				for (int64 VoxelIndex = 0; VoxelIndex < SVT::NumVoxelsPerPaddedTile; ++VoxelIndex)
 				{
@@ -795,131 +737,17 @@ void FResources::CompressTiles(int32 NumTiles, const TArray64<uint8>& PhysicalTi
 
 					if (TileOccupancyBits[WordIndex] & (1u << (static_cast<uint32>(VoxelIndex) % 32u)))
 					{
-						uint8* Dst = DstTileData[AttributesIdx] + VoxelDataWriteOffset * FormatSize[AttributesIdx];
-						const uint8* Src = (AttributesIdx == 0 ? PhysicalTileDataA.GetData() : PhysicalTileDataB.GetData()) + FormatSize[AttributesIdx] * (TileIndex * SVT::NumVoxelsPerPaddedTile + VoxelIndex);
-						FMemory::Memcpy(Dst, Src, FormatSize[AttributesIdx]);
-						++VoxelDataWriteOffset;
+						const uint8* Src = PhysicalTileData[AttributesIdx] + FormatSize[AttributesIdx] * (TileIndex * SVT::NumVoxelsPerPaddedTile + VoxelIndex);
+						FMemory::Memcpy(WritePtr, Src, FormatSize[AttributesIdx]);
+						WritePtr += FormatSize[AttributesIdx];
 					}
 				}
 			}
 		}
+		check((WritePtr - BaseWritePtr) == TileInfo.Size);
 	});
-}
-
-FPageTopology FResources::BuildTopology(const FTextureData& InTextureData, const TArray<uint8>& InRootData, const TArray<uint8>& InStreamableBulkData, const TArray<TMap<uint32, uint32>, TInlineAllocator<16>>& InCoordToPageIndexInMipData, const TArray<FMipLevelStreamingInfo>& InMipLevelStreamingInfo)
-{
-	FPageTopology Topology;
-
-	const int32 NumMipLevels = InMipLevelStreamingInfo.Num();
-
-	// Reserve memory and initialize the page offsets and counts for each mip level
-	{
-		Topology.MipInfo.SetNum(NumMipLevels);
-		uint32 NumInteriorPages = 0;
-		uint32 NumPages = 0;
-		for (int32 MipLevel = NumMipLevels - 1; MipLevel >= 0; --MipLevel)
-		{
-			// Page table bulk data stores 2 entries per non-zero page table entry
-			const uint32 NumPagesMip = InMipLevelStreamingInfo[MipLevel].PageTableSize / (2 * sizeof(uint32));
-			Topology.MipInfo[MipLevel].PageOffset = NumPages;
-			Topology.MipInfo[MipLevel].PageCount = NumPagesMip;
-			NumPages += NumPagesMip;
-			if (MipLevel > 0)
-			{
-				NumInteriorPages += NumPagesMip;
-			}
-		}
-		Topology.Pages.SetNumUninitialized(NumPages);
-		Topology.InteriorPageData.SetNumUninitialized(NumInteriorPages);
-	}
-
-	// Recursively build the page mip hierarchy. The order of nodes is identical to how they're stored in bulk data.
-	auto BuildHierarchy = [&](const FIntVector3& InPageTableMipResolution, int32 InMipLevel, int32 InPageX, int32 InPageY, int32 InPageZ, uint32 InParentIndex, auto& InRecursionLambda) -> uint32
-	{
-		// With non-cubic virtual volumes it's possible that the parent node doesn't have a cubic (square in 2D) footprint of child nodes, so early out if that's the case.
-		if (!IsInBounds(FIntVector3(InPageX, InPageY, InPageZ), FIntVector3::ZeroValue, InPageTableMipResolution))
-		{
-			return INDEX_NONE;
-		}
-
-		// Load entry from dense page table. 
-		const int32 LinearPageIndex = InPageZ * (InPageTableMipResolution.X * InPageTableMipResolution.Y) + (InPageY * InPageTableMipResolution.X) + InPageX;
-		const uint32 PageTableEntry = InTextureData.MipMaps[InMipLevel].PageTable[LinearPageIndex];
-		if (!PageTableEntry)
-		{
-			return INDEX_NONE;
-		}
-
-		FPageTopology::FPage Page;
-		Page.PackedPageTableCoord = PackX11Y11Z10(InPageX, InPageY, InPageZ);
-		Page.TileIndex = PageTableEntry - 1;
-		Page.ParentIndex = InParentIndex;
-
-		// Look up index of page in sparse list of non-zero pages
-		const uint32 PageIndex = InCoordToPageIndexInMipData[InMipLevel][Page.PackedPageTableCoord] + Topology.MipInfo[InMipLevel].PageOffset;
-
-		Topology.Pages[PageIndex] = Page;
-
-		// Process any potential children
-		if (InMipLevel > 0)
-		{
-			const int32 ChildMipLevel = InMipLevel - 1;
-			const FIntVector3 ChildPageTableMipResolution = ShiftRightAndMax(InTextureData.Header.PageTableVolumeResolution, ChildMipLevel, 1);
-
-			FPageTopology::FInteriorPageData InteriorPageData;
-			for (uint32 ChildIndex = 0; ChildIndex < 8; ++ChildIndex)
-			{
-				const int32 ChildPageX = InPageX * 2 + int32((ChildIndex >> 0u) & 1u);
-				const int32 ChildPageY = InPageY * 2 + int32((ChildIndex >> 1u) & 1u);
-				const int32 ChildPageZ = InPageZ * 2 + int32((ChildIndex >> 2u) & 1u);
-				const uint32 ChildPageIndex = InRecursionLambda(ChildPageTableMipResolution, ChildMipLevel, ChildPageX, ChildPageY, ChildPageZ, PageIndex, InRecursionLambda);
-				InteriorPageData.ChildIndices[ChildIndex] = ChildPageIndex;
-			}
-
-			Topology.InteriorPageData[PageIndex] = InteriorPageData;
-		}
-
-		return PageIndex;
-	};
-
-	// Recursively build the hierarchy, starting at the root mip.
-	const int32 NumPagesInLastMip = InTextureData.MipMaps.Last().PageTable.Num();
-	check(NumPagesInLastMip == 1);
-	const FIntVector3 RootMipPageTableResolution = FIntVector3(1, 1, 1);
-	BuildHierarchy(RootMipPageTableResolution, NumMipLevels - 1, 0, 0, 0, INDEX_NONE, BuildHierarchy);
-
-	// Validate
-	{
-#if DO_CHECK
-		for (int32 MipLevel = NumMipLevels - 1; MipLevel >= 0; --MipLevel)
-		{
-			const int32 NumPageTableEntries = (InMipLevelStreamingInfo[MipLevel].PageTableSize / (sizeof(uint32) * 2));
-			check(Topology.MipInfo[MipLevel].PageCount == NumPageTableEntries); // PageCount in topology must match number of pages in compressed bulk data
-			const bool bIsHighestMipLevel = MipLevel == (NumMipLevels - 1);
-			check(!bIsHighestMipLevel || NumPageTableEntries <= 1); // Highest mip has a maximum of 1 page
-			const TArray<uint8>& BulkData = bIsHighestMipLevel ? InRootData : InStreamableBulkData;
-			const uint32* PackedCoords = reinterpret_cast<const uint32*>(BulkData.GetData() + InMipLevelStreamingInfo[MipLevel].BulkOffset + InMipLevelStreamingInfo[MipLevel].PageTableOffset);
-			const uint32* PageEntries = PackedCoords + NumPageTableEntries;
-
-			const FIntVector3 PageTableRes = ShiftRightAndMax(InTextureData.Header.PageTableVolumeResolution, MipLevel, 1);
-
-			for (int32 i = 0; i < NumPageTableEntries; ++i)
-			{
-				const int32 PageIndex = Topology.MipInfo[MipLevel].PageOffset + i;
-				check(PackedCoords[i] == Topology.Pages[PageIndex].PackedPageTableCoord); // Topology stores pages in the same order as in the compressed bulk data
-				check(PageEntries[i] == Topology.Pages[PageIndex].TileIndex);
-				const int32 PageX = (PackedCoords[i] >> 0u) & 0x7FFu;
-				const int32 PageY = (PackedCoords[i] >> 11u) & 0x7FFu;
-				const int32 PageZ = (PackedCoords[i] >> 22u) & 0x3FFu;
-				const int32 LinearPageIndex = (PageZ * (PageTableRes.X * PageTableRes.Y)) + (PageY * PageTableRes.X) + PageX;
-				const uint32 PageTableEntry = InTextureData.MipMaps[MipLevel].PageTable[LinearPageIndex];
-				check(PageTableEntry != 0 && PageTableEntry != INDEX_NONE); // Unpacked page coord must map to a valid entry in the dense page table
-			}
-		}
-#endif // DO_CHECK
-	}
-
-	return Topology;
+	
+	return StreamingMetaData;
 }
 
 ////////////////////////////////////////////////////////////////////////////////////////////////
@@ -928,12 +756,6 @@ void FTextureRenderResources::GetPackedUniforms(FUintVector4& OutPacked0, FUintV
 {
 	check(IsInParallelRenderingThread());
 
-	auto AsUint = [](float X)
-	{
-		union { float F; uint32 U; } FU = { X };
-		return FU.U;
-	};
-
 	const FIntVector3 PageTableOffset = Header.PageTableVolumeAABBMin;
 	const FVector3f TileDataTexelSize = FVector3f(
 		1.0f / TileDataTextureResolution.X,
@@ -941,13 +763,13 @@ void FTextureRenderResources::GetPackedUniforms(FUintVector4& OutPacked0, FUintV
 		1.0f / TileDataTextureResolution.Z);
 	const FVector3f VolumePageResolution = FVector3f(GlobalVolumeResolution) / SPARSE_VOLUME_TILE_RES;
 
-	OutPacked0.X = AsUint(VolumePageResolution.X);
-	OutPacked0.Y = AsUint(VolumePageResolution.Y);
-	OutPacked0.Z = AsUint(VolumePageResolution.Z);
-	OutPacked0.W = PackX11Y11Z10(PageTableOffset.X, PageTableOffset.Y, PageTableOffset.Z);
-	OutPacked1.X = AsUint(TileDataTexelSize.X);
-	OutPacked1.Y = AsUint(TileDataTexelSize.Y);
-	OutPacked1.Z = AsUint(TileDataTexelSize.Z);
+	OutPacked0.X = FMath::AsUInt(VolumePageResolution.X);
+	OutPacked0.Y = FMath::AsUInt(VolumePageResolution.Y);
+	OutPacked0.Z = FMath::AsUInt(VolumePageResolution.Z);
+	OutPacked0.W = SVT::PackX11Y11Z10(PageTableOffset);
+	OutPacked1.X = FMath::AsUInt(TileDataTexelSize.X);
+	OutPacked1.Y = FMath::AsUInt(TileDataTexelSize.Y);
+	OutPacked1.Z = FMath::AsUInt(TileDataTexelSize.Z);
 	OutPacked1.W = 0;
 	OutPacked1.W |= (uint32)((FrameIndex & 0xFFFF) << 0);
 	OutPacked1.W |= (uint32)(((NumLogicalMipLevels - 1) & 0xFFFF) << 16);
@@ -976,6 +798,61 @@ void FTextureRenderResources::ReleaseRHI()
 	PhysicalTileDataBTextureReferenceRHI.SafeRelease();
 }
 
+FTileInfo FTileStreamingMetaData::GetTileInfo(uint32 TileIndex, uint32 FormatSizeA, uint32 FormatSizeB) const
+{
+	const uint32 NumVoxelsInA = NumVoxelsA[TileIndex];
+	const uint32 VoxelDataSizeA = NumVoxelsInA * FormatSizeA;
+
+	FTileInfo Result;
+	Result.Offset = TileDataOffsets[TileIndex];
+	Result.Size = TileDataOffsets[TileIndex + 1] - Result.Offset;
+	Result.OccupancyBitsOffsets[0] = 0;
+	Result.OccupancyBitsOffsets[1] = FormatSizeA > 0 ? SVT::OccupancyBitsSizePerPaddedTile : 0;
+	Result.OccupancyBitsSizes[0] = FormatSizeA > 0 ? SVT::OccupancyBitsSizePerPaddedTile : 0;
+	Result.OccupancyBitsSizes[1] = FormatSizeB > 0 ? SVT::OccupancyBitsSizePerPaddedTile : 0;
+	Result.VoxelDataOffsets[0] = Result.OccupancyBitsOffsets[1] + (FormatSizeB > 0 ? SVT::OccupancyBitsSizePerPaddedTile : 0);
+	Result.VoxelDataOffsets[1] = Result.VoxelDataOffsets[0] + VoxelDataSizeA;
+	Result.VoxelDataSizes[0] = VoxelDataSizeA;
+	Result.VoxelDataSizes[1] = Result.Size - Result.VoxelDataOffsets[1];
+	Result.NumVoxels[0] = NumVoxelsInA;
+	Result.NumVoxels[1] = FormatSizeB > 0 ? (Result.VoxelDataSizes[1] / FormatSizeB) : 0;
+	
+	if (TileIndex >= FirstStreamingTileIndex)
+	{
+		Result.Offset -= TileDataOffsets[FirstStreamingTileIndex];
+	}
+
+	check((Result.OccupancyBitsSizes[0] + Result.OccupancyBitsSizes[1] + Result.VoxelDataSizes[0] + Result.VoxelDataSizes[1]) == Result.Size);
+	check((Result.OccupancyBitsOffsets[0] + Result.OccupancyBitsSizes[0]) == Result.OccupancyBitsOffsets[1]);
+	check((Result.OccupancyBitsOffsets[1] + Result.OccupancyBitsSizes[1]) == Result.VoxelDataOffsets[0]);
+	check((Result.VoxelDataOffsets[0] + Result.VoxelDataSizes[0]) == Result.VoxelDataOffsets[1]);
+	check((Result.VoxelDataOffsets[1] + Result.VoxelDataSizes[1]) == Result.Size);
+
+	return Result;
+}
+
+void FTileStreamingMetaData::GetNumVoxelsInTileRange(uint32 TileOffset, uint32 TileCount, uint32 FormatSizeA, uint32 FormatSizeB, uint32& OutNumVoxelsA, uint32& OutNumVoxelsB) const
+{
+	check(FormatSizeA > 0 || FormatSizeB > 0);
+	const uint32 NumTextures = (FormatSizeA > 0 && FormatSizeB > 0) ? 2 : 1;
+	const uint32 TileOccupancyBitsSize = NumTextures * SVT::OccupancyBitsSizePerPaddedTile;
+	OutNumVoxelsA = 0;
+	OutNumVoxelsB = 0;
+	for (uint32 TileIndex = TileOffset; TileIndex < (TileOffset + TileCount); ++TileIndex)
+	{
+		const uint32 NumVoxelsATmp = NumVoxelsA[TileIndex];
+		OutNumVoxelsA += NumVoxelsATmp;
+		// For NumVoxelsB, we need to reconstruct that value based on the total tile size and the sizes of the other memory sections in the tile
+		if (FormatSizeB > 0)
+		{
+			const uint32 TileSize = GetTileMemorySize(TileIndex);
+			const uint32 VoxelsDataOffsetB = TileOccupancyBitsSize + NumVoxelsATmp * FormatSizeA;
+			check(TileSize >= VoxelsDataOffsetB);
+			OutNumVoxelsB += (TileSize - VoxelsDataOffsetB) / FormatSizeB;
+		}
+	}
+}
+
 }
 }
 
diff --git a/Engine/Source/Runtime/Engine/Private/SparseVolumeTexture/SparseVolumeTextureData.cpp b/Engine/Source/Runtime/Engine/Private/SparseVolumeTexture/SparseVolumeTextureData.cpp
--- a/Engine/Source/Runtime/Engine/Private/SparseVolumeTexture/SparseVolumeTextureData.cpp
+++ b/Engine/Source/Runtime/Engine/Private/SparseVolumeTexture/SparseVolumeTextureData.cpp
@@ -606,142 +606,40 @@ bool FTextureData::GenerateBorderVoxels(const FTextureDataAddressingInfo& Addres
 	return true;
 }
 
-bool FTextureData::DeduplicateTiles()
+bool FTextureData::BuildDerivedData(const FTextureDataAddressingInfo& AddressingInfo, int32 NumMipLevelsGlobal, bool bMoveMip0FromThis, FDerivedTextureData& OutDerivedData)
 {
-	const int32 NumMipLevels = MipMaps.Num();
-	const int32 FormatSize[] = { GPixelFormats[Header.AttributesFormats[0]].BlockBytes, GPixelFormats[Header.AttributesFormats[1]].BlockBytes };
-	const int32 TileByteSize[] = { SVT::NumVoxelsPerPaddedTile * FormatSize[0], SVT::NumVoxelsPerPaddedTile * FormatSize[1] };
-
-	int32 NumProcessedTiles = 0;
-	int32 NumSurvivingTiles = 0;
-
-	TArray<uint64> TileHashes; // Reused between mip level iterations
-	TArray<uint32> TileRemap; // Reused between mip level iterations
-
-	// Deduplicate each mip level individually
-	for (int32 MipLevel = 0; MipLevel < NumMipLevels; ++MipLevel)
-	{
-		TileHashes.Reset(MipMaps[MipLevel].NumPhysicalTiles);
-		TileRemap.Reset(MipMaps[MipLevel].NumPhysicalTiles);
-		uint8* PhysicalTileData[] = { MipMaps[MipLevel].PhysicalTileDataA.GetData(), MipMaps[MipLevel].PhysicalTileDataB.GetData() };
-		int32 TileWritePos = 0;
-
-		// Compute a hash per tile and compare against previous hashes.
-		for (int32 TileIdx = 0; TileIdx < MipMaps[MipLevel].NumPhysicalTiles; ++TileIdx)
-		{
-			uint64 Hash = 0;
-			if (FormatSize[0])
-			{
-				Hash = CityHash64((const char*)PhysicalTileData[0] + (SIZE_T)TileIdx * (SIZE_T)TileByteSize[0], (uint32)TileByteSize[0]);
-			}
-			if (FormatSize[1])
-			{
-				Hash = CityHash64WithSeed((const char*)PhysicalTileData[1] + (SIZE_T)TileIdx * (SIZE_T)TileByteSize[1], (uint32)TileByteSize[1], Hash);
-			}
-
-			const int32 OldNumUniqueTiles = TileHashes.Num();
-			const uint32 RemapIndex = TileHashes.AddUnique(Hash) + 1; // +1 because of implicit null tile at 0
-			TileRemap.Add(RemapIndex);
-			
-			const bool bAddedNewUniqueTile = TileHashes.Num() > OldNumUniqueTiles;
-			
-			// Move tile to compacted position
-			if (bAddedNewUniqueTile)
-			{
-				if (TileWritePos < TileIdx)
-				{
-					if (FormatSize[0])
-					{
-						FMemory::Memcpy(PhysicalTileData[0] + (SIZE_T)TileWritePos * (SIZE_T)TileByteSize[0], PhysicalTileData[0] + (SIZE_T)TileIdx * (SIZE_T)TileByteSize[0], (SIZE_T)TileByteSize[0]);
-					}
-					if (FormatSize[1])
-					{
-						FMemory::Memcpy(PhysicalTileData[1] + (SIZE_T)TileWritePos * (SIZE_T)TileByteSize[1], PhysicalTileData[1] + (SIZE_T)TileIdx * (SIZE_T)TileByteSize[1], (SIZE_T)TileByteSize[1]);
-					}
-				}
-				else
-				{
-					check(TileWritePos == TileIdx);
-				}
-				++TileWritePos;
-				++NumSurvivingTiles;
-			}
-
-			++NumProcessedTiles;
-		}
-
-		// Update mip info
-		MipMaps[MipLevel].NumPhysicalTiles = TileWritePos;
-
-		// Fix up page table
-		for (uint32& PageTableEntry : MipMaps[MipLevel].PageTable)
-		{
-			if (PageTableEntry != 0)
-			{
-				PageTableEntry = TileRemap[PageTableEntry - 1]; // -1 because of implicit null tile at 0
-			}
-		}
-
-		const int64 DeduplicatedTileSizeBytes[] = { (int64)TileWritePos * (int64)TileByteSize[0], (int64)TileWritePos * (int64)TileByteSize[1] };
-		if (DeduplicatedTileSizeBytes[0] > SVT::MaxResourceSize || DeduplicatedTileSizeBytes[1] > SVT::MaxResourceSize)
-		{
-			UE_LOG(LogSparseVolumeTextureData, Warning, TEXT("SparseVolumeTexture still requires too much memory (> 2048MB) after tile deduplication! The full mip chain will not fit in memory, so only lower resolution mip levels can be streamed in. Physical tile data A: %lld bytes. Physical tile data B: %lld bytes"),
-				(long long)DeduplicatedTileSizeBytes[0], (long long)DeduplicatedTileSizeBytes[1]);
-		}
-
-		if (!MipMaps[MipLevel].PhysicalTileDataA.IsEmpty())
-		{
-			MipMaps[MipLevel].PhysicalTileDataA.SetNum(DeduplicatedTileSizeBytes[0]);
-		}
-		if (!MipMaps[MipLevel].PhysicalTileDataB.IsEmpty())
-		{
-			MipMaps[MipLevel].PhysicalTileDataB.SetNum(DeduplicatedTileSizeBytes[1]);
-		}
-	}
-
-#if 0 // Disabled in order to not spam the log when importing long sequences
-	const int32 NumRemovedTiles = NumProcessedTiles - NumSurvivingTiles;
-	const float RemovedTilesPercentage = ((float)NumRemovedTiles / (float)NumProcessedTiles) * 100.0f;
-	UE_LOG(LogSparseVolumeTextureData, Display, TEXT("SparseVolumeTexture tile deduplication removed %i tiles out of %i (%f%%)"), NumRemovedTiles, NumProcessedTiles, RemovedTilesPercentage);
-#endif
-
-	return true;
-}
-
-bool FTextureData::BuildDerivedData(const FTextureDataAddressingInfo& AddressingInfo, int32 NumMipLevelsGlobal, bool bMoveMip0FromThis, FTextureData& OutDerivedData)
-{
-	OutDerivedData = FTextureData{};
-	OutDerivedData.Header = Header;
-	OutDerivedData.Header.UpdatePageTableFromGlobalNumMipLevels(NumMipLevelsGlobal);
-	if (!OutDerivedData.Header.Validate(true /*bPrintToLog*/))
+	FTextureData MippedTextureData{};
+	MippedTextureData.Header = Header;
+	MippedTextureData.Header.UpdatePageTableFromGlobalNumMipLevels(NumMipLevelsGlobal);
+	if (!MippedTextureData.Header.Validate(true /*bPrintToLog*/))
 	{
 		return false;
 	}
-	OutDerivedData.MipMaps.SetNum(1);
+	MippedTextureData.MipMaps.SetNum(1);
 	if (bMoveMip0FromThis)
 	{
-		OutDerivedData.MipMaps[0] = MoveTemp(MipMaps[0]);
+		MippedTextureData.MipMaps[0] = MoveTemp(MipMaps[0]);
 		MipMaps[0] = {};
 	}
 	else
 	{
-		OutDerivedData.MipMaps[0] = MipMaps[0];
+		MippedTextureData.MipMaps[0] = MipMaps[0];
 	}
 
 	// Page table changed due to UpdatePageTableFromGlobalNumMipLevels(), so we need to recreate it based on the old one
-	if (OutDerivedData.Header.PageTableVolumeAABBMin != Header.PageTableVolumeAABBMin
-		|| OutDerivedData.Header.PageTableVolumeAABBMax != Header.PageTableVolumeAABBMax
-		|| OutDerivedData.Header.PageTableVolumeResolution != Header.PageTableVolumeResolution)
+	if (MippedTextureData.Header.PageTableVolumeAABBMin != Header.PageTableVolumeAABBMin
+		|| MippedTextureData.Header.PageTableVolumeAABBMax != Header.PageTableVolumeAABBMax
+		|| MippedTextureData.Header.PageTableVolumeResolution != Header.PageTableVolumeResolution)
 	{
 		TArray<uint32> PageTableNew;
-		PageTableNew.SetNumZeroed(OutDerivedData.Header.PageTableVolumeResolution.X * OutDerivedData.Header.PageTableVolumeResolution.Y * OutDerivedData.Header.PageTableVolumeResolution.Z);
+		PageTableNew.SetNumZeroed(MippedTextureData.Header.PageTableVolumeResolution.X * MippedTextureData.Header.PageTableVolumeResolution.Y * MippedTextureData.Header.PageTableVolumeResolution.Z);
 
 		// Iterate over new page table and copy entries from old one
-		for (int32 Z = OutDerivedData.Header.PageTableVolumeAABBMin.Z; Z < OutDerivedData.Header.PageTableVolumeAABBMax.Z; ++Z)
+		for (int32 Z = MippedTextureData.Header.PageTableVolumeAABBMin.Z; Z < MippedTextureData.Header.PageTableVolumeAABBMax.Z; ++Z)
 		{
-			for (int32 Y = OutDerivedData.Header.PageTableVolumeAABBMin.Y; Y < OutDerivedData.Header.PageTableVolumeAABBMax.Y; ++Y)
+			for (int32 Y = MippedTextureData.Header.PageTableVolumeAABBMin.Y; Y < MippedTextureData.Header.PageTableVolumeAABBMax.Y; ++Y)
 			{
-				for (int32 X = OutDerivedData.Header.PageTableVolumeAABBMin.X; X < OutDerivedData.Header.PageTableVolumeAABBMax.X; ++X)
+				for (int32 X = MippedTextureData.Header.PageTableVolumeAABBMin.X; X < MippedTextureData.Header.PageTableVolumeAABBMax.X; ++X)
 				{
 					FIntVector3 GlobalPageTableCoord(X, Y, Z);
 					if (IsInBounds(GlobalPageTableCoord, Header.PageTableVolumeAABBMin, Header.PageTableVolumeAABBMax))
@@ -750,35 +648,35 @@ bool FTextureData::BuildDerivedData(const FTextureDataAddressingInfo& Addressing
 						const int32 PageTableIndexOld = (LocalPageTableCoordOld.Z * Header.PageTableVolumeResolution.Y * Header.PageTableVolumeResolution.X) 
 							+ (LocalPageTableCoordOld.Y * Header.PageTableVolumeResolution.X) 
 							+ LocalPageTableCoordOld.X;
-						check(OutDerivedData.MipMaps[0].PageTable.IsValidIndex(PageTableIndexOld));
+						check(MippedTextureData.MipMaps[0].PageTable.IsValidIndex(PageTableIndexOld));
 
-						const FIntVector3 LocalPageTableCoordNew = GlobalPageTableCoord - OutDerivedData.Header.PageTableVolumeAABBMin;
-						const int32 PageTableIndexNew = (LocalPageTableCoordNew.Z * OutDerivedData.Header.PageTableVolumeResolution.Y * OutDerivedData.Header.PageTableVolumeResolution.X) 
-							+ (LocalPageTableCoordNew.Y * OutDerivedData.Header.PageTableVolumeResolution.X) 
+						const FIntVector3 LocalPageTableCoordNew = GlobalPageTableCoord - MippedTextureData.Header.PageTableVolumeAABBMin;
+						const int32 PageTableIndexNew = (LocalPageTableCoordNew.Z * MippedTextureData.Header.PageTableVolumeResolution.Y * MippedTextureData.Header.PageTableVolumeResolution.X) 
+							+ (LocalPageTableCoordNew.Y * MippedTextureData.Header.PageTableVolumeResolution.X) 
 							+ LocalPageTableCoordNew.X;
 						check(PageTableNew.IsValidIndex(PageTableIndexNew));
 
-						PageTableNew[PageTableIndexNew] = OutDerivedData.MipMaps[0].PageTable[PageTableIndexOld];
+						PageTableNew[PageTableIndexNew] = MippedTextureData.MipMaps[0].PageTable[PageTableIndexOld];
 					}
 				}
 			}
 		}
 
-		OutDerivedData.MipMaps[0].PageTable = MoveTemp(PageTableNew);
+		MippedTextureData.MipMaps[0].PageTable = MoveTemp(PageTableNew);
 	}
 
 	// Generate border voxels of mip0
 	{
 		// Collect the page table coordinates of all non-zero pages
 		TArray<FIntVector3> PageCoords;
-		for (int32 PageZ = 0; PageZ < OutDerivedData.Header.PageTableVolumeResolution.Z; ++PageZ)
+		for (int32 PageZ = 0; PageZ < MippedTextureData.Header.PageTableVolumeResolution.Z; ++PageZ)
 		{
-			for (int32 PageY = 0; PageY < OutDerivedData.Header.PageTableVolumeResolution.Y; ++PageY)
+			for (int32 PageY = 0; PageY < MippedTextureData.Header.PageTableVolumeResolution.Y; ++PageY)
 			{
-				for (int32 PageX = 0; PageX < OutDerivedData.Header.PageTableVolumeResolution.X; ++PageX)
+				for (int32 PageX = 0; PageX < MippedTextureData.Header.PageTableVolumeResolution.X; ++PageX)
 				{
-					const int32 PageIndex = PageZ * (OutDerivedData.Header.PageTableVolumeResolution.Y * OutDerivedData.Header.PageTableVolumeResolution.X) + PageY * OutDerivedData.Header.PageTableVolumeResolution.X + PageX;
-					const uint32 PageTableEntry = OutDerivedData.MipMaps[0].PageTable[PageIndex];
+					const int32 PageIndex = PageZ * (MippedTextureData.Header.PageTableVolumeResolution.Y * MippedTextureData.Header.PageTableVolumeResolution.X) + PageY * MippedTextureData.Header.PageTableVolumeResolution.X + PageX;
+					const uint32 PageTableEntry = MippedTextureData.MipMaps[0].PageTable[PageIndex];
 					if (PageTableEntry != 0)
 					{
 						PageCoords.Add(FIntVector3(PageX, PageY, PageZ));
@@ -787,25 +685,191 @@ bool FTextureData::BuildDerivedData(const FTextureDataAddressingInfo& Addressing
 			}
 		}
 
-		if (!OutDerivedData.GenerateBorderVoxels(AddressingInfo, 0, PageCoords))
+		if (!MippedTextureData.GenerateBorderVoxels(AddressingInfo, 0, PageCoords))
 		{
 			return false;
 		}
 	}
 
 	// Generate all remaining mips. Also generates border voxels.
-	if (!OutDerivedData.GenerateMipMaps(AddressingInfo))
+	if (!MippedTextureData.GenerateMipMaps(AddressingInfo))
 	{
 		return false;
 	}
 
-	// Deduplicate tiles by potentially having multiple pages point to the same physical tile
-	if (!OutDerivedData.DeduplicateTiles())
+	OutDerivedData.Reset();
+	OutDerivedData.Build(MippedTextureData);
+
+	return true;
+}
+
+void FDerivedTextureData::Reset()
+{
+	NumPhysicalTiles = 0;
+	MipPageRanges.Reset();
+	PageTableCoords.Reset();
+	PageTableTileIndices.Reset();
+	PageTableParentIndices.Reset();
+	PhysicalTileDataA.Reset();
+	PhysicalTileDataB.Reset();
+}
+
+void FDerivedTextureData::Build(const FTextureData& MippedTextureData)
+{
+	Header = MippedTextureData.Header;
+
+	const int32 NumMipLevels = MippedTextureData.MipMaps.Num();
+	const SIZE_T FormatSize[] = { GPixelFormats[Header.AttributesFormats[0]].BlockBytes, GPixelFormats[Header.AttributesFormats[1]].BlockBytes };
+	const SIZE_T TileByteSize[] = { SVT::NumVoxelsPerPaddedTile * FormatSize[0], SVT::NumVoxelsPerPaddedTile * FormatSize[1] };
+
+	// Compute number of total tiles and pages, as well as the number and offset of pages per mip level
+	MipPageRanges.SetNum(NumMipLevels);
+	uint32 SrcNumPhysicalTilesTotal = 0;
+	uint32 NumPagesTotal = 0;
+	for (int32 MipLevel = NumMipLevels - 1; MipLevel >= 0; --MipLevel)
 	{
-		return false;
+		SrcNumPhysicalTilesTotal += MippedTextureData.MipMaps[MipLevel].NumPhysicalTiles;
+		
+		uint32 NumPagesMip = 0;
+		for (uint32 Entry : MippedTextureData.MipMaps[MipLevel].PageTable)
+		{
+			NumPagesMip += (Entry != 0) ? 1 : 0;
+		}
+		MipPageRanges[MipLevel].PageOffset = NumPagesTotal;
+		MipPageRanges[MipLevel].PageCount = NumPagesMip;
+		NumPagesTotal += NumPagesMip;
 	}
 
-	return true;
+	// Reserve memory
+	PageTableCoords.Reserve(NumPagesTotal);
+	PageTableTileIndices.Reserve(NumPagesTotal);
+	PageTableParentIndices.Reserve(NumPagesTotal);
+	PhysicalTileDataA.SetNumUninitialized(SrcNumPhysicalTilesTotal * TileByteSize[0]);
+	PhysicalTileDataB.SetNumUninitialized(SrcNumPhysicalTilesTotal * TileByteSize[1]);
+	
+	// Maps to help with deduplication and finding the parent index of a given page
+	TMap<uint64, uint32> PhysicalTileDataRemap;
+	TArray<TMap<uint32, uint32>> PageCoordToPageIndex;
+	PageCoordToPageIndex.SetNum(NumMipLevels);
+
+	uint32 NumWrittenTiles = 0;
+	uint8* DstPhysicalTileDataA = PhysicalTileDataA.GetData();
+	uint8* DstPhysicalTileDataB = PhysicalTileDataB.GetData();
+
+	for (int32 MipLevel = NumMipLevels - 1; MipLevel >= 0; --MipLevel)
+	{
+		// For now we don't deduplicate across mip levels. This simplifies logic in the current mip level based streaming manager.
+		PhysicalTileDataRemap.Reset();
+
+		const uint8* SrcPhysicalTileDataA = MippedTextureData.MipMaps[MipLevel].PhysicalTileDataA.GetData();
+		const uint8* SrcPhysicalTileDataB = MippedTextureData.MipMaps[MipLevel].PhysicalTileDataB.GetData();
+
+		// Iterate over dense page table of this mip
+		const FIntVector3 MipPageTableResolution = ShiftRightAndMax(MippedTextureData.Header.PageTableVolumeResolution, MipLevel, 1);
+		for (int32 PageZ = 0; PageZ < MipPageTableResolution.Z; ++PageZ)
+		{
+			for (int32 PageY = 0; PageY < MipPageTableResolution.Y; ++PageY)
+			{
+				for (int32 PageX = 0; PageX < MipPageTableResolution.X; ++PageX)
+				{
+					const int32 LinearPageTableIndex = (PageZ * MipPageTableResolution.Y * MipPageTableResolution.X) + (PageY * MipPageTableResolution.X) + PageX;
+					const uint32 Entry = MippedTextureData.MipMaps[MipLevel].PageTable[LinearPageTableIndex];
+
+					if (Entry)
+					{
+						const SIZE_T SrcTileIdx = Entry - 1; // 0 encodes the null tile, but we don't store that tile explicitely
+						
+						// Hash the tile voxel data
+						uint64 Hash = 0;
+						if (FormatSize[0])
+						{
+							Hash = CityHash64((const char*)SrcPhysicalTileDataA + SrcTileIdx * TileByteSize[0], (uint32)TileByteSize[0]);
+						}
+						if (FormatSize[1])
+						{
+							Hash = CityHash64WithSeed((const char*)SrcPhysicalTileDataB + SrcTileIdx * TileByteSize[1], (uint32)TileByteSize[1], Hash);
+						}
+
+						// Have we processed an identical tile before?
+						uint32 RemappedTileIndex = INDEX_NONE;
+						if (uint32* RemappedTileIndexPtr = PhysicalTileDataRemap.Find(Hash))
+						{
+							// If so, we can simply reuse that tile
+							RemappedTileIndex = *RemappedTileIndexPtr;
+						}
+						else
+						{
+							// Tile is unique, so copy it from the source data to the new compacted array
+							RemappedTileIndex = NumWrittenTiles++;
+							PhysicalTileDataRemap.Add(Hash, RemappedTileIndex);
+
+							if (FormatSize[0])
+							{
+								FMemory::Memcpy(DstPhysicalTileDataA + RemappedTileIndex * TileByteSize[0], SrcPhysicalTileDataA +SrcTileIdx * TileByteSize[0], TileByteSize[0]);
+							}
+							if (FormatSize[1])
+							{
+								FMemory::Memcpy(DstPhysicalTileDataB + RemappedTileIndex * TileByteSize[1], SrcPhysicalTileDataB + SrcTileIdx * TileByteSize[1], TileByteSize[1]);
+							}
+						}
+
+						// Insert page into hash map so child pages can look up its index
+						const uint32 PackedPageCoord = SVT::PackX11Y11Z10(FIntVector3(PageX, PageY, PageZ));
+						PageCoordToPageIndex[MipLevel].Add(PackedPageCoord, PageTableCoords.Num());
+						
+						PageTableCoords.Add(PackedPageCoord);
+						PageTableTileIndices.Add(RemappedTileIndex);
+						if (PageCoordToPageIndex.IsValidIndex(MipLevel + 1))
+						{
+							// Look up parent index based on mip level and parent page coord
+							const uint32 ParentPageCoord = SVT::PackX11Y11Z10(FIntVector3(PageX / 2, PageY / 2, PageZ / 2));
+							const uint32 ParentIndex = PageCoordToPageIndex[MipLevel + 1][ParentPageCoord];
+							PageTableParentIndices.Add(ParentIndex);
+						}
+						else
+						{
+							// Root mip has no parent
+							PageTableParentIndices.Add(INDEX_NONE);
+						}
+					}
+				}
+			}
+		}
+	}
+
+	// Resize the arrays but don't bother shrinking; the gains are probably not massive and this object is expected to have a fairly short lifetime anyways.
+	PhysicalTileDataA.SetNum(static_cast<int64>(NumWrittenTiles * TileByteSize[0]), EAllowShrinking::No);
+	PhysicalTileDataB.SetNum(static_cast<int64>(NumWrittenTiles * TileByteSize[1]), EAllowShrinking::No);
+	NumPhysicalTiles = NumWrittenTiles;
+
+	// Validate
+	{
+#if DO_CHECK
+		uint32 TotalPageCount = 0;
+		for (int32 MipLevel = NumMipLevels - 1; MipLevel >= 0; --MipLevel)
+		{
+			check(MipPageRanges[MipLevel].PageOffset == TotalPageCount);
+			const int32 NumPageTableEntries = MipPageRanges[MipLevel].PageCount;
+			TotalPageCount += NumPageTableEntries;
+			const bool bIsHighestMipLevel = MipLevel == (NumMipLevels - 1);
+			check(!bIsHighestMipLevel || NumPageTableEntries <= 1); // Highest mip has a maximum of 1 page
+			const FIntVector3 PageTableRes = ShiftRightAndMax(Header.PageTableVolumeResolution, MipLevel, 1);
+
+			for (int32 i = 0; i < NumPageTableEntries; ++i)
+			{
+				const int32 PageIndex = MipPageRanges[MipLevel].PageOffset + i;
+				const int32 PageX = (PageTableCoords[PageIndex] >> 0u) & 0x7FFu;
+				const int32 PageY = (PageTableCoords[PageIndex] >> 11u) & 0x7FFu;
+				const int32 PageZ = (PageTableCoords[PageIndex] >> 22u) & 0x3FFu;
+				const int32 LinearPageIndex = (PageZ * (PageTableRes.X * PageTableRes.Y)) + (PageY * PageTableRes.X) + PageX;
+				check(MippedTextureData.MipMaps[MipLevel].PageTable.IsValidIndex(LinearPageIndex));
+				const uint32 PageTableEntry = MippedTextureData.MipMaps[MipLevel].PageTable[LinearPageIndex];
+				check(PageTableEntry != 0 && PageTableEntry != INDEX_NONE); // Unpacked page coord must map to a valid entry in the dense page table
+			}
+		}
+		check(TotalPageCount == PageTableCoords.Num());
+#endif // DO_CHECK
+	}
 }
 
 }
diff --git a/Engine/Source/Runtime/Engine/Private/SparseVolumeTexture/SparseVolumeTextureStreamingManager.cpp b/Engine/Source/Runtime/Engine/Private/SparseVolumeTexture/SparseVolumeTextureStreamingManager.cpp
--- a/Engine/Source/Runtime/Engine/Private/SparseVolumeTexture/SparseVolumeTextureStreamingManager.cpp
+++ b/Engine/Source/Runtime/Engine/Private/SparseVolumeTexture/SparseVolumeTextureStreamingManager.cpp
@@ -392,17 +392,18 @@ void FStreamingManager::BeginAsyncUpdate(FRDGBuilder& GraphBuilder, bool bBlocki
 			double SumFrameMiB = 0.0;
 
 			const int32 NumFrames = Pair.Value->PerFrameInfo.Num();
-			UE_LOG(LogSparseVolumeTextureStreamingManager, Display, TEXT("Memory stats for SVT '%p': Each mip level is printed as a tuple of [PageTable Size | VoxelData Size | Total]"), Pair.Key);
+			UE_LOG(LogSparseVolumeTextureStreamingManager, Display, TEXT("Memory stats for SVT '%p': Each frame is displayed as a list of mip levels like this: [Mip0] [Mip1] [MipN]"), Pair.Key);
 
 			for (int32 FrameIdx = 0; FrameIdx < NumFrames; ++FrameIdx)
 			{
 				const FFrameInfo& FrameInfo = Pair.Value->PerFrameInfo[FrameIdx];
 				FString Str;
 				int32 TotalSize = 0;
-				for (const auto& SInfo : FrameInfo.Resources->MipLevelStreamingInfo)
+				for (int32 MipLevelIdx = 0; MipLevelIdx < FrameInfo.NumMipLevels; ++MipLevelIdx)
 				{
-					Str += FString::Printf(TEXT("[%5.2f KiB|%5.2f KiB|%5.2f KiB] "), SInfo.PageTableSize / 1024.0f, (SInfo.TileDataSize[0] + SInfo.TileDataSize[1]) / 1024.0f, (SInfo.PageTableSize + SInfo.TileDataSize[0] + SInfo.TileDataSize[1]) / 1024.0f);
-					TotalSize += SInfo.BulkSize;
+					const FMipTileReadInfo MipReadInfo = GetMipTileReadInfo(FrameInfo, MipLevelIdx);
+					Str += FString::Printf(TEXT("[%u Tiles, %5.2f KiB] "), MipReadInfo.TileCount, MipReadInfo.ReadSize / 1024.0f);
+					TotalSize += MipReadInfo.ReadSize;
 				}
 
 				MinFrameMiB = FMath::Min(MinFrameMiB, TotalSize / 1024.0 / 1024.0);
@@ -450,11 +451,16 @@ void FStreamingManager::BeginAsyncUpdate(FRDGBuilder& GraphBuilder, bool bBlocki
 			const int32 FormatSizeA = GPixelFormats[SVTInfo->FormatA].BlockBytes;
 			const int32 FormatSizeB = GPixelFormats[SVTInfo->FormatB].BlockBytes;
 			const FResources* Resources = SVTInfo->PerFrameInfo[PendingMipLevel.FrameIndex].Resources;
-			const int32 NumTilesToUpload = Resources->MipLevelStreamingInfo[PendingMipLevel.MipLevelIndex].NumPhysicalTiles;
-			const int32 NumVoxelsToUploadA = FormatSizeA > 0 ? Resources->MipLevelStreamingInfo[PendingMipLevel.MipLevelIndex].TileDataSize[0] / FormatSizeA : 0;
-			const int32 NumVoxelsToUploadB = FormatSizeB > 0 ? Resources->MipLevelStreamingInfo[PendingMipLevel.MipLevelIndex].TileDataSize[1] / FormatSizeB : 0;
+			const FPageTopology::FMip& MipInfo = Resources->Topology.MipInfo[PendingMipLevel.MipLevelIndex];
+			
+			uint32 TileRangeOffset = 0;
+			uint32 TileRangeCount = 0;
+			Resources->Topology.GetTileRange(MipInfo.PageOffset, MipInfo.PageCount, TileRangeOffset, TileRangeCount);
+			uint32 NumVoxelsA = 0;
+			uint32 NumVoxelsB = 0;
+			Resources->StreamingMetaData.GetNumVoxelsInTileRange(TileRangeOffset, TileRangeCount, FormatSizeA, FormatSizeB, NumVoxelsA, NumVoxelsB);
 
-			SVTInfo->TileDataTexture->ReserveUpload(NumTilesToUpload, NumVoxelsToUploadA, NumVoxelsToUploadB);
+			SVTInfo->TileDataTexture->ReserveUpload(TileRangeCount, NumVoxelsA, NumVoxelsB);
 
 			TileDataTexturesToUpdate.Add(SVTInfo->TileDataTexture.Get());
 		}
@@ -582,34 +588,28 @@ void FStreamingManager::AddInternal(FRDGBuilder& GraphBuilder, FNewSparseVolumeT
 		FFrameInfo& FrameInfo = SVTInfo.PerFrameInfo[FrameIdx];
 		check(FrameInfo.TextureRenderResources && FrameInfo.TextureRenderResources->IsInitialized());
 		const FResources* Resources = FrameInfo.Resources;
+		const FPageTopology& Topology = Resources->Topology;
+
+		const int32 NumPhysicalTiles = Resources->StreamingMetaData.GetNumTiles();
+		MaxNumPhysicalTiles = FMath::Max(NumPhysicalTiles, MaxNumPhysicalTiles);
 
-		FrameInfo.NumMipLevels = Resources->MipLevelStreamingInfo.Num();
+		FrameInfo.NumMipLevels = Topology.MipInfo.Num();
 		FrameInfo.LowestRequestedMipLevel = FrameInfo.NumMipLevels - 1;
 		FrameInfo.LowestResidentMipLevel = FrameInfo.NumMipLevels - 1;
-		FrameInfo.TileAllocations.SetNum(FrameInfo.NumMipLevels);
-		for (int32 MipLevel = 0; MipLevel < FrameInfo.NumMipLevels; ++MipLevel)
-		{
-			FrameInfo.TileAllocations[MipLevel].SetNumZeroed(Resources->MipLevelStreamingInfo[MipLevel].NumPhysicalTiles);
-		}
-		const int32 NumPagesTotal = Resources->Topology.Pages.Num();
+		FrameInfo.TileAllocations.SetNum(NumPhysicalTiles);
+
+		const int32 NumPagesTotal = Topology.NumPages();
 		FrameInfo.ResidentPages.SetNum(NumPagesTotal, false);
 		FrameInfo.ResidentPagesNew.SetNum(NumPagesTotal, false);
-		FrameInfo.PageEntries.SetNum(NumPagesTotal);
 		
-		int32 NumPhysicalTiles = 0;
-		for (const FMipLevelStreamingInfo& MipLevelStreamingInfo : Resources->MipLevelStreamingInfo)
-		{
-			NumPhysicalTiles += MipLevelStreamingInfo.NumPhysicalTiles;
-		}
-
-		MaxNumPhysicalTiles = FMath::Max(NumPhysicalTiles, MaxNumPhysicalTiles);
-		if (NumPhysicalTiles > 0)
+		if (!Topology.MipInfo.IsEmpty() && Topology.MipInfo.Last().PageCount > 0)
 		{
+			check(Topology.MipInfo.Last().PageOffset == 0);
+			check(Topology.MipInfo.Last().PageCount == 1);
 			++NumRootPhysicalTiles;
-			check(FormatSizes[0] == 0 || (Resources->MipLevelStreamingInfo.Last().TileDataSize[0] % FormatSizes[0]) == 0);
-			check(FormatSizes[1] == 0 || (Resources->MipLevelStreamingInfo.Last().TileDataSize[1] % FormatSizes[1]) == 0);
-			NumRootVoxelsA += FormatSizes[0] > 0 ? (Resources->MipLevelStreamingInfo.Last().TileDataSize[0] / FormatSizes[0]) : 0;
-			NumRootVoxelsB += FormatSizes[1] > 0 ? (Resources->MipLevelStreamingInfo.Last().TileDataSize[1] / FormatSizes[1]) : 0;
+			FTileInfo RootTileInfo = Resources->StreamingMetaData.GetTileInfo(0, FormatSizes[0], FormatSizes[1]);
+			NumRootVoxelsA += RootTileInfo.NumVoxels[0];
+			NumRootVoxelsB += RootTileInfo.NumVoxels[1];
 		}
 
 		for (int32 MipIdx = 0; MipIdx < SVTInfo.NumMipLevelsGlobal; ++MipIdx)
@@ -662,7 +662,8 @@ void FStreamingManager::AddInternal(FRDGBuilder& GraphBuilder, FNewSparseVolumeT
 		{
 			FFrameInfo& FrameInfo = SVTInfo.PerFrameInfo[FrameIdx];
 			const FResources* Resources = FrameInfo.Resources;
-			const int32 NumMipLevels = Resources->MipLevelStreamingInfo.Num();
+			const FPageTopology& Topology = Resources->Topology;
+			const int32 NumMipLevels = Topology.MipInfo.Num();
 
 			FrameInfo.LowestRequestedMipLevel = NumMipLevels - 1;
 			FrameInfo.LowestResidentMipLevel = NumMipLevels - 1;
@@ -697,16 +698,17 @@ void FStreamingManager::AddInternal(FRDGBuilder& GraphBuilder, FNewSparseVolumeT
 			FrameInfo.TextureRenderResources->NumLogicalMipLevels = NumMipLevels;
 
 			// Upload root mip data and update page tables
-			const FMipLevelStreamingInfo* RootStreamingInfo = !Resources->MipLevelStreamingInfo.IsEmpty() ? &Resources->MipLevelStreamingInfo.Last() : nullptr;
-			if (!Resources->RootData.IsEmpty() && RootStreamingInfo)
+			if (!Topology.MipInfo.IsEmpty() && Topology.MipInfo.Last().PageCount > 0)
 			{
+				check(!Resources->RootData.IsEmpty());
 				const uint32 TileCoord = SVTInfo.TileDataTexture->Allocate();
 				check(TileCoord != INDEX_NONE);
-				FrameInfo.TileAllocations.Last()[0] = TileCoord;
-				FrameInfo.PageEntries[Resources->Topology.MipInfo.Last().PageOffset] = TileCoord;
+				check(TileCoord != 0);
+				FrameInfo.TileAllocations[0] = TileCoord;
 
-				const int32 NumVoxelsA = FormatSizes[0] > 0 ? RootStreamingInfo->TileDataSize[0] / FormatSizes[0] : 0;
-				const int32 NumVoxelsB = FormatSizes[1] > 0 ? RootStreamingInfo->TileDataSize[1] / FormatSizes[1] : 0;
+				const FTileInfo RootTileInfo = Resources->StreamingMetaData.GetTileInfo(0, FormatSizes[0], FormatSizes[1]);
+				const int32 NumVoxelsA = RootTileInfo.NumVoxels[0];
+				const int32 NumVoxelsB = RootTileInfo.NumVoxels[1];
 				FTileUploader::FAddResult AddResult = RootTileUploader.Add_GetRef(1, NumVoxelsA, NumVoxelsB);
 
 				FMemory::Memcpy(AddResult.PackedPhysicalTileCoordsPtr, &TileCoord, sizeof(TileCoord));
@@ -715,26 +717,25 @@ void FStreamingManager::AddInternal(FRDGBuilder& GraphBuilder, FNewSparseVolumeT
 					if (FormatSizes[AttributesIdx] > 0)
 					{
 						// Occupancy bits
-						const uint8* SrcOccupancyBits = Resources->RootData.GetData() + RootStreamingInfo->OccupancyBitsOffset[AttributesIdx];
+						const uint8* SrcOccupancyBits = Resources->RootData.GetData() + RootTileInfo.OccupancyBitsOffsets[AttributesIdx];
 						check(AddResult.OccupancyBitsPtrs[AttributesIdx]);
-						FMemory::Memcpy(AddResult.OccupancyBitsPtrs[AttributesIdx], SrcOccupancyBits, RootStreamingInfo->OccupancyBitsSize[AttributesIdx]);
+						FMemory::Memcpy(AddResult.OccupancyBitsPtrs[AttributesIdx], SrcOccupancyBits, RootTileInfo.OccupancyBitsSizes[AttributesIdx]);
 
 						// Per-tile offsets into tile data
-						const uint32* SrcTileDataOffsets = reinterpret_cast<const uint32*>(Resources->RootData.GetData() + RootStreamingInfo->TileDataOffsetsOffset[AttributesIdx]);
 						check(AddResult.TileDataOffsetsPtrs[AttributesIdx]);
-						check(RootStreamingInfo->TileDataOffsetsSize[AttributesIdx] == sizeof(uint32));
-						reinterpret_cast<uint32*>(AddResult.TileDataOffsetsPtrs[AttributesIdx])[0] = AddResult.TileDataBaseOffsets[AttributesIdx] + SrcTileDataOffsets[0];
+						reinterpret_cast<uint32*>(AddResult.TileDataOffsetsPtrs[AttributesIdx])[0] = AddResult.TileDataBaseOffsets[AttributesIdx] + 0; // + 0 because we're only uploading this single tile in the current batch
 
 						// Tile data
-						const uint8* SrcTileData = Resources->RootData.GetData() + RootStreamingInfo->TileDataOffset[AttributesIdx];
+						const uint8* SrcTileData = Resources->RootData.GetData() + RootTileInfo.VoxelDataOffsets[AttributesIdx];
 						check(AddResult.TileDataPtrs[AttributesIdx]);
-						FMemory::Memcpy(AddResult.TileDataPtrs[AttributesIdx], SrcTileData, RootStreamingInfo->TileDataSize[AttributesIdx]);
+						FMemory::Memcpy(AddResult.TileDataPtrs[AttributesIdx], SrcTileData, RootTileInfo.VoxelDataSizes[AttributesIdx]);
 					}
 				}
-			}
 
-			const FPageTopology::FMip& TopologyMipInfo = Resources->Topology.MipInfo[NumMipLevels - 1];
-			FrameInfo.ResidentPagesNew.SetRange(TopologyMipInfo.PageOffset, TopologyMipInfo.PageCount, true);
+				// Mark page as resident
+				const FPageTopology::FMip& TopologyMipInfo = Topology.MipInfo.Last();
+				FrameInfo.ResidentPagesNew.SetRange(TopologyMipInfo.PageOffset, TopologyMipInfo.PageCount, true);
+			}
 
 			InvalidatedSVTFrames.Add(&FrameInfo);
 		}
@@ -883,7 +884,7 @@ void FStreamingManager::SelectHighestPriorityRequestsAndUpdateLRU(int32 MaxSelec
 				|| Request.Key.FrameIndex >= SVTInfo->PerFrameInfo.Num()
 				|| Request.Key.MipLevelIndex < 0
 				|| Request.Key.MipLevelIndex >= (SVTInfo->PerFrameInfo[Request.Key.FrameIndex].NumMipLevels - 1)
-				|| SVTInfo->PerFrameInfo[Request.Key.FrameIndex].Resources->MipLevelStreamingInfo[Request.Key.MipLevelIndex].BulkSize == 0)
+				|| SVTInfo->PerFrameInfo[Request.Key.FrameIndex].Resources->Topology.MipInfo[Request.Key.MipLevelIndex].PageCount == 0)
 			{
 				continue;
 			}
@@ -973,8 +974,9 @@ void FStreamingManager::IssueRequests(int32 MaxSelectedRequests)
 		check(SVTInfo->PerFrameInfo.Num() > SelectedKey.FrameIndex && SelectedKey.FrameIndex >= 0);
 		check(SVTInfo->PerFrameInfo[SelectedKey.FrameIndex].LowestRequestedMipLevel > SelectedKey.MipLevelIndex);
 		const FResources* Resources = SVTInfo->PerFrameInfo[SelectedKey.FrameIndex].Resources;
-		check((SelectedKey.MipLevelIndex + 1) < Resources->MipLevelStreamingInfo.Num()); // The lowest/last mip level is always resident and does not stream.
-		const FMipLevelStreamingInfo& MipLevelStreamingInfo = Resources->MipLevelStreamingInfo[SelectedKey.MipLevelIndex];
+		check((SelectedKey.MipLevelIndex + 1) < Resources->Topology.MipInfo.Num()); // The lowest/last mip level is always resident and does not stream.
+
+		const FMipTileReadInfo MipTileReadInfo = GetMipTileReadInfo(SVTInfo->PerFrameInfo[SelectedKey.FrameIndex], SelectedKey.MipLevelIndex);
 
 		TUniquePtr<FTileDataTexture>& TileDataTexture = SVTInfo->TileDataTexture;
 		check(TileDataTexture);
@@ -982,7 +984,7 @@ void FStreamingManager::IssueRequests(int32 MaxSelectedRequests)
 		// Ensure that enough tiles are available in the tile texture
 		const int32 TileDataTextureCapacity = TileDataTexture->GetTileCapacity();
 		const int32 NumAvailableTiles = TileDataTexture->GetNumAvailableTiles();
-		const int32 NumRequiredTiles = MipLevelStreamingInfo.NumPhysicalTiles;
+		const int32 NumRequiredTiles = MipTileReadInfo.TileCount;
 		if (NumAvailableTiles < NumRequiredTiles)
 		{
 #if SVT_STREAMING_LOG_VERBOSE
@@ -1002,7 +1004,12 @@ void FStreamingManager::IssueRequests(int32 MaxSelectedRequests)
 					if (Node.RefCount == 0 && Node.LastRequested < NextUpdateIndex)
 					{
 						MipLevelsToFree.Add(&Node);
-						NumNewlyAvailableTiles += SVTInfo->PerFrameInfo[Node.FrameIndex].Resources->MipLevelStreamingInfo[Node.MipLevelIndex].NumPhysicalTiles;
+						const FPageTopology& Topology = SVTInfo->PerFrameInfo[Node.FrameIndex].Resources->Topology;
+						const FPageTopology::FMip& TopologyMip = Topology.MipInfo[Node.MipLevelIndex];
+						uint32 TilesInMipCount = 0;
+						uint32 TilesInMipOffset = 0;
+						Topology.GetTileRange(TopologyMip.PageOffset, TopologyMip.PageCount, TilesInMipOffset, TilesInMipCount);
+						NumNewlyAvailableTiles += TilesInMipCount;
 
 						// Decrement ref count of mip levels higher up the chain
 						FLRUNode* Dependency = Node.NextHigherMipLevel;
@@ -1065,7 +1072,7 @@ void FStreamingManager::IssueRequests(int32 MaxSelectedRequests)
 		{
 			if (Resources->ResourceFlags & EResourceFlag_StreamingDataInDDC)
 			{
-				UE::DerivedData::FCacheGetChunkRequest DDCRequest = BuildDDCRequest(*Resources, MipLevelStreamingInfo, NextPendingMipLevelIndex);
+				UE::DerivedData::FCacheGetChunkRequest DDCRequest = BuildDDCRequest(*Resources, MipTileReadInfo.ReadOffset, MipTileReadInfo.ReadSize, NextPendingMipLevelIndex);
 				if (PendingMipLevel.bBlocking)
 				{
 					DDCRequestsBlocking.Add(DDCRequest);
@@ -1084,9 +1091,9 @@ void FStreamingManager::IssueRequests(int32 MaxSelectedRequests)
 		else
 #endif
 		{
-			PendingMipLevel.RequestBuffer = FIoBuffer(MipLevelStreamingInfo.BulkSize); // SVT_TODO: Use FIoBuffer::Wrap with preallocated memory
+			PendingMipLevel.RequestBuffer = FIoBuffer(MipTileReadInfo.ReadSize); // SVT_TODO: Use FIoBuffer::Wrap with preallocated memory
 			const EAsyncIOPriorityAndFlags Priority = PendingMipLevel.bBlocking ? AIOP_CriticalPath : AIOP_Low;
-			Batch.Read(BulkData, MipLevelStreamingInfo.BulkOffset, MipLevelStreamingInfo.BulkSize, Priority, PendingMipLevel.RequestBuffer, PendingMipLevel.Request);
+			Batch.Read(BulkData, MipTileReadInfo.ReadOffset, MipTileReadInfo.ReadSize, Priority, PendingMipLevel.RequestBuffer, PendingMipLevel.Request);
 			bIssueIOBatch = true;
 
 #if WITH_EDITORONLY_DATA
@@ -1102,13 +1109,19 @@ void FStreamingManager::IssueRequests(int32 MaxSelectedRequests)
 
 		// Allocate tiles in the tile data texture
 		{
-			TArray<uint32>& TileAllocations = FrameInfo.TileAllocations[SelectedKey.MipLevelIndex];
-			check(TileAllocations.Num() == NumRequiredTiles);
-			for (int32 TileIdx = 0; TileIdx < NumRequiredTiles; ++TileIdx)
+			const FPageTopology::FMip& MipInfo = FrameInfo.Resources->Topology.MipInfo[SelectedKey.MipLevelIndex];
+
+			uint32 TilesInMipCount = 0;
+			uint32 TilesInMipOffset = 0;
+			Resources->Topology.GetTileRange(MipInfo.PageOffset, MipInfo.PageCount, TilesInMipOffset, TilesInMipCount);
+
+			for (uint32 TileIndex = TilesInMipOffset; TileIndex < (TilesInMipOffset + TilesInMipCount); ++TileIndex)
 			{
+				check(FrameInfo.TileAllocations[TileIndex] == 0);
 				const int32 TileCoord = TileDataTexture->Allocate();
 				check(TileCoord != INDEX_NONE);
-				TileAllocations[TileIdx] = TileCoord;
+				check(TileCoord != 0);
+				FrameInfo.TileAllocations[TileIndex] = TileCoord;
 			}
 		}
 
@@ -1199,10 +1212,20 @@ void FStreamingManager::StreamOutMipLevel(FStreamingInfo* SVTInfo, FLRUNode* LRU
 	LRUNode->LastRequested = INDEX_NONE;
 
 	// Free allocated tiles
-	for (uint32& TileCoord : FrameInfo.TileAllocations[MipLevelIndex])
 	{
-		SVTInfo->TileDataTexture->Free(TileCoord);
-		TileCoord = 0;
+		const FPageTopology& Topology = FrameInfo.Resources->Topology;
+		const FPageTopology::FMip& TopologyMip = Topology.MipInfo[MipLevelIndex];
+		uint32 TilesInMipCount = 0;
+		uint32 TilesInMipOffset = 0;
+		Topology.GetTileRange(TopologyMip.PageOffset, TopologyMip.PageCount, TilesInMipOffset, TilesInMipCount);
+
+		for (uint32 TileIndex = TilesInMipOffset; TileIndex < (TilesInMipOffset + TilesInMipCount); ++TileIndex)
+		{
+			uint32& TileCoord = FrameInfo.TileAllocations[TileIndex];
+			check(TileCoord != 0);
+			SVTInfo->TileDataTexture->Free(TileCoord);
+			TileCoord = 0;
+		}
 	}
 
 	const FPageTopology::FMip& TopologyMipInfo = FrameInfo.Resources->Topology.MipInfo[MipLevelIndex];
@@ -1259,8 +1282,8 @@ int32 FStreamingManager::DetermineReadyMipLevels()
 					*Resources->ResourceName, PendingMipLevel.FrameIndex, PendingMipLevel.MipLevelIndex);
 			}
 
-			const FMipLevelStreamingInfo& MipLevelStreamingInfo = Resources->MipLevelStreamingInfo[PendingMipLevel.MipLevelIndex];
-			FCacheGetChunkRequest Request = BuildDDCRequest(*Resources, MipLevelStreamingInfo, PendingMipLevelIndex);
+			const FMipTileReadInfo MipTileReadInfo = GetMipTileReadInfo(SVTInfo->PerFrameInfo[PendingMipLevel.FrameIndex], PendingMipLevel.MipLevelIndex);
+			FCacheGetChunkRequest Request = BuildDDCRequest(*Resources, MipTileReadInfo.ReadOffset, MipTileReadInfo.ReadSize, PendingMipLevelIndex);
 			const bool bBlocking = GSVTStreamingForceBlockingRequests || PendingMipLevel.bBlocking;
 			RequestDDCData(MakeArrayView(&Request, 1), bBlocking);
 
@@ -1282,12 +1305,12 @@ int32 FStreamingManager::DetermineReadyMipLevels()
 				if (!PendingMipLevel.Request.IsOk())
 				{
 					// Retry if IO request failed for some reason
-					const FMipLevelStreamingInfo& MipLevelStreamingInfo = Resources->MipLevelStreamingInfo[PendingMipLevel.MipLevelIndex];
+					const FMipTileReadInfo MipTileReadInfo = GetMipTileReadInfo(SVTInfo->PerFrameInfo[PendingMipLevel.FrameIndex], PendingMipLevel.MipLevelIndex);
 					UE_LOG(LogSparseVolumeTextureStreamingManager, Warning, TEXT("SVT IO request failed for %p (frame %i, mip %i, offset %i, size %i). Retrying..."),
-						PendingMipLevel.SparseVolumeTexture, PendingMipLevel.FrameIndex, PendingMipLevel.MipLevelIndex, MipLevelStreamingInfo.BulkOffset, MipLevelStreamingInfo.BulkSize);
+						PendingMipLevel.SparseVolumeTexture, PendingMipLevel.FrameIndex, PendingMipLevel.MipLevelIndex, MipTileReadInfo.ReadOffset, MipTileReadInfo.ReadSize);
 					
 					FBulkDataBatchRequest::FBatchBuilder Batch = FBulkDataBatchRequest::NewBatch(1);
-					Batch.Read(Resources->StreamableMipLevels, MipLevelStreamingInfo.BulkOffset, MipLevelStreamingInfo.BulkSize, AIOP_Low, PendingMipLevel.RequestBuffer, PendingMipLevel.Request);
+					Batch.Read(Resources->StreamableMipLevels, MipTileReadInfo.ReadOffset, MipTileReadInfo.ReadSize, AIOP_Low, PendingMipLevel.RequestBuffer, PendingMipLevel.Request);
 					(void)Batch.Issue();
 					break;
 				}
@@ -1339,22 +1362,26 @@ void FStreamingManager::InstallReadyMipLevels()
 
 		FFrameInfo& FrameInfo = SVTInfo->PerFrameInfo[PendingMipLevel.FrameIndex];
 		const FResources* Resources = FrameInfo.Resources;
-		const FMipLevelStreamingInfo& MipLevelStreamingInfo = Resources->MipLevelStreamingInfo[PendingMipLevel.MipLevelIndex];
+		const FPageTopology& Topology = FrameInfo.Resources->Topology;
+		const FMipTileReadInfo MipTileReadInfo = GetMipTileReadInfo(FrameInfo, PendingMipLevel.MipLevelIndex);
 
 		const uint8* SrcPtr = nullptr;
+		const uint8* SrcEndPtr = nullptr;
 
 #if WITH_EDITORONLY_DATA
 		if (PendingMipLevel.State == FPendingMipLevel::EState::DDC_Ready)
 		{
 			check(Resources->ResourceFlags & EResourceFlag_StreamingDataInDDC);
 			SrcPtr = (const uint8*)PendingMipLevel.SharedBuffer.GetData();
+			SrcEndPtr = SrcPtr + PendingMipLevel.SharedBuffer.GetSize();
 		}
 		else if (PendingMipLevel.State == FPendingMipLevel::EState::Memory)
 		{
 			const uint8** BulkDataPtrPtr = ResourceToBulkPointer.Find(Resources);
 			if (BulkDataPtrPtr)
 			{
-				SrcPtr = *BulkDataPtrPtr + MipLevelStreamingInfo.BulkOffset;
+				SrcPtr = *BulkDataPtrPtr + MipTileReadInfo.ReadOffset;
+				SrcEndPtr = SrcPtr + MipTileReadInfo.ReadSize;
 			}
 			else
 			{
@@ -1362,7 +1389,8 @@ void FStreamingManager::InstallReadyMipLevels()
 				check(BulkData.IsBulkDataLoaded() && BulkData.GetBulkDataSize() > 0);
 				const uint8* BulkDataPtr = (const uint8*)BulkData.LockReadOnly();
 				ResourceToBulkPointer.Add(Resources, BulkDataPtr);
-				SrcPtr = BulkDataPtr + MipLevelStreamingInfo.BulkOffset;
+				SrcPtr = BulkDataPtr + MipTileReadInfo.ReadOffset;
+				SrcEndPtr = BulkDataPtr + BulkData.GetBulkDataSize();
 			}
 		}
 		else
@@ -1379,48 +1407,29 @@ void FStreamingManager::InstallReadyMipLevels()
 		const int32 FormatSizeA = GPixelFormats[SVTInfo->FormatA].BlockBytes;
 		const int32 FormatSizeB = GPixelFormats[SVTInfo->FormatB].BlockBytes;
 
-		const int32 NumPhysicalTiles = MipLevelStreamingInfo.NumPhysicalTiles;
-		const int32 NumVoxelsA = FormatSizeA > 0 ? MipLevelStreamingInfo.TileDataSize[0] / FormatSizeA : 0;
-		const int32 NumVoxelsB = FormatSizeB > 0 ? MipLevelStreamingInfo.TileDataSize[1] / FormatSizeB : 0;
-		TArray<uint32>& TileAllocations = FrameInfo.TileAllocations[PendingMipLevel.MipLevelIndex];
-		check(TileAllocations.Num() == NumPhysicalTiles);
-		check((MipLevelStreamingInfo.PageTableSize % (sizeof(uint32) * 2)) == 0);
-		const int32 NumPageTableUpdates = MipLevelStreamingInfo.PageTableSize / (sizeof(uint32) * 2);
+		for (uint32 TileIndex = MipTileReadInfo.TileOffset; TileIndex < (MipTileReadInfo.TileOffset + MipTileReadInfo.TileCount); ++TileIndex)
+		{
+			const FTileInfo TileInfo = Resources->StreamingMetaData.GetTileInfo(TileIndex, FormatSizeA, FormatSizeB);
 
-		FTileUploader::FAddResult TileDataAddResult = SVTInfo->TileDataTexture->AddUpload(NumPhysicalTiles, NumVoxelsA, NumVoxelsB);
+			check(TileInfo.Offset >= MipTileReadInfo.ReadOffset);
+			const uint8* TileSrcPtr = SrcPtr + (TileInfo.Offset - MipTileReadInfo.ReadOffset);
 
-		// Tile data
-		{
-			FUploadTask::FTileDataTask TileDataTask = {};
+			check((TileSrcPtr + TileInfo.Size) <= SrcEndPtr);
+
+			FTileUploader::FAddResult TileDataAddResult = SVTInfo->TileDataTexture->AddUpload(1, TileInfo.NumVoxels[0], TileInfo.NumVoxels[1]);
+
+			FTileDataTask& TileDataTask = UploadTasks.AddDefaulted_GetRef();
 			TileDataTask.DstOccupancyBitsPtrs = TileDataAddResult.OccupancyBitsPtrs;
 			TileDataTask.DstTileDataOffsetsPtrs = TileDataAddResult.TileDataOffsetsPtrs;
 			TileDataTask.DstTileDataPtrs = TileDataAddResult.TileDataPtrs;
-			TileDataTask.DstPhysicalTileCoords = TileDataAddResult.PackedPhysicalTileCoordsPtr;
-			TileDataTask.SrcOccupancyBitsPtrs[0] = SrcPtr + MipLevelStreamingInfo.OccupancyBitsOffset[0];
-			TileDataTask.SrcOccupancyBitsPtrs[1] = SrcPtr + MipLevelStreamingInfo.OccupancyBitsOffset[1];
-			TileDataTask.SrcTileDataOffsetsPtrs[0] = SrcPtr + MipLevelStreamingInfo.TileDataOffsetsOffset[0];
-			TileDataTask.SrcTileDataOffsetsPtrs[1] = SrcPtr + MipLevelStreamingInfo.TileDataOffsetsOffset[1];
-			TileDataTask.SrcTileDataPtrs[0] = SrcPtr + MipLevelStreamingInfo.TileDataOffset[0];
-			TileDataTask.SrcTileDataPtrs[1] = SrcPtr + MipLevelStreamingInfo.TileDataOffset[1];
-			TileDataTask.SrcPhysicalTileCoords = reinterpret_cast<const uint8*>(TileAllocations.GetData());
-			TileDataTask.TileDataBaseOffsets = TileDataAddResult.TileDataBaseOffsets;
-			TileDataTask.TileDataSizes = MipLevelStreamingInfo.TileDataSize;
-			TileDataTask.NumPhysicalTiles = NumPhysicalTiles;
-
-			FUploadTask& Task = UploadTasks.AddDefaulted_GetRef();
-			Task.Union.SetSubtype<FUploadTask::FTileDataTask>(TileDataTask);
-		}
-
-		// Page table
-		{
-			FUploadTask::FPageTableTask PageTableTask = {};
-			PageTableTask.PendingMipLevel = &PendingMipLevel;
-			PageTableTask.SrcPageCoords = SrcPtr + MipLevelStreamingInfo.PageTableOffset;
-			PageTableTask.SrcPageEntries = SrcPtr + MipLevelStreamingInfo.PageTableOffset + NumPageTableUpdates * sizeof(uint32);
-			PageTableTask.NumPageTableUpdates = NumPageTableUpdates;
-
-			FUploadTask& Task = UploadTasks.AddDefaulted_GetRef();
-			Task.Union.SetSubtype<FUploadTask::FPageTableTask>(PageTableTask);
+			TileDataTask.DstPhysicalTileCoordsPtr = TileDataAddResult.PackedPhysicalTileCoordsPtr;
+			TileDataTask.SrcOccupancyBitsPtrs[0] = TileSrcPtr + TileInfo.OccupancyBitsOffsets[0];
+			TileDataTask.SrcOccupancyBitsPtrs[1] = TileSrcPtr + TileInfo.OccupancyBitsOffsets[1];
+			TileDataTask.SrcVoxelDataPtrs[0] = TileSrcPtr + TileInfo.VoxelDataOffsets[0];
+			TileDataTask.SrcVoxelDataPtrs[1] = TileSrcPtr + TileInfo.VoxelDataOffsets[1];
+			TileDataTask.VoxelDataSizes = TileInfo.VoxelDataSizes;
+			TileDataTask.VoxelDataBaseOffsets = TileDataAddResult.TileDataBaseOffsets;
+			TileDataTask.PhysicalTileCoord = FrameInfo.TileAllocations[TileIndex];
 		}
 
 		// Cleanup
@@ -1444,82 +1453,48 @@ void FStreamingManager::InstallReadyMipLevels()
 		const int32 LRUNodeIndex = PendingMipLevel.FrameIndex * SVTInfo->NumMipLevelsGlobal + PendingMipLevel.MipLevelIndex;
 		SVTInfo->LRUNodes[LRUNodeIndex].PendingMipLevelIndex = INDEX_NONE;
 
-		const FPageTopology::FMip& TopologyMipInfo = FrameInfo.Resources->Topology.MipInfo[PendingMipLevel.MipLevelIndex];
+		const FPageTopology::FMip& TopologyMipInfo = Topology.MipInfo[PendingMipLevel.MipLevelIndex];
 		FrameInfo.ResidentPagesNew.SetRange(TopologyMipInfo.PageOffset, TopologyMipInfo.PageCount, true);
 	}
 
 	// Do all the memcpy's in parallel
-	ParallelFor(UploadTasks.Num(), [&](int32 TaskIndex)
-		{
-			FUploadTask& Task = UploadTasks[TaskIndex];
+	ParallelFor(TEXT("SVT::UploadTileDataTasks"), UploadTasks.Num(), 8, [&](int32 TaskIndex)
+	{
+		TRACE_CPUPROFILER_EVENT_SCOPE(SVT::StreamingTileDataUpload);
 
-			if (Task.Union.HasSubtype<FUploadTask::FPageTableTask>())
+		FTileDataTask& TileDataTask = UploadTasks[TaskIndex];
+		for (int32 TexIdx = 0; TexIdx < 2; ++TexIdx)
+		{
+			if (TileDataTask.DstOccupancyBitsPtrs[TexIdx])
 			{
-				TRACE_CPUPROFILER_EVENT_SCOPE(SVT::StreamingPageTableUpload);
-				FUploadTask::FPageTableTask& PageTableTask = Task.Union.GetSubtype<FUploadTask::FPageTableTask>();
-				if (PageTableTask.NumPageTableUpdates > 0)
-				{
-					const int32 FrameIndex = PageTableTask.PendingMipLevel->FrameIndex;
-					const int32 MipLevelIndex = PageTableTask.PendingMipLevel->MipLevelIndex;
-					FStreamingInfo* SVTInfo = FindStreamingInfo(PageTableTask.PendingMipLevel->SparseVolumeTexture);
-					FFrameInfo& FrameInfo = SVTInfo->PerFrameInfo[FrameIndex];
-					TArray<uint32>& TileAllocations = FrameInfo.TileAllocations[MipLevelIndex];
-					const uint32* SrcEntries = reinterpret_cast<const uint32*>(PageTableTask.SrcPageEntries);
-					uint32* EntriesForBookKeeping = FrameInfo.PageEntries.GetData() + FrameInfo.Resources->Topology.MipInfo[MipLevelIndex].PageOffset;
-					const uint32 MipLevelEntry = static_cast<uint32>(MipLevelIndex) << 24u;					
-					for (int32 i = 0; i < PageTableTask.NumPageTableUpdates; ++i)
-					{
-						const uint32 Entry = TileAllocations[SrcEntries[i]] | MipLevelEntry;
-						EntriesForBookKeeping[i] = Entry;
-					}
-				}
+				FMemory::Memcpy(TileDataTask.DstOccupancyBitsPtrs[TexIdx], TileDataTask.SrcOccupancyBitsPtrs[TexIdx], SVT::NumOccupancyWordsPerPaddedTile * sizeof(uint32));
 			}
-			else if(Task.Union.HasSubtype<FUploadTask::FTileDataTask>())
+			if (TileDataTask.VoxelDataSizes[TexIdx] > 0)
 			{
-				TRACE_CPUPROFILER_EVENT_SCOPE(SVT::StreamingTileDataUpload);
-				FUploadTask::FTileDataTask& TileDataTask = Task.Union.GetSubtype<FUploadTask::FTileDataTask>();
-				for (int32 i = 0; i < 2; ++i)
-				{
-					if (TileDataTask.DstOccupancyBitsPtrs[i])
-					{
-						FMemory::Memcpy(TileDataTask.DstOccupancyBitsPtrs[i], TileDataTask.SrcOccupancyBitsPtrs[i], TileDataTask.NumPhysicalTiles * SVT::NumOccupancyWordsPerPaddedTile * sizeof(uint32));
-					}
-					if (TileDataTask.TileDataSizes[i] > 0)
-					{
-						FMemory::Memcpy(TileDataTask.DstTileDataPtrs[i], TileDataTask.SrcTileDataPtrs[i], TileDataTask.TileDataSizes[i]);
-					}
-					if (TileDataTask.DstTileDataOffsetsPtrs[i])
-					{
-						uint32* Dst = reinterpret_cast<uint32*>(TileDataTask.DstTileDataOffsetsPtrs[i]);
-						const uint32* Src = reinterpret_cast<const uint32*>(TileDataTask.SrcTileDataOffsetsPtrs[i]);
-						for (int32 TileIndex = 0; TileIndex < TileDataTask.NumPhysicalTiles; ++TileIndex)
-						{
-							Dst[TileIndex] = Src[TileIndex] + TileDataTask.TileDataBaseOffsets[i];
-						}
-					}
-				}
-				FMemory::Memcpy(TileDataTask.DstPhysicalTileCoords, TileDataTask.SrcPhysicalTileCoords, TileDataTask.NumPhysicalTiles * sizeof(uint32));
+				FMemory::Memcpy(TileDataTask.DstTileDataPtrs[TexIdx], TileDataTask.SrcVoxelDataPtrs[TexIdx], TileDataTask.VoxelDataSizes[TexIdx]);
 			}
-			else
+			if (TileDataTask.DstTileDataOffsetsPtrs[TexIdx])
 			{
-				checkNoEntry();
+				FMemory::Memcpy(TileDataTask.DstTileDataOffsetsPtrs[TexIdx], &TileDataTask.VoxelDataBaseOffsets[TexIdx], sizeof(uint32));
 			}
-		});
+		}
+		FMemory::Memcpy(TileDataTask.DstPhysicalTileCoordsPtr, &TileDataTask.PhysicalTileCoord, sizeof(uint32));
+	});
 
 	ParallelFor(UploadCleanupTasks.Num(), [&](int32 TaskIndex)
-		{
-			TRACE_CPUPROFILER_EVENT_SCOPE(SVT::StreamingUploadCleanupTask);
+	{
+		TRACE_CPUPROFILER_EVENT_SCOPE(SVT::StreamingUploadCleanupTask);
 
-			FPendingMipLevel* PendingMipLevel = UploadCleanupTasks[TaskIndex];
+		FPendingMipLevel* PendingMipLevel = UploadCleanupTasks[TaskIndex];
 #if WITH_EDITORONLY_DATA
-			PendingMipLevel->SharedBuffer.Reset();
+		PendingMipLevel->SharedBuffer.Reset();
 #endif
-			if (!PendingMipLevel->Request.IsNone())
-			{
-				check(PendingMipLevel->Request.IsCompleted());
-				PendingMipLevel->Request.Reset();
-			}
-		});
+		if (!PendingMipLevel->Request.IsNone())
+		{
+			check(PendingMipLevel->Request.IsCompleted());
+			PendingMipLevel->Request.Reset();
+		}
+	});
 
 #if DO_CHECK // Clear processed pending mip levels for better debugging
 	for (int32 i = 0; i < AsyncState.NumReadyMipLevels; ++i)
@@ -1551,34 +1526,31 @@ void FStreamingManager::PatchPageTable(FRDGBuilder& GraphBuilder)
 		const FPageTopology& Topology = FrameInfo.Resources->Topology;
 
 		// Get all pages that were streamed in or out in this streaming update.
-		const TBitArray<> ResidentPagesDiff = TBitArray<>::BitwiseXOR(FrameInfo.ResidentPages, FrameInfo.ResidentPagesNew, EBitwiseOperatorFlags::MaxSize);
+		TBitArray<> ResidentPagesDiff = TBitArray<>::BitwiseXOR(FrameInfo.ResidentPages, FrameInfo.ResidentPagesNew, EBitwiseOperatorFlags::MaxSize);
 		
 		// Initialize InvalidatedPages with the diff. In the following loop, we then find all descendants of the newly streamed in/out pages and also mark them as invalidated.
 		FrameInfo.InvalidatedPages = ResidentPagesDiff;
+		// Iterate over all pages that were NOT changed in this update. We then dig down into their parents and try to figure out if they were newly streamed in/out. If so, mark the current page as also invalidated.
+		ResidentPagesDiff.BitwiseNOT();
 		for (TConstSetBitIterator It(ResidentPagesDiff); It; ++It)
 		{
 			const int32 PageIndex = It.GetIndex();
-			check(Topology.Pages.IsValidIndex(PageIndex));
+			check(Topology.IsValidPageIndex(PageIndex));
 
-			auto InvalidateDescendants = [&](uint32 InPageIndex, auto& InRecursiveLambda) -> void
+			// Only try to invalidate this page if it is not already resident in GPU memory
+			if (!FrameInfo.InvalidatedPages[PageIndex] && !FrameInfo.ResidentPages[PageIndex])
 			{
-				// Mip0 pages are leaf-nodes in the topology and don't have children
-				const bool bIsInteriorPage = Topology.InteriorPageData.IsValidIndex(InPageIndex);
-				if (bIsInteriorPage)
+				uint32 ParentPageIndex = Topology.ParentIndices[PageIndex];
+				while (ParentPageIndex != INDEX_NONE)
 				{
-					for (uint32 ChildPageIndex : Topology.InteriorPageData[InPageIndex].ChildIndices)
+					if (FrameInfo.InvalidatedPages[ParentPageIndex])
 					{
-						// Only process child if it is not already resident or included in the set of invalidated pages
-						if (ChildPageIndex != INDEX_NONE && !FrameInfo.InvalidatedPages[ChildPageIndex] && !FrameInfo.ResidentPages[ChildPageIndex])
-						{
-							FrameInfo.InvalidatedPages[ChildPageIndex] = true;
-							InRecursiveLambda(ChildPageIndex, InRecursiveLambda); // Let the child node walk its own children
-						}
+						FrameInfo.InvalidatedPages[PageIndex] = true;
+						break;
 					}
+					ParentPageIndex = Topology.ParentIndices[ParentPageIndex];
 				}
-			};
-
-			InvalidateDescendants(PageIndex, InvalidateDescendants);
+			}
 		}
 
 		NumUpdates += FrameInfo.InvalidatedPages.CountSetBits();
@@ -1606,7 +1578,7 @@ void FStreamingManager::PatchPageTable(FRDGBuilder& GraphBuilder)
 			for (TConstSetBitIterator It(FrameInfo.InvalidatedPages); It; ++It)
 			{
 				const int32 PageIndex = It.GetIndex();
-				check(Topology.Pages.IsValidIndex(PageIndex));
+				check(Topology.IsValidPageIndex(PageIndex));
 
 				auto IsInMipRange = [](const FPageTopology& InTopology, uint32 InIndex, int32 InMipLevel)
 				{
@@ -1638,35 +1610,34 @@ void FStreamingManager::PatchPageTable(FRDGBuilder& GraphBuilder)
 				uint32 PageTableEntry = 0;
 				if (FrameInfo.ResidentPagesNew[PageIndex])
 				{
-					// This page is already resident, so we can simply use its value in PageEntries
-					PageTableEntry = FrameInfo.PageEntries[PageIndex];
+					// This page is already resident, so we can simply use its value from TileAllocations
+					PageTableEntry = FrameInfo.TileAllocations[Topology.TileIndices[PageIndex]];
 					check(PageTableEntry);
 					PageTableEntry |= MipLevel << 24u;
 				}
 				else
 				{
 					// This page is not resident but needs a fallback value written to it, so we probe the parent pages until we find a resident one
-					uint32 ParentPageIndex = Topology.Pages[PageIndex].ParentIndex;
+					uint32 ParentPageIndex = Topology.ParentIndices[PageIndex];
 					int32 ParentMipLevel = MipLevel + 1;
 					while (ParentPageIndex != INDEX_NONE)
 					{
 						// The parent page is resident in GPU memory, so we can use it's cached value in PageEntries.
 						if (FrameInfo.ResidentPagesNew[ParentPageIndex])
 						{
-							PageTableEntry = FrameInfo.PageEntries[ParentPageIndex];
+							PageTableEntry = FrameInfo.TileAllocations[Topology.TileIndices[ParentPageIndex]];
 							check(PageTableEntry);
 							PageTableEntry |= ParentMipLevel << 24u;
-							FrameInfo.PageEntries[PageIndex] = PageTableEntry; // Don't forget to update the PageEntries value of our current page to reflect that it falls back to a coarser mip.
 							break;
 						}
-						ParentPageIndex = Topology.Pages[ParentPageIndex].ParentIndex;
+						ParentPageIndex = Topology.ParentIndices[ParentPageIndex];
 						++ParentMipLevel;
 					}
 					check(ParentPageIndex != INDEX_NONE); // If we hit this, then we tried to find the root node's parent. This should never happen as the root node should always be resident.
 				}
 
 				// Write the update to the upload buffer pointers
-				reinterpret_cast<uint32*>(DstCoordsPtr)[MipUpdateWriteIndex] = Topology.Pages[PageIndex].PackedPageTableCoord;
+				reinterpret_cast<uint32*>(DstCoordsPtr)[MipUpdateWriteIndex] = Topology.PackedPageTableCoords[PageIndex];
 				reinterpret_cast<uint32*>(DstEntryPtr)[MipUpdateWriteIndex] = PageTableEntry;
 				++MipUpdateWriteIndex;
 			}
@@ -1687,9 +1658,35 @@ FStreamingManager::FStreamingInfo* FStreamingManager::FindStreamingInfo(UStreama
 	return SVTInfoPtr ? SVTInfoPtr->Get() : nullptr;
 }
 
+FStreamingManager::FMipTileReadInfo FStreamingManager::GetMipTileReadInfo(const FFrameInfo& FrameInfo, int32 MipLevel)
+{
+	check(FrameInfo.Resources->Topology.MipInfo.IsValidIndex(MipLevel));
+	const bool bIsRootMipLevel = (MipLevel + 1) == FrameInfo.NumMipLevels;
+	const FPageTopology::FMip& MipInfo = FrameInfo.Resources->Topology.MipInfo[MipLevel];
+
+	uint32 TileRangeOffset = 0;
+	uint32 TileRangeCount = 0;
+	FrameInfo.Resources->Topology.GetTileRange(MipInfo.PageOffset, MipInfo.PageCount, TileRangeOffset, TileRangeCount);
+
+	const uint32 ReadOffsetLogical = FrameInfo.Resources->StreamingMetaData.TileDataOffsets[TileRangeOffset];
+	// TileDataOffsets are expressed as if all data is stored in the same buffer, when the root tile is actually stored separately. The root tile can't stream, so we simply subtract its size from the offset.
+	const uint32 RootTileSize = FrameInfo.Resources->RootData.Num();
+	check(RootTileSize == 0 || RootTileSize == FrameInfo.Resources->StreamingMetaData.GetTileMemorySize(0));
+	const uint32 ReadOffset = bIsRootMipLevel ? ReadOffsetLogical : (ReadOffsetLogical - RootTileSize);
+	const uint32 ReadSize = FrameInfo.Resources->StreamingMetaData.TileDataOffsets[TileRangeOffset + TileRangeCount] - ReadOffsetLogical;
+
+	FMipTileReadInfo ReadInfo;
+	ReadInfo.TileOffset = TileRangeOffset;
+	ReadInfo.TileCount = TileRangeCount;
+	ReadInfo.ReadOffset = ReadOffset;
+	ReadInfo.ReadSize = ReadSize;
+
+	return ReadInfo;
+}
+
 #if WITH_EDITORONLY_DATA
 
-UE::DerivedData::FCacheGetChunkRequest FStreamingManager::BuildDDCRequest(const FResources& Resources, const FMipLevelStreamingInfo& MipLevelStreamingInfo, const uint32 PendingMipLevelIndex)
+UE::DerivedData::FCacheGetChunkRequest FStreamingManager::BuildDDCRequest(const FResources& Resources, uint64 ReadOffset, uint64 ReadSize, uint32 PendingMipLevelIndex)
 {
 	using namespace UE::DerivedData;
 
@@ -1701,8 +1698,8 @@ UE::DerivedData::FCacheGetChunkRequest FStreamingManager::BuildDDCRequest(const
 	FCacheGetChunkRequest Request;
 	Request.Id = FValueId::FromName("SparseVolumeTextureStreamingData");
 	Request.Key = Key;
-	Request.RawOffset = MipLevelStreamingInfo.BulkOffset;
-	Request.RawSize = MipLevelStreamingInfo.BulkSize;
+	Request.RawOffset = ReadOffset;
+	Request.RawSize = ReadSize;
 	Request.RawHash = Resources.DDCRawHash;
 	Request.UserData = (((uint64)PendingMipLevelIndex) << uint64(32)) | (uint64)PendingMipLevels[PendingMipLevelIndex].RequestVersion;
 	return Request;
diff --git a/Engine/Source/Runtime/Engine/Private/SparseVolumeTexture/SparseVolumeTextureStreamingManager.h b/Engine/Source/Runtime/Engine/Private/SparseVolumeTexture/SparseVolumeTextureStreamingManager.h
--- a/Engine/Source/Runtime/Engine/Private/SparseVolumeTexture/SparseVolumeTextureStreamingManager.h
+++ b/Engine/Source/Runtime/Engine/Private/SparseVolumeTexture/SparseVolumeTextureStreamingManager.h
@@ -127,8 +127,7 @@ private:
 		int32 NumMipLevels; // Number of actual mip levels of this frame. Depending on virtual volume extents, frames can have different numbers of levels.
 		int32 LowestRequestedMipLevel; // Lowest mip level that should be resident. Can be lower than LowestResidentMipLevel when streaming in new mips. Stream-out is instant, so it should never be higher.
 		int32 LowestResidentMipLevel; // Actually resident on the GPU
-		TArray<TArray<uint32>> TileAllocations; // TileAllocations[MipLevel][PhysicalTileIndex]
-		TArray<uint32> PageEntries; // Flat array of sparse page table. Reflects the current state of the GPU page table
+		TArray<uint32> TileAllocations; // TileAllocations[PhysicalTileIndex]
 		TBitArray<> ResidentPages; //  One bit for every (non-zero) page (in the sparse page octree) for all mip levels, starting at the highest mip
 		TBitArray<> ResidentPagesNew; // Reflects changes made during an update of the streaming system
 		TBitArray<> InvalidatedPages; // All pages that need updated page table entries
@@ -208,32 +207,17 @@ private:
 
 	// Encapsulates all the src/dst pointers for memcpy-ing the streamed data into GPU memory.
 	// See the comment on FMipLevelStreamingInfo for details about these pointers.
-	struct FUploadTask
+	struct FTileDataTask
 	{
-		struct FTileDataTask
-		{
-			TStaticArray<uint8*, 2> DstOccupancyBitsPtrs;
-			TStaticArray<uint8*, 2> DstTileDataOffsetsPtrs;
-			TStaticArray<uint8*, 2> DstTileDataPtrs;
-			uint8* DstPhysicalTileCoords;
-			TStaticArray<const uint8*, 2> SrcOccupancyBitsPtrs;
-			TStaticArray<const uint8*, 2> SrcTileDataOffsetsPtrs;
-			TStaticArray<const uint8*, 2> SrcTileDataPtrs;
-			const uint8* SrcPhysicalTileCoords;
-			TStaticArray<uint32, 2> TileDataBaseOffsets;
-			TStaticArray<int32, 2> TileDataSizes;
-			int32 NumPhysicalTiles;
-		};
-
-		struct FPageTableTask
-		{
-			FPendingMipLevel* PendingMipLevel;
-			const uint8* SrcPageCoords;
-			const uint8* SrcPageEntries;
-			int32 NumPageTableUpdates;
-		};
-
-		TUnion<FTileDataTask, FPageTableTask> Union;
+		TStaticArray<uint8*, 2> DstOccupancyBitsPtrs;
+		TStaticArray<uint8*, 2> DstTileDataOffsetsPtrs;
+		TStaticArray<uint8*, 2> DstTileDataPtrs;
+		uint8* DstPhysicalTileCoordsPtr;
+		TStaticArray<const uint8*, 2> SrcOccupancyBitsPtrs;
+		TStaticArray<const uint8*, 2> SrcVoxelDataPtrs;
+		TStaticArray<uint32, 2> VoxelDataSizes;
+		TStaticArray<uint32, 2> VoxelDataBaseOffsets;
+		uint32 PhysicalTileCoord;
 	};
 
 	struct FAsyncState
@@ -255,6 +239,15 @@ private:
 		TArray<FFrameInfo> FrameInfo; // Only Resources and TextureRenderResources are initialized
 	};
 
+	// Helper struct representing a range of tiles within the cooked data corresponding to a given mip level
+	struct FMipTileReadInfo
+	{
+		uint32 TileOffset;
+		uint32 TileCount;
+		uint32 ReadOffset;
+		uint32 ReadSize;
+	};
+
 	TMap<UStreamableSparseVolumeTexture*, TUniquePtr<FStreamingInfo>> StreamingInfo; // Do not dereference the key! We just read the pointer itself.
 	TMap<FMipLevelKey, uint32> RequestsHashTable;
 	TArray<FPendingMipLevel> PendingMipLevels;
@@ -277,7 +270,7 @@ private:
 	TSet<FFrameInfo*> InvalidatedSVTFrames; // Set of SVT frames where pages have been streamed in or out. Used in PatchPageTable().
 	TArray<FStreamingRequest> PrioritizedRequestsHeap;
 	TArray<FStreamingRequest> SelectedRequests;
-	TArray<FUploadTask> UploadTasks; // accessed on the async thread
+	TArray<FTileDataTask> UploadTasks; // accessed on the async thread
 	TArray<FPendingMipLevel*> UploadCleanupTasks; // accessed on the async thread
 
 	void AddInternal(FRDGBuilder& GraphBuilder, FNewSparseVolumeTextureInfo&& NewSVTInfo);
@@ -292,9 +285,10 @@ private:
 	void InstallReadyMipLevels();
 	void PatchPageTable(FRDGBuilder& GraphBuilder); // Patches the page table to reflect streamed in/out pages and to ensure non-resident mip levels fall back to coarser mip level tile data
 	FStreamingInfo* FindStreamingInfo(UStreamableSparseVolumeTexture* Key); // Returns nullptr if the key can't be found
+	static FMipTileReadInfo GetMipTileReadInfo(const FFrameInfo& FrameInfo, int32 MipLevel);
 
 #if WITH_EDITORONLY_DATA
-	UE::DerivedData::FCacheGetChunkRequest BuildDDCRequest(const FResources& Resources, const FMipLevelStreamingInfo& MipLevelStreamingInfo, const uint32 PendingMipLevelIndex);
+	UE::DerivedData::FCacheGetChunkRequest BuildDDCRequest(const FResources& Resources, uint64 ReadOffset, uint64 ReadSize, uint32 PendingMipLevelIndex);
 	void RequestDDCData(TConstArrayView<UE::DerivedData::FCacheGetChunkRequest> DDCRequests, bool bBlocking);
 #endif
 };
diff --git a/Engine/Source/Runtime/Engine/Private/SparseVolumeTexture/SparseVolumeTextureTileDataTexture.cpp b/Engine/Source/Runtime/Engine/Private/SparseVolumeTexture/SparseVolumeTextureTileDataTexture.cpp
--- a/Engine/Source/Runtime/Engine/Private/SparseVolumeTexture/SparseVolumeTextureTileDataTexture.cpp
+++ b/Engine/Source/Runtime/Engine/Private/SparseVolumeTexture/SparseVolumeTextureTileDataTexture.cpp
@@ -135,6 +135,9 @@ void FTileDataTexture::BeginReserveUpload()
 void FTileDataTexture::ReserveUpload(int32 NumTiles, int32 NumVoxelsA, int32 NumVoxelsB)
 {
 	check(UploaderState == EUploaderState::Reserving);
+	check((NumReservedUploadTiles + NumTiles) >= 0);
+	check((NumReservedUploadVoxelsA + NumVoxelsA) >= 0);
+	check((NumReservedUploadVoxelsB + NumVoxelsB) >= 0);
 	NumReservedUploadTiles += NumTiles;
 	NumReservedUploadVoxelsA += NumVoxelsA;
 	NumReservedUploadVoxelsB += NumVoxelsB;
diff --git a/Engine/Source/Runtime/Engine/Private/SparseVolumeTexture/SparseVolumeTextureUpload.cpp b/Engine/Source/Runtime/Engine/Private/SparseVolumeTexture/SparseVolumeTextureUpload.cpp
--- a/Engine/Source/Runtime/Engine/Private/SparseVolumeTexture/SparseVolumeTextureUpload.cpp
+++ b/Engine/Source/Runtime/Engine/Private/SparseVolumeTexture/SparseVolumeTextureUpload.cpp
@@ -94,6 +94,9 @@ FTileUploader::FTileUploader()
 
 void FTileUploader::Init(FRDGBuilder& GraphBuilder, int32 InMaxNumTiles, int32 InMaxNumVoxelsA, int32 InMaxNumVoxelsB, EPixelFormat InFormatA, EPixelFormat InFormatB)
 {
+	check(InMaxNumTiles >= 0);
+	check(InMaxNumVoxelsA >= 0);
+	check(InMaxNumVoxelsB >= 0);
 	check(InFormatA != PF_Unknown || InFormatB != PF_Unknown);
 	ResetState();
 	MaxNumTiles = InMaxNumTiles;
diff --git a/Engine/Source/Runtime/Engine/Private/SparseVolumeTexture/SparseVolumeTextureUtility.cpp b/Engine/Source/Runtime/Engine/Private/SparseVolumeTexture/SparseVolumeTextureUtility.cpp
--- a/Engine/Source/Runtime/Engine/Private/SparseVolumeTexture/SparseVolumeTextureUtility.cpp
+++ b/Engine/Source/Runtime/Engine/Private/SparseVolumeTexture/SparseVolumeTextureUtility.cpp
@@ -24,6 +24,11 @@ namespace Private
 } // SVT
 } // UE
 
+uint32 UE::SVT::PackX11Y11Z10(const FIntVector3& Value)
+{
+	return (Value.X & 0x7FFu) | ((Value.Y & 0x7FFu) << 11u) | ((Value.Z & 0x3FFu) << 22u);
+}
+
 uint32 UE::SVT::PackPageTableEntry(const FIntVector3& Coord)
 {
 	uint32 Result = (Coord.X & 0xFFu) | ((Coord.Y & 0xFFu) << 8u) | ((Coord.Z & 0xFFu) << 16u);
diff --git a/Engine/Source/Runtime/Engine/Private/SparseVolumeTexture/SparseVolumeTextureUtility.h b/Engine/Source/Runtime/Engine/Private/SparseVolumeTexture/SparseVolumeTextureUtility.h
--- a/Engine/Source/Runtime/Engine/Private/SparseVolumeTexture/SparseVolumeTextureUtility.h
+++ b/Engine/Source/Runtime/Engine/Private/SparseVolumeTexture/SparseVolumeTextureUtility.h
@@ -16,7 +16,9 @@ namespace SVT
 	static constexpr int32 NumVoxelsPerTile = SPARSE_VOLUME_TILE_RES * SPARSE_VOLUME_TILE_RES * SPARSE_VOLUME_TILE_RES;
 	static constexpr int32 NumVoxelsPerPaddedTile = SPARSE_VOLUME_TILE_RES_PADDED * SPARSE_VOLUME_TILE_RES_PADDED * SPARSE_VOLUME_TILE_RES_PADDED;
 	static constexpr int32 NumOccupancyWordsPerPaddedTile = (NumVoxelsPerPaddedTile + 31) / 32;
+	static constexpr int32 OccupancyBitsSizePerPaddedTile = NumOccupancyWordsPerPaddedTile * sizeof(uint32);
 
+	uint32 PackX11Y11Z10(const FIntVector3& Value);
 	uint32 PackPageTableEntry(const FIntVector3& Coord);
 	FIntVector3 UnpackPageTableEntry(uint32 Packed);
 	FVector4f ReadVoxel(int64 VoxelIndex, const uint8* TileData, EPixelFormat Format);
diff --git a/Engine/Source/Runtime/Engine/Public/SparseVolumeTexture/SparseVolumeTextureData.h b/Engine/Source/Runtime/Engine/Public/SparseVolumeTexture/SparseVolumeTextureData.h
--- a/Engine/Source/Runtime/Engine/Public/SparseVolumeTexture/SparseVolumeTextureData.h
+++ b/Engine/Source/Runtime/Engine/Public/SparseVolumeTexture/SparseVolumeTextureData.h
@@ -55,19 +55,44 @@ struct FTextureData
 	ENGINE_API bool CreateFromDense(const FTextureDataCreateInfo& CreateInfo, const TArrayView64<uint8>& VoxelDataA, const TArrayView64<uint8>& VoxelDataB);
 	ENGINE_API void CreateDefault();
 	ENGINE_API FVector4f Load(const FIntVector3& VolumeCoord, int32 MipLevel, int32 AttributesIdx, const FTextureDataAddressingInfo& AddressingInfo) const;
-	ENGINE_API bool BuildDerivedData(const FTextureDataAddressingInfo& AddressingInfo, int32 NumMipLevelsGlobal, bool bMoveMip0FromThis, FTextureData& OutDerivedData);
+	ENGINE_API bool BuildDerivedData(const FTextureDataAddressingInfo& AddressingInfo, int32 NumMipLevelsGlobal, bool bMoveMip0FromThis, struct FDerivedTextureData& OutDerivedData);
 
 private:
 
 	bool GenerateMipMaps(const FTextureDataAddressingInfo& AddressingInfo, int32 NumMipLevels = -1);
 	bool GenerateBorderVoxels(const FTextureDataAddressingInfo& AddressingInfo, int32 MipLevel, const TArray<FIntVector3>& PageCoords);
-	bool DeduplicateTiles();
 
 	uint32 ReadPageTable(const FIntVector3& PageTableCoord, int32 MipLevel) const;
 	FVector4f ReadTileDataVoxel(int32 TileIndex, const FIntVector3& TileDataCoord, int32 MipLevel, int32 AttributesIdx) const;
 	void WriteTileDataVoxel(int32 TileIndex, const FIntVector3& TileDataCoord, int32 MipLevel, int32 AttributesIdx, const FVector4f& Value, int32 DstComponent = -1);
 };
 
+// Represents mip mapped SparseVolumeTexture data ready for compression/serializing. The data is stored as a set of arrays for the page table and a set of arrays for the tile data.
+// The page table arrays are essentially struct-of-arrays and store only those elements of the dense page page table which actually point to a tile (non-zero). Each stored page
+// consists of a coordinate in the 3D page table, an index to the tile it points to and another index to the parent page, which is the corresponding page in the next higher mip level.
+// Deduplicated tiles (of voxels) for all mip levels are stored in a single set of arrays as raw bytes.
+struct FDerivedTextureData
+{
+	// Pages are ordered by mip level. This struct represents the range of pages for a given mip level.
+	struct FMipPageRange
+	{
+		uint32 PageOffset;
+		uint32 PageCount;
+	};
+
+	FHeader Header = {};
+	uint32 NumPhysicalTiles;				// Number of tiles stored in PhysicalTileDataA/PhysicalTileDataB.
+	TArray<FMipPageRange> MipPageRanges;	// Page range for each mip level. Indexed by mip level.
+	TArray<uint32> PageTableCoords;			// 11|11|10 packed coords of the page within the logical dense 3D page table.
+	TArray<uint32> PageTableTileIndices;	// Index of tile in PhysicalTileDataA/PhysicalTileDataB pointed to by this page.
+	TArray<uint32> PageTableParentIndices;	// Index of parent page in the next higher mip level.
+	TArray64<uint8> PhysicalTileDataA;
+	TArray64<uint8> PhysicalTileDataB;
+
+	void Reset();
+	void Build(const FTextureData& MippedTextureData);
+};
+
 }
 }
 
