Renamed FPendingMipLevel to FPendingRequest (same for associated variables) in SVT streaming manager.

#rnx

diff --git a/Engine/Source/Runtime/Engine/Private/SparseVolumeTexture/SparseVolumeTextureStreamingManager.cpp b/Engine/Source/Runtime/Engine/Private/SparseVolumeTexture/SparseVolumeTextureStreamingManager.cpp
--- a/Engine/Source/Runtime/Engine/Private/SparseVolumeTexture/SparseVolumeTextureStreamingManager.cpp
+++ b/Engine/Source/Runtime/Engine/Private/SparseVolumeTexture/SparseVolumeTextureStreamingManager.cpp
@@ -74,11 +74,11 @@ static FAutoConsoleVariableRef CVarSVTStreamingPrintMemoryStats(
 	ECVF_RenderThreadSafe
 );
 
-static int32 GSVTStreamingMaxPendingMipLevels = 128;
-static FAutoConsoleVariableRef CVarSVTStreamingMaxPendingMipLevels(
-	TEXT("r.SparseVolumeTexture.Streaming.MaxPendingMipLevels"),
-	GSVTStreamingMaxPendingMipLevels,
-	TEXT("Maximum number of mip levels that can be pending for installation."),
+static int32 GSVTStreamingMaxPendingRequests = 128;
+static FAutoConsoleVariableRef CVarSVTStreamingMaxPendingRequests(
+	TEXT("r.SparseVolumeTexture.Streaming.MaxPendingRequests"),
+	GSVTStreamingMaxPendingRequests,
+	TEXT("Maximum number of IO requests that can be pending for installation."),
 	ECVF_RenderThreadSafe | ECVF_ReadOnly
 );
 
@@ -115,7 +115,7 @@ public:
 
 	void DoTask(ENamedThreads::Type CurrentThread, const FGraphEventRef& MyCompletionGraphEvent)
 	{
-		Parameters.StreamingManager->InstallReadyMipLevels();
+		Parameters.StreamingManager->InstallReadyRequests();
 	}
 
 	static ESubsequentsMode::Type	GetSubsequentsMode() { return ESubsequentsMode::TrackSubsequents; }
@@ -137,8 +137,8 @@ void FStreamingManager::InitRHI(FRHICommandListBase& RHICmdList)
 		return;
 	}
 
-	MaxPendingMipLevels = GSVTStreamingMaxPendingMipLevels;
-	PendingMipLevels.SetNum(MaxPendingMipLevels);
+	MaxPendingRequests = GSVTStreamingMaxPendingRequests;
+	PendingRequests.SetNum(MaxPendingRequests);
 	PageTableUpdater = MakeUnique<FPageTableUpdater>();
 
 #if WITH_EDITORONLY_DATA
@@ -420,22 +420,22 @@ void FStreamingManager::BeginAsyncUpdate(FRDGBuilder& GraphBuilder, bool bBlocki
 	}
 
 	AddParentRequests();
-	const int32 MaxSelectedRequests = MaxPendingMipLevels - NumPendingMipLevels;
+	const int32 MaxSelectedRequests = MaxPendingRequests - NumPendingRequests;
 	SelectHighestPriorityRequestsAndUpdateLRU(MaxSelectedRequests);
 	IssueRequests(MaxSelectedRequests);
-	AsyncState.NumReadyMipLevels = DetermineReadyMipLevels();
+	AsyncState.NumReadyRequests = DetermineReadyRequests();
 
 	// Do a first pass over all the mips to be uploaded to compute the upload buffer size requirements.
 	TileDataTexturesToUpdate.Reset();
 	{
-		const int32 StartPendingMipLevelIndex = (NextPendingMipLevelIndex + MaxPendingMipLevels - NumPendingMipLevels) % MaxPendingMipLevels;
-		for (int32 i = 0; i < AsyncState.NumReadyMipLevels; ++i)
+		const int32 StartPendingRequestIndex = (NextPendingRequestIndex + MaxPendingRequests - NumPendingRequests) % MaxPendingRequests;
+		for (int32 i = 0; i < AsyncState.NumReadyRequests; ++i)
 		{
-			const int32 PendingMipLevelIndex = (StartPendingMipLevelIndex + i) % MaxPendingMipLevels;
-			FPendingMipLevel& PendingMipLevel = PendingMipLevels[PendingMipLevelIndex];
+			const int32 PendingRequestIndex = (StartPendingRequestIndex + i) % MaxPendingRequests;
+			FPendingRequest& PendingRequest = PendingRequests[PendingRequestIndex];
 
-			FStreamingInfo* SVTInfo = FindStreamingInfo(PendingMipLevel.SparseVolumeTexture);
-			if (!SVTInfo || (SVTInfo->PerFrameInfo[PendingMipLevel.FrameIndex].LowestRequestedMipLevel > PendingMipLevel.MipLevelIndex))
+			FStreamingInfo* SVTInfo = FindStreamingInfo(PendingRequest.SparseVolumeTexture);
+			if (!SVTInfo || (SVTInfo->PerFrameInfo[PendingRequest.FrameIndex].LowestRequestedMipLevel > PendingRequest.MipLevelIndex))
 			{
 				continue; // Skip mip level install. SVT no longer exists or mip level was "streamed out" before it was even installed in the first place.
 			}
@@ -450,8 +450,8 @@ void FStreamingManager::BeginAsyncUpdate(FRDGBuilder& GraphBuilder, bool bBlocki
 
 			const int32 FormatSizeA = GPixelFormats[SVTInfo->FormatA].BlockBytes;
 			const int32 FormatSizeB = GPixelFormats[SVTInfo->FormatB].BlockBytes;
-			const FResources* Resources = SVTInfo->PerFrameInfo[PendingMipLevel.FrameIndex].Resources;
-			const FPageTopology::FMip& MipInfo = Resources->Topology.MipInfo[PendingMipLevel.MipLevelIndex];
+			const FResources* Resources = SVTInfo->PerFrameInfo[PendingRequest.FrameIndex].Resources;
+			const FPageTopology::FMip& MipInfo = Resources->Topology.MipInfo[PendingRequest.MipLevelIndex];
 			
 			uint32 TileRangeOffset = 0;
 			uint32 TileRangeCount = 0;
@@ -484,7 +484,7 @@ void FStreamingManager::BeginAsyncUpdate(FRDGBuilder& GraphBuilder, bool bBlocki
 	}
 	else
 	{
-		InstallReadyMipLevels();
+		InstallReadyRequests();
 	}
 }
 
@@ -519,8 +519,8 @@ void FStreamingManager::EndAsyncUpdate(FRDGBuilder& GraphBuilder)
 	// Update page table with newly streamed in/out pages and make sure descendant pages in the hierarchy have correct fallback values
 	PatchPageTable(GraphBuilder);
 
-	check(AsyncState.NumReadyMipLevels <= NumPendingMipLevels);
-	NumPendingMipLevels -= AsyncState.NumReadyMipLevels;
+	check(AsyncState.NumReadyRequests <= NumPendingRequests);
+	NumPendingRequests -= AsyncState.NumReadyRequests;
 	++NextUpdateIndex;
 	AsyncState.bUpdateActive = false;
 	AsyncState.bUpdateIsAsync = false;
@@ -779,11 +779,11 @@ void FStreamingManager::RemoveInternal(UStreamableSparseVolumeTexture* SparseVol
 		}
 
 		// Cancel any pending mip levels
-		for (FPendingMipLevel& PendingMipLevel : PendingMipLevels)
+		for (FPendingRequest& PendingRequest : PendingRequests)
 		{
-			if (PendingMipLevel.SparseVolumeTexture == SparseVolumeTexture)
+			if (PendingRequest.SparseVolumeTexture == SparseVolumeTexture)
 			{
-				PendingMipLevel.Reset();
+				PendingRequest.Reset();
 			}
 		}
 
@@ -989,7 +989,7 @@ void FStreamingManager::IssueRequests(int32 MaxSelectedRequests)
 		{
 #if SVT_STREAMING_LOG_VERBOSE
 			UE_LOG(LogSparseVolumeTextureStreamingManager, Display, TEXT("(%i)%i IssueRequests() Frame %i Mip %i: Not enough tiles available (%i) to fit mip level (%i)"), 
-				NextUpdateIndex, NextPendingMipLevelIndex, SelectedKey.FrameIndex, SelectedKey.MipLevelIndex, NumAvailableTiles, NumRequiredTiles);
+				NextUpdateIndex, NextPendingRequestIndex, SelectedKey.FrameIndex, SelectedKey.MipLevelIndex, NumAvailableTiles, NumRequiredTiles);
 #endif
 
 			// Try to free old mip levels, starting at higher resolution mips and going up the mip chain
@@ -1045,20 +1045,20 @@ void FStreamingManager::IssueRequests(int32 MaxSelectedRequests)
 		}
 
 #if DO_CHECK
-		for (auto& Pending : PendingMipLevels)
+		for (auto& Pending : PendingRequests)
 		{
 			check(Pending.SparseVolumeTexture != SelectedKey.SVT || Pending.FrameIndex != SelectedKey.FrameIndex || Pending.MipLevelIndex != SelectedKey.MipLevelIndex); //-V1013
 		}
 #endif
 
-		const int32 PendingMipLevelIndex = NextPendingMipLevelIndex;
-		FPendingMipLevel& PendingMipLevel = PendingMipLevels[PendingMipLevelIndex];
-		PendingMipLevel.Reset();
-		PendingMipLevel.SparseVolumeTexture = SelectedKey.SVT;
-		PendingMipLevel.FrameIndex = SelectedKey.FrameIndex;
-		PendingMipLevel.MipLevelIndex = SelectedKey.MipLevelIndex;
-		PendingMipLevel.IssuedInFrame = NextUpdateIndex;
-		PendingMipLevel.bBlocking = GSVTStreamingForceBlockingRequests || (SelectedRequest.Priority == FStreamingRequest::BlockingPriority);
+		const int32 PendingRequestIndex = NextPendingRequestIndex;
+		FPendingRequest& PendingRequest = PendingRequests[PendingRequestIndex];
+		PendingRequest.Reset();
+		PendingRequest.SparseVolumeTexture = SelectedKey.SVT;
+		PendingRequest.FrameIndex = SelectedKey.FrameIndex;
+		PendingRequest.MipLevelIndex = SelectedKey.MipLevelIndex;
+		PendingRequest.IssuedInFrame = NextUpdateIndex;
+		PendingRequest.bBlocking = GSVTStreamingForceBlockingRequests || (SelectedRequest.Priority == FStreamingRequest::BlockingPriority);
 
 		const FByteBulkData& BulkData = Resources->StreamableMipLevels;
 #if WITH_EDITORONLY_DATA
@@ -1072,8 +1072,8 @@ void FStreamingManager::IssueRequests(int32 MaxSelectedRequests)
 		{
 			if (Resources->ResourceFlags & EResourceFlag_StreamingDataInDDC)
 			{
-				UE::DerivedData::FCacheGetChunkRequest DDCRequest = BuildDDCRequest(*Resources, MipTileReadInfo.ReadOffset, MipTileReadInfo.ReadSize, NextPendingMipLevelIndex);
-				if (PendingMipLevel.bBlocking)
+				UE::DerivedData::FCacheGetChunkRequest DDCRequest = BuildDDCRequest(*Resources, MipTileReadInfo.ReadOffset, MipTileReadInfo.ReadSize, NextPendingRequestIndex);
+				if (PendingRequest.bBlocking)
 				{
 					DDCRequestsBlocking.Add(DDCRequest);
 				}
@@ -1081,29 +1081,29 @@ void FStreamingManager::IssueRequests(int32 MaxSelectedRequests)
 				{
 					DDCRequests.Add(DDCRequest);
 				}
-				PendingMipLevel.State = FPendingMipLevel::EState::DDC_Pending;
+				PendingRequest.State = FPendingRequest::EState::DDC_Pending;
 			}
 			else
 			{
-				PendingMipLevel.State = FPendingMipLevel::EState::Memory;
+				PendingRequest.State = FPendingRequest::EState::Memory;
 			}
 		}
 		else
 #endif
 		{
-			PendingMipLevel.RequestBuffer = FIoBuffer(MipTileReadInfo.ReadSize); // SVT_TODO: Use FIoBuffer::Wrap with preallocated memory
-			const EAsyncIOPriorityAndFlags Priority = PendingMipLevel.bBlocking ? AIOP_CriticalPath : AIOP_Low;
-			Batch.Read(BulkData, MipTileReadInfo.ReadOffset, MipTileReadInfo.ReadSize, Priority, PendingMipLevel.RequestBuffer, PendingMipLevel.Request);
+			PendingRequest.RequestBuffer = FIoBuffer(MipTileReadInfo.ReadSize); // SVT_TODO: Use FIoBuffer::Wrap with preallocated memory
+			const EAsyncIOPriorityAndFlags Priority = PendingRequest.bBlocking ? AIOP_CriticalPath : AIOP_Low;
+			Batch.Read(BulkData, MipTileReadInfo.ReadOffset, MipTileReadInfo.ReadSize, Priority, PendingRequest.RequestBuffer, PendingRequest.Request);
 			bIssueIOBatch = true;
 
 #if WITH_EDITORONLY_DATA
-			PendingMipLevel.State = FPendingMipLevel::EState::Disk;
+			PendingRequest.State = FPendingRequest::EState::Disk;
 #endif
 		}
 
-		NextPendingMipLevelIndex = (NextPendingMipLevelIndex + 1) % MaxPendingMipLevels;
-		check(NumPendingMipLevels < MaxPendingMipLevels);
-		++NumPendingMipLevels;
+		NextPendingRequestIndex = (NextPendingRequestIndex + 1) % MaxPendingRequests;
+		check(NumPendingRequests < MaxPendingRequests);
+		++NumPendingRequests;
 
 		FFrameInfo& FrameInfo = SVTInfo->PerFrameInfo[SelectedKey.FrameIndex];
 
@@ -1131,7 +1131,7 @@ void FStreamingManager::IssueRequests(int32 MaxSelectedRequests)
 			FLRUNode* LRUNode = &SVTInfo->LRUNodes[LRUNodeIndex];
 			check(!LRUNode->IsInList());
 			LRUNode->LastRequested = NextUpdateIndex;
-			LRUNode->PendingMipLevelIndex = PendingMipLevelIndex;
+			LRUNode->PendingRequestIndex = PendingRequestIndex;
 
 			FLRUNode* Dependency = LRUNode->NextHigherMipLevel;
 			while (Dependency)
@@ -1145,7 +1145,7 @@ void FStreamingManager::IssueRequests(int32 MaxSelectedRequests)
 
 #if SVT_STREAMING_LOG_VERBOSE
 		UE_LOG(LogSparseVolumeTextureStreamingManager, Display, TEXT("(%i)%i StreamIn Frame %i OldReqMip %i, NewReqMip %i, ResMip %i"),
-			PendingMipLevel.IssuedInFrame, PendingMipLevelIndex,
+			PendingRequest.IssuedInFrame, PendingRequestIndex,
 			SelectedKey.FrameIndex, 
 			FrameInfo.LowestRequestedMipLevel, SelectedKey.MipLevelIndex, 
 			FrameInfo.LowestResidentMipLevel);
@@ -1185,18 +1185,18 @@ void FStreamingManager::StreamOutMipLevel(FStreamingInfo* SVTInfo, FLRUNode* LRU
 	check(FrameInfo.LowestRequestedMipLevel == MipLevelIndex);
 
 	// Cancel potential IO request
-	check((MipLevelIndex < FrameInfo.LowestResidentMipLevel) == (LRUNode->PendingMipLevelIndex != INDEX_NONE));
-	if (LRUNode->PendingMipLevelIndex != INDEX_NONE)
+	check((MipLevelIndex < FrameInfo.LowestResidentMipLevel) == (LRUNode->PendingRequestIndex != INDEX_NONE));
+	if (LRUNode->PendingRequestIndex != INDEX_NONE)
 	{
-		PendingMipLevels[LRUNode->PendingMipLevelIndex].Reset();
-		LRUNode->PendingMipLevelIndex = INDEX_NONE;
+		PendingRequests[LRUNode->PendingRequestIndex].Reset();
+		LRUNode->PendingRequestIndex = INDEX_NONE;
 	}
 
 	const int32 NewLowestRequestedMipLevel = MipLevelIndex + 1;
 	const int32 NewLowestResidentMipLevel = FMath::Max(MipLevelIndex + 1, FrameInfo.LowestResidentMipLevel);
 #if SVT_STREAMING_LOG_VERBOSE
 	UE_LOG(LogSparseVolumeTextureStreamingManager, Display, TEXT("(%i)%i StreamOut Frame %i OldReqMip %i, NewReqMip %i, OldResMip %i, NewResMip %i"),
-		NextUpdateIndex, NextPendingMipLevelIndex,
+		NextUpdateIndex, NextPendingRequestIndex,
 		FrameIndex,
 		FrameInfo.LowestRequestedMipLevel, NewLowestRequestedMipLevel,
 		FrameInfo.LowestResidentMipLevel, NewLowestResidentMipLevel);
@@ -1232,65 +1232,65 @@ void FStreamingManager::StreamOutMipLevel(FStreamingInfo* SVTInfo, FLRUNode* LRU
 	FrameInfo.ResidentPagesNew.SetRange(TopologyMipInfo.PageOffset, TopologyMipInfo.PageCount, false);
 }
 
-int32 FStreamingManager::DetermineReadyMipLevels()
+int32 FStreamingManager::DetermineReadyRequests()
 {
 	using namespace UE::DerivedData;
 
 	TRACE_CPUPROFILER_EVENT_SCOPE(SVT::StreamingDetermineReadyRequests);
 
-	const int32 StartPendingMipLevelIndex = (NextPendingMipLevelIndex + MaxPendingMipLevels - NumPendingMipLevels) % MaxPendingMipLevels;
-	int32 NumReadyMipLevels = 0;
+	const int32 StartPendingRequestIndex = (NextPendingRequestIndex + MaxPendingRequests - NumPendingRequests) % MaxPendingRequests;
+	int32 NumReadyRequests = 0;
 
-	for (int32 i = 0; i < NumPendingMipLevels; ++i)
+	for (int32 i = 0; i < NumPendingRequests; ++i)
 	{
-		const int32 PendingMipLevelIndex = (StartPendingMipLevelIndex + i) % MaxPendingMipLevels;
-		FPendingMipLevel& PendingMipLevel = PendingMipLevels[PendingMipLevelIndex];
+		const int32 PendingRequestIndex = (StartPendingRequestIndex + i) % MaxPendingRequests;
+		FPendingRequest& PendingRequest = PendingRequests[PendingRequestIndex];
 
-		FStreamingInfo* SVTInfo = FindStreamingInfo(PendingMipLevel.SparseVolumeTexture);
+		FStreamingInfo* SVTInfo = FindStreamingInfo(PendingRequest.SparseVolumeTexture);
 		if (!SVTInfo)
 		{
 #if WITH_EDITORONLY_DATA
 			// Resource is no longer there. Just mark as ready so it will be skipped later
-			PendingMipLevel.State = FPendingMipLevel::EState::DDC_Ready;
+			PendingRequest.State = FPendingRequest::EState::DDC_Ready;
 #endif
 			continue; 
 		}
 
-		const FResources* Resources = SVTInfo->PerFrameInfo[PendingMipLevel.FrameIndex].Resources;
+		const FResources* Resources = SVTInfo->PerFrameInfo[PendingRequest.FrameIndex].Resources;
 
 #if WITH_EDITORONLY_DATA
-		if (PendingMipLevel.State == FPendingMipLevel::EState::DDC_Ready)
+		if (PendingRequest.State == FPendingRequest::EState::DDC_Ready)
 		{
-			if (PendingMipLevel.RetryCount > 0)
+			if (PendingRequest.RetryCount > 0)
 			{
 				check(SVTInfo);
 				UE_LOG(LogSparseVolumeTextureStreamingManager, Display, TEXT("SVT DDC retry succeeded for '%s' (frame %i, mip %i) on %i attempt."), 
-					*Resources->ResourceName, PendingMipLevel.FrameIndex, PendingMipLevel.MipLevelIndex, PendingMipLevel.RetryCount);
+					*Resources->ResourceName, PendingRequest.FrameIndex, PendingRequest.MipLevelIndex, PendingRequest.RetryCount);
 			}
 		}
-		else if (PendingMipLevel.State == FPendingMipLevel::EState::DDC_Pending)
+		else if (PendingRequest.State == FPendingRequest::EState::DDC_Pending)
 		{
 			break;
 		}
-		else if (PendingMipLevel.State == FPendingMipLevel::EState::DDC_Failed)
+		else if (PendingRequest.State == FPendingRequest::EState::DDC_Failed)
 		{
-			PendingMipLevel.State = FPendingMipLevel::EState::DDC_Pending;
+			PendingRequest.State = FPendingRequest::EState::DDC_Pending;
 
-			if (PendingMipLevel.RetryCount == 0) // Only warn on first retry to prevent spam
+			if (PendingRequest.RetryCount == 0) // Only warn on first retry to prevent spam
 			{
 				UE_LOG(LogSparseVolumeTextureStreamingManager, Warning, TEXT("SVT DDC request failed for '%s' (frame %i, mip %i). Retrying..."),
-					*Resources->ResourceName, PendingMipLevel.FrameIndex, PendingMipLevel.MipLevelIndex);
+					*Resources->ResourceName, PendingRequest.FrameIndex, PendingRequest.MipLevelIndex);
 			}
 
-			const FMipTileReadInfo MipTileReadInfo = GetMipTileReadInfo(SVTInfo->PerFrameInfo[PendingMipLevel.FrameIndex], PendingMipLevel.MipLevelIndex);
-			FCacheGetChunkRequest Request = BuildDDCRequest(*Resources, MipTileReadInfo.ReadOffset, MipTileReadInfo.ReadSize, PendingMipLevelIndex);
-			const bool bBlocking = GSVTStreamingForceBlockingRequests || PendingMipLevel.bBlocking;
+			const FMipTileReadInfo MipTileReadInfo = GetMipTileReadInfo(SVTInfo->PerFrameInfo[PendingRequest.FrameIndex], PendingRequest.MipLevelIndex);
+			FCacheGetChunkRequest Request = BuildDDCRequest(*Resources, MipTileReadInfo.ReadOffset, MipTileReadInfo.ReadSize, PendingRequestIndex);
+			const bool bBlocking = GSVTStreamingForceBlockingRequests || PendingRequest.bBlocking;
 			RequestDDCData(MakeArrayView(&Request, 1), bBlocking);
 
-			++PendingMipLevel.RetryCount;
+			++PendingRequest.RetryCount;
 			break;
 		}
-		else if (PendingMipLevel.State == FPendingMipLevel::EState::Memory)
+		else if (PendingRequest.State == FPendingRequest::EState::Memory)
 		{
 			// Memory is always ready
 		}
@@ -1298,19 +1298,19 @@ int32 FStreamingManager::DetermineReadyMipLevels()
 #endif // WITH_EDITORONLY_DATA
 		{
 #if WITH_EDITORONLY_DATA
-			check(PendingMipLevel.State == FPendingMipLevel::EState::Disk);
+			check(PendingRequest.State == FPendingRequest::EState::Disk);
 #endif
-			if (PendingMipLevel.Request.IsCompleted())
+			if (PendingRequest.Request.IsCompleted())
 			{
-				if (!PendingMipLevel.Request.IsOk())
+				if (!PendingRequest.Request.IsOk())
 				{
 					// Retry if IO request failed for some reason
-					const FMipTileReadInfo MipTileReadInfo = GetMipTileReadInfo(SVTInfo->PerFrameInfo[PendingMipLevel.FrameIndex], PendingMipLevel.MipLevelIndex);
+					const FMipTileReadInfo MipTileReadInfo = GetMipTileReadInfo(SVTInfo->PerFrameInfo[PendingRequest.FrameIndex], PendingRequest.MipLevelIndex);
 					UE_LOG(LogSparseVolumeTextureStreamingManager, Warning, TEXT("SVT IO request failed for %p (frame %i, mip %i, offset %i, size %i). Retrying..."),
-						PendingMipLevel.SparseVolumeTexture, PendingMipLevel.FrameIndex, PendingMipLevel.MipLevelIndex, MipTileReadInfo.ReadOffset, MipTileReadInfo.ReadSize);
+						PendingRequest.SparseVolumeTexture, PendingRequest.FrameIndex, PendingRequest.MipLevelIndex, MipTileReadInfo.ReadOffset, MipTileReadInfo.ReadSize);
 					
 					FBulkDataBatchRequest::FBatchBuilder Batch = FBulkDataBatchRequest::NewBatch(1);
-					Batch.Read(Resources->StreamableMipLevels, MipTileReadInfo.ReadOffset, MipTileReadInfo.ReadSize, AIOP_Low, PendingMipLevel.RequestBuffer, PendingMipLevel.Request);
+					Batch.Read(Resources->StreamableMipLevels, MipTileReadInfo.ReadOffset, MipTileReadInfo.ReadSize, AIOP_Low, PendingRequest.RequestBuffer, PendingRequest.Request);
 					(void)Batch.Issue();
 					break;
 				}
@@ -1321,17 +1321,17 @@ int32 FStreamingManager::DetermineReadyMipLevels()
 			}
 		}
 
-		++NumReadyMipLevels;
+		++NumReadyRequests;
 	}
 
-	return NumReadyMipLevels;
+	return NumReadyRequests;
 }
 
-void FStreamingManager::InstallReadyMipLevels()
+void FStreamingManager::InstallReadyRequests()
 {
 	check(AsyncState.bUpdateActive);
-	check(AsyncState.NumReadyMipLevels <= PendingMipLevels.Num())
-	if (AsyncState.NumReadyMipLevels <= 0)
+	check(AsyncState.NumReadyRequests <= PendingRequests.Num())
+	if (AsyncState.NumReadyRequests <= 0)
 	{
 		return;
 	}
@@ -1339,7 +1339,7 @@ void FStreamingManager::InstallReadyMipLevels()
 	TRACE_CPUPROFILER_EVENT_SCOPE(SVT::StreamingInstallReadyRequests);
 
 	UploadTasks.Reset();
-	UploadTasks.Reserve(AsyncState.NumReadyMipLevels * 2 /*slack for splitting large uploads*/);
+	UploadTasks.Reserve(AsyncState.NumReadyRequests * 2 /*slack for splitting large uploads*/);
 	UploadCleanupTasks.Reset();
 
 #if WITH_EDITORONLY_DATA
@@ -1347,35 +1347,35 @@ void FStreamingManager::InstallReadyMipLevels()
 #endif
 
 	// Do a second pass over all ready mip levels, claiming memory in the upload buffers and creating FUploadTasks
-	const int32 StartPendingMipLevelIndex = (NextPendingMipLevelIndex + MaxPendingMipLevels - NumPendingMipLevels) % MaxPendingMipLevels;
-	for (int32 i = 0; i < AsyncState.NumReadyMipLevels; ++i)
+	const int32 StartPendingRequestIndex = (NextPendingRequestIndex + MaxPendingRequests - NumPendingRequests) % MaxPendingRequests;
+	for (int32 i = 0; i < AsyncState.NumReadyRequests; ++i)
 	{
-		const int32 PendingMipLevelIndex = (StartPendingMipLevelIndex + i) % MaxPendingMipLevels;
-		FPendingMipLevel& PendingMipLevel = PendingMipLevels[PendingMipLevelIndex];
+		const int32 PendingRequestIndex = (StartPendingRequestIndex + i) % MaxPendingRequests;
+		FPendingRequest& PendingRequest = PendingRequests[PendingRequestIndex];
 
-		FStreamingInfo* SVTInfo = FindStreamingInfo(PendingMipLevel.SparseVolumeTexture);
-		if (!SVTInfo || (SVTInfo->PerFrameInfo[PendingMipLevel.FrameIndex].LowestRequestedMipLevel > PendingMipLevel.MipLevelIndex))
+		FStreamingInfo* SVTInfo = FindStreamingInfo(PendingRequest.SparseVolumeTexture);
+		if (!SVTInfo || (SVTInfo->PerFrameInfo[PendingRequest.FrameIndex].LowestRequestedMipLevel > PendingRequest.MipLevelIndex))
 		{
-			PendingMipLevel.Reset();
+			PendingRequest.Reset();
 			continue; // Skip mip level install. SVT no longer exists or mip level was "streamed out" before it was even installed in the first place.
 		}
 
-		FFrameInfo& FrameInfo = SVTInfo->PerFrameInfo[PendingMipLevel.FrameIndex];
+		FFrameInfo& FrameInfo = SVTInfo->PerFrameInfo[PendingRequest.FrameIndex];
 		const FResources* Resources = FrameInfo.Resources;
 		const FPageTopology& Topology = FrameInfo.Resources->Topology;
-		const FMipTileReadInfo MipTileReadInfo = GetMipTileReadInfo(FrameInfo, PendingMipLevel.MipLevelIndex);
+		const FMipTileReadInfo MipTileReadInfo = GetMipTileReadInfo(FrameInfo, PendingRequest.MipLevelIndex);
 
 		const uint8* SrcPtr = nullptr;
 		const uint8* SrcEndPtr = nullptr;
 
 #if WITH_EDITORONLY_DATA
-		if (PendingMipLevel.State == FPendingMipLevel::EState::DDC_Ready)
+		if (PendingRequest.State == FPendingRequest::EState::DDC_Ready)
 		{
 			check(Resources->ResourceFlags & EResourceFlag_StreamingDataInDDC);
-			SrcPtr = (const uint8*)PendingMipLevel.SharedBuffer.GetData();
-			SrcEndPtr = SrcPtr + PendingMipLevel.SharedBuffer.GetSize();
+			SrcPtr = (const uint8*)PendingRequest.SharedBuffer.GetData();
+			SrcEndPtr = SrcPtr + PendingRequest.SharedBuffer.GetSize();
 		}
-		else if (PendingMipLevel.State == FPendingMipLevel::EState::Memory)
+		else if (PendingRequest.State == FPendingRequest::EState::Memory)
 		{
 			const uint8** BulkDataPtrPtr = ResourceToBulkPointer.Find(Resources);
 			if (BulkDataPtrPtr)
@@ -1397,9 +1397,9 @@ void FStreamingManager::InstallReadyMipLevels()
 #endif
 		{
 #if WITH_EDITORONLY_DATA
-			check(PendingMipLevel.State == FPendingMipLevel::EState::Disk);
+			check(PendingRequest.State == FPendingRequest::EState::Disk);
 #endif
-			SrcPtr = PendingMipLevel.RequestBuffer.GetData();
+			SrcPtr = PendingRequest.RequestBuffer.GetData();
 		}
 
 		check(SrcPtr);
@@ -1434,26 +1434,26 @@ void FStreamingManager::InstallReadyMipLevels()
 
 		// Cleanup
 		{
-			UploadCleanupTasks.Add(&PendingMipLevel);
+			UploadCleanupTasks.Add(&PendingRequest);
 		}
 	
 #if SVT_STREAMING_LOG_VERBOSE
 		UE_LOG(LogSparseVolumeTextureStreamingManager, Display, TEXT("(%i)%i Install Frame %i OldResMip %i, NewResMip %i, ReqMip %i"),
-			PendingMipLevel.IssuedInFrame, PendingMipLevelIndex,
-			PendingMipLevel.FrameIndex, 
-			FrameInfo.LowestResidentMipLevel, PendingMipLevel.MipLevelIndex,
+			PendingRequest.IssuedInFrame, PendingRequest,
+			PendingRequest.FrameIndex,
+			FrameInfo.LowestResidentMipLevel, PendingRequest.MipLevelIndex,
 			FrameInfo.LowestRequestedMipLevel);
 #endif
 
-		check(FrameInfo.LowestResidentMipLevel == (PendingMipLevel.MipLevelIndex + 1));
-		FrameInfo.LowestResidentMipLevel = PendingMipLevel.MipLevelIndex;
+		check(FrameInfo.LowestResidentMipLevel == (PendingRequest.MipLevelIndex + 1));
+		FrameInfo.LowestResidentMipLevel = PendingRequest.MipLevelIndex;
 
 		InvalidatedSVTFrames.Add(&FrameInfo);
 
-		const int32 LRUNodeIndex = PendingMipLevel.FrameIndex * SVTInfo->NumMipLevelsGlobal + PendingMipLevel.MipLevelIndex;
-		SVTInfo->LRUNodes[LRUNodeIndex].PendingMipLevelIndex = INDEX_NONE;
+		const int32 LRUNodeIndex = PendingRequest.FrameIndex * SVTInfo->NumMipLevelsGlobal + PendingRequest.MipLevelIndex;
+		SVTInfo->LRUNodes[LRUNodeIndex].PendingRequestIndex = INDEX_NONE;
 
-		const FPageTopology::FMip& TopologyMipInfo = Topology.MipInfo[PendingMipLevel.MipLevelIndex];
+		const FPageTopology::FMip& TopologyMipInfo = Topology.MipInfo[PendingRequest.MipLevelIndex];
 		FrameInfo.ResidentPagesNew.SetRange(TopologyMipInfo.PageOffset, TopologyMipInfo.PageCount, true);
 	}
 
@@ -1485,22 +1485,22 @@ void FStreamingManager::InstallReadyMipLevels()
 	{
 		TRACE_CPUPROFILER_EVENT_SCOPE(SVT::StreamingUploadCleanupTask);
 
-		FPendingMipLevel* PendingMipLevel = UploadCleanupTasks[TaskIndex];
+		FPendingRequest* PendingRequest = UploadCleanupTasks[TaskIndex];
 #if WITH_EDITORONLY_DATA
-		PendingMipLevel->SharedBuffer.Reset();
+		PendingRequest->SharedBuffer.Reset();
 #endif
-		if (!PendingMipLevel->Request.IsNone())
+		if (!PendingRequest->Request.IsNone())
 		{
-			check(PendingMipLevel->Request.IsCompleted());
-			PendingMipLevel->Request.Reset();
+			check(PendingRequest->Request.IsCompleted());
+			PendingRequest->Request.Reset();
 		}
 	});
 
 #if DO_CHECK // Clear processed pending mip levels for better debugging
-	for (int32 i = 0; i < AsyncState.NumReadyMipLevels; ++i)
+	for (int32 i = 0; i < AsyncState.NumReadyRequests; ++i)
 	{
-		const int32 PendingMipLevelIndex = (StartPendingMipLevelIndex + i) % MaxPendingMipLevels;
-		PendingMipLevels[PendingMipLevelIndex].Reset();
+		const int32 PendingRequestIndex = (StartPendingRequestIndex + i) % MaxPendingRequests;
+		PendingRequests[PendingRequestIndex].Reset();
 	}
 #endif
 
@@ -1686,7 +1686,7 @@ FStreamingManager::FMipTileReadInfo FStreamingManager::GetMipTileReadInfo(const
 
 #if WITH_EDITORONLY_DATA
 
-UE::DerivedData::FCacheGetChunkRequest FStreamingManager::BuildDDCRequest(const FResources& Resources, uint64 ReadOffset, uint64 ReadSize, uint32 PendingMipLevelIndex)
+UE::DerivedData::FCacheGetChunkRequest FStreamingManager::BuildDDCRequest(const FResources& Resources, uint64 ReadOffset, uint64 ReadSize, uint32 PendingRequestIndex)
 {
 	using namespace UE::DerivedData;
 
@@ -1701,7 +1701,7 @@ UE::DerivedData::FCacheGetChunkRequest FStreamingManager::BuildDDCRequest(const
 	Request.RawOffset = ReadOffset;
 	Request.RawSize = ReadSize;
 	Request.RawHash = Resources.DDCRawHash;
-	Request.UserData = (((uint64)PendingMipLevelIndex) << uint64(32)) | (uint64)PendingMipLevels[PendingMipLevelIndex].RequestVersion;
+	Request.UserData = (((uint64)PendingRequestIndex) << uint64(32)) | (uint64)PendingRequests[PendingRequestIndex].RequestVersion;
 	return Request;
 }
 
@@ -1715,26 +1715,26 @@ void FStreamingManager::RequestDDCData(TConstArrayView<UE::DerivedData::FCacheGe
 		GetCache().GetChunks(DDCRequests, *RequestOwnerPtr,
 			[this](FCacheGetChunkResponse&& Response)
 			{
-				const uint32 PendingMipLevelIndex = (uint32)(Response.UserData >> uint64(32));
+				const uint32 PendingRequestIndex = (uint32)(Response.UserData >> uint64(32));
 				const uint32 RequestVersion = (uint32)Response.UserData;
 
-				// In case the request returned after the mip level was already streamed out again we need to abort so that we do not overwrite data in the FPendingMipLevel slot.
-				if (RequestVersion < PendingMipLevels[PendingMipLevelIndex].RequestVersion)
+				// In case the request returned after the mip level was already streamed out again we need to abort so that we do not overwrite data in the FPendingRequest slot.
+				if (RequestVersion < PendingRequests[PendingRequestIndex].RequestVersion)
 				{
 					return;
 				}
 
-				FPendingMipLevel& PendingMipLevel = PendingMipLevels[PendingMipLevelIndex];
-				check(PendingMipLevel.SparseVolumeTexture); // A valid PendingMipLevel should have a non-nullptr here
+				FPendingRequest& PendingRequest = PendingRequests[PendingRequestIndex];
+				check(PendingRequest.SparseVolumeTexture); // A valid PendingRequest should have a non-nullptr here
 
 				if (Response.Status == EStatus::Ok)
 				{
-					PendingMipLevel.SharedBuffer = MoveTemp(Response.RawData);
-					PendingMipLevel.State = FPendingMipLevel::EState::DDC_Ready;
+					PendingRequest.SharedBuffer = MoveTemp(Response.RawData);
+					PendingRequest.State = FPendingRequest::EState::DDC_Ready;
 				}
 				else
 				{
-					PendingMipLevel.State = FPendingMipLevel::EState::DDC_Failed;
+					PendingRequest.State = FPendingRequest::EState::DDC_Failed;
 				}
 			});
 	}
diff --git a/Engine/Source/Runtime/Engine/Private/SparseVolumeTexture/SparseVolumeTextureStreamingManager.h b/Engine/Source/Runtime/Engine/Private/SparseVolumeTexture/SparseVolumeTextureStreamingManager.h
--- a/Engine/Source/Runtime/Engine/Private/SparseVolumeTexture/SparseVolumeTextureStreamingManager.h
+++ b/Engine/Source/Runtime/Engine/Private/SparseVolumeTexture/SparseVolumeTextureStreamingManager.h
@@ -140,7 +140,7 @@ private:
 		FLRUNode* NextHigherMipLevel = nullptr;
 		uint32 LastRequested = INDEX_NONE;
 		uint32 RefCount = 0; // Keep track of how many lower mip levels have a dependency on this one
-		uint32 PendingMipLevelIndex = INDEX_NONE;
+		uint32 PendingRequestIndex = INDEX_NONE;
 		int16 FrameIndex = INDEX_NONE;
 		int16 MipLevelIndex = INDEX_NONE;
 	};
@@ -163,7 +163,7 @@ private:
 	};
 
 	// Represents an IO request for a mip level
-	struct FPendingMipLevel
+	struct FPendingRequest
 	{
 		UStreamableSparseVolumeTexture* SparseVolumeTexture = nullptr; // Do not dereference!
 		int32 FrameIndex = INDEX_NONE;
@@ -222,7 +222,7 @@ private:
 
 	struct FAsyncState
 	{
-		int32 NumReadyMipLevels = 0;
+		int32 NumReadyRequests = 0;
 		bool bUpdateActive = false;
 		bool bUpdateIsAsync = false;
 	};
@@ -250,7 +250,7 @@ private:
 
 	TMap<UStreamableSparseVolumeTexture*, TUniquePtr<FStreamingInfo>> StreamingInfo; // Do not dereference the key! We just read the pointer itself.
 	TMap<FMipLevelKey, uint32> RequestsHashTable;
-	TArray<FPendingMipLevel> PendingMipLevels;
+	TArray<FPendingRequest> PendingRequests;
 #if WITH_EDITORONLY_DATA
 	TUniquePtr<UE::DerivedData::FRequestOwner> RequestOwner;
 	TUniquePtr<UE::DerivedData::FRequestOwner> RequestOwnerBlocking;
@@ -259,9 +259,9 @@ private:
 	TUniquePtr<class FPageTableUpdater> PageTableUpdater;
 	FGraphEventArray AsyncTaskEvents;
 	FAsyncState AsyncState;
-	int32 MaxPendingMipLevels = 0;
-	int32 NumPendingMipLevels = 0;
-	int32 NextPendingMipLevelIndex = 0;
+	int32 MaxPendingRequests = 0;
+	int32 NumPendingRequests = 0;
+	int32 NextPendingRequestIndex = 0;
 	uint32 NextUpdateIndex = 1;
 
 	// Transient lifetime
@@ -271,7 +271,7 @@ private:
 	TArray<FStreamingRequest> PrioritizedRequestsHeap;
 	TArray<FStreamingRequest> SelectedRequests;
 	TArray<FTileDataTask> UploadTasks; // accessed on the async thread
-	TArray<FPendingMipLevel*> UploadCleanupTasks; // accessed on the async thread
+	TArray<FPendingRequest*> UploadCleanupTasks; // accessed on the async thread
 
 	void AddInternal(FRDGBuilder& GraphBuilder, FNewSparseVolumeTextureInfo&& NewSVTInfo);
 	void RemoveInternal(UStreamableSparseVolumeTexture* SparseVolumeTexture);
@@ -281,14 +281,14 @@ private:
 	void SelectHighestPriorityRequestsAndUpdateLRU(int32 MaxSelectedRequests);
 	void IssueRequests(int32 MaxSelectedRequests);
 	void StreamOutMipLevel(FStreamingInfo* SVTInfo, FLRUNode* LRUNode);
-	int32 DetermineReadyMipLevels();
-	void InstallReadyMipLevels();
+	int32 DetermineReadyRequests();
+	void InstallReadyRequests();
 	void PatchPageTable(FRDGBuilder& GraphBuilder); // Patches the page table to reflect streamed in/out pages and to ensure non-resident mip levels fall back to coarser mip level tile data
 	FStreamingInfo* FindStreamingInfo(UStreamableSparseVolumeTexture* Key); // Returns nullptr if the key can't be found
 	static FMipTileReadInfo GetMipTileReadInfo(const FFrameInfo& FrameInfo, int32 MipLevel);
 
 #if WITH_EDITORONLY_DATA
-	UE::DerivedData::FCacheGetChunkRequest BuildDDCRequest(const FResources& Resources, uint64 ReadOffset, uint64 ReadSize, uint32 PendingMipLevelIndex);
+	UE::DerivedData::FCacheGetChunkRequest BuildDDCRequest(const FResources& Resources, uint64 ReadOffset, uint64 ReadSize, uint32 PendingRequestIndex);
 	void RequestDDCData(TConstArrayView<UE::DerivedData::FCacheGetChunkRequest> DDCRequests, bool bBlocking);
 #endif
 };
