Fixed regression in APlayerController::ClientFlushLevelStreaming_Implementation introduced by 1730015 (inverted logic).
ClientFlushLevelStreaming was not flusing which was causing replication issues when a client joins a server which previously streamed sublevels.

diff --git a/Engine/Source/Runtime/Engine/Private/PlayerController.cpp b/Engine/Source/Runtime/Engine/Private/PlayerController.cpp
--- a/Engine/Source/Runtime/Engine/Private/PlayerController.cpp
+++ b/Engine/Source/Runtime/Engine/Private/PlayerController.cpp
@@ -421,15 +421,11 @@ void APlayerController::ClientAckUpdateLevelVisibility_Implementation(FName Pack
 void APlayerController::ClientFlushLevelStreaming_Implementation()
 {
 	UWorld* World = GetWorld();
-	// if we're already doing a map change, requesting another blocking load is just wasting time	
-	if (GEngine->ShouldCommitPendingMapChange(World))
-	{
-		// request level streaming be flushed next frame
-		World->UpdateLevelStreaming();
-		World->bRequestedBlockOnAsyncLoading = true;
-		// request GC as soon as possible to remove any unloaded levels from memory
-		GEngine->ForceGarbageCollection();
-	}
+	// request level streaming be flushed next frame
+	World->UpdateLevelStreaming();
+	World->bRequestedBlockOnAsyncLoading = true;
+	// request GC as soon as possible to remove any unloaded levels from memory
+	GEngine->ForceGarbageCollection();
 }
 
 
