World Partition: make sure to skip loading of actor folders when loading packages for instanced actors from level instaqnces in PIE.

diff --git a/Engine/Source/Runtime/Engine/Classes/Engine/Level.h b/Engine/Source/Runtime/Engine/Classes/Engine/Level.h
--- a/Engine/Source/Runtime/Engine/Classes/Engine/Level.h
+++ b/Engine/Source/Runtime/Engine/Classes/Engine/Level.h
@@ -732,6 +732,7 @@ public:
 
 	ENGINE_API static const FName LoadAllExternalObjectsTag;
 	ENGINE_API static const FName DontLoadExternalObjectsTag;
+	ENGINE_API static const FName DontLoadExternalFoldersTag;
 
 	ENGINE_API bool GetPromptWhenAddingToLevelOutsideBounds() const;
 	ENGINE_API bool GetPromptWhenAddingToLevelBeforeCheckout() const;
diff --git a/Engine/Source/Runtime/Engine/Private/Level.cpp b/Engine/Source/Runtime/Engine/Private/Level.cpp
--- a/Engine/Source/Runtime/Engine/Private/Level.cpp
+++ b/Engine/Source/Runtime/Engine/Private/Level.cpp
@@ -381,6 +381,7 @@ TMap<FName, TWeakObjectPtr<UWorld> > ULevel::StreamedLevelsOwningWorld;
 #if WITH_EDITOR
 const FName ULevel::LoadAllExternalObjectsTag(TEXT("LoadAllExternalObjectsTag"));
 const FName ULevel::DontLoadExternalObjectsTag(TEXT("DontLoadExternalObjectsTag"));
+const FName ULevel::DontLoadExternalFoldersTag(TEXT("DontLoadExternalFoldersTag"));
 #endif
 
 ULevel::ULevel( const FObjectInitializer& ObjectInitializer )
@@ -1153,11 +1154,13 @@ void ULevel::PostLoad()
 	}
 
 #if WITH_EDITOR
+	const FLinkerLoad* Linker = GetLinker();
+	check(Linker || bWasDuplicated);
+
 	if ((IsUsingExternalActors() || IsUsingExternalObjects())
 		&& OwningWorld
 		&& OwningWorld->WorldType == EWorldType::Editor)
 	{
-		const FLinkerLoad* Linker = GetLinker();
 		if (Linker && Linker->IsPackageRelocated())
 		{
 			const FString LevelPackageName = GetPackage()->GetName();
@@ -1177,7 +1180,7 @@ void ULevel::PostLoad()
 
 	if (IsUsingActorFolders() && IsUsingExternalObjects() && IsActorFolderObjectsFeatureAvailable())
 	{
-		if (!bWasDuplicated && !FPackageName::IsTempPackage(GetPackage()->GetName()))
+		if (!bWasDuplicated && !Linker->GetInstancingContext().HasTag(ULevel::DontLoadExternalFoldersTag))
 		{
 			// Load all folders for this level
 			FExternalPackageHelper::LoadObjectsFromExternalPackages<UActorFolder>(this, [this](UActorFolder* LoadedFolder)
@@ -1191,9 +1194,6 @@ void ULevel::PostLoad()
 	// if we use external actors, load dynamic actors here
 	if (IsUsingExternalActors())
 	{
-		const FLinkerLoad* Linker = GetLinker();
-		check(Linker || bWasDuplicated);
-
 		if (!bWasDuplicated &&
 			(!bIsPartitioned || Linker->GetInstancingContext().HasTag(ULevel::LoadAllExternalObjectsTag)) &&
 			!Linker->GetInstancingContext().HasTag(ULevel::DontLoadExternalObjectsTag))
diff --git a/Engine/Source/Runtime/Engine/Private/WorldPartition/WorldPartitionLevelHelper.cpp b/Engine/Source/Runtime/Engine/Private/WorldPartition/WorldPartitionLevelHelper.cpp
--- a/Engine/Source/Runtime/Engine/Private/WorldPartition/WorldPartitionLevelHelper.cpp
+++ b/Engine/Source/Runtime/Engine/Private/WorldPartition/WorldPartitionLevelHelper.cpp
@@ -446,6 +446,7 @@ bool FWorldPartitionLevelHelper::LoadActors(const FLoadActorsParams& InParams)
 				NewContext.SetSoftObjectPathRemappingEnabled(false); 
 			
 				NewContext.AddTag(ULevel::DontLoadExternalObjectsTag);
+				NewContext.AddTag(ULevel::DontLoadExternalFoldersTag);
 				NewContext.AddPackageMapping(PackageObjectMapping.ContainerPackage, ContainerPackageInstanceName);
 				Context = &NewContext;
 			}
