Defer FWorldPartitionRuntimeSpatialHashGridPreviewer material initialization (loading) to avoid calling LoadObject during async load.

diff --git a/Engine/Source/Runtime/Engine/Private/WorldPartition/WorldPartitionRuntimeSpatialHashGridPreviewer.cpp b/Engine/Source/Runtime/Engine/Private/WorldPartition/WorldPartitionRuntimeSpatialHashGridPreviewer.cpp
--- a/Engine/Source/Runtime/Engine/Private/WorldPartition/WorldPartitionRuntimeSpatialHashGridPreviewer.cpp
+++ b/Engine/Source/Runtime/Engine/Private/WorldPartition/WorldPartitionRuntimeSpatialHashGridPreviewer.cpp
@@ -17,19 +17,20 @@ FWorldPartitionRuntimeSpatialHashGridPreviewer::FWorldPartitionRuntimeSpatialHas
 #if WITH_EDITORONLY_DATA
 	: MID(nullptr)
 	, Volume(nullptr)
+	, bIsInitialized(false)
 #endif
 {
-#if WITH_EDITORONLY_DATA
-	if (!IsRunningCookCommandlet())
-	{
-		Material = LoadObject<UMaterial>(nullptr, TEXT("/Engine/EditorMaterials/WorldPartition/WorldPartitionSpatialHashGridPreviewMaterial"));
-	}
-#endif
 }
 
 #if WITH_EDITOR
 void FWorldPartitionRuntimeSpatialHashGridPreviewer::Draw(UWorld* World, const TArray<FSpatialHashRuntimeGrid>& Grids, bool bEnabled, int32 PreviewGridLevel, bool bUseAlignedGridLevels)
 {
+	if (!bIsInitialized)
+	{
+		Material = LoadObject<UMaterial>(nullptr, TEXT("/Engine/EditorMaterials/WorldPartition/WorldPartitionSpatialHashGridPreviewMaterial"));
+		bIsInitialized = true;
+	}
+
 	if (bEnabled && Material)
 	{
 		if (!Volume)
diff --git a/Engine/Source/Runtime/Engine/Public/WorldPartition/WorldPartitionRuntimeSpatialHashGridPreviewer.h b/Engine/Source/Runtime/Engine/Public/WorldPartition/WorldPartitionRuntimeSpatialHashGridPreviewer.h
--- a/Engine/Source/Runtime/Engine/Public/WorldPartition/WorldPartitionRuntimeSpatialHashGridPreviewer.h
+++ b/Engine/Source/Runtime/Engine/Public/WorldPartition/WorldPartitionRuntimeSpatialHashGridPreviewer.h
@@ -30,6 +30,8 @@ struct FWorldPartitionRuntimeSpatialHashGridPreviewer
 	UPROPERTY()
 	TObjectPtr<APostProcessVolume> Volume;
 
+	bool bIsInitialized;
+
 	struct FGridParametersCache
 	{
 		FGridParametersCache()
