Fixed potential crash caused by incorrectly setting up a check() condition in the SVT streaming manager when running a cooked game.

#rnx

diff --git a/Engine/Source/Runtime/Engine/Private/SparseVolumeTexture/SparseVolumeTextureStreamingManager.cpp b/Engine/Source/Runtime/Engine/Private/SparseVolumeTexture/SparseVolumeTextureStreamingManager.cpp
--- a/Engine/Source/Runtime/Engine/Private/SparseVolumeTexture/SparseVolumeTextureStreamingManager.cpp
+++ b/Engine/Source/Runtime/Engine/Private/SparseVolumeTexture/SparseVolumeTextureStreamingManager.cpp
@@ -1400,6 +1400,7 @@ void FStreamingManager::InstallReadyRequests()
 			check(PendingRequest.State == FPendingRequest::EState::Disk);
 #endif
 			SrcPtr = PendingRequest.RequestBuffer.GetData();
+			SrcEndPtr = SrcPtr + PendingRequest.RequestBuffer.DataSize();
 		}
 
 		check(SrcPtr);
