Added an array of debugging info string to represent the expressions stored in ExpressionCollection.Expressions inside UMaterialFunctionEditorOnlyData to aid debugging an unsolved issue that sporadically causes some expressions in the list to be null when a MaterialFunction is loaded.

diff --git a/Engine/Source/Runtime/Engine/Classes/Materials/MaterialExpression.h b/Engine/Source/Runtime/Engine/Classes/Materials/MaterialExpression.h
--- a/Engine/Source/Runtime/Engine/Classes/Materials/MaterialExpression.h
+++ b/Engine/Source/Runtime/Engine/Classes/Materials/MaterialExpression.h
@@ -167,6 +167,11 @@ struct FMaterialExpressionCollection
 	/** Array of material expressions, excluding Comments.  Used by the material editor. */
 	UPROPERTY()
 	TArray<TObjectPtr<UMaterialExpression>> Expressions;
+	
+	/** Temporary debug info strings for each expression to help debugging the sporadic
+	 *  unexpected occurrence of nulls in Expressions on PostLoad() (UE-198712) */
+	UPROPERTY()
+	TArray<FString> DebugExpressionInfos;
 
 	/** Array of comments associated with this material; viewed in the material editor. */
 	UPROPERTY()
diff --git a/Engine/Source/Runtime/Engine/Private/Materials/MaterialExpressions.cpp b/Engine/Source/Runtime/Engine/Private/Materials/MaterialExpressions.cpp
--- a/Engine/Source/Runtime/Engine/Private/Materials/MaterialExpressions.cpp
+++ b/Engine/Source/Runtime/Engine/Private/Materials/MaterialExpressions.cpp
@@ -14909,6 +14909,24 @@ void UMaterialFunction::ForceRecompileForRendering(FMaterialUpdateContext& Updat
 
 void UMaterialFunction::Serialize(FArchive& Ar)
 {
+#if WITH_EDITOR
+	// Temporary debugging code. This will populate the DebugExpressionInfos with information about each expression
+	// in ExpressionCollection.Expressions, to gather more information when some expression is null upon function
+	// PostLoad() to help solve UE-198712.
+	UMaterialFunctionEditorOnlyData* EditorOnly = GetEditorOnlyData();
+	if (EditorOnly  && Ar.IsSaving() && !Ar.IsCooking() && !Ar.IsObjectReferenceCollector())
+	{
+		EditorOnly->ExpressionCollection.DebugExpressionInfos.Empty();
+		EditorOnly->ExpressionCollection.DebugExpressionInfos.Reserve(EditorOnly->ExpressionCollection.Expressions.Num());
+		for (UMaterialExpression* Expression : EditorOnly->ExpressionCollection.Expressions)
+		{
+			check(Expression);
+			FString Info = FString::Printf(TEXT("Name: '%s', Type: '%s'"), *Expression->GetFullName(), *Expression->GetClass()->GetFullName());
+			EditorOnly->ExpressionCollection.DebugExpressionInfos.Push(MoveTemp(Info));
+		}
+	}
+#endif
+
 	Super::Serialize(Ar);
 
 #if WITH_EDITOR
@@ -15000,6 +15018,19 @@ void UMaterialFunction::PostLoad()
 
 	if (GIsEditor && EditorOnly)
 	{
+		if (EditorOnly->ExpressionCollection.DebugExpressionInfos.Num() == EditorOnly->ExpressionCollection.Expressions.Num())
+		{
+			for (int i = 0; i < EditorOnly->ExpressionCollection.Expressions.Num(); ++i)
+			{
+				UMaterialExpression* Expression = EditorOnly->ExpressionCollection.Expressions[i].Get();
+				if (!Expression)
+				{
+					UE_LOG(LogMaterial, Log, TEXT("Expression in function expression collection with index %d was null. Expression Info: %s"), i, *EditorOnly->ExpressionCollection.DebugExpressionInfos[i]);
+					EditorOnly->ExpressionCollection.DebugExpressionInfos.RemoveAt(i);
+				}
+			}
+		}
+
 		// Clean up any removed material expression classes	
 		if (EditorOnly->ExpressionCollection.Expressions.Remove(nullptr) != 0)
 		{
