fire the OnInputHardwareDeviceChanged event whenever the hardware device id changes from the last known platform user id map, not the input device id map.

This gives us the desired behavior of firing a delegate any time the Platform User changes their "most recently used" input device, instead of only ever firing when an input device is connected.

diff --git a/Engine/Source/Runtime/Engine/Private/GameFramework/InputDeviceSubsystem.cpp b/Engine/Source/Runtime/Engine/Private/GameFramework/InputDeviceSubsystem.cpp
--- a/Engine/Source/Runtime/Engine/Private/GameFramework/InputDeviceSubsystem.cpp
+++ b/Engine/Source/Runtime/Engine/Private/GameFramework/InputDeviceSubsystem.cpp
@@ -503,8 +503,10 @@ FHardwareDeviceIdentifier UInputDeviceSubsystem::GetInputDeviceHardwareIdentifie
 
 void UInputDeviceSubsystem::SetMostRecentlyUsedHardwareDevice(const FInputDeviceId InDeviceId, const FHardwareDeviceIdentifier& InHardwareId)
 {	
-	// If the hardware hasn't changed, then just ignore it because we don't need to update anything.
-	if (const FHardwareDeviceIdentifier* ExistingDevice = LatestInputDeviceIdentifiers.Find(InDeviceId))
+	const FPlatformUserId OwningUserId = IPlatformInputDeviceMapper::Get().GetUserForInputDevice(InDeviceId);
+
+	// If this hardware is the same as what the platform user already has, then there is no need to fire this event
+	if (const FHardwareDeviceIdentifier* ExistingDevice = LatestUserDeviceIdentifiers.Find(OwningUserId))
 	{
 		if (InHardwareId == *ExistingDevice)
 		{
@@ -512,8 +514,6 @@ void UInputDeviceSubsystem::SetMostRecentlyUsedHardwareDevice(const FInputDevice
 		}
 	}
 
-	FPlatformUserId OwningUserId = IPlatformInputDeviceMapper::Get().GetUserForInputDevice(InDeviceId);
-
 	// Keep track of each input device's latest hardware id
 	LatestInputDeviceIdentifiers.Add(InDeviceId, InHardwareId);
 
