Replace LocalPlayer logic backed out in CL 31434853 with a more robust solution that still captures the needed data for Ortho without ignoring the existing camera settings.

diff --git a/Engine/Source/Runtime/Engine/Private/LocalPlayer.cpp b/Engine/Source/Runtime/Engine/Private/LocalPlayer.cpp
--- a/Engine/Source/Runtime/Engine/Private/LocalPlayer.cpp
+++ b/Engine/Source/Runtime/Engine/Private/LocalPlayer.cpp
@@ -26,6 +26,7 @@
 
 #include "IXRTrackingSystem.h"
 #include "IXRCamera.h"
+#include "Camera/CameraComponent.h"
 #include "SceneView.h"
 #include "SceneViewExtension.h"
 #include "Net/DataChannel.h"
@@ -687,6 +688,29 @@ void ULocalPlayer::GetViewPoint(FMinimalViewInfo& OutViewInfo) const
 			OutViewInfo = PlayerController->PlayerCameraManager->GetCameraCacheView();
 			OutViewInfo.FOV = PlayerController->PlayerCameraManager->GetFOVAngle();
 			PlayerController->GetPlayerViewPoint(/*out*/ OutViewInfo.Location, /*out*/ OutViewInfo.Rotation);
+			
+			bool bCameraFound = false;
+			if (const AActor* Target = PlayerController->GetViewTarget())
+			{
+				if (const UCameraComponent* Camera = Target->GetComponentByClass<const UCameraComponent>())
+				{
+					OutViewInfo.bAutoCalculateOrthoPlanes = Camera->bAutoCalculateOrthoPlanes;
+					OutViewInfo.SetOrthoCameraArmLengthFromOwnerLocation(Target->GetActorLocation());
+					bCameraFound = true;
+				}
+			}
+
+			if (!bCameraFound)
+			{
+				OutViewInfo.bAutoCalculateOrthoPlanes = PlayerController->PlayerCameraManager->bAutoCalculateOrthoPlanes;
+				if(OutViewInfo.bAutoCalculateOrthoPlanes)
+				{
+					if (const APawn* Owner = PlayerController->GetPawn())
+					{
+						OutViewInfo.SetOrthoCameraArmLengthFromOwnerLocation(Owner->GetActorLocation());
+					}
+				}
+			}
 		}
 		else
 		{
