Tear off the gamma space for all layers, and make sure to integrate any format limitations.

diff --git a/Engine/Source/Runtime/Engine/Classes/Engine/Texture.h b/Engine/Source/Runtime/Engine/Classes/Engine/Texture.h
--- a/Engine/Source/Runtime/Engine/Classes/Engine/Texture.h
+++ b/Engine/Source/Runtime/Engine/Classes/Engine/Texture.h
@@ -585,8 +585,8 @@ private:
 	/** TextureClass == Owner->GetTextureClass(); a copy is kept here for torn-off **/
 	ETextureClass TornOffTextureClass;
 	/** if Owner != null, check Owner->GetGammaSpace , if it is null, use TornOffGammaSpace
-	* do not check this directly, use GetGammaSpace **/
-	EGammaSpace TornOffGammaSpace;
+	* do not check this directly, use GetGammaSpace. **/
+	TArray<EGammaSpace, TInlineAllocator<1>> TornOffGammaSpace;
 #endif
 	/** The bulk source data. */
 	UE::Serialization::FEditorBulkData BulkData;
diff --git a/Engine/Source/Runtime/Engine/Private/Texture.cpp b/Engine/Source/Runtime/Engine/Private/Texture.cpp
--- a/Engine/Source/Runtime/Engine/Private/Texture.cpp
+++ b/Engine/Source/Runtime/Engine/Private/Texture.cpp
@@ -1928,7 +1928,6 @@ FTextureSource::FTextureSource()
 #if WITH_EDITOR
 	  Owner(nullptr),
 	  TornOffTextureClass(ETextureClass::Invalid),
-	  TornOffGammaSpace(EGammaSpace::Invalid),
 #endif
 	  NumLockedMips(0u)
 	, LockState(ELockState::None)
@@ -2303,7 +2302,13 @@ FTextureSource FTextureSource::CopyTornOff() const
 	// Result can't talk to Owner any more, so save info we need :
 	check( Owner != nullptr );
 	check( &(Owner->Source) == this );
-	Result.TornOffGammaSpace = Owner->GetGammaSpace();
+
+	Result.TornOffGammaSpace.SetNumZeroed(NumLayers);
+	for (int32 LayerIndex = 0; LayerIndex < NumLayers; LayerIndex++)
+	{
+		// Make sure we save the gamma space for each layer, as well as with any format adjustments from our GetGammaSpace().
+		Result.TornOffGammaSpace[LayerIndex] = this->GetGammaSpace(LayerIndex);
+	}
 	Result.TornOffTextureClass = Owner->GetTextureClass();
 	return Result;
 }
@@ -2903,8 +2908,15 @@ EGammaSpace FTextureSource::GetGammaSpace(int LayerIndex) const
 	else
 	{
 		// torn off, should have saved TornOffGammaSpace
-		check( TornOffGammaSpace != EGammaSpace::Invalid );
-		return TornOffGammaSpace;
+		check(LayerIndex < TornOffGammaSpace.Num());
+		if (LayerIndex >= TornOffGammaSpace.Num())
+		{
+			UE_LOG(LogTexture, Error, TEXT("Torn off texture source doesn't have gamma copied!"));
+			return EGammaSpace::Linear;
+		}
+
+		check( TornOffGammaSpace[LayerIndex] != EGammaSpace::Invalid );
+		return TornOffGammaSpace[LayerIndex];
 	}
 }
 
