Add additional validation to prevent incompatible streaming mips with runtime virtual texture. This is extra defence against any code logic issues that could generate an incompatible SVT format. If that does happen it will generate D3D device errors when trying to copy the incompatible SVT format to the physical texture.

diff --git a/Engine/Source/Runtime/Engine/Private/Components/RuntimeVirtualTextureComponent.cpp b/Engine/Source/Runtime/Engine/Private/Components/RuntimeVirtualTextureComponent.cpp
--- a/Engine/Source/Runtime/Engine/Private/Components/RuntimeVirtualTextureComponent.cpp
+++ b/Engine/Source/Runtime/Engine/Private/Components/RuntimeVirtualTextureComponent.cpp
@@ -14,6 +14,7 @@
 #include "UObject/UObjectIterator.h"
 #include "UObject/UnrealType.h"
 #include "VT/RuntimeVirtualTexture.h"
+#include "VT/VirtualTexture.h"
 #include "VT/VirtualTextureBuilder.h"
 #include "RenderUtils.h"
 #include "RHIGlobals.h"
@@ -241,6 +242,13 @@ bool URuntimeVirtualTextureComponent::IsStreamingLowMips(EShadingPath ShadingPat
 	return VirtualTexture != nullptr && StreamingTexture != nullptr && StreamingTexture->GetVirtualTexture(ShadingPath) != nullptr;
 }
 
+bool IsCompatibleFormat(URuntimeVirtualTexture const& RuntimeVirtualTexture, UVirtualTexture2D const& StreamingVirtualTexture)
+{
+	// During texture compilation we can't validate anything other than first layer, so restrict validation to that.
+	// This should catch any 99% of issues anyway. 
+	return (RuntimeVirtualTexture.GetLayerFormat(0) != StreamingVirtualTexture.GetPixelFormat(0));
+}
+
 bool URuntimeVirtualTextureComponent::IsStreamingTextureInvalid(EShadingPath ShadingPath) const
 {
 	checkf(IsActiveInWorld(), TEXT("This function should never be called for a world where we're inactive"));
@@ -249,6 +257,7 @@ bool URuntimeVirtualTextureComponent::IsStreamingTextureInvalid(EShadingPath Sha
 		VirtualTexture != nullptr && 
 		StreamingTexture != nullptr && 
 		StreamingTexture->GetVirtualTexture(ShadingPath) != nullptr && 
+		IsCompatibleFormat(*VirtualTexture, *StreamingTexture->GetVirtualTexture(ShadingPath)) &&
 		StreamingTexture->BuildHash != CalculateStreamingTextureSettingsHash();
 }
 
@@ -261,6 +270,8 @@ bool URuntimeVirtualTextureComponent::IsStreamingTextureInvalid() const
 		StreamingTexture != nullptr && 
 		StreamingTexture->GetVirtualTexture(EShadingPath::Deferred) != nullptr && 
 		StreamingTexture->GetVirtualTexture(EShadingPath::Mobile) != nullptr && 
+		IsCompatibleFormat(*VirtualTexture, *StreamingTexture->GetVirtualTexture(EShadingPath::Deferred)) &&
+		IsCompatibleFormat(*VirtualTexture, *StreamingTexture->GetVirtualTexture(EShadingPath::Mobile)) &&
 		StreamingTexture->BuildHash != CalculateStreamingTextureSettingsHash();
 }
 
