Add unconditional sleep of 16ms to AsyncCompilationHelpers::FinishCompilation to free up the core while we know there is outstanding work in flight.

diff --git a/Engine/Source/Runtime/Engine/Private/AsyncCompilationHelpers.cpp b/Engine/Source/Runtime/Engine/Private/AsyncCompilationHelpers.cpp
--- a/Engine/Source/Runtime/Engine/Private/AsyncCompilationHelpers.cpp
+++ b/Engine/Source/Runtime/Engine/Private/AsyncCompilationHelpers.cpp
@@ -197,15 +197,14 @@ namespace AsyncCompilationHelpers
 				const bool bBlockOnGlobalShaderCompletion = false;
 				GShaderCompilingManager->ProcessAsyncResults(bLimitExecutionTime, bBlockOnGlobalShaderCompletion);
 			}
-			else if(NumDone < Num)
-			{
-				// Jobs are still in flight so give them some time to complete
-				FPlatformProcess::Sleep(0.016);
-			}
-			else
+
+			if(NumDone >= Num)
 			{
 				break;
 			}
+
+			// Jobs are still in flight so give them some time to complete
+			FPlatformProcess::Sleep(0.016);
 		}
 
 		SaveStallStack(FPlatformTime::Cycles64() - StartTime);
