SparseVolumeTextures can now use reserved resources to reduce memory fragmentation and to get around the 2GB resource size limit.

#rnx

diff --git a/Engine/Source/Runtime/Engine/Private/SparseVolumeTexture/SparseVolumeTextureTileDataTexture.cpp b/Engine/Source/Runtime/Engine/Private/SparseVolumeTexture/SparseVolumeTextureTileDataTexture.cpp
--- a/Engine/Source/Runtime/Engine/Private/SparseVolumeTexture/SparseVolumeTextureTileDataTexture.cpp
+++ b/Engine/Source/Runtime/Engine/Private/SparseVolumeTexture/SparseVolumeTextureTileDataTexture.cpp
@@ -5,11 +5,53 @@
 #include "SparseVolumeTextureStreamingManager.h" // LogSparseVolumeTextureStreamingManager
 #include "RenderTargetPool.h"
 
+static int32 GSVTStreamingReservedResources = 1;
+static FAutoConsoleVariableRef CVarSVTStreamingReservedResources(
+	TEXT("r.SparseVolumeTexture.Streaming.UseReservedResources"),
+	GSVTStreamingReservedResources,
+	TEXT("Allocate the SVT tile data texture (streaming pool) as a reserved/virtual texture, backed by N small physical memory allocations to reduce fragmentation. This lifts the 2GB resource size limit and also allows for better GPU memory management when allocating the texture."),
+	ECVF_ReadOnly | ECVF_RenderThreadSafe
+);
+
+static int32 GSVTStreamingReservedResourcesMemoryLimit = -1;
+static FAutoConsoleVariableRef CVarSVTStreamingReservedResourcesMemoryLimit(
+	TEXT("r.SparseVolumeTexture.Streaming.ReservedResourcesMemoryLimit"),
+	GSVTStreamingReservedResourcesMemoryLimit,
+	TEXT("Memory limit in MiB on the maximum size of the streaming pool textures when using reserved resources. Without this limit it is theoretically possible to allocate enormous amounts of memory. Set to -1 to disable the limit. Default: -1"),
+	ECVF_ReadOnly | ECVF_RenderThreadSafe
+);
+
 namespace UE
 {
 namespace SVT
 {
 
+bool FTileDataTexture::ShouldUseReservedResources()
+{
+	const bool bRHISupportsReservedVolumeTextures = GRHIGlobals.ReservedResources.SupportsVolumeTextures;
+	const bool bCVarEnabled = GSVTStreamingReservedResources != 0;
+	return bRHISupportsReservedVolumeTextures && bCVarEnabled;
+}
+
+int64 FTileDataTexture::GetMaxTileDataTextureResourceSize(int32 InVoxelMemSize)
+{
+	// When using reserved resources, we are no longer bound by the 2GB resource size limit, so we can set our own limit.
+	int64 MaxStreamingPoolResourceSize = SVT::MaxResourceSize;
+	if (ShouldUseReservedResources())
+	{
+		// First compute the maximum size based on the texture format and the maximum texture dimensions. This limit still applies even with reserved resources.
+		constexpr int64 MaxNumVoxels = (int64)SVT::MaxVolumeTextureDim * (int64)SVT::MaxVolumeTextureDim * (int64)SVT::MaxVolumeTextureDim;
+		MaxStreamingPoolResourceSize = MaxNumVoxels * (int64)InVoxelMemSize;
+
+		if (GSVTStreamingReservedResourcesMemoryLimit > 0)
+		{
+			constexpr int64 OneMiB = 1024LL * 1024LL;
+			MaxStreamingPoolResourceSize = FMath::Min(MaxStreamingPoolResourceSize, (int64)GSVTStreamingReservedResourcesMemoryLimit) * OneMiB;
+		}
+	}
+	return MaxStreamingPoolResourceSize;
+}
+
 FIntVector3 FTileDataTexture::GetVolumeResolutionInTiles(int32 InNumRequiredTiles)
 {
 	int32 TileVolumeResolutionCube = 1;
@@ -38,27 +80,29 @@ FIntVector3 FTileDataTexture::GetVolumeResolutionInTiles(int32 InNumRequiredTile
 
 FIntVector3 FTileDataTexture::GetLargestPossibleVolumeResolutionInTiles(int32 InVoxelMemSize)
 {
+	const int64 MaxStreamingPoolResourceSize = GetMaxTileDataTextureResourceSize(InVoxelMemSize);
 	const int64 TileMemSize = SVT::NumVoxelsPerPaddedTile * InVoxelMemSize;
-	const int64 NumMaxTiles = SVT::MaxResourceSize / TileMemSize;
+	const int64 NumMaxTiles = MaxStreamingPoolResourceSize / TileMemSize;
+	constexpr int32 MaxDimensionInTiles = SVT::MaxVolumeTextureDim / SPARSE_VOLUME_TILE_RES_PADDED;
 
-	// Find a cube with a volume as close to NumMaxTiles as possible
+	// Find a cube with a volume as close to NumMaxTiles as possible but not exceeding MaxDimensionInTiles.
 	int32 TileVolumeResolutionCube = 1;
-	while (((TileVolumeResolutionCube + 1) * (TileVolumeResolutionCube + 1) * (TileVolumeResolutionCube + 1)) <= NumMaxTiles)
+	while ((((TileVolumeResolutionCube + 1) * (TileVolumeResolutionCube + 1) * (TileVolumeResolutionCube + 1)) <= NumMaxTiles) && ((TileVolumeResolutionCube + 1) <= MaxDimensionInTiles))
 	{
 		++TileVolumeResolutionCube;
 	}
 
-	// Try to add to the sides to get closer to NumMaxTiles
+	// Try to add to the sides to get closer to NumMaxTiles without exceeding MaxDimensionInTiles
 	FIntVector3 ResolutionInTiles = FIntVector3(TileVolumeResolutionCube, TileVolumeResolutionCube, TileVolumeResolutionCube);
-	if (((ResolutionInTiles.X + 1) * ResolutionInTiles.Y * ResolutionInTiles.Z) <= NumMaxTiles)
+	if ((((ResolutionInTiles.X + 1) * ResolutionInTiles.Y * ResolutionInTiles.Z) <= NumMaxTiles) && ((ResolutionInTiles.X + 1) <= MaxDimensionInTiles))
 	{
 		++ResolutionInTiles.X;
 	}
-	if ((ResolutionInTiles.X * (ResolutionInTiles.Y + 1) * ResolutionInTiles.Z) <= NumMaxTiles)
+	if (((ResolutionInTiles.X * (ResolutionInTiles.Y + 1) * ResolutionInTiles.Z) <= NumMaxTiles) && ((ResolutionInTiles.Y + 1) <= MaxDimensionInTiles))
 	{
 		++ResolutionInTiles.Y;
 	}
-	if ((ResolutionInTiles.X * ResolutionInTiles.Y * (ResolutionInTiles.Z + 1)) <= NumMaxTiles)
+	if (((ResolutionInTiles.X * ResolutionInTiles.Y * (ResolutionInTiles.Z + 1)) <= NumMaxTiles) && ((ResolutionInTiles.Z + 1) <= MaxDimensionInTiles))
 	{
 		++ResolutionInTiles.Z;
 	}
@@ -66,7 +110,7 @@ FIntVector3 FTileDataTexture::GetLargestPossibleVolumeResolutionInTiles(int32 In
 	const FIntVector3 ResolutionInVoxels = ResolutionInTiles * SPARSE_VOLUME_TILE_RES_PADDED;
 	check(IsInBounds(ResolutionInVoxels, FIntVector3::ZeroValue, FIntVector3(SVT::MaxVolumeTextureDim + 1)));
 	check(ResolutionInVoxels.X <= SVT::MaxVolumeTextureDim && ResolutionInVoxels.Y <= SVT::MaxVolumeTextureDim && ResolutionInVoxels.Z <= SVT::MaxVolumeTextureDim);
-	check(((int64)ResolutionInVoxels.X * (int64)ResolutionInVoxels.Y * (int64)ResolutionInVoxels.Z * (int64)InVoxelMemSize) <= SVT::MaxResourceSize);
+	check(((int64)ResolutionInVoxels.X * (int64)ResolutionInVoxels.Y * (int64)ResolutionInVoxels.Z * (int64)InVoxelMemSize) <= MaxStreamingPoolResourceSize);
 
 	return ResolutionInTiles;
 }
@@ -100,8 +144,9 @@ FTileDataTexture::FTileDataTexture(const FIntVector3& InResolutionInTiles, EPixe
 
 	const FIntVector3 Resolution = ResolutionInTiles * SPARSE_VOLUME_TILE_RES_PADDED;
 	check(Resolution.X <= SVT::MaxVolumeTextureDim && Resolution.Y <= SVT::MaxVolumeTextureDim && Resolution.Z <= SVT::MaxVolumeTextureDim);
-	check(((int64)Resolution.X * (int64)Resolution.Y * (int64)Resolution.Z * (int64)GPixelFormats[FormatA].BlockBytes) <= SVT::MaxResourceSize);
-	check(((int64)Resolution.X * (int64)Resolution.Y * (int64)Resolution.Z * (int64)GPixelFormats[FormatB].BlockBytes) <= SVT::MaxResourceSize);
+	const int64 MaxTileDataTextureResourceSize = GetMaxTileDataTextureResourceSize(MaxFormatSize);
+	check(((int64)Resolution.X * (int64)Resolution.Y * (int64)Resolution.Z * (int64)GPixelFormats[FormatA].BlockBytes) <= MaxTileDataTextureResourceSize);
+	check(((int64)Resolution.X * (int64)Resolution.Y * (int64)Resolution.Z * (int64)GPixelFormats[FormatB].BlockBytes) <= MaxTileDataTextureResourceSize);
 
 	TileCoords.SetNum(PhysicalTilesCapacity);
 
@@ -172,7 +217,8 @@ void FTileDataTexture::EndUpload(FRDGBuilder& GraphBuilder)
 void FTileDataTexture::InitRHI(FRHICommandListBase& RHICmdList)
 {
 	const FIntVector3 Resolution = ResolutionInTiles * SPARSE_VOLUME_TILE_RES_PADDED;
-	const ETextureCreateFlags Flags = TexCreate_ShaderResource | TexCreate_UAV | TexCreate_3DTiling | TexCreate_ReduceMemoryWithTilingMode;
+	const ETextureCreateFlags ReservedResourceFlags = ShouldUseReservedResources() ? (TexCreate_ReservedResource | TexCreate_ImmediateCommit) : TexCreate_None;
+	const ETextureCreateFlags Flags = ReservedResourceFlags | TexCreate_ShaderResource | TexCreate_UAV | TexCreate_3DTiling | TexCreate_ReduceMemoryWithTilingMode;
 	if (FormatA != PF_Unknown)
 	{
 		TileDataTextureA = GRenderTargetPool.FindFreeElement(RHICmdList, FRDGTextureDesc::Create3D(Resolution, FormatA, FClearValueBinding::Black, Flags), TEXT("SparseVolumeTexture.PhysicalTileDataA.RHITexture"));
diff --git a/Engine/Source/Runtime/Engine/Private/SparseVolumeTexture/SparseVolumeTextureTileDataTexture.h b/Engine/Source/Runtime/Engine/Private/SparseVolumeTexture/SparseVolumeTextureTileDataTexture.h
--- a/Engine/Source/Runtime/Engine/Private/SparseVolumeTexture/SparseVolumeTextureTileDataTexture.h
+++ b/Engine/Source/Runtime/Engine/Private/SparseVolumeTexture/SparseVolumeTextureTileDataTexture.h
@@ -29,6 +29,8 @@ public:
 
 	static constexpr uint32 PhysicalCoordMask = (1u << 24u) - 1u; // Lower 24 bits are used for storing XYZ in 8 bit each. Upper 8 bit can be used by the caller. 
 
+	static bool ShouldUseReservedResources();
+	static int64 GetMaxTileDataTextureResourceSize(int32 InVoxelMemSize);
 	static FIntVector3 GetVolumeResolutionInTiles(int32 InNumRequiredTiles);
 	static FIntVector3 GetLargestPossibleVolumeResolutionInTiles(int32 InVoxelMemSize);
 
