World Partition - HLOD: Improved AActor::IsHLODRelevant() so that it can better reflect the state of an actor for the HLOD relevant view
* Only spatially loaded actors can be HLOD relevant in partitioned worlds
* Editor only actors should not be relevant to HLOD

Those two properties were already considered during HLOD generation (which relies on WP Streaming Generation).

diff --git a/Engine/Source/Runtime/Engine/Private/Actor.cpp b/Engine/Source/Runtime/Engine/Private/Actor.cpp
--- a/Engine/Source/Runtime/Engine/Private/Actor.cpp
+++ b/Engine/Source/Runtime/Engine/Private/Actor.cpp
@@ -6015,11 +6015,24 @@ bool AActor::IsHLODRelevant() const
 		return false;
 	}
 
+	if (IsEditorOnly())
+	{
+		return false;
+	}
+
 	if (!bEnableAutoLODGeneration)
 	{
 		return false;
 	}
 
+#if WITH_EDITOR 
+	// Only spatially loaded actors can be HLOD relevant in partitioned worlds
+	if (UWorld::IsPartitionedWorld(GetWorld()) && !GetIsSpatiallyLoaded())
+	{
+		return false;
+	}
+#endif
+
 	FVector Origin, Extent;
 	GetActorBounds(false, Origin, Extent);
 	if (Extent.SizeSquared() <= 0.1)
