GitHub 11421 : Fix crash when mounting pak file after FRenderAssetStreamingManager was deleted.

#11421 tag

diff --git a/Engine/Source/Runtime/Engine/Private/Streaming/StreamingManagerTexture.cpp b/Engine/Source/Runtime/Engine/Private/Streaming/StreamingManagerTexture.cpp
--- a/Engine/Source/Runtime/Engine/Private/Streaming/StreamingManagerTexture.cpp
+++ b/Engine/Source/Runtime/Engine/Private/Streaming/StreamingManagerTexture.cpp
@@ -184,14 +184,14 @@ FRenderAssetStreamingManager::FRenderAssetStreamingManager()
 
 	FCoreUObjectDelegates::GetPreGarbageCollectDelegate().AddRaw(this, &FRenderAssetStreamingManager::OnPreGarbageCollect);
 
-	FCoreDelegates::GetOnPakFileMounted2().AddLambda([this](const IPakFile& PakFile)
+	OnPakFileMounted2Handle = FCoreDelegates::GetOnPakFileMounted2().AddLambda([this](const IPakFile& PakFile)
 	{
 		FScopeLock ScopeLock(&MountedStateDirtyFilesCS);
 		bRecacheAllFiles = true;
 		MountedStateDirtyFiles.Empty();
 	});
 
-	FCoreDelegates::NewFileAddedDelegate.AddLambda([this](const FString& FileName)
+	NewFileAddedDelegateHandle = FCoreDelegates::NewFileAddedDelegate.AddLambda([this](const FString& FileName)
 	{
 		MarkMountedStateDirty(MakeIoFilenameHash(FileName));
 	});
@@ -206,6 +206,8 @@ FRenderAssetStreamingManager::~FRenderAssetStreamingManager()
 
 	RenderAssetInstanceAsyncWork->EnsureCompletion();
 	
+	FCoreDelegates::GetOnPakFileMounted2().Remove(OnPakFileMounted2Handle);
+	FCoreDelegates::NewFileAddedDelegate.Remove(NewFileAddedDelegateHandle);
 	FCoreUObjectDelegates::GetPreGarbageCollectDelegate().RemoveAll(this);
 
 	// Clear the stats
diff --git a/Engine/Source/Runtime/Engine/Private/Streaming/StreamingManagerTexture.h b/Engine/Source/Runtime/Engine/Private/Streaming/StreamingManagerTexture.h
--- a/Engine/Source/Runtime/Engine/Private/Streaming/StreamingManagerTexture.h
+++ b/Engine/Source/Runtime/Engine/Private/Streaming/StreamingManagerTexture.h
@@ -427,6 +427,9 @@ private:
 	typedef TSet<FIoFilenameHash> FIoFilenameHashSet;
 	FIoFilenameHashSet MountedStateDirtyFiles;
 
+	FDelegateHandle OnPakFileMounted2Handle;
+	FDelegateHandle NewFileAddedDelegateHandle;
+
 	ENGINE_API virtual void MarkMountedStateDirty(FIoFilenameHash FilenameHash) override;
 
 	// A critical section use around code that could be called in parallel with NotifyPrimitiveUpdated() or NotifyPrimitiveUpdated_Concurrent().
