Clamp montage playback when stopping at end of section

diff --git a/Engine/Source/Runtime/Engine/Private/Animation/AnimMontage.cpp b/Engine/Source/Runtime/Engine/Private/Animation/AnimMontage.cpp
--- a/Engine/Source/Runtime/Engine/Private/Animation/AnimMontage.cpp
+++ b/Engine/Source/Runtime/Engine/Private/Animation/AnimMontage.cpp
@@ -2556,13 +2556,14 @@ void FAnimMontageInstance::Advance(float DeltaTime, struct FRootMotionMovementPa
 						// Get recent NextSectionIndex in case it's been changed by previous events.
 						const int32 CurrentSectionIndex = MontageSubStepper.GetCurrentSectionIndex();
 						const int32 RecentNextSectionIndex = bPlayingForward ? NextSections[CurrentSectionIndex] : PrevSections[CurrentSectionIndex];
+						const float EndOffset = UE_KINDA_SMALL_NUMBER / 2.f; //KINDA_SMALL_NUMBER/2 because we use KINDA_SMALL_NUMBER to offset notifies for triggering and SMALL_NUMBER is too small
+
 						if (RecentNextSectionIndex != INDEX_NONE)
 						{
 							float LatestNextSectionStartTime, LatestNextSectionEndTime;
 							Montage->GetSectionStartAndEndTime(RecentNextSectionIndex, LatestNextSectionStartTime, LatestNextSectionEndTime);
 
 							// Jump to next section's appropriate starting point (start or end).
-							const float EndOffset = UE_KINDA_SMALL_NUMBER / 2.f; //KINDA_SMALL_NUMBER/2 because we use KINDA_SMALL_NUMBER to offset notifies for triggering and SMALL_NUMBER is too small
 							Position = bPlayingForward ? LatestNextSectionStartTime : (LatestNextSectionEndTime - EndOffset);
 							SubStepResult = EMontageSubStepResult::Moved;
 						}
@@ -2570,10 +2571,16 @@ void FAnimMontageInstance::Advance(float DeltaTime, struct FRootMotionMovementPa
 						{
 							// If there is no next section and we've reached the end of this one, exit
 
-							// Stop playing to prevent playing animation data past the end of the current section
+							// Stop playing and clamp position to prevent playing animation data past the end of the current section
 							// We already called Stop above if needed, like if bEnableAutoBlendOut is true
 							bPlaying = false;
 
+							float CurrentSectionStartTime, CurrentSectionEndTime;
+							Montage->GetSectionStartAndEndTime(CurrentSectionIndex, CurrentSectionStartTime, CurrentSectionEndTime);
+
+							Position = bPlayingForward ? (CurrentSectionEndTime - EndOffset) : CurrentSectionStartTime;
+							SubStepResult = EMontageSubStepResult::Moved;
+
 							break;
 						}
 					}
