Fixing live attenuation updates.

Made the max distance value use the attenuation setting asset and not use the cached value.

diff --git a/Engine/Plugins/AudioGameplay/Source/AudioGameplay/Private/SoundHandleSubsystem.cpp b/Engine/Plugins/AudioGameplay/Source/AudioGameplay/Private/SoundHandleSubsystem.cpp
--- a/Engine/Plugins/AudioGameplay/Source/AudioGameplay/Private/SoundHandleSubsystem.cpp
+++ b/Engine/Plugins/AudioGameplay/Source/AudioGameplay/Private/SoundHandleSubsystem.cpp
@@ -84,8 +84,6 @@ Audio::EResult USoundHandleSubsystem::Play(const Audio::FSoundHandleID ID)
 				Handle->ActiveSound.AttenuationSettings = *AttenuationSettings;
 				Handle->ActiveSound.FocusData.PriorityScale = AttenuationSettings->GetFocusPriorityScale(GetAudioDeviceHandle().GetAudioDevice()->GetGlobalFocusSettings(), FocusFactor);
 			}
-
-			Handle->ActiveSound.MaxDistance = MaxDistance;
 		}
 		
 		AudioDevice->AddNewActiveSound(Handle->ActiveSound);		
diff --git a/Engine/Source/Runtime/Engine/Private/ActiveSound.cpp b/Engine/Source/Runtime/Engine/Private/ActiveSound.cpp
--- a/Engine/Source/Runtime/Engine/Private/ActiveSound.cpp
+++ b/Engine/Source/Runtime/Engine/Private/ActiveSound.cpp
@@ -1399,6 +1399,28 @@ FWaveInstance& FActiveSound::AddWaveInstance(const UPTRINT WaveInstanceHash)
 	return *WaveInstance;
 }
 
+float FActiveSound::GetMaxDistance() const
+{
+	float OutMaxDistance = 0.0f;
+	float OutFocusFactor = 1.0f;
+	GetMaxDistanceAndFocusFactor(OutMaxDistance, OutFocusFactor);
+	return OutMaxDistance;
+}
+
+void FActiveSound::GetMaxDistanceAndFocusFactor(float& OutMaxDistance, float& OutFocusFactor) const
+{
+	OutMaxDistance = 0.0f;
+	OutFocusFactor = 1.0f;
+
+	if (!Sound)
+	{
+		return;
+	}
+
+	check(AudioDevice);
+	AudioDevice->GetMaxDistanceAndFocusFactor(Sound, GetWorld(), LastLocation, GetAttenuationSettings(), OutMaxDistance, OutFocusFactor);
+}
+
 void FActiveSound::ApplyRadioFilter(const FSoundParseParameters& ParseParams)
 {
 	check(AudioDevice);
@@ -1674,7 +1696,7 @@ void FActiveSound::UpdateAttenuation(float DeltaTime, FSoundParseParameters& Par
 	UpdateAttenuation(DeltaTime, ParseParams, Listener.ListenerIndex, SettingsAttenuationNode);
 }
 
-void FActiveSound::UpdateAttenuation(float DeltaTime, FSoundParseParameters& ParseParams, int32 ListenerIndex, const FSoundAttenuationSettings* SettingsAttenuationNode)
+const FSoundAttenuationSettings* FActiveSound::GetAttenuationSettings(const FSoundAttenuationSettings* SettingsAttenuationNode) const
 {
 	// We default to using the copied off "overridden" settings (or default constructed settings)
 	const FSoundAttenuationSettings* Settings = &AttenuationSettings;
@@ -1693,6 +1715,13 @@ void FActiveSound::UpdateAttenuation(float DeltaTime, FSoundParseParameters& Par
 			Settings = &SoundAttenuation->Attenuation;
 		}
 	}
+	return Settings;
+}
+
+void FActiveSound::UpdateAttenuation(float DeltaTime, FSoundParseParameters& ParseParams, int32 ListenerIndex, const FSoundAttenuationSettings* SettingsAttenuationNode)
+{
+	// We default to using the copied off "overridden" settings (or default constructed settings)
+	const FSoundAttenuationSettings* Settings = GetAttenuationSettings(SettingsAttenuationNode);
 
 	// Reset Focus data and recompute if necessary
 	FAttenuationFocusData FocusDataToApply;
diff --git a/Engine/Source/Runtime/Engine/Private/AudioDebug.cpp b/Engine/Source/Runtime/Engine/Private/AudioDebug.cpp
--- a/Engine/Source/Runtime/Engine/Private/AudioDebug.cpp
+++ b/Engine/Source/Runtime/Engine/Private/AudioDebug.cpp
@@ -1110,7 +1110,7 @@ namespace Audio
 			const FString Name = Sound->GetName();
 			const FTransform CurTransform = ActiveSound.Transform;
 			FColor TextColor = FColor::White;
-			const float CurMaxDistance = ActiveSound.MaxDistance;
+			const float CurMaxDistance = ActiveSound.GetMaxDistance();
 			float DisplayValue = 0.0f;
 			float FilterValue = 0.0f;
 			if (ActiveSoundVisualizeModeCVar == 1 || ActiveSoundVisualizeModeCVar == 2)
diff --git a/Engine/Source/Runtime/Engine/Private/AudioDevice.cpp b/Engine/Source/Runtime/Engine/Private/AudioDevice.cpp
--- a/Engine/Source/Runtime/Engine/Private/AudioDevice.cpp
+++ b/Engine/Source/Runtime/Engine/Private/AudioDevice.cpp
@@ -6397,8 +6397,6 @@ void FAudioDevice::PlaySoundAtLocation(USoundBase* Sound, UWorld* World, float V
 			NewActiveSound.FocusData.DistanceScale = AttenuationSettingsToApply->GetFocusDistanceScale(FocusSettings, FocusFactor);
 		}
 
-		NewActiveSound.MaxDistance = MaxDistance;
-
 		if (Concurrency)
 		{
 			NewActiveSound.ConcurrencySet.Add(Concurrency);
diff --git a/Engine/Source/Runtime/Engine/Private/AudioVirtualLoop.cpp b/Engine/Source/Runtime/Engine/Private/AudioVirtualLoop.cpp
--- a/Engine/Source/Runtime/Engine/Private/AudioVirtualLoop.cpp
+++ b/Engine/Source/Runtime/Engine/Private/AudioVirtualLoop.cpp
@@ -150,7 +150,7 @@ void FAudioVirtualLoop::CalculateUpdateInterval()
 	check(AudioDevice);
 
 	const float DistanceToListener = AudioDevice->GetDistanceToNearestListener(ActiveSound->Transform.GetLocation());
-	const float DistanceRatio = (DistanceToListener - ActiveSound->MaxDistance) / FMath::Max(VirtualLoopsPerfDistanceCVar, 1.0f);
+	const float DistanceRatio = (DistanceToListener - ActiveSound->GetMaxDistance()) / FMath::Max(VirtualLoopsPerfDistanceCVar, 1.0f);
 	const float DistanceRatioClamped = FMath::Clamp(DistanceRatio, 0.0f, 1.0f);
 	UpdateInterval = FMath::Lerp(VirtualLoopsUpdateRateMinCVar, VirtualLoopsUpdateRateMaxCVar, DistanceRatioClamped);
 }
@@ -215,7 +215,7 @@ bool FAudioVirtualLoop::IsInAudibleRange(const FActiveSound& InActiveSound, cons
 
 	DistanceScale = FMath::Max(DistanceScale, UE_KINDA_SMALL_NUMBER);
 	const FVector Location = InActiveSound.Transform.GetLocation();
-	return AudioDevice->LocationIsAudible(Location, InActiveSound.MaxDistance / DistanceScale);
+	return AudioDevice->LocationIsAudible(Location, InActiveSound.GetMaxDistance() / DistanceScale);
 }
 
 void FAudioVirtualLoop::UpdateFocusData(float DeltaTime)
diff --git a/Engine/Source/Runtime/Engine/Private/Components/AudioComponent.cpp b/Engine/Source/Runtime/Engine/Private/Components/AudioComponent.cpp
--- a/Engine/Source/Runtime/Engine/Private/Components/AudioComponent.cpp
+++ b/Engine/Source/Runtime/Engine/Private/Components/AudioComponent.cpp
@@ -817,8 +817,6 @@ void UAudioComponent::PlayInternal(const PlayInternalRequestData& InPlayRequestD
 	// Pass down any source buffer listener we have
 	NewActiveSound.SetSourceListener(SourceBufferListener, bShouldSourceBufferListenerZeroBuffer);
 
-	NewActiveSound.MaxDistance = MaxDistance;
-
 	// Setup the submix and bus sends that may have been set before playing
 	for (FSoundSubmixSendInfo& SubmixSendInfo : PendingSubmixSends)
 	{
@@ -868,6 +866,9 @@ void UAudioComponent::PlayInternal(const PlayInternalRequestData& InPlayRequestD
 	TArray<FAudioParameter> InstanceParamsCopy = InstanceParameters;
 	FAudioParameter::Merge(MoveTemp(InstanceParamsCopy), SoundParams);
 
+	// Make sure our new active sound has the correct audio device
+	NewActiveSound.SetAudioDevice(AudioDevice);
+
 	AudioDevice->AddNewActiveSound(NewActiveSound, MoveTemp(SoundParams));
 
 	LastSoundPlayOrder = NewActiveSound.GetPlayOrder();
diff --git a/Engine/Source/Runtime/Engine/Public/ActiveSound.h b/Engine/Source/Runtime/Engine/Public/ActiveSound.h
--- a/Engine/Source/Runtime/Engine/Public/ActiveSound.h
+++ b/Engine/Source/Runtime/Engine/Public/ActiveSound.h
@@ -371,6 +371,12 @@ public:
 	/** Whether or not the active sound is currently playing audible sound. */
 	bool IsPlayingAudio() const { return bIsPlayingAudio; }
 
+	// Returns the max distance of the active sound
+	float GetMaxDistance() const;
+	void GetMaxDistanceAndFocusFactor(float& OutMaxDistance, float& OutFocusFactor) const;
+
+	const FSoundAttenuationSettings* GetAttenuationSettings(const FSoundAttenuationSettings* SettingsAttenuationNode = nullptr) const;
+
 	/** Whether or not sound reference is valid and set to play when silent. */
 	ENGINE_API bool IsPlayWhenSilent() const;
 
@@ -584,7 +590,7 @@ public:
 	/** Last time we checked for occlusion */
 	float LastOcclusionCheckTime;
 
-	/** The max distance this sound will be audible. */
+	UE_DEPRECATED(5.4, "Please do not use MaxDistance. Use GetMaxDistance() instead.")
 	float MaxDistance;
 
 	FTransform Transform;
