Fixed crash on load when a canvas render target was saved with sample counts above one.

The bCanCreateUAV member is not serialized, and thus needs to be updated accordingly.

#rnx

diff --git a/Engine/Source/Runtime/Engine/Private/CanvasRenderTarget2D.cpp b/Engine/Source/Runtime/Engine/Private/CanvasRenderTarget2D.cpp
--- a/Engine/Source/Runtime/Engine/Private/CanvasRenderTarget2D.cpp
+++ b/Engine/Source/Runtime/Engine/Private/CanvasRenderTarget2D.cpp
@@ -26,6 +26,11 @@ UCanvasRenderTarget2D::UCanvasRenderTarget2D( const FObjectInitializer& ObjectIn
 
 void UCanvasRenderTarget2D::UpdateResource()
 {
+	if (SampleCount != ETextureRenderTargetSampleCount::RTSC_1)
+	{
+		bCanCreateUAV = false;
+	}
+
 	// Call parent implementation
 	Super::UpdateResource();
 
