USplineMeshComponent: Remove deprecated GetAxisValue functions

diff --git a/Engine/Source/Runtime/Engine/Classes/Components/SplineMeshComponent.h b/Engine/Source/Runtime/Engine/Classes/Components/SplineMeshComponent.h
--- a/Engine/Source/Runtime/Engine/Classes/Components/SplineMeshComponent.h
+++ b/Engine/Source/Runtime/Engine/Classes/Components/SplineMeshComponent.h
@@ -414,16 +414,6 @@ public:
 	 */
 	ENGINE_API FTransform CalcSliceTransformAtSplineOffset(const float Alpha, const float MinT=0.f, const float MaxT=1.0f) const;
 
-	UE_DEPRECATED(5.2, "Use GetAxisValueRef() instead.")
-	static const double& GetAxisValue(const FVector3d& InVector, ESplineMeshAxis::Type InAxis) { return GetAxisValueRef(InVector, InAxis); }
-	UE_DEPRECATED(5.2, "Use GetAxisValueRef() instead.")
-	static double& GetAxisValue(FVector3d& InVector, ESplineMeshAxis::Type InAxis) { return GetAxisValueRef(InVector, InAxis); }
-
-	UE_DEPRECATED(5.2, "Use GetAxisValueRef() instead.")
-	static const float& GetAxisValue(const FVector3f& InVector, ESplineMeshAxis::Type InAxis) { return GetAxisValueRef(InVector, InAxis); }
-	UE_DEPRECATED(5.2, "Use GetAxisValueRef() instead.")
-	static float& GetAxisValue(FVector3f& InVector, ESplineMeshAxis::Type InAxis) { return GetAxisValueRef(InVector, InAxis); }
-
 	inline static const double& GetAxisValueRef(const FVector3d& InVector, ESplineMeshAxis::Type InAxis);
 	inline static double& GetAxisValueRef(FVector3d& InVector, ESplineMeshAxis::Type InAxis);
 
