Fixed case where breaking a level instance not using actor folders inside a partitioned map would not properly set the ActorFolderGuid of all child actors which could make future folder operation on these actors not working as expected.

diff --git a/Engine/Source/Runtime/Engine/Private/Actor.cpp b/Engine/Source/Runtime/Engine/Private/Actor.cpp
--- a/Engine/Source/Runtime/Engine/Private/Actor.cpp
+++ b/Engine/Source/Runtime/Engine/Private/Actor.cpp
@@ -3636,7 +3636,7 @@ void AActor::PostEditImport()
 	ULevel* Level = GetLevel();
 	if (Level && Level->IsUsingActorFolders())
 	{
-		if (GetFolderPath().IsNone() && !FolderPath.IsNone())
+		if ((GetFolderPath().IsNone() || !FolderGuid.IsValid()) && !FolderPath.IsNone())
 		{
 			SetFolderPath(FolderPath);
 		}
