World Partition: include source partition grid name and optional cell bounds to runtime streaming cells.

diff --git a/Engine/Source/Runtime/Engine/Private/WorldPartition/RuntimeHashSet/RuntimePartitionLHGrid.cpp b/Engine/Source/Runtime/Engine/Private/WorldPartition/RuntimeHashSet/RuntimePartitionLHGrid.cpp
--- a/Engine/Source/Runtime/Engine/Private/WorldPartition/RuntimeHashSet/RuntimePartitionLHGrid.cpp
+++ b/Engine/Source/Runtime/Engine/Private/WorldPartition/RuntimeHashSet/RuntimePartitionLHGrid.cpp
@@ -48,6 +48,18 @@ struct FCellCoord
 		);
 	}
 
+	static inline FBox GetCellBounds(const FCellCoord& InCellCoord, int32 InCellSize)
+	{
+		check(InCellCoord.Level >= 0);
+		const int64 CellSizeForLevel = (int64)InCellSize * (1LL << InCellCoord.Level);
+		const FVector Min = FVector(
+			static_cast<FVector::FReal>(InCellCoord.X * CellSizeForLevel), 
+			static_cast<FVector::FReal>(InCellCoord.Y * CellSizeForLevel), 
+			static_cast<FVector::FReal>(InCellCoord.Z * CellSizeForLevel)
+		);
+		const FVector Max = Min + FVector(static_cast<double>(CellSizeForLevel));
+		return FBox(Min, Max);
+	}
 	friend uint32 GetTypeHash(const FCellCoord& CellCoord)
 	{
 		FHashBuilder HashBuilder;
@@ -137,7 +149,8 @@ bool URuntimePartitionLHGrid::GenerateStreaming(const FGenerateStreamingParams&
 
 	for (auto& [CellCoord, CellActorSetInstances] : CellsActorSetInstances)
 	{
-		OutResult.RuntimeCellDescs.Emplace(CreateCellDesc(CellCoord.ToString(), true, CellCoord.Level, CellActorSetInstances));
+		URuntimePartition::FCellDesc& CellDesc = OutResult.RuntimeCellDescs.Emplace_GetRef(CreateCellDesc(CellCoord.ToString(), true, CellCoord.Level, CellActorSetInstances));
+		CellDesc.CellBounds = FCellCoord::GetCellBounds(CellCoord, CellSize);
 	}
 
 	return true;
diff --git a/Engine/Source/Runtime/Engine/Private/WorldPartition/RuntimeHashSet/WorldPartitionRuntimeHashSetStreamingGeneration.cpp b/Engine/Source/Runtime/Engine/Private/WorldPartition/RuntimeHashSet/WorldPartitionRuntimeHashSetStreamingGeneration.cpp
--- a/Engine/Source/Runtime/Engine/Private/WorldPartition/RuntimeHashSet/WorldPartitionRuntimeHashSetStreamingGeneration.cpp
+++ b/Engine/Source/Runtime/Engine/Private/WorldPartition/RuntimeHashSet/WorldPartitionRuntimeHashSetStreamingGeneration.cpp
@@ -192,8 +192,10 @@ bool UWorldPartitionRuntimeHashSet::GenerateStreaming(UWorldPartitionStreamingPo
 
 		UWorldPartitionRuntimeCellData* RuntimeCellData = RuntimeCell->RuntimeCellData;
 		RuntimeCellData->DebugName = CellUniqueId.Name;
+		RuntimeCellData->CellBounds = CellDescInstance.CellBounds;
 		RuntimeCellData->HierarchicalLevel = CellDescInstance.Level;
 		RuntimeCellData->Priority = CellDescInstance.Priority;
+		RuntimeCellData->GridName = CellDescInstance.SourcePartition->Name;
 
 		return RuntimeCell;
 	};
diff --git a/Engine/Source/Runtime/Engine/Private/WorldPartition/WorldPartitionRuntimeCellData.cpp b/Engine/Source/Runtime/Engine/Private/WorldPartition/WorldPartitionRuntimeCellData.cpp
--- a/Engine/Source/Runtime/Engine/Private/WorldPartition/WorldPartitionRuntimeCellData.cpp
+++ b/Engine/Source/Runtime/Engine/Private/WorldPartition/WorldPartitionRuntimeCellData.cpp
@@ -112,7 +112,7 @@ const FBox& UWorldPartitionRuntimeCellData::GetContentBounds() const
 
 FBox UWorldPartitionRuntimeCellData::GetCellBounds() const
 {
-	return ContentBounds;
+	return CellBounds.IsSet() ? *CellBounds : ContentBounds;
 }
 
 FString UWorldPartitionRuntimeCellData::GetDebugName() const
diff --git a/Engine/Source/Runtime/Engine/Public/WorldPartition/RuntimeHashSet/RuntimePartition.h b/Engine/Source/Runtime/Engine/Public/WorldPartition/RuntimeHashSet/RuntimePartition.h
--- a/Engine/Source/Runtime/Engine/Public/WorldPartition/RuntimeHashSet/RuntimePartition.h
+++ b/Engine/Source/Runtime/Engine/Public/WorldPartition/RuntimeHashSet/RuntimePartition.h
@@ -33,6 +33,9 @@ public:
 		/** Optional level value that can be used to filter debug display */
 		int32 Level;
 
+		/** Optional cell bounds for partitions that work on uniform grids */
+		TOptional<FBox> CellBounds;
+
 		TArray<const IStreamingGenerationContext::FActorSetInstance*> ActorSetInstances;
 	};
 
diff --git a/Engine/Source/Runtime/Engine/Public/WorldPartition/WorldPartitionRuntimeCellData.h b/Engine/Source/Runtime/Engine/Public/WorldPartition/WorldPartitionRuntimeCellData.h
--- a/Engine/Source/Runtime/Engine/Public/WorldPartition/WorldPartitionRuntimeCellData.h
+++ b/Engine/Source/Runtime/Engine/Public/WorldPartition/WorldPartitionRuntimeCellData.h
@@ -55,6 +55,13 @@ class UWorldPartitionRuntimeCellData : public UObject
 	UPROPERTY()
 	FBox ContentBounds;
 
+	// Optional cell bounds
+	UPROPERTY()
+	TOptional<FBox> CellBounds;
+
+	UPROPERTY()
+	FName GridName;
+
 	UPROPERTY()
 	int32 Priority;
 
diff --git a/Engine/Source/Runtime/Engine/Public/WorldPartition/WorldPartitionRuntimeCellDataSpatialHash.h b/Engine/Source/Runtime/Engine/Public/WorldPartition/WorldPartitionRuntimeCellDataSpatialHash.h
--- a/Engine/Source/Runtime/Engine/Public/WorldPartition/WorldPartitionRuntimeCellDataSpatialHash.h
+++ b/Engine/Source/Runtime/Engine/Public/WorldPartition/WorldPartitionRuntimeCellDataSpatialHash.h
@@ -28,9 +28,6 @@ class UWorldPartitionRuntimeCellDataSpatialHash : public UWorldPartitionRuntimeC
 	UPROPERTY()
 	float Extent;
 
-	UPROPERTY()
-	FName GridName;
-
 private:
 	ENGINE_API float ComputeSourceToCellAngleFactor(const FSphericalSector& SourceShape) const;
 
