Texture limit output size in bytes from platform size limit
This limit is a hard limit on the size of a texture resource that will not crash the game
this is not a dimensional limit, or a heuristic limit to make small platform-friendly textures
rather this enforces the absolute limit of texture sizes that will not crash the runtime or fail to package

diff --git a/Engine/Source/Developer/TargetPlatform/Private/TargetPlatformBase.cpp b/Engine/Source/Developer/TargetPlatform/Private/TargetPlatformBase.cpp
--- a/Engine/Source/Developer/TargetPlatform/Private/TargetPlatformBase.cpp
+++ b/Engine/Source/Developer/TargetPlatform/Private/TargetPlatformBase.cpp
@@ -1,6 +1,7 @@
 // Copyright Epic Games, Inc. All Rights Reserved.
 
 #include "Common/TargetPlatformBase.h"
+#include "Common/TargetPlatformControlsBase.h"
 #include "HAL/IConsoleManager.h"
 #include "DeviceBrowserDefaultPlatformWidgetCreator.h"
 #include "Interfaces/IProjectBuildMutatorFeature.h"
@@ -401,6 +402,11 @@ void FTargetPlatformBase::GetWaveFormatModuleHints(TArray<FName>& OutModuleNames
 	GetAudioFormatSettings().GetWaveFormatModuleHints(OutModuleNames);
 }
 
+void FTargetPlatformBase::GetTextureSizeLimits(uint64 & OutMaximumSurfaceBytes, uint64 & OutMaximumPackageBytes) const
+{
+	FTargetPlatformControlsBase::GetTextureSizeLimitsDefault(GetConfigSystem(),OutMaximumSurfaceBytes,OutMaximumPackageBytes);
+}
+
 #endif // WITH_ENGINE
 
 #undef LOCTEXT_NAMESPACE
diff --git a/Engine/Source/Developer/TargetPlatform/Private/TargetPlatformControlsBase.cpp b/Engine/Source/Developer/TargetPlatform/Private/TargetPlatformControlsBase.cpp
--- a/Engine/Source/Developer/TargetPlatform/Private/TargetPlatformControlsBase.cpp
+++ b/Engine/Source/Developer/TargetPlatform/Private/TargetPlatformControlsBase.cpp
@@ -1,4 +1,5 @@
 // Copyright Epic Games, Inc. All Rights Reserved.
+// Copyright Epic Games, Inc. All Rights Reserved.
 
 #include "Common/TargetPlatformControlsBase.h"
 #include "Interfaces/IProjectBuildMutatorFeature.h"
@@ -294,6 +295,30 @@ void FTargetPlatformControlsBase::GetWaveFormatModuleHints(TArray<FName>& OutMod
 	GetAudioFormatSettings().GetWaveFormatModuleHints(OutModuleNames);
 }
 
+/* static */ void FTargetPlatformControlsBase::GetTextureSizeLimitsDefault(FConfigCacheIni* ConfigSystem,uint64 & OutMaximumSurfaceBytes, uint64 & OutMaximumPackageBytes)
+{
+	OutMaximumSurfaceBytes = 1ULL << 31; // 2 GB
+	//OutMaximumPackageBytes = 1ULL << 32; // 4 GB  seems to work on some platforms
+	OutMaximumPackageBytes = 1ULL << 31; // 2 GB
+
+	#if 0
+	// for stress testing
+	OutMaximumSurfaceBytes = 32 * 1024 * 1024;
+	OutMaximumPackageBytes = OutMaximumSurfaceBytes * 2;
+	#endif
+
+	int64 MaxChunkSize = 0;
+	if ( ConfigSystem->GetInt64(TEXT("/Script/UnrealEd.ProjectPackagingSettings"), TEXT("MaxChunkSize"), MaxChunkSize, GGameIni) &&
+		MaxChunkSize != 0 )
+	{
+		check( MaxChunkSize > 0 );
+
+		OutMaximumPackageBytes = FMath::Min<uint64>(OutMaximumPackageBytes,MaxChunkSize);
+	}
+
+	OutMaximumSurfaceBytes = FMath::Min<uint64>(OutMaximumSurfaceBytes,OutMaximumPackageBytes);
+}
+
 #endif // WITH_ENGINE
 
 #undef LOCTEXT_NAMESPACE
\ No newline at end of file
diff --git a/Engine/Source/Developer/TargetPlatform/Public/Common/TargetPlatformBase.h b/Engine/Source/Developer/TargetPlatform/Public/Common/TargetPlatformBase.h
--- a/Engine/Source/Developer/TargetPlatform/Public/Common/TargetPlatformBase.h
+++ b/Engine/Source/Developer/TargetPlatform/Public/Common/TargetPlatformBase.h
@@ -241,7 +241,8 @@ public:
 	virtual void GetShaderFormatModuleHints(TArray<FName>& OutModuleNames) const override
 	{
 	}
-
+	
+	// beware this is duplicated between TargetPlatformBase and TargetPlatformControlsBase
 	virtual void GetTextureFormatModuleHints(TArray<FName>& OutModuleNames) const override
 	{
 		// these are the default texture format modules, since many platforms 
@@ -257,6 +258,9 @@ public:
 		}
 	}
 	
+	// beware this is duplicated between TargetPlatformBase and TargetPlatformControlsBase
+	TARGETPLATFORM_API virtual void GetTextureSizeLimits(uint64 & OutMaximumSurfaceBytes, uint64 & OutMaximumPackageBytes) const override;
+
 	TARGETPLATFORM_API virtual FName GetWaveFormat(const class USoundWave* Wave) const override;
 	
 	TARGETPLATFORM_API virtual void GetAllWaveFormats(TArray<FName>& OutFormats) const override;
@@ -871,6 +875,11 @@ public:
 	{
 		TargetPlatformControls->GetAllTextureFormats(OutFormats);
 	}
+	
+	virtual void GetTextureSizeLimits(uint64 & OutMaximumSurfaceBytes, uint64 & OutMaximumPackageBytes) const override
+	{
+		TargetPlatformControls->GetTextureSizeLimits(OutMaximumSurfaceBytes,OutMaximumPackageBytes);
+	}
 
 	virtual void GetTextureFormatModuleHints(TArray<FName>& OutModuleNames) const override 
 	{
diff --git a/Engine/Source/Developer/TargetPlatform/Public/Common/TargetPlatformControlsBase.h b/Engine/Source/Developer/TargetPlatform/Public/Common/TargetPlatformControlsBase.h
--- a/Engine/Source/Developer/TargetPlatform/Public/Common/TargetPlatformControlsBase.h
+++ b/Engine/Source/Developer/TargetPlatform/Public/Common/TargetPlatformControlsBase.h
@@ -157,7 +157,8 @@ public:
 		static const FName NAME_MeshBuilder(TEXT("MeshBuilder"));
 		return NAME_MeshBuilder;
 	}
-
+	
+	// beware this is duplicated between TargetPlatformBase and TargetPlatformControlsBase
 	virtual void GetTextureFormatModuleHints(TArray<FName>& OutModuleNames) const override
 	{
 		// these are the default texture format modules, since many platforms 
@@ -178,6 +179,14 @@ public:
 	TARGETPLATFORM_API virtual void GetAllWaveFormats(TArray<FName>& OutFormats) const override;
 
 	TARGETPLATFORM_API virtual void GetWaveFormatModuleHints(TArray<FName>& OutModuleNames) const override;
+	
+	static TARGETPLATFORM_API void GetTextureSizeLimitsDefault(FConfigCacheIni* ConfigSystem,uint64 & OutMaximumSurfaceBytes, uint64 & OutMaximumPackageBytes);
+	
+	// beware this is duplicated between TargetPlatformBase and TargetPlatformControlsBase
+	TARGETPLATFORM_API virtual void GetTextureSizeLimits(uint64 & OutMaximumSurfaceBytes, uint64 & OutMaximumPackageBytes) const override
+	{
+		GetTextureSizeLimitsDefault(TargetPlatformSettings->GetConfigSystem(),OutMaximumSurfaceBytes,OutMaximumPackageBytes);
+	}
 
 #endif
 
diff --git a/Engine/Source/Developer/TargetPlatform/Public/Interfaces/ITargetPlatformControls.h b/Engine/Source/Developer/TargetPlatform/Public/Interfaces/ITargetPlatformControls.h
--- a/Engine/Source/Developer/TargetPlatform/Public/Interfaces/ITargetPlatformControls.h
+++ b/Engine/Source/Developer/TargetPlatform/Public/Interfaces/ITargetPlatformControls.h
@@ -345,6 +345,13 @@ public:
 	 */
 	virtual void GetTextureFormats(const class UTexture* Texture, TArray< TArray<FName> >& OutFormats) const = 0;
 
+	/**
+	 * Get the limitations on texture size (bytes) on this platform
+	 * this limitation is in addition to any dimension limits (specified in device profiles)
+	 * all built textures must be <= this size
+	 */
+	virtual void GetTextureSizeLimits(uint64 & OutMaximumSurfaceBytes, uint64 & OutMaximumPackageBytes) const = 0;
+
 	/**
 	 * Gets the texture formats this platform can use
 	 *
diff --git a/Engine/Source/Runtime/Engine/Classes/Engine/Texture.h b/Engine/Source/Runtime/Engine/Classes/Engine/Texture.h
--- a/Engine/Source/Runtime/Engine/Classes/Engine/Texture.h
+++ b/Engine/Source/Runtime/Engine/Classes/Engine/Texture.h
@@ -1648,8 +1648,19 @@ public:
 	 * Get the dimensions of the largest mip of the texture when built for the target platform
 	 *   accounting for LODBias and other constraints
 	 */
-	ENGINE_API void GetBuiltTextureSize( const ITargetPlatform* TargetPlatform , int32 & OutSizeX, int32 & OutSizeY ) const;
-	ENGINE_API void GetBuiltTextureSize( const class ITargetPlatformSettings* TargetPlatformSettings, const class ITargetPlatformControls* TargetPlatformControls, int32 & OutSizeX, int32 & OutSizeY ) const;
+	ENGINE_API void GetBuiltTextureSize( const ITargetPlatform* TargetPlatform , int32 & OutSizeX, int32 & OutSizeY, int32 & OutSizeZ ) const;
+	ENGINE_API void GetBuiltTextureSize( const class ITargetPlatformSettings* TargetPlatformSettings, const class ITargetPlatformControls* TargetPlatformControls, int32 & OutSizeX, int32 & OutSizeY, int32 & OutSizeZ ) const;
+	
+	ENGINE_API void GetBuiltTextureSize( const ITargetPlatform* TargetPlatform , int32 & OutSizeX, int32 & OutSizeY) const
+	{
+		int32 IgnoredSizeZ;
+		GetBuiltTextureSize(TargetPlatform,OutSizeX,OutSizeY,IgnoredSizeZ);
+	}
+	ENGINE_API void GetBuiltTextureSize( const class ITargetPlatformSettings* TargetPlatformSettings, const class ITargetPlatformControls* TargetPlatformControls, int32 & OutSizeX, int32 & OutSizeY) const
+	{
+		int32 IgnoredSizeZ;
+		GetBuiltTextureSize(TargetPlatformSettings,TargetPlatformControls,OutSizeX,OutSizeY,IgnoredSizeZ);
+	}
 
 	/**
 	 * Serializes cooked platform data.
diff --git a/Engine/Source/Runtime/Engine/Private/Texture.cpp b/Engine/Source/Runtime/Engine/Private/Texture.cpp
--- a/Engine/Source/Runtime/Engine/Private/Texture.cpp
+++ b/Engine/Source/Runtime/Engine/Private/Texture.cpp
@@ -3567,45 +3567,32 @@ static FName ConditionalGetPrefixedFormat(FName TextureFormatName, const ITarget
 	return ConditionalGetPrefixedFormat(TextureFormatName, TargetPlatform->GetTargetPlatformSettings(), bOodleTextureSdkVersionIsNone);
 }
 
-void UTexture::GetBuiltTextureSize(const ITargetPlatformSettings* TargetPlatformSettings, const ITargetPlatformControls* TargetPlatformControls, int32 & OutSizeX, int32 & OutSizeY ) const
+void UTexture::GetBuiltTextureSize(const ITargetPlatformSettings* TargetPlatformSettings, const ITargetPlatformControls* TargetPlatformControls, int32 & OutSizeX, int32 & OutSizeY, int32& OutSizeZ ) const
 {
-	// @todo Oodle : SizeZ
 	// @todo Oodle : verify against TextureCompressorModule
 	// @todo Oodle : with cinematic mips or not? maybe add a bool arg
 	
-	int32 SizeX,SizeY;
+	int32 SizeX,SizeY,SizeZ;
 
 #if WITH_EDITORONLY_DATA
 	FIntPoint SourceSize = Source.GetLogicalSize();
 	SizeX = SourceSize.X;
 	SizeY = SourceSize.Y;
-
-	if (PowerOfTwoMode == ETexturePowerOfTwoSetting::PadToPowerOfTwo || PowerOfTwoMode == ETexturePowerOfTwoSetting::PadToSquarePowerOfTwo ||
-		PowerOfTwoMode == ETexturePowerOfTwoSetting::StretchToPowerOfTwo || PowerOfTwoMode == ETexturePowerOfTwoSetting::StretchToSquarePowerOfTwo)
-	{
-		SizeX = FMath::RoundUpToPowerOfTwo(SizeX);
-		SizeY = FMath::RoundUpToPowerOfTwo(SizeY);
-
-		if (PowerOfTwoMode == ETexturePowerOfTwoSetting::PadToSquarePowerOfTwo || PowerOfTwoMode == ETexturePowerOfTwoSetting::StretchToSquarePowerOfTwo)
-		{
-			SizeX = SizeY = FMath::Max(SizeX, SizeY);
-		}
-	}
-	else if (PowerOfTwoMode == ETexturePowerOfTwoSetting::ResizeToSpecificResolution)
-	{
-		if (ResizeDuringBuildX)
-		{
-			SizeX = ResizeDuringBuildX;
-		}
-		if (ResizeDuringBuildY)
-		{
-			SizeY = ResizeDuringBuildY;
-		}
-	}
-	else
+	
+	SizeZ = Source.GetNumSlices();
+	if ( Source.IsLongLatCubemap() )
 	{
-		checkf(PowerOfTwoMode == ETexturePowerOfTwoSetting::None, TEXT("Unknown entry in ETexturePowerOfTwoSetting::Type"));
+		SizeZ *= 6;
 	}
+	
+	// Volumes mip down Z, other types don't
+	ETextureClass TextureClass = GetTextureClass();
+	bool bIsVolume = ( TextureClass == ETextureClass::Volume );
+	
+	UE::TextureBuildUtilities::GetPowerOfTwoTargetTextureSize(
+		SizeX, SizeY, SizeZ,
+		bIsVolume, PowerOfTwoMode, ResizeDuringBuildX, ResizeDuringBuildY,
+		SizeX, SizeY, SizeZ);
 
 	if (Source.IsLongLatCubemap())
 	{
@@ -3623,6 +3610,10 @@ void UTexture::GetBuiltTextureSize(const ITargetPlatformSettings* TargetPlatform
 		{
 			SizeX = FMath::Max(SizeX>>1,1);
 			SizeY = FMath::Max(SizeY>>1,1);
+			if ( bIsVolume )
+			{
+				SizeZ = FMath::Max(SizeZ>>1,1);
+			}
 		}
 	}
 
@@ -3633,8 +3624,10 @@ void UTexture::GetBuiltTextureSize(const ITargetPlatformSettings* TargetPlatform
  	const uint32 LODBiasNoCinematics = FMath::Max<int32>(LODSettings.CalculateLODBias(SizeX, SizeY, MaxTextureSize, LODGroup, LODBias, 0, MipGenSettings, bVirtualTextureStreaming), 0);
 	SizeX = FMath::Max<int32>(SizeX >> LODBiasNoCinematics, 1);
 	SizeY = FMath::Max<int32>(SizeY >> LODBiasNoCinematics, 1);
-
-	// @todo Oodle : check against GetCookedPlatformData ?
+	if ( bIsVolume )
+	{
+		SizeZ = FMath::Max<int32>(SizeZ >> LODBiasNoCinematics, 1);
+	}
 
 #else // WITH_EDITORONLY_DATA
 
@@ -3650,21 +3643,22 @@ void UTexture::GetBuiltTextureSize(const ITargetPlatformSettings* TargetPlatform
 	{
 		SizeX = (*PPlatformData)->SizeX;
 		SizeY = (*PPlatformData)->SizeY;
+		SizeZ = (*PPlatformData)->GetNumSlices();
 	}
 	else
 	{
-		SizeX = 0;
-		SizeY = 0;
+		SizeX = SizeY = SizeZ = 0;
 	}
 
 #endif // WITH_EDITORONLY_DATA
 
 	OutSizeX = SizeX;
 	OutSizeY = SizeY;
+	OutSizeZ = SizeZ;
 }
-void UTexture::GetBuiltTextureSize(const ITargetPlatform* TargetPlatform, int32& OutSizeX, int32& OutSizeY) const
+void UTexture::GetBuiltTextureSize(const ITargetPlatform* TargetPlatform, int32& OutSizeX, int32& OutSizeY, int32& OutSizeZ) const
 {
-	return GetBuiltTextureSize(TargetPlatform->GetTargetPlatformSettings(), TargetPlatform->GetTargetPlatformControls(), OutSizeX, OutSizeY);
+	return GetBuiltTextureSize(TargetPlatform->GetTargetPlatformSettings(), TargetPlatform->GetTargetPlatformControls(), OutSizeX, OutSizeY, OutSizeZ);
 }
 // this should not be called directly; it is called from TargetPlatform GetTextureFormats
 //	entry point API is GetPlatformTextureFormatNamesWithPrefix
diff --git a/Engine/Source/Runtime/Engine/Private/TextureDerivedData.cpp b/Engine/Source/Runtime/Engine/Private/TextureDerivedData.cpp
--- a/Engine/Source/Runtime/Engine/Private/TextureDerivedData.cpp
+++ b/Engine/Source/Runtime/Engine/Private/TextureDerivedData.cpp
@@ -619,6 +619,177 @@ static void GetEncodeSpeedOptions(ETextureEncodeSpeed InEncodeSpeed, FTextureEnc
 	}
 }
 
+// this should be a strict over-estimate
+// SizeZ is 6 for cubes, can be slices for arrays, etc
+//	or it's volume depth and set IsVolume = true
+static void GetBuiltTextureSizeBytesEstimate(
+	const FTextureBuildSettings& BuildSettings,
+	const ITextureFormat* TextureFormat,
+	int64 TopMipSizeX,int64 TopMipSizeY,int64 TopMipSizeZ,
+	bool bIsVolume,EPixelFormat PixelFormat,
+	// fills :
+	uint64 & OutTopMipSizeBytes, uint64 & OutTotalImageSizeBytes)
+{
+	check( PixelFormat != PF_Unknown );
+	const FPixelFormatInfo & PFI = GPixelFormats[PixelFormat];
+
+	uint64 TopMipSizeBytes = 0;
+	uint64 TotalImageSizeBytes = 0;
+
+	int64 SizeX = TopMipSizeX;
+	int64 SizeY = TopMipSizeY;
+	int64 SizeZ = TopMipSizeZ;
+	
+	int64 NumMips = FImageCoreUtils::GetMipCountFromDimensions(TopMipSizeX,TopMipSizeY,TopMipSizeZ,bIsVolume);
+	check( NumMips > 0 );
+
+	// calculate bytes for linear unpadded/untiled layout :
+
+	for(int64 Mip=0;Mip<NumMips;Mip++)
+	{
+		int64 NumBlocksXY = 
+			FMath::DivideAndRoundUp<int64>( SizeX, PFI.BlockSizeX ) * 
+			FMath::DivideAndRoundUp<int64>( SizeY, PFI.BlockSizeY ); 
+		uint64 SizeBytes = NumBlocksXY * PFI.BlockBytes * SizeZ;
+
+		if ( Mip == 0 )
+		{
+			// save size of first mip
+			TopMipSizeBytes = SizeBytes;
+		}
+
+		TotalImageSizeBytes += SizeBytes;
+
+		SizeX = FMath::Max(1,SizeX>>1);
+		SizeY = FMath::Max(1,SizeY>>1);
+		if ( bIsVolume )
+		{
+			SizeZ = FMath::Max(1,SizeZ>>1);
+		}
+	}
+	
+	check( TotalImageSizeBytes > 0 );
+
+	OutTopMipSizeBytes = TopMipSizeBytes;
+	OutTotalImageSizeBytes = TotalImageSizeBytes;
+	
+	// if alpha is unknown, assume yes to be conservative about pixel size
+	bool bHasAlpha = ( BuildSettings.bKnowAlphaTransparency ) ? BuildSettings.bHasTransparentAlpha : true;
+
+	FEncodedTextureDescription TextureDescription;
+	BuildSettings.GetEncodedTextureDescription(&TextureDescription, TextureFormat, TopMipSizeX, TopMipSizeY, TopMipSizeZ, NumMips, bHasAlpha);
+	check( TextureDescription.PixelFormat == PixelFormat );
+
+	int32 LODBias = 0;
+	FEncodedTextureExtendedData ExtendedData = TextureFormat->GetExtendedDataForTexture(TextureDescription, LODBias);
+	if ( ExtendedData.MipSizesInBytes.Num() > 0 )
+	{
+		// ExtendedData is only valid for platform/tiled images
+
+		TopMipSizeBytes = ExtendedData.MipSizesInBytes[0];
+		
+		// tiled size should be bigger than linear
+		check( TopMipSizeBytes >= OutTopMipSizeBytes );
+
+		TotalImageSizeBytes = 0;
+		for(const uint64 & MipSize : ExtendedData.MipSizesInBytes )
+		{
+			TotalImageSizeBytes += MipSize;
+		}
+		
+		OutTopMipSizeBytes = TopMipSizeBytes;
+		OutTotalImageSizeBytes = TotalImageSizeBytes;
+	}
+
+}
+
+// may reduce OutSettings.MaxTextureResolution
+//	nop if called again
+//	does not change anything else in OutSettings
+//	OutSettings must be otherwise fully set up
+static void ModifyMaxTextureResolutionBuildSettingsForPlatformLimit(
+	const UTexture& Texture, 
+	const ITargetPlatform* TargetPlatform,
+	
+	const ITextureFormat* TextureFormat,
+
+	FTextureBuildSettings& OutSettings)
+{
+	check( ! OutSettings.bVirtualStreamable );
+	check( OutSettings.TextureFormatName != NAME_None );
+			
+	// GetBuiltTextureSize is the size after LODBias
+	int32 BuiltSizeX=0,BuiltSizeY=0,BuiltSizeZ=0;
+	Texture.GetBuiltTextureSize(TargetPlatform,BuiltSizeX,BuiltSizeY,BuiltSizeZ);
+		
+	uint64 MaxSurfaceBytes,MaxPackageBytes;
+	TargetPlatform->GetTextureSizeLimits(MaxSurfaceBytes,MaxPackageBytes);
+	
+	EPixelFormat PixelFormat = GetOutputPixelFormat(OutSettings);
+
+	if ( PixelFormat == PF_Unknown )
+	{
+		UE_LOG(LogTexture, Error, TEXT("Texture %s failed GetOutputPixelFormat (format=%s)"), 
+			*Texture.GetPathName(),
+			*OutSettings.TextureFormatName.ToString());
+			
+		PixelFormat = PF_FloatRGBA;
+	}
+
+	uint64 SurfaceBytes,TotalBytes;
+	GetBuiltTextureSizeBytesEstimate(OutSettings,TextureFormat, BuiltSizeX,BuiltSizeY,BuiltSizeZ,OutSettings.bVolume,PixelFormat,SurfaceBytes,TotalBytes);
+
+	if ( SurfaceBytes > MaxSurfaceBytes || TotalBytes > MaxPackageBytes )
+	{
+		UE_LOG(LogTexture, Warning, TEXT("Texture %s exceeds maximum size of surface or package: %d x %d x %d x %s = {%lld,%lld bytes} exceeds limit {%lld,%lld bytes} shrinking..."), *Texture.GetPathName(),
+			BuiltSizeX, BuiltSizeY, BuiltSizeZ, GetPixelFormatString(PixelFormat),
+			SurfaceBytes,TotalBytes,
+			MaxSurfaceBytes,MaxPackageBytes
+			);
+		
+		// OriginalMaxTextureResolution is uint32_max if Texture did not have a max size set
+		uint32 OriginalMaxTextureResolution = OutSettings.MaxTextureResolution;
+
+		do
+		{
+			// change MaxTextureResolution so that it causes us to do one mip step down
+			//	and adjust BuiltSize accordingly
+			
+			// BuiltSizeZ not affected by MaxTextureResolution
+			OutSettings.MaxTextureResolution = FMath::RoundUpToPowerOfTwo( FMath::Max(BuiltSizeX,BuiltSizeY) )/2;
+			check( (int64)OutSettings.MaxTextureResolution < (int64)BuiltSizeX || (int64)OutSettings.MaxTextureResolution < (int64)BuiltSizeY );
+				
+			BuiltSizeX = FMath::Max(1,BuiltSizeX>>1);
+			BuiltSizeY = FMath::Max(1,BuiltSizeY>>1);
+			if ( OutSettings.bVolume )
+			{
+				BuiltSizeZ = FMath::Max(1,BuiltSizeZ>>1);
+			}
+
+			check( (int64)BuiltSizeX <= (int64)OutSettings.MaxTextureResolution && (int64)BuiltSizeY <= (int64)OutSettings.MaxTextureResolution );
+
+			// recalc size in bytes :
+			GetBuiltTextureSizeBytesEstimate(OutSettings,TextureFormat, BuiltSizeX,BuiltSizeY,BuiltSizeZ,OutSettings.bVolume,PixelFormat,SurfaceBytes,TotalBytes);
+		}
+		while ( SurfaceBytes > MaxSurfaceBytes || TotalBytes > MaxPackageBytes );
+		
+		{
+			// compensate for LODBias that will be applied
+			// after scaling to MaxTextureResolution, LODBiasNoCinematics will be applied
+			
+			const UTextureLODSettings& LODSettings = TargetPlatform->GetTextureLODSettings();
+ 			const uint32 LODBiasNoCinematics = FMath::Max<int32>(LODSettings.CalculateLODBias(BuiltSizeX, BuiltSizeY, Texture.MaxTextureSize, Texture.LODGroup, Texture.LODBias, 0, Texture.MipGenSettings, OutSettings.bVirtualStreamable), 0);
+
+			int64 MaxTextureResolutionUp = ((int64)OutSettings.MaxTextureResolution)<<LODBiasNoCinematics;
+
+			OutSettings.MaxTextureResolution = (uint32) FMath::Min<int64>((int64)OriginalMaxTextureResolution,MaxTextureResolutionUp);
+		}
+
+		// ensure MaxTextureResolution never goes up :
+		OutSettings.MaxTextureResolution = FMath::Min(OriginalMaxTextureResolution,OutSettings.MaxTextureResolution);
+	}
+}
+
 
 // Convert the baseline build settings for all layers to one for the given layer.
 // Note this gets called twice for layer 0, so needs to be idempotent.
@@ -682,57 +853,57 @@ static void FinalizeBuildSettingsForLayer(
 	}
 
 	// Now that we know the texture format, we can make decisions based on it.
+	
+	ITargetPlatformManagerModule* TPM = GetTargetPlatformManager();
+	const ITextureFormat* TextureFormat = nullptr;
+	// this is called once first with NAME_None and then called again after Name is set up
+	if ( ! OutSettings.TextureFormatName.IsNone() )
+	{
+		TextureFormat = TPM->FindTextureFormat(OutSettings.TextureFormatName);
+	}
 
 	bool bSupportsEncodeSpeed = false;
+
+	// Can be null with first finalize (at the end of GetTextureBuildSettings)
+	if (TextureFormat)
 	{
-		ITargetPlatformManagerModule* TPM = GetTargetPlatformManager();
-		if (TPM)
-		{
-			// Can be null with first finalize (at the end of GetTextureBuildSettings)
-			const ITextureFormat* TextureFormat = TPM->FindTextureFormat(OutSettings.TextureFormatName);
-			if (TextureFormat)
-			{
-				bSupportsEncodeSpeed = TextureFormat->SupportsEncodeSpeed(OutSettings.TextureFormatName);
-				const FChildTextureFormat* ChildTextureFormat = TextureFormat->GetChildFormat();
+		bSupportsEncodeSpeed = TextureFormat->SupportsEncodeSpeed(OutSettings.TextureFormatName);
+		const FChildTextureFormat* ChildTextureFormat = TextureFormat->GetChildFormat();
 
-				if (ChildTextureFormat)
-				{
-					OutSettings.BaseTextureFormatName = ChildTextureFormat->GetBaseFormatName(OutSettings.TextureFormatName);
-				}
-				else
-				{
-					OutSettings.BaseTextureFormatName = OutSettings.TextureFormatName;
-				}
+		if (ChildTextureFormat)
+		{
+			OutSettings.BaseTextureFormatName = ChildTextureFormat->GetBaseFormatName(OutSettings.TextureFormatName);
+		}
+		else
+		{
+			OutSettings.BaseTextureFormatName = OutSettings.TextureFormatName;
+		}
 
-				if (OutBuildResultMetadata)
-				{
-					OutBuildResultMetadata->Encoder = TextureFormat->GetEncoderName(OutSettings.TextureFormatName);
-					OutBuildResultMetadata->bIsValid = true;
-					OutBuildResultMetadata->bSupportsEncodeSpeed = bSupportsEncodeSpeed;
-				}
+		if (OutBuildResultMetadata)
+		{
+			OutBuildResultMetadata->Encoder = TextureFormat->GetEncoderName(OutSettings.TextureFormatName);
+			OutBuildResultMetadata->bIsValid = true;
+			OutBuildResultMetadata->bSupportsEncodeSpeed = bSupportsEncodeSpeed;
+		}
 			
-				{
-					if (FResolvedTextureEncodingSettings::Get().Project.bSharedLinearTextureEncoding)
-					{
-						//
-						// We want to separate out textures involved in shared linear encoding in order to facilitate
-						// fixing bugs without invalidating the world (even though we expect the exact same data to
-						// get generated). However, virtual textures never tile, and so are exempt from this separation.
-						//
-						if (OutSettings.bVirtualStreamable == false)
-						{
-							OutSettings.bAffectedBySharedLinearEncoding = true;
-						}
+		if (FResolvedTextureEncodingSettings::Get().Project.bSharedLinearTextureEncoding)
+		{
+			//
+			// We want to separate out textures involved in shared linear encoding in order to facilitate
+			// fixing bugs without invalidating the world (even though we expect the exact same data to
+			// get generated). However, virtual textures never tile, and so are exempt from this separation.
+			//
+			if (OutSettings.bVirtualStreamable == false)
+			{
+				OutSettings.bAffectedBySharedLinearEncoding = true;
+			}
 
-						// Shared linear encoding can only work if the base texture format does not expect to
-						// do the tiling itself (SupportsTiling == false).
-						if (ChildTextureFormat && ChildTextureFormat->GetBaseFormatObject(OutSettings.TextureFormatName)->SupportsTiling() == false)
-						{
-							OutSettings.Tiler = ChildTextureFormat->GetTiler();
-						}
-					} // end if enabled
-				} // end if ddc2
-			} // end if texture format found.
+			// Shared linear encoding can only work if the base texture format does not expect to
+			// do the tiling itself (SupportsTiling == false).
+			if (ChildTextureFormat && ChildTextureFormat->GetBaseFormatObject(OutSettings.TextureFormatName)->SupportsTiling() == false)
+			{
+				OutSettings.Tiler = ChildTextureFormat->GetTiler();
+			}
 		}
 	}
 
@@ -783,6 +954,13 @@ static void FinalizeBuildSettingsForLayer(
 			OutBuildResultMetadata->OodleUniversalTiling = OutSettings.OodleUniversalTiling;
 		}
 	}
+	
+	// this is called once first with NAME_None and then called again after Name is set up
+	if ( ! OutSettings.bVirtualStreamable && ! OutSettings.TextureFormatName.IsNone() )
+	{
+		check( LayerIndex == 0 );
+		ModifyMaxTextureResolutionBuildSettingsForPlatformLimit(Texture,TargetPlatform,TextureFormat,OutSettings);
+	}
 }
 
 ENGINE_API ETextureEncodeSpeed UTexture::GetDesiredEncodeSpeed() const
diff --git a/Engine/Source/Runtime/Engine/Private/TextureDerivedDataTask.cpp b/Engine/Source/Runtime/Engine/Private/TextureDerivedDataTask.cpp
--- a/Engine/Source/Runtime/Engine/Private/TextureDerivedDataTask.cpp
+++ b/Engine/Source/Runtime/Engine/Private/TextureDerivedDataTask.cpp
@@ -888,7 +888,7 @@ static void DDC1_BuildTexture(
 }
 
 
-static EPixelFormat GetOutputPixelFormat(const FTextureBuildSettings & BuildSettings, bool bHasAlpha)
+EPixelFormat GetOutputPixelFormat(const FTextureBuildSettings & BuildSettings)
 {
 	// get the TextureFormat so we can get the output pixel format :
 		
@@ -908,6 +908,9 @@ static EPixelFormat GetOutputPixelFormat(const FTextureBuildSettings & BuildSett
 			
 		return PF_Unknown; /* Unknown */
 	}
+	
+	// if alpha is unknown, assume yes to be conservative about pixel size
+	bool bHasAlpha = ( BuildSettings.bKnowAlphaTransparency ) ? BuildSettings.bHasTransparentAlpha : true;
 
 	EPixelFormat PixelFormat = TextureFormat->GetEncodedPixelFormat(BuildSettings,bHasAlpha);
 	check( PixelFormat != PF_Unknown );
@@ -915,7 +918,7 @@ static EPixelFormat GetOutputPixelFormat(const FTextureBuildSettings & BuildSett
 	return PixelFormat;
 }
 
-static int GetWithinSliceRDOMemoryUsePerPixel(const FName TextureFormatName, bool bHasAlpha)
+static int GetWithinSliceRDOMemoryUsePerPixel(EPixelFormat PixelFormat)
 {
 	// Memory use of RDO data structures, per pixel, within each slice
 	// not counting per-image memory use
@@ -926,39 +929,25 @@ static int GetWithinSliceRDOMemoryUsePerPixel(const FName TextureFormatName, boo
 	const int MemUse_BC7 = 30;
 	const int MemUse_BC3 = MemUse_BC4; // max of BC1,BC4
 	
-	if ( TextureFormatName == "DXT1" ||
-		(TextureFormatName == "AutoDXT" && ! bHasAlpha) )
+	switch(PixelFormat)
 	{
+	case PF_DXT1:
 		return MemUse_BC1;
-	}
-	else if ( TextureFormatName == "DXT3" || 
-		TextureFormatName == "DXT5" ||
-		TextureFormatName == "DXT5n" ||
-		TextureFormatName == "AutoDXT")
-	{
+	case PF_DXT3:
+	case PF_DXT5:
 		return MemUse_BC3;
-	}
-	else if ( TextureFormatName == "BC4" )
-	{
+	case PF_BC4:
 		return MemUse_BC4;
-	}
-	else if ( TextureFormatName == "BC5" )
-	{
+	case PF_BC5:
 		return MemUse_BC5;
-	}
-	else if ( TextureFormatName == "BC6H" )
-	{
+	case PF_BC6H:
 		return MemUse_BC6;
-	}
-	else if ( TextureFormatName == "BC7" )
-	{
+	case PF_BC7:
 		return MemUse_BC7;
-	}
-	else
-	{
+	default:
 		// is this possible?
 		UE_CALL_ONCE( [&](){
-			UE_LOG(LogTexture, Display, TEXT("Unexpected non-BC TextureFormatName: %s."), *TextureFormatName.ToString());
+			UE_LOG(LogTexture, Display, TEXT("Unexpected non-BC PixelFormat: %d."), (int)PixelFormat);
 		} );
 
 		return 100;
@@ -971,12 +960,7 @@ static int64 GetBuildRequiredMemoryEstimate(UTexture* InTexture,
 	const FTextureSource & Source = InTexture->Source;
 
 	const bool bIsVT = InSettingsPerLayerFetchFirst[0].bVirtualStreamable;
-
-	const bool bHasAlpha = true; 
-	// @todo Oodle : need bHasAlpha for AutoDXT ; we currently over-estimate, treat all AutoDXT as BC3
-	// BEWARE : you must use the larger mem use of the two to be conservative
-	// BC1 has twice as many pixels per slice as BC3 so it's not trivially true that the mem use for BC3 is higher
-
+	
 	const bool bRDO = true;
 	// @todo Oodle : be careful about using BuildSettings for this as there are two buildsettingses, just assume its on for now
 	//   <- FIX ME, allow lower mem estimates for non-RDO
@@ -1077,7 +1061,7 @@ static int64 GetBuildRequiredMemoryEstimate(UTexture* InTexture,
 			//  exception would be lots of udim blocks + lots of layers
 			//  because IntermediateFloatColorBytes is per block/layer but output is held for all
 			
-			EPixelFormat PixelFormat = GetOutputPixelFormat(BuildSettings,bHasAlpha);
+			EPixelFormat PixelFormat = GetOutputPixelFormat(BuildSettings);
 
 			if ( PixelFormat == PF_Unknown )
 			{
@@ -1096,9 +1080,7 @@ static int64 GetBuildRequiredMemoryEstimate(UTexture* InTexture,
 
 				if ( bRDO )
 				{
-					const FName TextureFormatName = UE::TextureBuildUtilities::TextureFormatRemovePrefixFromName(BuildSettings.TextureFormatName);
-				
-					int RDOMemUse = GetWithinSliceRDOMemoryUsePerPixel(TextureFormatName,bHasAlpha);
+					int RDOMemUse = GetWithinSliceRDOMemoryUsePerPixel(PixelFormat);
 					CurPerPixelEncoderMemUse += 4; // activity
 					CurPerPixelEncoderMemUse += RDOMemUse;
 					CurPerPixelEncoderMemUse += 1; // output again
@@ -1189,7 +1171,7 @@ static int64 GetBuildRequiredMemoryEstimate(UTexture* InTexture,
 		
 		int64 MemoryEstimate = TotalSourceBytes + IntermediateFloatColorBytes;
 	
-		EPixelFormat PixelFormat = GetOutputPixelFormat(BuildSettings,bHasAlpha);
+		EPixelFormat PixelFormat = GetOutputPixelFormat(BuildSettings);
 
 		if ( PixelFormat == PF_Unknown )
 		{
@@ -1218,18 +1200,16 @@ static int64 GetBuildRequiredMemoryEstimate(UTexture* InTexture,
 			//	this copy is done in TFO
 			//  Oodle then allocs another copy to swizzle into blocks before encoding
 
-			const FName TextureFormatName = UE::TextureBuildUtilities::TextureFormatRemovePrefixFromName(BuildSettings.TextureFormatName);
-				
 			int IntermediateBytesPerPixel;
 			bool bNeedsIntermediateCopy = true;
 
 			// this matches the logic in TextureFormatOodle :
-			if ( TextureFormatName == "BC6H" )
+			if ( PixelFormat == PF_BC6H )
 			{
 				IntermediateBytesPerPixel = 16; //RGBAF32
 				bNeedsIntermediateCopy = false; // no intermediate used in TFO (float source kept), 1 blocksurf
 			}
-			else if ( TextureFormatName == "BC4" || TextureFormatName == "BC5" )
+			else if ( PixelFormat == PF_BC4 || PixelFormat == PF_BC5 )
 			{
 				// changed: TFO uses 2_U16 now (4 byte intermediate)
 				IntermediateBytesPerPixel = 8; // RGBA16
@@ -1252,8 +1232,8 @@ static int64 GetBuildRequiredMemoryEstimate(UTexture* InTexture,
 				// Phase1 = computing activity map
 				int ActivityBytesPerPixel;
 
-				if ( TextureFormatName == "BC4" ) ActivityBytesPerPixel = 12;
-				else if ( TextureFormatName == "BC5" ) ActivityBytesPerPixel = 16;
+				if ( PixelFormat == PF_BC4 ) ActivityBytesPerPixel = 12;
+				else if ( PixelFormat == PF_BC5 ) ActivityBytesPerPixel = 16;
 				else ActivityBytesPerPixel = 24;
 
 				int64 RDOPhase1MemUse = ActivityBytesPerPixel * TotalNumPixels;
@@ -1271,7 +1251,7 @@ static int64 GetBuildRequiredMemoryEstimate(UTexture* InTexture,
 					PixelsPerSlice = TotalNumPixels / NumberofSlices;
 				}
 			
-				int64 MemoryUsePerWorker = PixelsPerSlice * GetWithinSliceRDOMemoryUsePerPixel(TextureFormatName,bHasAlpha);
+				int64 MemoryUsePerWorker = PixelsPerSlice * GetWithinSliceRDOMemoryUsePerPixel(PixelFormat);
 					// MemoryUsePerWorker is around 10 MB
 				int64 NumberOfWorkers = FMath::Min(NumberofSlices,MaxNumberOfWorkers);
 			
diff --git a/Engine/Source/Runtime/Engine/Private/TextureDerivedDataTask.h b/Engine/Source/Runtime/Engine/Private/TextureDerivedDataTask.h
--- a/Engine/Source/Runtime/Engine/Private/TextureDerivedDataTask.h
+++ b/Engine/Source/Runtime/Engine/Private/TextureDerivedDataTask.h
@@ -322,4 +322,6 @@ FTexturePlatformData::FStructuredDerivedDataKey CreateTextureDerivedDataKey(
 	ETextureCacheFlags CacheFlags,
 	const FTextureBuildSettings& Settings);
 
+EPixelFormat GetOutputPixelFormat(const FTextureBuildSettings & BuildSettings);
+
 #endif // WITH_EDITOR
