Remove ECVF_ScalabilityGroup flag from sg.ResolutionQuality cvar as it not actually a scalability group.
This fixes issue where scalability system initializes all ECVF_ScalabilityGroup cvars to a 3 by default, setting screen resolution to a 3%

diff --git a/Engine/Source/Runtime/Engine/Private/DeviceProfiles/DeviceProfileManager.cpp b/Engine/Source/Runtime/Engine/Private/DeviceProfiles/DeviceProfileManager.cpp
--- a/Engine/Source/Runtime/Engine/Private/DeviceProfiles/DeviceProfileManager.cpp
+++ b/Engine/Source/Runtime/Engine/Private/DeviceProfiles/DeviceProfileManager.cpp
@@ -141,7 +141,7 @@ static void GetCVarsFromDPFragmentIncludes(const FString& CurrentSectionName, co
 static void ExpandScalabilityCVar(FConfigCacheIni* ConfigSystem, const FString& CVarKey, const FString CVarValue, TMap<FString, FString>& ExpandedCVars, bool bOverwriteExistingValue)
 {
 	// load scalability settings directly from ini instead of using scalability system, so as not to inadvertantly mess anything up
-	// if the DP had sg.ResolutionQuality=3, we would read [ResolutionQuality@3]
+	// if the DP had sg.ViewDistanceQuality=3, we would read [ViewDistanceQuality@3]
 	FString SectionName = FString::Printf(TEXT("%s@%s"), *CVarKey.Mid(3), *CVarValue);
 	// walk over the scalability section and add them in, unless already done
 	const FConfigSection* ScalabilitySection = ConfigSystem->GetSection(*SectionName, false, GScalabilityIni);
diff --git a/Engine/Source/Runtime/Engine/Private/Scalability.cpp b/Engine/Source/Runtime/Engine/Private/Scalability.cpp
--- a/Engine/Source/Runtime/Engine/Private/Scalability.cpp
+++ b/Engine/Source/Runtime/Engine/Private/Scalability.cpp
@@ -35,7 +35,7 @@ static TAutoConsoleVariable<float> CVarResolutionQuality(
 	0.0f,
 	TEXT("Scalability quality state (internally used by scalability system, ini load/save or using SCALABILITY console command)\n")
 	TEXT(" 10..100, default: 100"),
-	ECVF_ScalabilityGroup | ECVF_Preview);
+	ECVF_Scalability);
 
 static TAutoConsoleVariable<int32> CVarViewDistanceQuality(
 	TEXT("sg.ViewDistanceQuality"),
