SparseVolumeTextures now stream at tile granularity. Requests are still done at mip level granularity, but they are translated into a range of pages, which is then translated into a set of ranges of tiles. Everything after this mip level -> tile range(s) translation operates only on tiles. Resident pages and changes of page residency are detected based on tile residency at the end of the streaming update, so tiles are now a first class citizen. The next step will be to implement bandwidth throttling of IO requests.

Also added a new cvar to control SVT log verbosity.

#rnx

diff --git a/Engine/Source/Runtime/Engine/Classes/SparseVolumeTexture/SparseVolumeTexture.h b/Engine/Source/Runtime/Engine/Classes/SparseVolumeTexture/SparseVolumeTexture.h
--- a/Engine/Source/Runtime/Engine/Classes/SparseVolumeTexture/SparseVolumeTexture.h
+++ b/Engine/Source/Runtime/Engine/Classes/SparseVolumeTexture/SparseVolumeTexture.h
@@ -155,6 +155,11 @@ struct FTileStreamingMetaData
 		return GetNumTiles() - FirstStreamingTileIndex;
 	}
 
+	bool HasRootTile() const
+	{
+		return FirstStreamingTileIndex > 0;
+	}
+
 	uint32 GetRootTileSize() const
 	{
 		return TileDataOffsets[FirstStreamingTileIndex] - TileDataOffsets[0];
@@ -165,9 +170,18 @@ struct FTileStreamingMetaData
 		return TileDataOffsets[TileIndex + 1] - TileDataOffsets[TileIndex];
 	}
 
+	void GetTileRangeMemoryOffsetSize(uint32 TileOffset, uint32 TileCount, uint32& OutMemoryOffset, uint32& OutMemorySize) const
+	{
+		checkf(!(TileOffset < FirstStreamingTileIndex && (TileOffset + TileCount) > FirstStreamingTileIndex), TEXT("Tile range must not straddle the root tile and streaming tiles!"));
+		const uint32 RootTileSize = TileOffset < FirstStreamingTileIndex ? 0 : GetRootTileSize();
+		OutMemoryOffset = TileDataOffsets[TileOffset] - RootTileSize;
+		const uint32 ReadEnd = TileDataOffsets[TileOffset + TileCount] - RootTileSize;
+		OutMemorySize = ReadEnd - OutMemoryOffset;
+	}
+
 	FTileInfo GetTileInfo(uint32 TileIndex, uint32 FormatSizeA, uint32 FormatSizeB) const;
 
-	void GetNumVoxelsInTileRange(uint32 TileOffset, uint32 TileCount, uint32 FormatSizeA, uint32 FormatSizeB, uint32& OutNumVoxelsA, uint32& OutNumVoxelsB) const;
+	void GetNumVoxelsInTileRange(uint32 TileOffset, uint32 TileCount, uint32 FormatSizeA, uint32 FormatSizeB, const TBitArray<>* OptionalValidTiles, uint32& OutNumVoxelsA, uint32& OutNumVoxelsB) const;
 };
 
 enum EResourceFlag : uint32
@@ -230,7 +244,7 @@ private:
 	void EndRebuildBulkDataFromCache();
 #endif
 
-	static FTileStreamingMetaData CompressTiles(const FPageTopology& Topology, const struct FDerivedTextureData& DerivedTextureData, TArray<uint8>& OutRootBulkData, TArray<uint8>& OutStreamingBulkData);
+	static FTileStreamingMetaData CompressTiles(const FPageTopology& Topology, const struct FDerivedTextureData& DerivedTextureData, TArray<uint8>& OutRootBulkData, TArray64<uint8>& OutStreamingBulkData);
 };
 
 // Encapsulates RHI resources needed to render a SparseVolumeTexture.
diff --git a/Engine/Source/Runtime/Engine/Private/SparseVolumeTexture/SparseVolumeTexture.cpp b/Engine/Source/Runtime/Engine/Private/SparseVolumeTexture/SparseVolumeTexture.cpp
--- a/Engine/Source/Runtime/Engine/Private/SparseVolumeTexture/SparseVolumeTexture.cpp
+++ b/Engine/Source/Runtime/Engine/Private/SparseVolumeTexture/SparseVolumeTexture.cpp
@@ -112,7 +112,7 @@ static int32 ComputeNumMipLevels(const FIntVector3& InVirtualVolumeMin, const FI
 
 static const FString& GetDerivedDataVersion()
 {
-	static FString CachedVersionString = TEXT("6C460857-3C9C-4580-98F5-63811E5EA4D5");	// Bump this if you want to ignore all cached data so far.
+	static FString CachedVersionString = TEXT("2113A905-4C31-4EFD-B55C-82375C6A7B56");	// Bump this if you want to ignore all cached data so far.
 	return CachedVersionString;
 }
 
@@ -385,7 +385,7 @@ bool FResources::Build(USparseVolumeTextureFrame* Owner, UE::Serialization::FEdi
 		Topology.ParentIndices = MoveTemp(DerivedTextureData.PageTableParentIndices);
 
 		// Compress tile data
-		TArray<uint8> StreamableBulkData;
+		TArray64<uint8> StreamableBulkData;
 		StreamingMetaData = CompressTiles(Topology, DerivedTextureData, RootData, StreamableBulkData);
 
 		// Store StreamableMipLevels
@@ -608,7 +608,7 @@ void FResources::EndRebuildBulkDataFromCache()
 
 #endif // WITH_EDITORONLY_DATA
 
-FTileStreamingMetaData FResources::CompressTiles(const FPageTopology& Topology, const FDerivedTextureData& DerivedTextureData, TArray<uint8>& OutRootBulkData, TArray<uint8>& OutStreamingBulkData)
+FTileStreamingMetaData FResources::CompressTiles(const FPageTopology& Topology, const FDerivedTextureData& DerivedTextureData, TArray<uint8>& OutRootBulkData, TArray64<uint8>& OutStreamingBulkData)
 {
 	TRACE_CPUPROFILER_EVENT_SCOPE(SVT::FResources::Build::CompressTiles);
 
@@ -670,6 +670,7 @@ FTileStreamingMetaData FResources::CompressTiles(const FPageTopology& Topology,
 	uint32 CurrentOffset = 0;
 	for (uint32 TileIndex = 0; TileIndex < NumTiles; ++TileIndex)
 	{
+		uint32 PrevOffset = CurrentOffset;
 		TileDataOffsets[TileIndex] = CurrentOffset;
 		CurrentOffset += SVT::NumOccupancyWordsPerPaddedTile * sizeof(uint32) * NumTextures;
 		for (int32 AttributesIdx = 0; AttributesIdx < 2; ++AttributesIdx)
@@ -679,6 +680,7 @@ FTileStreamingMetaData FResources::CompressTiles(const FPageTopology& Topology,
 				CurrentOffset += NumVoxels[AttributesIdx][TileIndex] * FormatSize[AttributesIdx];
 			}
 		}
+		checkf(PrevOffset < CurrentOffset, TEXT("SVT streaming data overflowed the uint32 range!"));
 	}
 
 	// Write final size at the end of the array so we can compute individual tile sizes as (Offsets[N+1] - Offsets[N])
@@ -831,7 +833,7 @@ FTileInfo FTileStreamingMetaData::GetTileInfo(uint32 TileIndex, uint32 FormatSiz
 	return Result;
 }
 
-void FTileStreamingMetaData::GetNumVoxelsInTileRange(uint32 TileOffset, uint32 TileCount, uint32 FormatSizeA, uint32 FormatSizeB, uint32& OutNumVoxelsA, uint32& OutNumVoxelsB) const
+void FTileStreamingMetaData::GetNumVoxelsInTileRange(uint32 TileOffset, uint32 TileCount, uint32 FormatSizeA, uint32 FormatSizeB, const TBitArray<>* OptionalValidTiles, uint32& OutNumVoxelsA, uint32& OutNumVoxelsB) const
 {
 	check(FormatSizeA > 0 || FormatSizeB > 0);
 	const uint32 NumTextures = (FormatSizeA > 0 && FormatSizeB > 0) ? 2 : 1;
@@ -840,6 +842,10 @@ void FTileStreamingMetaData::GetNumVoxelsInTileRange(uint32 TileOffset, uint32 T
 	OutNumVoxelsB = 0;
 	for (uint32 TileIndex = TileOffset; TileIndex < (TileOffset + TileCount); ++TileIndex)
 	{
+		if (OptionalValidTiles && !(*OptionalValidTiles)[TileIndex])
+		{
+			continue;
+		}
 		const uint32 NumVoxelsATmp = NumVoxelsA[TileIndex];
 		OutNumVoxelsA += NumVoxelsATmp;
 		// For NumVoxelsB, we need to reconstruct that value based on the total tile size and the sizes of the other memory sections in the tile
diff --git a/Engine/Source/Runtime/Engine/Private/SparseVolumeTexture/SparseVolumeTextureData.cpp b/Engine/Source/Runtime/Engine/Private/SparseVolumeTexture/SparseVolumeTextureData.cpp
--- a/Engine/Source/Runtime/Engine/Private/SparseVolumeTexture/SparseVolumeTextureData.cpp
+++ b/Engine/Source/Runtime/Engine/Private/SparseVolumeTexture/SparseVolumeTextureData.cpp
@@ -758,9 +758,6 @@ void FDerivedTextureData::Build(const FTextureData& MippedTextureData)
 
 	for (int32 MipLevel = NumMipLevels - 1; MipLevel >= 0; --MipLevel)
 	{
-		// For now we don't deduplicate across mip levels. This simplifies logic in the current mip level based streaming manager.
-		PhysicalTileDataRemap.Reset();
-
 		const uint8* SrcPhysicalTileDataA = MippedTextureData.MipMaps[MipLevel].PhysicalTileDataA.GetData();
 		const uint8* SrcPhysicalTileDataB = MippedTextureData.MipMaps[MipLevel].PhysicalTileDataB.GetData();
 
diff --git a/Engine/Source/Runtime/Engine/Private/SparseVolumeTexture/SparseVolumeTextureStreamingManager.cpp b/Engine/Source/Runtime/Engine/Private/SparseVolumeTexture/SparseVolumeTextureStreamingManager.cpp
--- a/Engine/Source/Runtime/Engine/Private/SparseVolumeTexture/SparseVolumeTextureStreamingManager.cpp
+++ b/Engine/Source/Runtime/Engine/Private/SparseVolumeTexture/SparseVolumeTextureStreamingManager.cpp
@@ -22,10 +22,6 @@ DEFINE_LOG_CATEGORY(LogSparseVolumeTextureStreamingManager);
 
 DECLARE_GPU_STAT(SVTStreaming);
 
-#ifndef SVT_STREAMING_LOG_VERBOSE
-#define SVT_STREAMING_LOG_VERBOSE 0
-#endif
-
 static int32 GSVTStreamingNumPrefetchFrames = 3;
 static FAutoConsoleVariableRef CVarSVTStreamingNumPrefetchFrames(
 	TEXT("r.SparseVolumeTexture.Streaming.NumPrefetchFrames"),
@@ -74,7 +70,15 @@ static FAutoConsoleVariableRef CVarSVTStreamingPrintMemoryStats(
 	ECVF_RenderThreadSafe
 );
 
-static int32 GSVTStreamingMaxPendingRequests = 128;
+static int32 GSVTStreamingLogVerbosity = 1;
+static FAutoConsoleVariableRef CVarSVTStreamingLogVerbosity(
+	TEXT("r.SparseVolumeTexture.Streaming.LogVerbosity"),
+	GSVTStreamingLogVerbosity,
+	TEXT("0: no logging, 1: basic logging, 2: additional logging (might spam the log) 3: log everything (will spam the log)"),
+	ECVF_RenderThreadSafe
+);
+
+static int32 GSVTStreamingMaxPendingRequests = 1024;
 static FAutoConsoleVariableRef CVarSVTStreamingMaxPendingRequests(
 	TEXT("r.SparseVolumeTexture.Streaming.MaxPendingRequests"),
 	GSVTStreamingMaxPendingRequests,
@@ -165,11 +169,11 @@ void FStreamingManager::Add_GameThread(UStreamableSparseVolumeTexture* SparseVol
 	FNewSparseVolumeTextureInfo NewSVTInfo{};
 	const int32 NumFrames = SparseVolumeTexture->GetNumFrames();
 	NewSVTInfo.SVT = SparseVolumeTexture;
+	NewSVTInfo.SVTName = SparseVolumeTexture->GetFName();
 	NewSVTInfo.FormatA = SparseVolumeTexture->GetFormat(0);
 	NewSVTInfo.FormatB = SparseVolumeTexture->GetFormat(1);
 	NewSVTInfo.FallbackValueA = SparseVolumeTexture->GetFallbackValue(0);
 	NewSVTInfo.FallbackValueB = SparseVolumeTexture->GetFallbackValue(1);
-	NewSVTInfo.NumMipLevelsGlobal = SparseVolumeTexture->GetNumMipLevels();
 	NewSVTInfo.FrameInfo.SetNum(NumFrames);
 
 	for (int32 FrameIdx = 0; FrameIdx < NumFrames; ++FrameIdx)
@@ -318,14 +322,16 @@ void FStreamingManager::Request(UStreamableSparseVolumeTexture* SparseVolumeText
 			// Wrap around on both positive and negative numbers, assuming (i + NumFrames) >= 0. See the comment on NumPrefetchFrames.
 			const int32 RequestFrameIndex = (static_cast<int32>(FrameIndex) + i + NumFrames) % NumFrames;
 			const int32 RequestMipLevelOffset = FMath::Abs(i) + GSVTStreamingPrefetchMipLevelBias;
+			const uint32 MipLevelIndex = FMath::Clamp(MipLevel + RequestMipLevelOffset, 0, SVTInfo->PerFrameInfo[RequestFrameIndex].NumMipLevels);
+			const uint8 Priority = bBlocking ? UINT8_MAX : FMath::Max(0, OffsetMagnitude - FMath::Abs(i));
 			FStreamingRequest Request;
-			Request.Key.SVT = SparseVolumeTexture;
+			Request.Key.SVTHandle = SVTInfo->SVTHandle;
 			Request.Key.FrameIndex = RequestFrameIndex;
-			Request.Key.MipLevelIndex = FMath::Clamp(MipLevel + RequestMipLevelOffset, 0, SVTInfo->PerFrameInfo[RequestFrameIndex].NumMipLevels);
-			Request.Priority = FMath::Max(0, OffsetMagnitude - FMath::Abs(i));
-			if (bBlocking)
+			Request.Payload.MipLevelMask = ~0u << MipLevelIndex;
+			Request.Payload.LowestMipFraction = 1.0f;
+			for (int32 j = MipLevelIndex; j < Request.Payload.Priorities.Num(); ++j)
 			{
-				Request.Priority = FStreamingRequest::BlockingPriority;
+				Request.Payload.Priorities[j] = Priority;
 			}
 			AddRequest(Request);
 		}
@@ -349,9 +355,10 @@ void FStreamingManager::BeginAsyncUpdate(FRDGBuilder& GraphBuilder, bool bBlocki
 	RDG_CSV_STAT_EXCLUSIVE_SCOPE(GraphBuilder, SVTStreaming);
 	SCOPED_NAMED_EVENT_TEXT("SVT::StreamingBeginAsyncUpdate", FColor::Green);
 
-#if SVT_STREAMING_LOG_VERBOSE
-	UE_LOG(LogSparseVolumeTextureStreamingManager, Display, TEXT("SVT Streaming Update %i"), NextUpdateIndex);
-#endif
+	if (GSVTStreamingLogVerbosity > 2)
+	{
+		UE_LOG(LogSparseVolumeTextureStreamingManager, Display, TEXT("SVT Streaming Update %i"), NextUpdateIndex);
+	}
 
 	AsyncState = {};
 	AsyncState.bUpdateActive = true;
@@ -359,58 +366,67 @@ void FStreamingManager::BeginAsyncUpdate(FRDGBuilder& GraphBuilder, bool bBlocki
 	// For debugging, we can stream out ALL tiles
 	if (GSVTStreamingEmptyPhysicalTileTextures != 0)
 	{
-		TArray<FLRUNode*> MipLevelsToFree;
-		for (auto& Pair : StreamingInfo)
+		for (TUniquePtr<FStreamingInfo>& SVTInfo : StreamingInfo)
 		{
-			MipLevelsToFree.Reset();
-			FStreamingInfo* SVTInfo = Pair.Value.Get();
+			if (!SVTInfo)
+			{
+				continue;
+			}
 			const int32 NumFrames = SVTInfo->PerFrameInfo.Num();
-			const int32 NumMipLevelsGlobal = SVTInfo->NumMipLevelsGlobal;
-			
-			for (int32 MipLevel = 0; MipLevel < NumMipLevelsGlobal; ++MipLevel)
+			for (int32 FrameIndex = 0; FrameIndex < NumFrames; ++FrameIndex)
 			{
-				for (auto& Node : SVTInfo->PerMipLRULists[MipLevel])
+				FFrameInfo& FrameInfo = SVTInfo->PerFrameInfo[FrameIndex];
+				const bool bHasRootTile = FrameInfo.Resources->StreamingMetaData.HasRootTile();
+				for (TConstSetBitIterator It(FrameInfo.StreamingTiles, bHasRootTile ? 1 : 0); It; ++It)
 				{
-					MipLevelsToFree.Add(&Node);
+					const uint32 TileIndex = It.GetIndex();
+					FrameInfo.StreamingTiles[TileIndex] = false;
+					FrameInfo.ResidentTiles[TileIndex] = false;
+					check(FrameInfo.TileAllocations[TileIndex] != 0);
+					SVTInfo->TileAllocator.Free(FrameInfo.TileAllocations[TileIndex]);
+					FrameInfo.TileAllocations[TileIndex] = 0;
 				}
-			}
-			for (FLRUNode* Node : MipLevelsToFree)
-			{
-				StreamOutMipLevel(SVTInfo, Node);
+				FrameInfo.TileIndexToPendingRequestIndex.Reset();
+				InvalidatedSVTFrames.Add(&FrameInfo);
 			}
 		}
 
+		// Cancel all IO requests
+		for (FPendingRequest& PendingRequest : PendingRequests)
+		{
+			PendingRequest.Reset();
+		}
+
 		GSVTStreamingEmptyPhysicalTileTextures = 0;
 	}
 
 	if (GSVTStreamingPrintMemoryStats != 0)
 	{
-		for (auto& Pair : StreamingInfo)
+		for (TUniquePtr<FStreamingInfo>& SVTInfo : StreamingInfo)
 		{
+			if (!SVTInfo)
+			{
+				continue;
+			}
+
 			double MinFrameMiB = FLT_MAX;
 			double MaxFrameMiB = -FLT_MAX;
 			double SumFrameMiB = 0.0;
 
-			const int32 NumFrames = Pair.Value->PerFrameInfo.Num();
-			UE_LOG(LogSparseVolumeTextureStreamingManager, Display, TEXT("Memory stats for SVT '%p': Each frame is displayed as a list of mip levels like this: [Mip0] [Mip1] [MipN]"), Pair.Key);
-
+			const int32 NumFrames = SVTInfo->PerFrameInfo.Num();
+			UE_LOG(LogSparseVolumeTextureStreamingManager, Display, TEXT("Memory stats for SVT '%s':"), *SVTInfo->SVTName.ToString());
+		
 			for (int32 FrameIdx = 0; FrameIdx < NumFrames; ++FrameIdx)
 			{
-				const FFrameInfo& FrameInfo = Pair.Value->PerFrameInfo[FrameIdx];
-				FString Str;
-				int32 TotalSize = 0;
-				for (int32 MipLevelIdx = 0; MipLevelIdx < FrameInfo.NumMipLevels; ++MipLevelIdx)
-				{
-					const FMipTileReadInfo MipReadInfo = GetMipTileReadInfo(FrameInfo, MipLevelIdx);
-					Str += FString::Printf(TEXT("[%u Tiles, %5.2f KiB] "), MipReadInfo.TileCount, MipReadInfo.ReadSize / 1024.0f);
-					TotalSize += MipReadInfo.ReadSize;
-				}
+				const FFrameInfo& FrameInfo = SVTInfo->PerFrameInfo[FrameIdx];
+				const uint32 Size = FrameInfo.Resources->StreamingMetaData.TileDataOffsets.Last() - FrameInfo.Resources->StreamingMetaData.GetRootTileSize();
+				const uint32 NumTiles = FrameInfo.Resources->StreamingMetaData.GetNumStreamingTiles();
 
-				MinFrameMiB = FMath::Min(MinFrameMiB, TotalSize / 1024.0 / 1024.0);
-				MaxFrameMiB = FMath::Max(MaxFrameMiB, TotalSize / 1024.0 / 1024.0);
-				SumFrameMiB += TotalSize / 1024.0 / 1024.0;
+				MinFrameMiB = FMath::Min(MinFrameMiB, Size / 1024.0 / 1024.0);
+				MaxFrameMiB = FMath::Max(MaxFrameMiB, Size / 1024.0 / 1024.0);
+				SumFrameMiB += Size / 1024.0 / 1024.0;
 
-				UE_LOG(LogSparseVolumeTextureStreamingManager, Display, TEXT("SVT Frame %3i: TotalSize: %3.2f MiB %s"), FrameIdx, TotalSize / 1024.0f / 1024.0f, *Str);
+				UE_LOG(LogSparseVolumeTextureStreamingManager, Display, TEXT("SVT Frame %3u: Size: %5.4f MiB Num Tiles: %u"), FrameIdx, Size / 1024.0f / 1024.0f, NumTiles);
 			}
 
 			UE_LOG(LogSparseVolumeTextureStreamingManager, Display, TEXT("SVT Frame Stats: Min: %3.2f MiB, Max: %3.2f MiB, Avg: %3.2f, Total All: %3.2f MiB"), (float)MinFrameMiB, (float)MaxFrameMiB, (float)(SumFrameMiB / NumFrames), SumFrameMiB);
@@ -419,13 +435,11 @@ void FStreamingManager::BeginAsyncUpdate(FRDGBuilder& GraphBuilder, bool bBlocki
 		GSVTStreamingPrintMemoryStats = 0;
 	}
 
-	AddParentRequests();
-	const int32 MaxSelectedRequests = MaxPendingRequests - NumPendingRequests;
-	SelectHighestPriorityRequestsAndUpdateLRU(MaxSelectedRequests);
-	IssueRequests(MaxSelectedRequests);
+	FilterRequests();
+	IssueRequests();
 	AsyncState.NumReadyRequests = DetermineReadyRequests();
 
-	// Do a first pass over all the mips to be uploaded to compute the upload buffer size requirements.
+	// Do a first pass over all the tiles to be uploaded to compute the upload buffer size requirements.
 	TileDataTexturesToUpdate.Reset();
 	{
 		const int32 StartPendingRequestIndex = (NextPendingRequestIndex + MaxPendingRequests - NumPendingRequests) % MaxPendingRequests;
@@ -434,12 +448,15 @@ void FStreamingManager::BeginAsyncUpdate(FRDGBuilder& GraphBuilder, bool bBlocki
 			const int32 PendingRequestIndex = (StartPendingRequestIndex + i) % MaxPendingRequests;
 			FPendingRequest& PendingRequest = PendingRequests[PendingRequestIndex];
 
-			FStreamingInfo* SVTInfo = FindStreamingInfo(PendingRequest.SparseVolumeTexture);
-			if (!SVTInfo || (SVTInfo->PerFrameInfo[PendingRequest.FrameIndex].LowestRequestedMipLevel > PendingRequest.MipLevelIndex))
+			// Skip if request was cancelled
+			if (!PendingRequest.IsValid())
 			{
-				continue; // Skip mip level install. SVT no longer exists or mip level was "streamed out" before it was even installed in the first place.
+				continue;
 			}
 
+			FStreamingInfo* SVTInfo = FindStreamingInfo(PendingRequest.SVTHandle);
+			check(SVTInfo);
+
 			// Prepare tile data texture for upload
 			const FTileDataTexture::EUploaderState UploaderState = SVTInfo->TileDataTexture->GetUploaderState();
 			check(UploaderState == FTileDataTexture::EUploaderState::Ready || UploaderState == FTileDataTexture::EUploaderState::Reserving);
@@ -450,17 +467,13 @@ void FStreamingManager::BeginAsyncUpdate(FRDGBuilder& GraphBuilder, bool bBlocki
 
 			const int32 FormatSizeA = GPixelFormats[SVTInfo->FormatA].BlockBytes;
 			const int32 FormatSizeB = GPixelFormats[SVTInfo->FormatB].BlockBytes;
-			const FResources* Resources = SVTInfo->PerFrameInfo[PendingRequest.FrameIndex].Resources;
-			const FPageTopology::FMip& MipInfo = Resources->Topology.MipInfo[PendingRequest.MipLevelIndex];
-			
-			uint32 TileRangeOffset = 0;
-			uint32 TileRangeCount = 0;
-			Resources->Topology.GetTileRange(MipInfo.PageOffset, MipInfo.PageCount, TileRangeOffset, TileRangeCount);
+			FFrameInfo& FrameInfo = SVTInfo->PerFrameInfo[PendingRequest.FrameIndex];
+
 			uint32 NumVoxelsA = 0;
 			uint32 NumVoxelsB = 0;
-			Resources->StreamingMetaData.GetNumVoxelsInTileRange(TileRangeOffset, TileRangeCount, FormatSizeA, FormatSizeB, NumVoxelsA, NumVoxelsB);
+			FrameInfo.Resources->StreamingMetaData.GetNumVoxelsInTileRange(PendingRequest.TileOffset, PendingRequest.TileCount, FormatSizeA, FormatSizeB, &FrameInfo.StreamingTiles, NumVoxelsA, NumVoxelsB);
 
-			SVTInfo->TileDataTexture->ReserveUpload(TileRangeCount, NumVoxelsA, NumVoxelsB);
+			SVTInfo->TileDataTexture->ReserveUpload(PendingRequest.TileCount, NumVoxelsA, NumVoxelsB);
 
 			TileDataTexturesToUpdate.Add(SVTInfo->TileDataTexture.Get());
 		}
@@ -524,37 +537,13 @@ void FStreamingManager::EndAsyncUpdate(FRDGBuilder& GraphBuilder)
 	++NextUpdateIndex;
 	AsyncState.bUpdateActive = false;
 	AsyncState.bUpdateIsAsync = false;
-
-#if DO_CHECK
-	for (const auto& Pair : StreamingInfo)
-	{
-#if SVT_STREAMING_LOG_VERBOSE
-		FString ResidentMipLevelsStr = TEXT("");
-#endif
-		const int32 NumFrames = Pair.Value->PerFrameInfo.Num();
-		for (int32 FrameIndex = 0; FrameIndex < NumFrames; ++FrameIndex)
-		{
-			const auto& FrameInfo = Pair.Value->PerFrameInfo[FrameIndex];
-			check(FrameInfo.LowestResidentMipLevel <= (FrameInfo.NumMipLevels - 1));
-			check(FrameInfo.LowestRequestedMipLevel <= FrameInfo.LowestResidentMipLevel);
-			check(FrameInfo.TextureRenderResources->GetNumLogicalMipLevels() == FrameInfo.NumMipLevels);
-
-#if SVT_STREAMING_LOG_VERBOSE
-			ResidentMipLevelsStr += FString::Printf(TEXT("%i"), FrameInfo.LowestResidentMipLevel);
-#endif
-		}
-#if SVT_STREAMING_LOG_VERBOSE
-		UE_LOG(LogSparseVolumeTextureStreamingManager, Display, TEXT("%s"), *ResidentMipLevelsStr);
-#endif
-	}
-#endif // DO_CHECK
 }
 
 void FStreamingManager::AddInternal(FRDGBuilder& GraphBuilder, FNewSparseVolumeTextureInfo&& NewSVTInfo)
 {
 	check(IsInRenderingThread());
 	check(!AsyncState.bUpdateActive);
-	if (!ensure(!StreamingInfo.Contains(NewSVTInfo.SVT)))
+	if (!ensure(!SparseVolumeTextureToHandle.Contains(NewSVTInfo.SVT)))
 	{
 		return;
 	}
@@ -564,18 +553,21 @@ void FStreamingManager::AddInternal(FRDGBuilder& GraphBuilder, FNewSparseVolumeT
 	RDG_CSV_STAT_EXCLUSIVE_SCOPE(GraphBuilder, SVTStreaming);
 	SCOPED_NAMED_EVENT_TEXT("SVT::StreamingAddInternal", FColor::Green);
 
+	// Allocate SVTHandle, which is actually simply an index into the sparse StreamingInfo array.
+	int32 LowestFreeIndexSearchStart = 0;
+	const uint16 SVTHandle = StreamingInfo.EmplaceAtLowestFreeIndex(LowestFreeIndexSearchStart, MakeUnique<FStreamingInfo>());
+	SparseVolumeTextureToHandle.Add(NewSVTInfo.SVT, SVTHandle);
+
 	const int32 NumFrames = NewSVTInfo.FrameInfo.Num();
 
-	FStreamingInfo& SVTInfo = *StreamingInfo.Emplace(NewSVTInfo.SVT, MakeUnique<FStreamingInfo>());
+	FStreamingInfo& SVTInfo = *StreamingInfo[SVTHandle];
+	SVTInfo.SVTHandle = SVTHandle;
+	SVTInfo.SVTName = NewSVTInfo.SVTName;
 	SVTInfo.FormatA = NewSVTInfo.FormatA;
 	SVTInfo.FormatB = NewSVTInfo.FormatB;
 	SVTInfo.FallbackValueA = NewSVTInfo.FallbackValueA;
 	SVTInfo.FallbackValueB = NewSVTInfo.FallbackValueB;
-	SVTInfo.NumMipLevelsGlobal = NewSVTInfo.NumMipLevelsGlobal;
-	SVTInfo.LastRequested = 0;
 	SVTInfo.PerFrameInfo = MoveTemp(NewSVTInfo.FrameInfo);
-	SVTInfo.LRUNodes.SetNum(NumFrames * SVTInfo.NumMipLevelsGlobal);
-	SVTInfo.PerMipLRULists.SetNum(SVTInfo.NumMipLevelsGlobal);
 
 	const int32 FormatSizes[] = { GPixelFormats[SVTInfo.FormatA].BlockBytes, GPixelFormats[SVTInfo.FormatB].BlockBytes };
 
@@ -594,13 +586,13 @@ void FStreamingManager::AddInternal(FRDGBuilder& GraphBuilder, FNewSparseVolumeT
 		MaxNumPhysicalTiles = FMath::Max(NumPhysicalTiles, MaxNumPhysicalTiles);
 
 		FrameInfo.NumMipLevels = Topology.MipInfo.Num();
-		FrameInfo.LowestRequestedMipLevel = FrameInfo.NumMipLevels - 1;
-		FrameInfo.LowestResidentMipLevel = FrameInfo.NumMipLevels - 1;
-		FrameInfo.TileAllocations.SetNum(NumPhysicalTiles);
+		FrameInfo.TileAllocations.SetNumZeroed(NumPhysicalTiles);
 
 		const int32 NumPagesTotal = Topology.NumPages();
 		FrameInfo.ResidentPages.SetNum(NumPagesTotal, false);
-		FrameInfo.ResidentPagesNew.SetNum(NumPagesTotal, false);
+		FrameInfo.InvalidatedPages.SetNum(NumPagesTotal, false);
+		FrameInfo.ResidentTiles.SetNum(NumPhysicalTiles, false);
+		FrameInfo.StreamingTiles.SetNum(NumPhysicalTiles, false);
 		
 		if (!Topology.MipInfo.IsEmpty() && Topology.MipInfo.Last().PageCount > 0)
 		{
@@ -611,19 +603,6 @@ void FStreamingManager::AddInternal(FRDGBuilder& GraphBuilder, FNewSparseVolumeT
 			NumRootVoxelsA += RootTileInfo.NumVoxels[0];
 			NumRootVoxelsB += RootTileInfo.NumVoxels[1];
 		}
-
-		for (int32 MipIdx = 0; MipIdx < SVTInfo.NumMipLevelsGlobal; ++MipIdx)
-		{
-			FLRUNode& LRUNode = SVTInfo.LRUNodes[FrameIdx * SVTInfo.NumMipLevelsGlobal + MipIdx];
-			LRUNode.Reset();
-			LRUNode.FrameIndex = FrameIdx;
-			LRUNode.MipLevelIndex = MipIdx < FrameInfo.NumMipLevels ? MipIdx : INDEX_NONE;
-
-			if ((MipIdx + 1) < FrameInfo.NumMipLevels)
-			{
-				LRUNode.NextHigherMipLevel = &SVTInfo.LRUNodes[FrameIdx * SVTInfo.NumMipLevelsGlobal + (MipIdx + 1)];
-			}
-		}
 	}
 
 	// Create RHI resources and upload root tile data
@@ -634,14 +613,14 @@ void FStreamingManager::AddInternal(FRDGBuilder& GraphBuilder, FNewSparseVolumeT
 
 		SVTInfo.TileDataTexture = MakeUnique<FTileDataTexture>(TileDataVolumeResolutionInTiles, SVTInfo.FormatA, SVTInfo.FormatB, SVTInfo.FallbackValueA, SVTInfo.FallbackValueB);
 		SVTInfo.TileDataTexture->InitResource(GraphBuilder.RHICmdList);
+		SVTInfo.TileAllocator.Init(SVTInfo.TileDataTexture->GetResolutionInTiles());
 
 		FTileUploader RootTileUploader;
 		RootTileUploader.Init(GraphBuilder, NumRootPhysicalTiles + 1 /*null tile*/, NumRootVoxelsA, NumRootVoxelsB, SVTInfo.FormatA, SVTInfo.FormatB);
 
-		// Allocate null tile
+		// Create null tile
 		{
-			const uint32 NullTileCoord = SVTInfo.TileDataTexture->Allocate();
-			check(NullTileCoord == 0);
+			const uint32 NullTileCoord = 0; // Implicitly allocated when we created the tile allocator
 			FTileUploader::FAddResult AddResult = RootTileUploader.Add_GetRef(1 /*NumTiles*/, 0 /*NumVoxelsA*/, 0 /*NumVoxelsB*/);
 			FMemory::Memcpy(AddResult.PackedPhysicalTileCoordsPtr, &NullTileCoord, sizeof(NullTileCoord));
 			if (SVTInfo.FormatA != PF_Unknown)
@@ -665,9 +644,6 @@ void FStreamingManager::AddInternal(FRDGBuilder& GraphBuilder, FNewSparseVolumeT
 			const FPageTopology& Topology = Resources->Topology;
 			const int32 NumMipLevels = Topology.MipInfo.Num();
 
-			FrameInfo.LowestRequestedMipLevel = NumMipLevels - 1;
-			FrameInfo.LowestResidentMipLevel = NumMipLevels - 1;
-
 			// Create page table
 			{
 				FIntVector3 PageTableResolution = Resources->Header.PageTableVolumeResolution;
@@ -701,10 +677,14 @@ void FStreamingManager::AddInternal(FRDGBuilder& GraphBuilder, FNewSparseVolumeT
 			if (!Topology.MipInfo.IsEmpty() && Topology.MipInfo.Last().PageCount > 0)
 			{
 				check(!Resources->RootData.IsEmpty());
-				const uint32 TileCoord = SVTInfo.TileDataTexture->Allocate();
+				FTileAllocator::FAllocation PreviousAllocation;
+				const uint32 TileCoord = SVTInfo.TileAllocator.Allocate(NextUpdateIndex, 0 /*FreeThreshold*/, FrameIdx, 0 /*TileIndexInFrame*/, -1 /*TilePriority*/, true, PreviousAllocation);
 				check(TileCoord != INDEX_NONE);
 				check(TileCoord != 0);
+				check(!PreviousAllocation.bIsAllocated); // We should never have to evict another tile to upload a root tile
 				FrameInfo.TileAllocations[0] = TileCoord;
+				FrameInfo.ResidentTiles[0] = true;
+				FrameInfo.StreamingTiles[0] = true; // This tile is not technically streaming, but StreamingTiles is a super set of ResidentTiles and represents tiles that are resident or about to be
 
 				const FTileInfo RootTileInfo = Resources->StreamingMetaData.GetTileInfo(0, FormatSizes[0], FormatSizes[1]);
 				const int32 NumVoxelsA = RootTileInfo.NumVoxels[0];
@@ -731,10 +711,6 @@ void FStreamingManager::AddInternal(FRDGBuilder& GraphBuilder, FNewSparseVolumeT
 						FMemory::Memcpy(AddResult.TileDataPtrs[AttributesIdx], SrcTileData, RootTileInfo.VoxelDataSizes[AttributesIdx]);
 					}
 				}
-
-				// Mark page as resident
-				const FPageTopology::FMip& TopologyMipInfo = Topology.MipInfo.Last();
-				FrameInfo.ResidentPagesNew.SetRange(TopologyMipInfo.PageOffset, TopologyMipInfo.PageCount, true);
 			}
 
 			InvalidatedSVTFrames.Add(&FrameInfo);
@@ -746,15 +722,12 @@ void FStreamingManager::AddInternal(FRDGBuilder& GraphBuilder, FNewSparseVolumeT
 	// Add requests for all mips the first frame. This is necessary for cases where UAnimatedSparseVolumeTexture or UStaticSparseVolumeTexture
 	// are directly bound to the material without getting a specific frame through USparseVolumeTextureFrame::GetFrameAndIssueStreamingRequest().
 	const int32 NumMipLevelsFrame0 = SVTInfo.PerFrameInfo[0].NumMipLevels;
-	for (int32 MipLevel = 0; (MipLevel + 1) < NumMipLevelsFrame0; ++MipLevel)
-	{
-		FStreamingRequest Request;
-		Request.Key.SVT = NewSVTInfo.SVT;
-		Request.Key.FrameIndex = 0;
-		Request.Key.MipLevelIndex = MipLevel;
-		Request.Priority = MipLevel;
-		AddRequest(Request);
-	}
+	FStreamingRequest Request;
+	Request.Key.SVTHandle = SVTHandle;
+	Request.Key.FrameIndex = 0;
+	Request.Payload.MipLevelMask = UINT16_MAX;
+	Request.Payload.LowestMipFraction = 1.0f;
+	AddRequest(Request);
 }
 
 void FStreamingManager::RemoveInternal(UStreamableSparseVolumeTexture* SparseVolumeTexture)
@@ -765,23 +738,23 @@ void FStreamingManager::RemoveInternal(UStreamableSparseVolumeTexture* SparseVol
 	if (SVTInfo)
 	{
 		// Remove any requests for this SVT
-		TArray<FMipLevelKey> RequestsToRemove;
+		TArray<FFrameKey> RequestsToRemove;
 		for (auto& Pair : RequestsHashTable)
 		{
-			if (Pair.Key.SVT == SparseVolumeTexture)
+			if (Pair.Key.SVTHandle == SVTInfo->SVTHandle)
 			{
 				RequestsToRemove.Add(Pair.Key);
 			}
 		}
-		for (const FMipLevelKey& Key : RequestsToRemove)
+		for (const FFrameKey& Key : RequestsToRemove)
 		{
 			RequestsHashTable.Remove(Key);
 		}
 
-		// Cancel any pending mip levels
+		// Cancel any pending requests
 		for (FPendingRequest& PendingRequest : PendingRequests)
 		{
-			if (PendingRequest.SparseVolumeTexture == SparseVolumeTexture)
+			if (PendingRequest.SVTHandle == SVTInfo->SVTHandle)
 			{
 				PendingRequest.Reset();
 			}
@@ -799,157 +772,170 @@ void FStreamingManager::RemoveInternal(UStreamableSparseVolumeTexture* SparseVol
 			SVTInfo->TileDataTexture.Reset();
 		}
 
-		StreamingInfo.Remove(SparseVolumeTexture);
+		StreamingInfo.RemoveAt(SVTInfo->SVTHandle, 1);
 	}
 }
 
-bool FStreamingManager::AddRequest(const FStreamingRequest& Request)
+void FStreamingManager::AddRequest(const FStreamingRequest& Request)
 {
-	uint32* ExistingRequestPriority = RequestsHashTable.Find(Request.Key);
-	if (ExistingRequestPriority)
+	FRequestPayload* ExistingRequestPayload = RequestsHashTable.Find(Request.Key);
+	if (ExistingRequestPayload)
 	{
-		if (Request.Priority > *ExistingRequestPriority)
+		ExistingRequestPayload->MipLevelMask |= Request.Payload.MipLevelMask;
+		const int32 ExistingLowestMip = FMath::CountTrailingZeros(ExistingRequestPayload->MipLevelMask);
+		const int32 IncomingLowestMip = FMath::CountTrailingZeros(Request.Payload.MipLevelMask);
+		if (IncomingLowestMip < ExistingLowestMip)
 		{
-			*ExistingRequestPriority = Request.Priority;
-			return true;
+			ExistingRequestPayload->LowestMipFraction = Request.Payload.LowestMipFraction;
 		}
-		else
+		for (int32 i = 0; i < Request.Payload.Priorities.Num(); ++i)
 		{
-			return false;
+			ExistingRequestPayload->Priorities[i] = FMath::Max(ExistingRequestPayload->Priorities[i], Request.Payload.Priorities[i]);
 		}
 	}
 	else
 	{
-		RequestsHashTable.Add(Request.Key, Request.Priority);
-		return true;
+		RequestsHashTable.Add(Request.Key, Request.Payload);
 	}
 }
 
-void FStreamingManager::AddParentRequests()
+void FStreamingManager::FilterRequests()
 {
-	TRACE_CPUPROFILER_EVENT_SCOPE(SVT::StreamingAddParentRequests);
+	TRACE_CPUPROFILER_EVENT_SCOPE(SVT::StreamingFilterRequests);
 
-	ParentRequestsToAdd.Reset();
-	for (const auto& Request : RequestsHashTable)
+	TileRangesToStream.Reset();
+
+	// Convert requests into contiguous ranges of tiles with the same priority
+	for (const auto& RequestPair : RequestsHashTable)
 	{
-		FStreamingInfo* SVTInfo = FindStreamingInfo(Request.Key.SVT);
-		check(SVTInfo);
-		const int32 NumStreamableMipLevels = SVTInfo->PerFrameInfo[Request.Key.FrameIndex].NumMipLevels - 1;
-		uint32 Priority = Request.Value == FStreamingRequest::BlockingPriority ? FStreamingRequest::BlockingPriority : (Request.Value + 1);
-		for (int32 MipLevelIndex = Request.Key.MipLevelIndex + 1; MipLevelIndex < NumStreamableMipLevels; ++MipLevelIndex)
+		const FFrameKey& RequestKey = RequestPair.Key;
+		const FRequestPayload& RequestValue = RequestPair.Value;
+		FStreamingInfo* SVTInfo = FindStreamingInfo(RequestKey.SVTHandle);
+		if (!SVTInfo || !SVTInfo->PerFrameInfo.IsValidIndex(RequestKey.FrameIndex))
 		{
-			FMipLevelKey ParentKey = Request.Key;
-			ParentKey.MipLevelIndex = MipLevelIndex;
-			
-			uint32* ExistingParentRequestPriority = RequestsHashTable.Find(ParentKey);
-			if (ExistingParentRequestPriority && Priority > *ExistingParentRequestPriority)
-			{
-				*ExistingParentRequestPriority = Priority;
-			}
-			else
-			{
-				ParentRequestsToAdd.Add(FStreamingRequest{ ParentKey, Priority });
-			}
-
-			if (Priority != FStreamingRequest::BlockingPriority)
-			{
-				++Priority;
-			}
+			continue;
 		}
-	}
 
-	for (const FStreamingRequest& Request : ParentRequestsToAdd)
-	{
-		AddRequest(Request);
-	}
-}
+		FFrameInfo& FrameInfo = SVTInfo->PerFrameInfo[RequestKey.FrameIndex];
+		const FPageTopology& Topology = FrameInfo.Resources->Topology;
+		const int32 FirstStreamingMipLevel = FrameInfo.NumMipLevels - 2;
+		const int32 LowestRequestedMip = FMath::CountTrailingZeros((uint32)RequestValue.MipLevelMask);
 
-void FStreamingManager::SelectHighestPriorityRequestsAndUpdateLRU(int32 MaxSelectedRequests)
-{
-	PrioritizedRequestsHeap.Reset();
-	SelectedRequests.Reset();
+		// Keep track of all tiles requested in this frame and also of the tiles requested only within the current priority
+		const uint32 NumTiles = FrameInfo.TileAllocations.Num();
+		TBitArray<> RequestedTiles = FrameInfo.StreamingTiles;
+		TBitArray<> RequestedTilesInCurrentPriority = TBitArray<>(false, NumTiles);
 
-	if (!RequestsHashTable.IsEmpty())
-	{
-		TRACE_CPUPROFILER_EVENT_SCOPE(SVT::StreamingFilterRequests);
+		// StreamingTiles must be a super set of ResidentTiles
+		check(TBitArray<>::BitwiseAND(FrameInfo.StreamingTiles, FrameInfo.ResidentTiles, EBitwiseOperatorFlags::MaxSize) == FrameInfo.ResidentTiles);
 
-		for (const auto& Request : RequestsHashTable)
+		// Creates a FTileRange for each contiguous range of set bits in RequestedTilesInCurrentPriority
+		auto MakeTileRanges = [&](uint8 Priority)
 		{
-			FStreamingInfo* SVTInfo = FindStreamingInfo(Request.Key.SVT);
-			check(SVTInfo);
+			bool bResetCounters = true;
+			uint32 TileOffset = 0;
+			uint32 TileCount = 0;
+			for (TConstSetBitIterator It(RequestedTilesInCurrentPriority); It; ++It)
+			{
+				const uint32 TileIndex = It.GetIndex();
+				if (bResetCounters)
+				{
+					TileOffset = TileIndex;
+					TileCount = 0;
+					bResetCounters = false;
+				}
+				++TileCount;
+
+				// Is the next bit unset? If so, this tile is the end of the current range
+				if (!RequestedTilesInCurrentPriority.IsValidIndex(TileIndex + 1) || !RequestedTilesInCurrentPriority[TileIndex + 1])
+				{
+					FTileRange& TileRange = TileRangesToStream.AddDefaulted_GetRef();
+					TileRange.SVTHandle = RequestKey.SVTHandle;
+					TileRange.TileOffset = TileOffset;
+					TileRange.TileCount = TileCount;
+					TileRange.FrameIndex = RequestKey.FrameIndex;
+					TileRange.Priority = Priority;
+
+					bResetCounters = true;
+				}
+			}
+		};
 
-			// Discard invalid requests: frame index out of bounds, mip level index out of bounds (or root mip level) and mip levels without any data.
-			// There can never be lower mip levels with data depending on higher mip levels without any data, so discarding such requests is ok.
-			if (Request.Key.FrameIndex < 0
-				|| Request.Key.FrameIndex >= SVTInfo->PerFrameInfo.Num()
-				|| Request.Key.MipLevelIndex < 0
-				|| Request.Key.MipLevelIndex >= (SVTInfo->PerFrameInfo[Request.Key.FrameIndex].NumMipLevels - 1)
-				|| SVTInfo->PerFrameInfo[Request.Key.FrameIndex].Resources->Topology.MipInfo[Request.Key.MipLevelIndex].PageCount == 0)
+		uint32 PrevPriority = 0xFFFFFFFFu;
+		for (int32 MipLevelIndex = FirstStreamingMipLevel; MipLevelIndex >= LowestRequestedMip; --MipLevelIndex)
+		{
+			// Skip if mip level was not requested
+			if (((uint32)RequestValue.MipLevelMask & (1u << MipLevelIndex)) == 0)
 			{
 				continue;
 			}
 
-			const int32 LRUNodeIndex = Request.Key.FrameIndex * SVTInfo->NumMipLevelsGlobal + Request.Key.MipLevelIndex;
-			FLRUNode* LRUNode = &SVTInfo->LRUNodes[LRUNodeIndex];
-#if DO_CHECK
-			bool bFoundNodeInList = false;
-			for (auto& Node : SVTInfo->PerMipLRULists[Request.Key.MipLevelIndex])
+			// When the priority changes, we need to create tile ranges based on the currently set bets in RequestedTilesInCurrentPriority.
+			const uint8 Priority = RequestValue.Priorities[MipLevelIndex];
+			if (Priority != PrevPriority)
 			{
-				if (&Node == LRUNode)
-				{
-					bFoundNodeInList = true;
-					break;
-				}
+				MakeTileRanges(Priority);
+				RequestedTilesInCurrentPriority.SetRange(0, NumTiles, false);
 			}
-#endif
+			PrevPriority = Priority;
 
-			const bool bIsAlreadyStreaming = Request.Key.MipLevelIndex >= SVTInfo->PerFrameInfo[Request.Key.FrameIndex].LowestRequestedMipLevel;
-			if (bIsAlreadyStreaming)
+			// Get the range of pages to try and stream in
+			const FPageTopology::FMip& MipInfo = Topology.MipInfo[MipLevelIndex];
+			uint32 PageCount = MipInfo.PageCount;
+			if (MipLevelIndex == LowestRequestedMip)
 			{
-				check(bFoundNodeInList);
-				// Update LastRequested and move to front of LRU
-				LRUNode->LastRequested = NextUpdateIndex;
+				PageCount = FMath::CeilToInt32((float)PageCount * FMath::Clamp(RequestValue.LowestMipFraction, 0.0f, 1.0f));
+				check(PageCount <= MipInfo.PageCount);
+			}
+			const uint32 PageBegin = MipInfo.PageOffset;
+			const uint32 PageEnd = MipInfo.PageOffset + PageCount;
 
-				// Unlink
-				LRUNode->Remove();
+			// Mark requested tiles
+			for (uint32 PageIndex = PageBegin; PageIndex < PageEnd; ++PageIndex)
+			{
+				const uint32 TileIndex = Topology.TileIndices[PageIndex];
+				if (!RequestedTiles[TileIndex])
+				{
+					RequestedTiles[TileIndex] = true;
+					RequestedTilesInCurrentPriority[TileIndex] = true;
+				}
 
-				// Insert at the end of the LRU list
-				SVTInfo->PerMipLRULists[Request.Key.MipLevelIndex].AddTail(LRUNode);
+				// Tell the allocator that we're still interested in this tile
+				if (FrameInfo.StreamingTiles[TileIndex])
+				{
+					check(FrameInfo.TileAllocations[TileIndex] != 0);
+					const uint32 TilePriority = FrameInfo.ResidentTiles.Num() - TileIndex; // Higher value means higher priority, but we want higher tiles to be streamed out first
+					SVTInfo->TileAllocator.UpdateUsage(NextUpdateIndex, FrameInfo.TileAllocations[TileIndex], TilePriority);
+				}
 			}
-			else
+
+			// After processing the lowest/last mip level, we need to create one final batch of tile ranges
+			if (MipLevelIndex == LowestRequestedMip)
 			{
-				check(!bFoundNodeInList);
-				PrioritizedRequestsHeap.Add(FStreamingRequest{ Request.Key, Request.Value });
+				MakeTileRanges(Priority);
 			}
 		}
+	}
 
-		// Sort by priority but make sure to load higher mip levels with the same priority first. This can happen when a blocking priority is used.
-		auto PriorityPredicate = [](const auto& A, const auto& B) { return A.Priority != B.Priority ? (A.Priority > B.Priority) : A.Key.MipLevelIndex > B.Key.MipLevelIndex; };
-		PrioritizedRequestsHeap.Heapify(PriorityPredicate);
+	RequestsHashTable.Reset();
 
-		while (SelectedRequests.Num() < MaxSelectedRequests && PrioritizedRequestsHeap.Num() > 0)
-		{
-			FStreamingRequest SelectedRequest;
-			PrioritizedRequestsHeap.HeapPop(SelectedRequest, PriorityPredicate, EAllowShrinking::No);
+	// Sort by priority
+	auto PriorityPredicate = [](const auto& A, const auto& B) { return A.Priority != B.Priority ? (A.Priority > B.Priority) : A.TileOffset < B.TileOffset; };
+	TileRangesToStream.Sort(PriorityPredicate);
 
-			FStreamingInfo* SVTInfo = FindStreamingInfo(SelectedRequest.Key.SVT);
-			if (SVTInfo)
-			{
-				check(SelectedRequest.Key.FrameIndex < SVTInfo->PerFrameInfo.Num());
-				check(SelectedRequest.Key.MipLevelIndex < SVTInfo->PerFrameInfo[SelectedRequest.Key.FrameIndex].NumMipLevels);
-				SelectedRequests.Push(SelectedRequest);
-			}
-		}
+	// SVT_TODO: Here we need to shave off tiles depending on available bandwidth
 
-		RequestsHashTable.Reset();
-	}
+	const int32 MaxSelectedRequests = MaxPendingRequests - NumPendingRequests;
+	const int32 NumSelectedRequests = FMath::Min(MaxSelectedRequests, TileRangesToStream.Num());
+
+	TileRangesToStream.SetNum(NumSelectedRequests, EAllowShrinking::No);
 }
 
-void FStreamingManager::IssueRequests(int32 MaxSelectedRequests)
+void FStreamingManager::IssueRequests()
 {
 	using namespace UE::DerivedData;
 
-	if (SelectedRequests.IsEmpty())
+	if (TileRangesToStream.IsEmpty())
 	{
 		return;
 	}
@@ -958,278 +944,179 @@ void FStreamingManager::IssueRequests(int32 MaxSelectedRequests)
 
 #if WITH_EDITORONLY_DATA
 	TArray<FCacheGetChunkRequest> DDCRequests;
-	DDCRequests.Reserve(MaxSelectedRequests);
+	DDCRequests.Reserve(TileRangesToStream.Num());
 	TArray<FCacheGetChunkRequest> DDCRequestsBlocking;
-	DDCRequestsBlocking.Reserve(MaxSelectedRequests);
+	DDCRequestsBlocking.Reserve(TileRangesToStream.Num());
 #endif
 
-	FBulkDataBatchRequest::FBatchBuilder Batch = FBulkDataBatchRequest::NewBatch(SelectedRequests.Num());
+	FBulkDataBatchRequest::FBatchBuilder Batch = FBulkDataBatchRequest::NewBatch(TileRangesToStream.Num());
 	bool bIssueIOBatch = false;
 
-	for (const FStreamingRequest& SelectedRequest : SelectedRequests)
+	// Process all tile ranges selected for streaming, allocate a slot in the tile data texture for every tile and finally create IO requests for every range.
+	for (FTileRange& TileRange : TileRangesToStream)
 	{
-		const FMipLevelKey& SelectedKey = SelectedRequest.Key;
-		FStreamingInfo* SVTInfo = FindStreamingInfo(SelectedKey.SVT);
-		check(SVTInfo);
-		check(SVTInfo->PerFrameInfo.Num() > SelectedKey.FrameIndex && SelectedKey.FrameIndex >= 0);
-		check(SVTInfo->PerFrameInfo[SelectedKey.FrameIndex].LowestRequestedMipLevel > SelectedKey.MipLevelIndex);
-		const FResources* Resources = SVTInfo->PerFrameInfo[SelectedKey.FrameIndex].Resources;
-		check((SelectedKey.MipLevelIndex + 1) < Resources->Topology.MipInfo.Num()); // The lowest/last mip level is always resident and does not stream.
-
-		const FMipTileReadInfo MipTileReadInfo = GetMipTileReadInfo(SVTInfo->PerFrameInfo[SelectedKey.FrameIndex], SelectedKey.MipLevelIndex);
-
-		TUniquePtr<FTileDataTexture>& TileDataTexture = SVTInfo->TileDataTexture;
-		check(TileDataTexture);
+		FStreamingInfo* SVTInfo = FindStreamingInfo(TileRange.SVTHandle);
+		FFrameInfo& FrameInfo = SVTInfo->PerFrameInfo[TileRange.FrameIndex];
+		const FResources* Resources = FrameInfo.Resources;
+		FTileDataTexture* TileDataTexture = SVTInfo->TileDataTexture.Get();
+		const int32 NumFrames = SVTInfo->PerFrameInfo.Num();
 
-		// Ensure that enough tiles are available in the tile texture
-		const int32 TileDataTextureCapacity = TileDataTexture->GetTileCapacity();
-		const int32 NumAvailableTiles = TileDataTexture->GetNumAvailableTiles();
-		const int32 NumRequiredTiles = MipTileReadInfo.TileCount;
-		if (NumAvailableTiles < NumRequiredTiles)
+		// Allocate tiles
+		uint32 NumSuccessfulAllocations = 0;
+		for (uint32 i = 0; i < TileRange.TileCount; ++i)
 		{
-#if SVT_STREAMING_LOG_VERBOSE
-			UE_LOG(LogSparseVolumeTextureStreamingManager, Display, TEXT("(%i)%i IssueRequests() Frame %i Mip %i: Not enough tiles available (%i) to fit mip level (%i)"), 
-				NextUpdateIndex, NextPendingRequestIndex, SelectedKey.FrameIndex, SelectedKey.MipLevelIndex, NumAvailableTiles, NumRequiredTiles);
-#endif
+			const uint32 TileIndex = TileRange.TileOffset + i;
+			check(FrameInfo.TileAllocations[TileIndex] == 0);
 
-			// Try to free old mip levels, starting at higher resolution mips and going up the mip chain
-			TArray<FLRUNode*, TInlineAllocator<16>> MipLevelsToFree;
-			int32 NumNewlyAvailableTiles = 0;
-			const int32 NumMipLevelsGlobal = SVTInfo->NumMipLevelsGlobal;
-			for (int32 MipLevel = 0; MipLevel < NumMipLevelsGlobal && (NumAvailableTiles + NumNewlyAvailableTiles) < NumRequiredTiles; ++MipLevel)
+			// Allocate tile
+			const uint32 TilePriority = FrameInfo.ResidentTiles.Num() - TileIndex; // Higher value means higher priority, but we want higher tiles to be streamed out first
+			FTileAllocator::FAllocation PreviousAllocation;
+			const uint32 TileCoord = SVTInfo->TileAllocator.Allocate(NextUpdateIndex, 0 /*FreeThreshold*/, TileRange.FrameIndex, TileIndex, TilePriority, false /*bLocked*/, PreviousAllocation);
+			
+			// Allocation failed, cut the TileRange short and upload what we have.
+			if (TileCoord == INDEX_NONE)
 			{
-				for (auto& Node : SVTInfo->PerMipLRULists[MipLevel])
+				if (GSVTStreamingLogVerbosity > 1)
 				{
-					// Only free "leaf" mip levels with no higher resolution mip levels resident. Don't free mip levels requested this frame.
-					if (Node.RefCount == 0 && Node.LastRequested < NextUpdateIndex)
-					{
-						MipLevelsToFree.Add(&Node);
-						const FPageTopology& Topology = SVTInfo->PerFrameInfo[Node.FrameIndex].Resources->Topology;
-						const FPageTopology::FMip& TopologyMip = Topology.MipInfo[Node.MipLevelIndex];
-						uint32 TilesInMipCount = 0;
-						uint32 TilesInMipOffset = 0;
-						Topology.GetTileRange(TopologyMip.PageOffset, TopologyMip.PageCount, TilesInMipOffset, TilesInMipCount);
-						NumNewlyAvailableTiles += TilesInMipCount;
-
-						// Decrement ref count of mip levels higher up the chain
-						FLRUNode* Dependency = Node.NextHigherMipLevel;
-						while (Dependency)
-						{
-							check(Dependency->RefCount > 0);
-							--Dependency->RefCount;
-							Dependency = Dependency->NextHigherMipLevel;
-						}
-					}
-
-					// Exit once we freed enough tiles
-					if ((NumAvailableTiles + NumNewlyAvailableTiles) >= NumRequiredTiles)
-					{
-						break;
-					}
+					UE_LOG(LogSparseVolumeTextureStreamingManager, Warning, TEXT("Streaming pool of SVT '%s' is oversubscribed! Tried to allocate %u tiles for frame %u, but only %u were available."),
+						*SVTInfo->SVTName.ToString(), TileRange.TileCount, TileRange.FrameIndex, NumSuccessfulAllocations);
 				}
+				break;
 			}
+			++NumSuccessfulAllocations;
+			FrameInfo.TileAllocations[TileIndex] = TileCoord;
 
-			// Free mip levels
-			for (FLRUNode* MipLevelToFree : MipLevelsToFree)
+			// Slot in the tile data texture was previously used by a different tile -> mark it as no longer resident/streaming.
+			if (PreviousAllocation.bIsAllocated)
 			{
-				StreamOutMipLevel(SVTInfo, MipLevelToFree);
-			}
+				check(SVTInfo->PerFrameInfo.IsValidIndex(PreviousAllocation.FrameIndex));
+				check(SVTInfo->PerFrameInfo[PreviousAllocation.FrameIndex].ResidentTiles.IsValidIndex(PreviousAllocation.TileIndexInFrame));
+				FFrameInfo& FrameInfoTmp = SVTInfo->PerFrameInfo[PreviousAllocation.FrameIndex];
+
+				// Mark the tile as no longer resident/streaming and make sure to set its entry in TileAllocations to 0 (for validation purposes).
+				check(FrameInfoTmp.StreamingTiles[PreviousAllocation.TileIndexInFrame]);
+				check(FrameInfoTmp.TileAllocations[PreviousAllocation.TileIndexInFrame] != 0);
+				FrameInfoTmp.StreamingTiles[PreviousAllocation.TileIndexInFrame] = false;
+				FrameInfoTmp.ResidentTiles[PreviousAllocation.TileIndexInFrame] = false;
+				FrameInfoTmp.TileAllocations[PreviousAllocation.TileIndexInFrame] = 0;
+
+				// Cancel potential IO request
+				if (uint32* PendingRequestIndexPtr = FrameInfoTmp.TileIndexToPendingRequestIndex.Find((uint32)PreviousAllocation.TileIndexInFrame))
+				{
+					FPendingRequest& PendingRequest = PendingRequests[*PendingRequestIndexPtr];
+					check(PendingRequest.IsValid());
+					if (FrameInfoTmp.StreamingTiles.CountSetBits(PendingRequest.TileOffset, PendingRequest.TileOffset + PendingRequest.TileCount) == 0)
+					{
+						// None of the tiles in the tile range of the IO request is still requested, so we can cancel the request
+						PendingRequest.Reset();
+					}
 
-			// Couldn't free enough tiles, so skip this mip level
-			if ((NumAvailableTiles + NumNewlyAvailableTiles) < NumRequiredTiles)
-			{
-				UE_LOG(LogSparseVolumeTextureStreamingManager, Warning, TEXT("IssueRequests() SVT %p Frame %i Mip %i: Not enough tiles available (%i) to fit mip level (%i) even after freeing"),
-					SelectedKey.SVT, SelectedKey.FrameIndex, SelectedKey.MipLevelIndex, (NumAvailableTiles + NumNewlyAvailableTiles), NumRequiredTiles);
-				continue;
+					// Remove the tile -> IO request mapping
+					FrameInfoTmp.TileIndexToPendingRequestIndex.Remove((uint32)PreviousAllocation.TileIndexInFrame);
+				}
+
+				// We changed the tile residency of this frame, so add it to the set of invalidated frames that need a page table fixup
+				InvalidatedSVTFrames.Add(&FrameInfoTmp);
 			}
 		}
 
-#if DO_CHECK
-		for (auto& Pending : PendingRequests)
+		TileRange.TileCount = NumSuccessfulAllocations;
+		if (TileRange.TileCount == 0)
 		{
-			check(Pending.SparseVolumeTexture != SelectedKey.SVT || Pending.FrameIndex != SelectedKey.FrameIndex || Pending.MipLevelIndex != SelectedKey.MipLevelIndex); //-V1013
+			continue;
 		}
-#endif
 
-		const int32 PendingRequestIndex = NextPendingRequestIndex;
-		FPendingRequest& PendingRequest = PendingRequests[PendingRequestIndex];
-		PendingRequest.Reset();
-		PendingRequest.SparseVolumeTexture = SelectedKey.SVT;
-		PendingRequest.FrameIndex = SelectedKey.FrameIndex;
-		PendingRequest.MipLevelIndex = SelectedKey.MipLevelIndex;
-		PendingRequest.IssuedInFrame = NextUpdateIndex;
-		PendingRequest.bBlocking = GSVTStreamingForceBlockingRequests || (SelectedRequest.Priority == FStreamingRequest::BlockingPriority);
-
-		const FByteBulkData& BulkData = Resources->StreamableMipLevels;
+		// Create request
+		{
+			const int32 PendingRequestIndex = NextPendingRequestIndex;
+			FPendingRequest& PendingRequest = PendingRequests[PendingRequestIndex];
+			PendingRequest.Reset();
+			PendingRequest.SVTHandle = TileRange.SVTHandle;
+			PendingRequest.FrameIndex = TileRange.FrameIndex;
+			PendingRequest.TileOffset = TileRange.TileOffset;
+			PendingRequest.TileCount = TileRange.TileCount;
+			PendingRequest.IssuedInFrame = NextUpdateIndex;
+			PendingRequest.bBlocking = GSVTStreamingForceBlockingRequests || (TileRange.Priority == uint8(-1));
+
+			uint32 ReadOffset = 0;
+			uint32 ReadSize = 0;
+			Resources->StreamingMetaData.GetTileRangeMemoryOffsetSize(PendingRequest.TileOffset, PendingRequest.TileCount, ReadOffset, ReadSize);
+
+			const FByteBulkData& BulkData = Resources->StreamableMipLevels;
 #if WITH_EDITORONLY_DATA
-		const bool bDiskRequest = (!(Resources->ResourceFlags & EResourceFlag_StreamingDataInDDC) && !BulkData.IsBulkDataLoaded());
+			const bool bDiskRequest = (!(Resources->ResourceFlags & EResourceFlag_StreamingDataInDDC) && !BulkData.IsBulkDataLoaded());
 #else
-		const bool bDiskRequest = true;
+			const bool bDiskRequest = true;
 #endif
 
 #if WITH_EDITORONLY_DATA
-		if (!bDiskRequest)
-		{
-			if (Resources->ResourceFlags & EResourceFlag_StreamingDataInDDC)
+			if (!bDiskRequest)
 			{
-				UE::DerivedData::FCacheGetChunkRequest DDCRequest = BuildDDCRequest(*Resources, MipTileReadInfo.ReadOffset, MipTileReadInfo.ReadSize, NextPendingRequestIndex);
-				if (PendingRequest.bBlocking)
+				if (Resources->ResourceFlags & EResourceFlag_StreamingDataInDDC)
 				{
-					DDCRequestsBlocking.Add(DDCRequest);
+					UE::DerivedData::FCacheGetChunkRequest DDCRequest = BuildDDCRequest(*Resources, ReadOffset, ReadSize, NextPendingRequestIndex);
+					if (PendingRequest.bBlocking)
+					{
+						DDCRequestsBlocking.Add(DDCRequest);
+					}
+					else
+					{
+						DDCRequests.Add(DDCRequest);
+					}
+					PendingRequest.State = FPendingRequest::EState::DDC_Pending;
 				}
 				else
 				{
-					DDCRequests.Add(DDCRequest);
+					PendingRequest.State = FPendingRequest::EState::Memory;
 				}
-				PendingRequest.State = FPendingRequest::EState::DDC_Pending;
 			}
 			else
-			{
-				PendingRequest.State = FPendingRequest::EState::Memory;
-			}
-		}
-		else
 #endif
-		{
-			PendingRequest.RequestBuffer = FIoBuffer(MipTileReadInfo.ReadSize); // SVT_TODO: Use FIoBuffer::Wrap with preallocated memory
-			const EAsyncIOPriorityAndFlags Priority = PendingRequest.bBlocking ? AIOP_CriticalPath : AIOP_Low;
-			Batch.Read(BulkData, MipTileReadInfo.ReadOffset, MipTileReadInfo.ReadSize, Priority, PendingRequest.RequestBuffer, PendingRequest.Request);
-			bIssueIOBatch = true;
+			{
+				PendingRequest.RequestBuffer = FIoBuffer(ReadSize); // SVT_TODO: Use FIoBuffer::Wrap with preallocated memory
+				const EAsyncIOPriorityAndFlags Priority = PendingRequest.bBlocking ? AIOP_CriticalPath : AIOP_Low;
+				Batch.Read(BulkData, ReadOffset, ReadSize, Priority, PendingRequest.RequestBuffer, PendingRequest.Request);
+				bIssueIOBatch = true;
 
 #if WITH_EDITORONLY_DATA
-			PendingRequest.State = FPendingRequest::EState::Disk;
+				PendingRequest.State = FPendingRequest::EState::Disk;
 #endif
-		}
-
-		NextPendingRequestIndex = (NextPendingRequestIndex + 1) % MaxPendingRequests;
-		check(NumPendingRequests < MaxPendingRequests);
-		++NumPendingRequests;
-
-		FFrameInfo& FrameInfo = SVTInfo->PerFrameInfo[SelectedKey.FrameIndex];
-
-		// Allocate tiles in the tile data texture
-		{
-			const FPageTopology::FMip& MipInfo = FrameInfo.Resources->Topology.MipInfo[SelectedKey.MipLevelIndex];
-
-			uint32 TilesInMipCount = 0;
-			uint32 TilesInMipOffset = 0;
-			Resources->Topology.GetTileRange(MipInfo.PageOffset, MipInfo.PageCount, TilesInMipOffset, TilesInMipCount);
-
-			for (uint32 TileIndex = TilesInMipOffset; TileIndex < (TilesInMipOffset + TilesInMipCount); ++TileIndex)
-			{
-				check(FrameInfo.TileAllocations[TileIndex] == 0);
-				const int32 TileCoord = TileDataTexture->Allocate();
-				check(TileCoord != INDEX_NONE);
-				check(TileCoord != 0);
-				FrameInfo.TileAllocations[TileIndex] = TileCoord;
 			}
-		}
 
-		// Add to tail of LRU list
-		{
-			const int32 LRUNodeIndex = SelectedKey.FrameIndex * SVTInfo->NumMipLevelsGlobal + SelectedKey.MipLevelIndex;
-			FLRUNode* LRUNode = &SVTInfo->LRUNodes[LRUNodeIndex];
-			check(!LRUNode->IsInList());
-			LRUNode->LastRequested = NextUpdateIndex;
-			LRUNode->PendingRequestIndex = PendingRequestIndex;
-
-			FLRUNode* Dependency = LRUNode->NextHigherMipLevel;
-			while (Dependency)
+			for (uint32 TileIndex = TileRange.TileOffset; TileIndex < (TileRange.TileOffset + TileRange.TileCount); ++TileIndex)
 			{
-				++Dependency->RefCount;
-				Dependency = Dependency->NextHigherMipLevel;
+				FrameInfo.TileIndexToPendingRequestIndex.Add(TileIndex, (uint32)NextPendingRequestIndex);
 			}
 
-			SVTInfo->PerMipLRULists[SelectedKey.MipLevelIndex].AddTail(LRUNode);
+			NextPendingRequestIndex = (NextPendingRequestIndex + 1) % MaxPendingRequests;
+			check(NumPendingRequests < MaxPendingRequests);
+			++NumPendingRequests;
 		}
 
-#if SVT_STREAMING_LOG_VERBOSE
-		UE_LOG(LogSparseVolumeTextureStreamingManager, Display, TEXT("(%i)%i StreamIn Frame %i OldReqMip %i, NewReqMip %i, ResMip %i"),
-			PendingRequest.IssuedInFrame, PendingRequestIndex,
-			SelectedKey.FrameIndex, 
-			FrameInfo.LowestRequestedMipLevel, SelectedKey.MipLevelIndex, 
-			FrameInfo.LowestResidentMipLevel);
-#endif
-
-		check(FrameInfo.LowestRequestedMipLevel == (SelectedKey.MipLevelIndex + 1));
-		FrameInfo.LowestRequestedMipLevel = SelectedKey.MipLevelIndex;
+		// Mark tiles as streaming. Once they've actually been loaded into GPU memory, they'll also be marked as resident.
+		FrameInfo.StreamingTiles.SetRange(TileRange.TileOffset, TileRange.TileCount, true);
 	}
 
-#if WITH_EDITORONLY_DATA
-	if (!DDCRequests.IsEmpty())
-	{
-		RequestDDCData(DDCRequests, false /*bBlocking*/);
-		DDCRequests.Empty();
-	}
-	if (!DDCRequestsBlocking.IsEmpty())
+	// Now we can finally issue the requests
 	{
-		RequestDDCData(DDCRequestsBlocking, true /*bBlocking*/);
-		DDCRequestsBlocking.Empty();
-	}
-#endif
-
-	if (bIssueIOBatch)
-	{
-		(void)Batch.Issue();
-	}
-}
-
-void FStreamingManager::StreamOutMipLevel(FStreamingInfo* SVTInfo, FLRUNode* LRUNode)
-{
-	const int32 FrameIndex = LRUNode->FrameIndex;
-	const int32 MipLevelIndex = LRUNode->MipLevelIndex;
-
-	FFrameInfo& FrameInfo = SVTInfo->PerFrameInfo[FrameIndex];
-
-	check(FrameInfo.LowestResidentMipLevel >= MipLevelIndex); // mip might not have streamed in yet, so use >= instead of ==
-	check(FrameInfo.LowestRequestedMipLevel == MipLevelIndex);
-
-	// Cancel potential IO request
-	check((MipLevelIndex < FrameInfo.LowestResidentMipLevel) == (LRUNode->PendingRequestIndex != INDEX_NONE));
-	if (LRUNode->PendingRequestIndex != INDEX_NONE)
-	{
-		PendingRequests[LRUNode->PendingRequestIndex].Reset();
-		LRUNode->PendingRequestIndex = INDEX_NONE;
-	}
-
-	const int32 NewLowestRequestedMipLevel = MipLevelIndex + 1;
-	const int32 NewLowestResidentMipLevel = FMath::Max(MipLevelIndex + 1, FrameInfo.LowestResidentMipLevel);
-#if SVT_STREAMING_LOG_VERBOSE
-	UE_LOG(LogSparseVolumeTextureStreamingManager, Display, TEXT("(%i)%i StreamOut Frame %i OldReqMip %i, NewReqMip %i, OldResMip %i, NewResMip %i"),
-		NextUpdateIndex, NextPendingRequestIndex,
-		FrameIndex,
-		FrameInfo.LowestRequestedMipLevel, NewLowestRequestedMipLevel,
-		FrameInfo.LowestResidentMipLevel, NewLowestResidentMipLevel);
+#if WITH_EDITORONLY_DATA
+		if (!DDCRequests.IsEmpty())
+		{
+			RequestDDCData(DDCRequests, false /*bBlocking*/);
+			DDCRequests.Empty();
+		}
+		if (!DDCRequestsBlocking.IsEmpty())
+		{
+			RequestDDCData(DDCRequestsBlocking, true /*bBlocking*/);
+			DDCRequestsBlocking.Empty();
+		}
 #endif
 
-	FrameInfo.LowestRequestedMipLevel = NewLowestRequestedMipLevel;
-	FrameInfo.LowestResidentMipLevel = NewLowestResidentMipLevel;
-	
-	InvalidatedSVTFrames.Add(&FrameInfo);
-
-	// Unlink
-	LRUNode->Remove();
-	LRUNode->LastRequested = INDEX_NONE;
-
-	// Free allocated tiles
-	{
-		const FPageTopology& Topology = FrameInfo.Resources->Topology;
-		const FPageTopology::FMip& TopologyMip = Topology.MipInfo[MipLevelIndex];
-		uint32 TilesInMipCount = 0;
-		uint32 TilesInMipOffset = 0;
-		Topology.GetTileRange(TopologyMip.PageOffset, TopologyMip.PageCount, TilesInMipOffset, TilesInMipCount);
-
-		for (uint32 TileIndex = TilesInMipOffset; TileIndex < (TilesInMipOffset + TilesInMipCount); ++TileIndex)
+		if (bIssueIOBatch)
 		{
-			uint32& TileCoord = FrameInfo.TileAllocations[TileIndex];
-			check(TileCoord != 0);
-			SVTInfo->TileDataTexture->Free(TileCoord);
-			TileCoord = 0;
+			(void)Batch.Issue();
 		}
 	}
-
-	const FPageTopology::FMip& TopologyMipInfo = FrameInfo.Resources->Topology.MipInfo[MipLevelIndex];
-	FrameInfo.ResidentPagesNew.SetRange(TopologyMipInfo.PageOffset, TopologyMipInfo.PageCount, false);
 }
 
 int32 FStreamingManager::DetermineReadyRequests()
@@ -1246,26 +1133,27 @@ int32 FStreamingManager::DetermineReadyRequests()
 		const int32 PendingRequestIndex = (StartPendingRequestIndex + i) % MaxPendingRequests;
 		FPendingRequest& PendingRequest = PendingRequests[PendingRequestIndex];
 
-		FStreamingInfo* SVTInfo = FindStreamingInfo(PendingRequest.SparseVolumeTexture);
-		if (!SVTInfo)
+		// Check if request was cancelled
+		if (!PendingRequest.IsValid())
 		{
 #if WITH_EDITORONLY_DATA
-			// Resource is no longer there. Just mark as ready so it will be skipped later
+			// Just mark as ready so it will be skipped later
 			PendingRequest.State = FPendingRequest::EState::DDC_Ready;
 #endif
 			continue; 
 		}
-
+		FStreamingInfo* SVTInfo = FindStreamingInfo(PendingRequest.SVTHandle);
+		check(SVTInfo);
 		const FResources* Resources = SVTInfo->PerFrameInfo[PendingRequest.FrameIndex].Resources;
 
 #if WITH_EDITORONLY_DATA
 		if (PendingRequest.State == FPendingRequest::EState::DDC_Ready)
 		{
-			if (PendingRequest.RetryCount > 0)
+			if (PendingRequest.RetryCount > 0 && GSVTStreamingLogVerbosity > 0)
 			{
 				check(SVTInfo);
-				UE_LOG(LogSparseVolumeTextureStreamingManager, Display, TEXT("SVT DDC retry succeeded for '%s' (frame %i, mip %i) on %i attempt."), 
-					*Resources->ResourceName, PendingRequest.FrameIndex, PendingRequest.MipLevelIndex, PendingRequest.RetryCount);
+				UE_LOG(LogSparseVolumeTextureStreamingManager, Display, TEXT("SVT DDC retry succeeded for '%s' (frame %i, tile offset %i, tile count %i) on %i attempt."),
+					*SVTInfo->SVTName.ToString(), PendingRequest.FrameIndex, PendingRequest.TileOffset, PendingRequest.TileCount, PendingRequest.RetryCount);
 			}
 		}
 		else if (PendingRequest.State == FPendingRequest::EState::DDC_Pending)
@@ -1276,14 +1164,16 @@ int32 FStreamingManager::DetermineReadyRequests()
 		{
 			PendingRequest.State = FPendingRequest::EState::DDC_Pending;
 
-			if (PendingRequest.RetryCount == 0) // Only warn on first retry to prevent spam
+			if (GSVTStreamingLogVerbosity > 1 || (GSVTStreamingLogVerbosity > 0 && PendingRequest.RetryCount == 0)) // Only warn on first retry to prevent spam
 			{
-				UE_LOG(LogSparseVolumeTextureStreamingManager, Warning, TEXT("SVT DDC request failed for '%s' (frame %i, mip %i). Retrying..."),
-					*Resources->ResourceName, PendingRequest.FrameIndex, PendingRequest.MipLevelIndex);
+				UE_LOG(LogSparseVolumeTextureStreamingManager, Warning, TEXT("SVT DDC request failed for '%s' (frame %i, tile offset %i, tile count %i). Retrying..."),
+					*SVTInfo->SVTName.ToString(), PendingRequest.FrameIndex, PendingRequest.TileOffset, PendingRequest.TileCount);
 			}
 
-			const FMipTileReadInfo MipTileReadInfo = GetMipTileReadInfo(SVTInfo->PerFrameInfo[PendingRequest.FrameIndex], PendingRequest.MipLevelIndex);
-			FCacheGetChunkRequest Request = BuildDDCRequest(*Resources, MipTileReadInfo.ReadOffset, MipTileReadInfo.ReadSize, PendingRequestIndex);
+			uint32 ReadOffset = 0;
+			uint32 ReadSize = 0;
+			Resources->StreamingMetaData.GetTileRangeMemoryOffsetSize(PendingRequest.TileOffset, PendingRequest.TileCount, ReadOffset, ReadSize);
+			FCacheGetChunkRequest Request = BuildDDCRequest(*Resources, ReadOffset, ReadSize, PendingRequestIndex);
 			const bool bBlocking = GSVTStreamingForceBlockingRequests || PendingRequest.bBlocking;
 			RequestDDCData(MakeArrayView(&Request, 1), bBlocking);
 
@@ -1305,12 +1195,17 @@ int32 FStreamingManager::DetermineReadyRequests()
 				if (!PendingRequest.Request.IsOk())
 				{
 					// Retry if IO request failed for some reason
-					const FMipTileReadInfo MipTileReadInfo = GetMipTileReadInfo(SVTInfo->PerFrameInfo[PendingRequest.FrameIndex], PendingRequest.MipLevelIndex);
-					UE_LOG(LogSparseVolumeTextureStreamingManager, Warning, TEXT("SVT IO request failed for %p (frame %i, mip %i, offset %i, size %i). Retrying..."),
-						PendingRequest.SparseVolumeTexture, PendingRequest.FrameIndex, PendingRequest.MipLevelIndex, MipTileReadInfo.ReadOffset, MipTileReadInfo.ReadSize);
+					uint32 ReadOffset = 0;
+					uint32 ReadSize = 0;
+					Resources->StreamingMetaData.GetTileRangeMemoryOffsetSize(PendingRequest.TileOffset, PendingRequest.TileCount, ReadOffset, ReadSize);
+					if (GSVTStreamingLogVerbosity > 0)
+					{
+						UE_LOG(LogSparseVolumeTextureStreamingManager, Warning, TEXT("SVT IO request failed for %s (frame %i, tile offset %i, tile count %i, offset %i, size %i). Retrying..."),
+							*SVTInfo->SVTName.ToString(), PendingRequest.FrameIndex, PendingRequest.TileOffset, PendingRequest.TileCount, ReadOffset, ReadSize);
+					}
 					
 					FBulkDataBatchRequest::FBatchBuilder Batch = FBulkDataBatchRequest::NewBatch(1);
-					Batch.Read(Resources->StreamableMipLevels, MipTileReadInfo.ReadOffset, MipTileReadInfo.ReadSize, AIOP_Low, PendingRequest.RequestBuffer, PendingRequest.Request);
+					Batch.Read(Resources->StreamableMipLevels, ReadOffset, ReadSize, AIOP_Low, PendingRequest.RequestBuffer, PendingRequest.Request);
 					(void)Batch.Issue();
 					break;
 				}
@@ -1340,30 +1235,36 @@ void FStreamingManager::InstallReadyRequests()
 
 	UploadTasks.Reset();
 	UploadTasks.Reserve(AsyncState.NumReadyRequests * 2 /*slack for splitting large uploads*/);
-	UploadCleanupTasks.Reset();
+	RequestsToCleanUp.Reset();
 
 #if WITH_EDITORONLY_DATA
 	TMap<const FResources*, const uint8*> ResourceToBulkPointer;
 #endif
 
-	// Do a second pass over all ready mip levels, claiming memory in the upload buffers and creating FUploadTasks
+	// Do a second pass over all ready requests, claiming memory in the upload buffers and creating FUploadTasks
 	const int32 StartPendingRequestIndex = (NextPendingRequestIndex + MaxPendingRequests - NumPendingRequests) % MaxPendingRequests;
 	for (int32 i = 0; i < AsyncState.NumReadyRequests; ++i)
 	{
 		const int32 PendingRequestIndex = (StartPendingRequestIndex + i) % MaxPendingRequests;
 		FPendingRequest& PendingRequest = PendingRequests[PendingRequestIndex];
 
-		FStreamingInfo* SVTInfo = FindStreamingInfo(PendingRequest.SparseVolumeTexture);
-		if (!SVTInfo || (SVTInfo->PerFrameInfo[PendingRequest.FrameIndex].LowestRequestedMipLevel > PendingRequest.MipLevelIndex))
+		// Skip if request was cancelled
+		if (!PendingRequest.IsValid())
 		{
-			PendingRequest.Reset();
-			continue; // Skip mip level install. SVT no longer exists or mip level was "streamed out" before it was even installed in the first place.
+			continue;
 		}
 
+		FStreamingInfo* SVTInfo = FindStreamingInfo(PendingRequest.SVTHandle);
+		check(SVTInfo);
+		const int32 FormatSizeA = GPixelFormats[SVTInfo->FormatA].BlockBytes;
+		const int32 FormatSizeB = GPixelFormats[SVTInfo->FormatB].BlockBytes;
+
 		FFrameInfo& FrameInfo = SVTInfo->PerFrameInfo[PendingRequest.FrameIndex];
 		const FResources* Resources = FrameInfo.Resources;
 		const FPageTopology& Topology = FrameInfo.Resources->Topology;
-		const FMipTileReadInfo MipTileReadInfo = GetMipTileReadInfo(FrameInfo, PendingRequest.MipLevelIndex);
+		uint32 ReadOffset = 0;
+		uint32 ReadSize = 0;
+		Resources->StreamingMetaData.GetTileRangeMemoryOffsetSize(PendingRequest.TileOffset, PendingRequest.TileCount, ReadOffset, ReadSize);
 
 		const uint8* SrcPtr = nullptr;
 		const uint8* SrcEndPtr = nullptr;
@@ -1380,8 +1281,8 @@ void FStreamingManager::InstallReadyRequests()
 			const uint8** BulkDataPtrPtr = ResourceToBulkPointer.Find(Resources);
 			if (BulkDataPtrPtr)
 			{
-				SrcPtr = *BulkDataPtrPtr + MipTileReadInfo.ReadOffset;
-				SrcEndPtr = SrcPtr + MipTileReadInfo.ReadSize;
+				SrcPtr = *BulkDataPtrPtr + ReadOffset;
+				SrcEndPtr = SrcPtr + ReadSize;
 			}
 			else
 			{
@@ -1389,7 +1290,7 @@ void FStreamingManager::InstallReadyRequests()
 				check(BulkData.IsBulkDataLoaded() && BulkData.GetBulkDataSize() > 0);
 				const uint8* BulkDataPtr = (const uint8*)BulkData.LockReadOnly();
 				ResourceToBulkPointer.Add(Resources, BulkDataPtr);
-				SrcPtr = BulkDataPtr + MipTileReadInfo.ReadOffset;
+				SrcPtr = BulkDataPtr + ReadOffset;
 				SrcEndPtr = BulkDataPtr + BulkData.GetBulkDataSize();
 			}
 		}
@@ -1403,18 +1304,19 @@ void FStreamingManager::InstallReadyRequests()
 			SrcEndPtr = SrcPtr + PendingRequest.RequestBuffer.DataSize();
 		}
 
-		check(SrcPtr);
-
-		const int32 FormatSizeA = GPixelFormats[SVTInfo->FormatA].BlockBytes;
-		const int32 FormatSizeB = GPixelFormats[SVTInfo->FormatB].BlockBytes;
-
-		for (uint32 TileIndex = MipTileReadInfo.TileOffset; TileIndex < (MipTileReadInfo.TileOffset + MipTileReadInfo.TileCount); ++TileIndex)
+		for (uint32 TileIndex = PendingRequest.TileOffset; TileIndex < (PendingRequest.TileOffset + PendingRequest.TileCount); ++TileIndex)
 		{
-			const FTileInfo TileInfo = Resources->StreamingMetaData.GetTileInfo(TileIndex, FormatSizeA, FormatSizeB);
+			SVTInfo->PerFrameInfo[PendingRequest.FrameIndex].TileIndexToPendingRequestIndex.Remove(TileIndex);
+			if (!SVTInfo->PerFrameInfo[PendingRequest.FrameIndex].StreamingTiles[TileIndex])
+			{
+				continue; // Skip tile install. Tile was "streamed out" before it was even installed in the first place.
+			}
 
-			check(TileInfo.Offset >= MipTileReadInfo.ReadOffset);
-			const uint8* TileSrcPtr = SrcPtr + (TileInfo.Offset - MipTileReadInfo.ReadOffset);
+			const FTileInfo TileInfo = Resources->StreamingMetaData.GetTileInfo(TileIndex, FormatSizeA, FormatSizeB);
 
+			check(SrcPtr);
+			check(TileInfo.Offset >= ReadOffset);
+			const uint8* TileSrcPtr = SrcPtr + (TileInfo.Offset - ReadOffset);
 			check((TileSrcPtr + TileInfo.Size) <= SrcEndPtr);
 
 			FTileUploader::FAddResult TileDataAddResult = SVTInfo->TileDataTexture->AddUpload(1, TileInfo.NumVoxels[0], TileInfo.NumVoxels[1]);
@@ -1431,34 +1333,19 @@ void FStreamingManager::InstallReadyRequests()
 			TileDataTask.VoxelDataSizes = TileInfo.VoxelDataSizes;
 			TileDataTask.VoxelDataBaseOffsets = TileDataAddResult.TileDataBaseOffsets;
 			TileDataTask.PhysicalTileCoord = FrameInfo.TileAllocations[TileIndex];
+
+			FrameInfo.ResidentTiles[TileIndex] = true;
 		}
 
 		// Cleanup
 		{
-			UploadCleanupTasks.Add(&PendingRequest);
+			RequestsToCleanUp.Add(PendingRequestIndex);
 		}
-	
-#if SVT_STREAMING_LOG_VERBOSE
-		UE_LOG(LogSparseVolumeTextureStreamingManager, Display, TEXT("(%i)%i Install Frame %i OldResMip %i, NewResMip %i, ReqMip %i"),
-			PendingRequest.IssuedInFrame, PendingRequest,
-			PendingRequest.FrameIndex,
-			FrameInfo.LowestResidentMipLevel, PendingRequest.MipLevelIndex,
-			FrameInfo.LowestRequestedMipLevel);
-#endif
-
-		check(FrameInfo.LowestResidentMipLevel == (PendingRequest.MipLevelIndex + 1));
-		FrameInfo.LowestResidentMipLevel = PendingRequest.MipLevelIndex;
 
 		InvalidatedSVTFrames.Add(&FrameInfo);
-
-		const int32 LRUNodeIndex = PendingRequest.FrameIndex * SVTInfo->NumMipLevelsGlobal + PendingRequest.MipLevelIndex;
-		SVTInfo->LRUNodes[LRUNodeIndex].PendingRequestIndex = INDEX_NONE;
-
-		const FPageTopology::FMip& TopologyMipInfo = Topology.MipInfo[PendingRequest.MipLevelIndex];
-		FrameInfo.ResidentPagesNew.SetRange(TopologyMipInfo.PageOffset, TopologyMipInfo.PageCount, true);
 	}
 
-	// Do all the memcpy's in parallel
+	// Do all the memcpys in parallel
 	ParallelFor(TEXT("SVT::UploadTileDataTasks"), UploadTasks.Num(), 8, [&](int32 TaskIndex)
 	{
 		TRACE_CPUPROFILER_EVENT_SCOPE(SVT::StreamingTileDataUpload);
@@ -1482,22 +1369,21 @@ void FStreamingManager::InstallReadyRequests()
 		FMemory::Memcpy(TileDataTask.DstPhysicalTileCoordsPtr, &TileDataTask.PhysicalTileCoord, sizeof(uint32));
 	});
 
-	ParallelFor(UploadCleanupTasks.Num(), [&](int32 TaskIndex)
+	// Clean up requests
+	for (int32 PendingRequestIndex : RequestsToCleanUp)
 	{
-		TRACE_CPUPROFILER_EVENT_SCOPE(SVT::StreamingUploadCleanupTask);
-
-		FPendingRequest* PendingRequest = UploadCleanupTasks[TaskIndex];
+		FPendingRequest& PendingRequest = PendingRequests[PendingRequestIndex];
 #if WITH_EDITORONLY_DATA
-		PendingRequest->SharedBuffer.Reset();
+		PendingRequest.SharedBuffer.Reset();
 #endif
-		if (!PendingRequest->Request.IsNone())
+		if (!PendingRequest.Request.IsNone())
 		{
-			check(PendingRequest->Request.IsCompleted());
-			PendingRequest->Request.Reset();
+			check(PendingRequest.Request.IsCompleted());
+			PendingRequest.Request.Reset();
 		}
-	});
+	}
 
-#if DO_CHECK // Clear processed pending mip levels for better debugging
+#if DO_CHECK // Clear processed pending requests for better debugging
 	for (int32 i = 0; i < AsyncState.NumReadyRequests; ++i)
 	{
 		const int32 PendingRequestIndex = (StartPendingRequestIndex + i) % MaxPendingRequests;
@@ -1525,9 +1411,21 @@ void FStreamingManager::PatchPageTable(FRDGBuilder& GraphBuilder)
 	{
 		FFrameInfo& FrameInfo = *FrameInfoPtr;
 		const FPageTopology& Topology = FrameInfo.Resources->Topology;
+		const uint32 NumPages = Topology.NumPages();
 
-		// Get all pages that were streamed in or out in this streaming update.
-		TBitArray<> ResidentPagesDiff = TBitArray<>::BitwiseXOR(FrameInfo.ResidentPages, FrameInfo.ResidentPagesNew, EBitwiseOperatorFlags::MaxSize);
+		// Mark all pages with resident tiles as resident. For now we also enforce that a page is only resident if its parent is resident, but in theory this is not needed.
+		ResidentPagesNew.SetNumUninitialized(NumPages);
+		ResidentPagesNew.SetRange(0, NumPages, false);
+		for (uint32 PageIndex = 0; PageIndex < NumPages; ++PageIndex)
+		{
+			const uint32 ParentIndex = Topology.ParentIndices[PageIndex];
+			const uint32 TileIndex = Topology.TileIndices[PageIndex];
+			ResidentPagesNew[PageIndex] = (ParentIndex == INDEX_NONE || ResidentPagesNew[ParentIndex]) && FrameInfo.ResidentTiles[TileIndex];
+		}
+
+		// Get all pages that were streamed in or out in this streaming update by doing a bitwise XOR with the resident pages from before the current streaming update.
+		ResidentPagesDiff = FrameInfo.ResidentPages;
+		ResidentPagesDiff.CombineWithBitwiseXOR(ResidentPagesNew, EBitwiseOperatorFlags::MaxSize);
 		
 		// Initialize InvalidatedPages with the diff. In the following loop, we then find all descendants of the newly streamed in/out pages and also mark them as invalidated.
 		FrameInfo.InvalidatedPages = ResidentPagesDiff;
@@ -1554,6 +1452,9 @@ void FStreamingManager::PatchPageTable(FRDGBuilder& GraphBuilder)
 			}
 		}
 
+		// Update the bitset of resident pages
+		FrameInfo.ResidentPages = ResidentPagesNew;
+		
 		NumUpdates += FrameInfo.InvalidatedPages.CountSetBits();
 	}
 
@@ -1608,8 +1509,9 @@ void FStreamingManager::PatchPageTable(FRDGBuilder& GraphBuilder)
 					PageTableUpdater->Add_GetRef(FrameInfo.PageTableTexture, MipLevel, NumUpdatesThisMip, DstCoordsPtr, DstEntryPtr);
 				}
 
+				// Get the tile data texture coordinate we need to write into the GPU page table texture.
 				uint32 PageTableEntry = 0;
-				if (FrameInfo.ResidentPagesNew[PageIndex])
+				if (FrameInfo.ResidentPages[PageIndex])
 				{
 					// This page is already resident, so we can simply use its value from TileAllocations
 					PageTableEntry = FrameInfo.TileAllocations[Topology.TileIndices[PageIndex]];
@@ -1624,7 +1526,7 @@ void FStreamingManager::PatchPageTable(FRDGBuilder& GraphBuilder)
 					while (ParentPageIndex != INDEX_NONE)
 					{
 						// The parent page is resident in GPU memory, so we can use it's cached value in PageEntries.
-						if (FrameInfo.ResidentPagesNew[ParentPageIndex])
+						if (FrameInfo.ResidentPages[ParentPageIndex])
 						{
 							PageTableEntry = FrameInfo.TileAllocations[Topology.TileIndices[ParentPageIndex]];
 							check(PageTableEntry);
@@ -1642,8 +1544,6 @@ void FStreamingManager::PatchPageTable(FRDGBuilder& GraphBuilder)
 				reinterpret_cast<uint32*>(DstEntryPtr)[MipUpdateWriteIndex] = PageTableEntry;
 				++MipUpdateWriteIndex;
 			}
-
-			FrameInfo.ResidentPages = FrameInfo.ResidentPagesNew;
 		}
 
 		PageTableUpdater->Apply(GraphBuilder);
@@ -1652,37 +1552,23 @@ void FStreamingManager::PatchPageTable(FRDGBuilder& GraphBuilder)
 	InvalidatedSVTFrames.Reset();
 }
 
-FStreamingManager::FStreamingInfo* FStreamingManager::FindStreamingInfo(UStreamableSparseVolumeTexture* Key)
+FStreamingManager::FStreamingInfo* FStreamingManager::FindStreamingInfo(uint16 SparseVolumeTextureHandle)
 {
-	TUniquePtr<FStreamingInfo>* SVTInfoPtr = StreamingInfo.Find(Key);
-	check(!SVTInfoPtr || SVTInfoPtr->Get());
-	return SVTInfoPtr ? SVTInfoPtr->Get() : nullptr;
+	if (!StreamingInfo.IsValidIndex(SparseVolumeTextureHandle) || !StreamingInfo.IsAllocated(SparseVolumeTextureHandle))
+	{
+		return nullptr;
+	}
+	TUniquePtr<FStreamingInfo>& SVTInfo = StreamingInfo[SparseVolumeTextureHandle];
+	return SVTInfo.Get();
 }
 
-FStreamingManager::FMipTileReadInfo FStreamingManager::GetMipTileReadInfo(const FFrameInfo& FrameInfo, int32 MipLevel)
+FStreamingManager::FStreamingInfo* FStreamingManager::FindStreamingInfo(UStreamableSparseVolumeTexture* SparseVolumeTexture)
 {
-	check(FrameInfo.Resources->Topology.MipInfo.IsValidIndex(MipLevel));
-	const bool bIsRootMipLevel = (MipLevel + 1) == FrameInfo.NumMipLevels;
-	const FPageTopology::FMip& MipInfo = FrameInfo.Resources->Topology.MipInfo[MipLevel];
-
-	uint32 TileRangeOffset = 0;
-	uint32 TileRangeCount = 0;
-	FrameInfo.Resources->Topology.GetTileRange(MipInfo.PageOffset, MipInfo.PageCount, TileRangeOffset, TileRangeCount);
-
-	const uint32 ReadOffsetLogical = FrameInfo.Resources->StreamingMetaData.TileDataOffsets[TileRangeOffset];
-	// TileDataOffsets are expressed as if all data is stored in the same buffer, when the root tile is actually stored separately. The root tile can't stream, so we simply subtract its size from the offset.
-	const uint32 RootTileSize = FrameInfo.Resources->RootData.Num();
-	check(RootTileSize == 0 || RootTileSize == FrameInfo.Resources->StreamingMetaData.GetTileMemorySize(0));
-	const uint32 ReadOffset = bIsRootMipLevel ? ReadOffsetLogical : (ReadOffsetLogical - RootTileSize);
-	const uint32 ReadSize = FrameInfo.Resources->StreamingMetaData.TileDataOffsets[TileRangeOffset + TileRangeCount] - ReadOffsetLogical;
-
-	FMipTileReadInfo ReadInfo;
-	ReadInfo.TileOffset = TileRangeOffset;
-	ReadInfo.TileCount = TileRangeCount;
-	ReadInfo.ReadOffset = ReadOffset;
-	ReadInfo.ReadSize = ReadSize;
-
-	return ReadInfo;
+	if (uint16* HandlePtr = SparseVolumeTextureToHandle.Find(SparseVolumeTexture))
+	{
+		return FindStreamingInfo(*HandlePtr);
+	}
+	return nullptr;
 }
 
 #if WITH_EDITORONLY_DATA
@@ -1726,7 +1612,7 @@ void FStreamingManager::RequestDDCData(TConstArrayView<UE::DerivedData::FCacheGe
 				}
 
 				FPendingRequest& PendingRequest = PendingRequests[PendingRequestIndex];
-				check(PendingRequest.SparseVolumeTexture); // A valid PendingRequest should have a non-nullptr here
+				check(PendingRequest.IsValid());
 
 				if (Response.Status == EStatus::Ok)
 				{
@@ -1748,5 +1634,129 @@ void FStreamingManager::RequestDDCData(TConstArrayView<UE::DerivedData::FCacheGe
 
 #endif // WITH_EDITORONLY_DATA
 
+void FTileAllocator::Init(const FIntVector3& InResolutionInTiles)
+{
+	check(InResolutionInTiles.X > 0 && InResolutionInTiles.Y > 0 && InResolutionInTiles.Z > 0);
+	FreeHeap.Clear();
+	ResolutionInTiles = InResolutionInTiles;
+	TileCapacity = ResolutionInTiles.X * ResolutionInTiles.Y * ResolutionInTiles.Z;
+	NumAllocated = 1; // Implicitly allocate null tile
+	Allocations.SetNumZeroed(TileCapacity);
+	Allocations[0] = FAllocation(INDEX_NONE, INDEX_NONE, true, true);
+	
+	// Populate free heap with packed tile coordinates, ready for allocation
+	for (int32 TileZ = 0; TileZ < ResolutionInTiles.Z; ++TileZ)
+	{
+		for (int32 TileY = 0; TileY < ResolutionInTiles.Y; ++TileY)
+		{
+			for (int32 TileX = 0; TileX < ResolutionInTiles.X; ++TileX)
+			{
+				// Skip (0,0,0): That's where we implicitly allocated the null tile
+				const uint32 PackedCoord = SVT::PackPageTableEntry(FIntVector3(TileX, TileY, TileZ));
+				if (PackedCoord)
+				{
+					FreeHeap.Add(0, PackedCoord);
+				}
+			}
+		}
+	}
+}
+
+uint32 FTileAllocator::Allocate(uint32 UpdateIndex, uint32 FreeThreshold, uint16 FrameIndex, uint32 TileIndexInFrame, uint32 TilePriority, bool bLocked, FAllocation& OutPreviousAllocation)
+{
+	if (FreeHeap.IsEmpty())
+	{
+		// No free tiles available!
+		return INDEX_NONE;
+	}
+
+	const uint32 TileCoord = FreeHeap.Top();
+
+	// There are non-locked tiles in the heap, but they've all been recently requested, so we can't reuse them for this allocation!
+	const uint32 LastRequested = uint32(FreeHeap.GetKey(TileCoord) >> 32u);
+	if ((LastRequested + FreeThreshold) >= UpdateIndex)
+	{
+		return INDEX_NONE;
+	}
+
+	const FIntVector3 UnpackedTileCoord = SVT::UnpackPageTableEntry(TileCoord);
+	const uint32 LinearTileCoord = (UnpackedTileCoord.Z * ResolutionInTiles.Y * ResolutionInTiles.X) + (UnpackedTileCoord.Y * ResolutionInTiles.X) + UnpackedTileCoord.X;
+	FAllocation& Allocation = Allocations[LinearTileCoord];
+
+	// Tile was already allocated
+	if (Allocation.bIsAllocated)
+	{
+		check(!Allocation.bIsLocked);
+		OutPreviousAllocation = Allocation;
+		--NumAllocated;
+	}
+	else
+	{
+		// Tile wasn't allocated, but for consistency and debugging, simply zero the out parameter.
+		FMemory::Memzero(OutPreviousAllocation);
+	}
+
+	// Mark as used by the new SVT frame
+	Allocation.FrameIndex = FrameIndex;
+	Allocation.TileIndexInFrame = TileIndexInFrame;
+	Allocation.bIsLocked = bLocked;
+	Allocation.bIsAllocated = true;
+
+	if (bLocked)
+	{
+		// Allocation is locked: remove it from the FreeHeap so it can't ever be reused unless manually freed.
+		FreeHeap.Pop();
+	}
+	else
+	{
+		// Allocation is not locked: update it with the new key
+		const uint64 HeapKey = (((uint64)UpdateIndex) << 32u) + TilePriority;
+		FreeHeap.Update(HeapKey, TileCoord);
+	}
+
+	++NumAllocated;
+	check(NumAllocated <= TileCapacity);
+
+	return TileCoord;
+}
+
+void FTileAllocator::UpdateUsage(uint32 UpdateIndex, uint32 TileCoord, uint32 TilePriority)
+{
+	if (FreeHeap.IsPresent(TileCoord))
+	{
+		// Mark tile as requested in the given UpdateIndex
+		const uint64 HeapKey = (((uint64)UpdateIndex) << 32u) + TilePriority;
+		FreeHeap.Update(HeapKey, TileCoord);
+	}
+}
+
+void FTileAllocator::Free(uint32 TileCoord)
+{
+	TileCoord &= PhysicalCoordMask; // Clear the upper 8 bit in case the caller wrote "user data" (mip level index) there.
+	check(TileCoord);
+
+	// Make sure the passed in TileCoord actually has a valid value.
+	const FIntVector3 UnpackedTileCoord = SVT::UnpackPageTableEntry(TileCoord);
+	check(IsInBounds(UnpackedTileCoord, FIntVector3::ZeroValue, ResolutionInTiles));
+	
+	// Zero the allocation info
+	const uint32 LinearTileCoord = (UnpackedTileCoord.Z * ResolutionInTiles.Y * ResolutionInTiles.X) + (UnpackedTileCoord.Y * ResolutionInTiles.X) + UnpackedTileCoord.X;
+	FMemory::Memzero(Allocations[LinearTileCoord]);
+
+	if (FreeHeap.IsPresent(TileCoord))
+	{
+		// Tile was already in the heap; simply put it at the front so it'll be used in the next allocation.
+		FreeHeap.Update(0, TileCoord);
+	}
+	else
+	{
+		// Tile wasn't in the heap (must have been a locked allocation), so we add it now.
+		FreeHeap.Add(0, TileCoord);
+	}
+
+	check(NumAllocated > 1); // Null tile is permanently allocated
+	--NumAllocated;
+}
+
 }
 }
diff --git a/Engine/Source/Runtime/Engine/Private/SparseVolumeTexture/SparseVolumeTextureStreamingManager.h b/Engine/Source/Runtime/Engine/Private/SparseVolumeTexture/SparseVolumeTextureStreamingManager.h
--- a/Engine/Source/Runtime/Engine/Private/SparseVolumeTexture/SparseVolumeTextureStreamingManager.h
+++ b/Engine/Source/Runtime/Engine/Private/SparseVolumeTexture/SparseVolumeTextureStreamingManager.h
@@ -13,6 +13,7 @@
 #include "Containers/Map.h"
 #include "Containers/StaticArray.h"
 #include "Containers/Union.h"
+#include "Containers/BinaryHeap.h"
 #include "RenderGraphBuilder.h"
 
 DECLARE_LOG_CATEGORY_EXTERN(LogSparseVolumeTextureStreamingManager, Log, All);
@@ -41,44 +42,48 @@ class FTileDataTexture;
 class FTileUploader;
 class FPageTableUpdater;
 
-// Uniquely identifies the mip level of a frame in a static or animated SparseVolumeTexture
-struct FMipLevelKey
+// Helper class for managing slots in the tile data texture (SVT streaming pool). It uses a priority queue to reuse older slots that haven't been referenced for some (render) frames.
+class FTileAllocator
 {
-	UStreamableSparseVolumeTexture* SVT = nullptr; // This struct is used on the rendering thread. Do not dereference!
-	uint16 FrameIndex = INDEX_NONE;
-	uint16 MipLevelIndex = INDEX_NONE;
-
-	friend FORCEINLINE uint32 GetTypeHash(const FMipLevelKey& Key)
-	{
-		return HashCombine(GetTypeHash(Key.SVT), GetTypeHash(((uint32)Key.FrameIndex << 16u) | (uint32)Key.MipLevelIndex));
-	}
+public:
+	static constexpr uint32 PhysicalCoordMask = (1u << 24u) - 1u; // Lower 24 bits are used for storing XYZ in 8 bit each. Upper 8 bit can be used by the caller. 
 
-	FORCEINLINE bool operator==(const FMipLevelKey& Other) const
+	// Describes a slot in the tile data texture and the current tile that is mapped to it.
+	struct FAllocation
 	{
-		return (SVT == Other.SVT) && (FrameIndex == Other.FrameIndex) && (MipLevelIndex == Other.MipLevelIndex);
-	}
+		uint32 TileIndexInFrame;	// Index of the tile in the array of tiles of the SVT frame.
+		uint16 FrameIndex;			// Index of the SVT frame this tile belongs to.
+		uint16 bIsLocked : 1;		// Locked tiles are not in the priority queue and will never be streamed out/replaced unless manually freed.
+		uint16 bIsAllocated : 1;	// Whether this slot is currently occupied by a tile.
+
+		FAllocation() : TileIndexInFrame(0), FrameIndex(0), bIsLocked(0), bIsAllocated(0) {}
+		FAllocation(uint16 InFrameIndex, uint32 InTileIndex, bool bInIsLocked, bool bInIsAllocated) 
+			: TileIndexInFrame(InTileIndex), FrameIndex(InFrameIndex), bIsLocked(bInIsLocked), bIsAllocated(bInIsAllocated) {}
+	};
 
-	FORCEINLINE bool operator!=(const FMipLevelKey& Other) const
-	{
-		return !(*this == Other);
-	}
+	void Init(const FIntVector3& InResolutionInTiles);
 
-	FORCEINLINE bool operator<(const FMipLevelKey& Other) const
-	{
-		return SVT != Other.SVT ? SVT < Other.SVT : FrameIndex != Other.FrameIndex ? FrameIndex < Other.FrameIndex : MipLevelIndex > Other.MipLevelIndex;
-	}
-};
-
-struct FStreamingRequest
-{
-	static constexpr uint32 BlockingPriority = 0xFFFFFFFFu;
-	FMipLevelKey Key;
-	uint32 Priority; // A higher value means a higher priority, 0xFFFFFF means blocking
+	// Allocate a tile. Returns a packed 8|8|8 coordinate into the tile data texture or INDEX_NONE if there is no more space.
+	// UpdateIndex is the index of the current tick/update.
+	// FreeThreshold is an additional buffer to how many updates must have passed until a tile may be reused.
+	// FrameIndex is the index of the SVT frame and TileIndexInFrame is the index of the tile for which we try to allocate a slot.
+	// TilePriority is used to sort tiles in the priority queue. Higher values make the tile stream out after other tiles with lower values.
+	// If bLocked is true, then the tile will never be automatically streamed out.
+	// OutPreviousAllocation is a description of the old tile that occupied the newly allocated slot.
+	uint32 Allocate(uint32 UpdateIndex, uint32 FreeThreshold, uint16 FrameIndex, uint32 TileIndexInFrame, uint32 TilePriority, bool bLocked, FAllocation& OutPreviousAllocation);
+	
+	// Marks a tile as still in use. See Allocate() for a description of the parameters.
+	void UpdateUsage(uint32 UpdateIndex, uint32 TileCoord, uint32 TilePriority);
+	
+	// Frees a given tile. Expects the output of Allocate(), i.e. a packed 8|8|8 coordinate. Freed tiles are placed at the front of the priority queue for tile reuse.
+	void Free(uint32 TileCoord);
 
-	FORCEINLINE bool operator<(const FStreamingRequest& Other) const
-	{
-		return Key != Other.Key ? Key < Other.Key : Priority > Other.Priority;
-	}
+private:
+	FBinaryHeap<uint64, uint32> FreeHeap;
+	TArray<FAllocation> Allocations; // One entry for every slot in the tile data texture. Maps to the current SVT frame occupying that tile
+	FIntVector3 ResolutionInTiles = FIntVector3::ZeroValue;
+	uint32 TileCapacity = 0;
+	uint32 NumAllocated = 0;
 };
 
 class FStreamingManager : public FRenderResource, public IStreamingManager
@@ -124,50 +129,39 @@ private:
 		const FResources* Resources;						// Initialized on the game thread
 		FTextureRenderResources* TextureRenderResources;	// Initialized on the game thread
 
-		int32 NumMipLevels; // Number of actual mip levels of this frame. Depending on virtual volume extents, frames can have different numbers of levels.
-		int32 LowestRequestedMipLevel; // Lowest mip level that should be resident. Can be lower than LowestResidentMipLevel when streaming in new mips. Stream-out is instant, so it should never be higher.
-		int32 LowestResidentMipLevel; // Actually resident on the GPU
+		int32 NumMipLevels = 0; // Number of actual mip levels of this frame. Depending on virtual volume extents, frames can have different numbers of levels.
 		TArray<uint32> TileAllocations; // TileAllocations[PhysicalTileIndex]
 		TBitArray<> ResidentPages; //  One bit for every (non-zero) page (in the sparse page octree) for all mip levels, starting at the highest mip
-		TBitArray<> ResidentPagesNew; // Reflects changes made during an update of the streaming system
-		TBitArray<> InvalidatedPages; // All pages that need updated page table entries
+		TBitArray<> InvalidatedPages; // Temporarly needed when patching the page table. Marks all pages that require an update
+		TBitArray<> ResidentTiles; // Tiles that are actually resident in GPU memory
+		TBitArray<> StreamingTiles; // Tiles that have logically been streamed in but may not actually be resident in GPU memory yet because the IO request hasn't finished
+		TMap<uint32, uint32> TileIndexToPendingRequestIndex;
 		TRefCountPtr<IPooledRenderTarget> PageTableTexture;
 	};
 
-	// Used to keep track of the least-recently-used order of mip levels. This is done per mip level and not per frame, so that lower resolution mip levels are kept in memory for longer.
-	struct FLRUNode : public TIntrusiveDoubleLinkedListNode<FLRUNode>
-	{
-		FLRUNode* NextHigherMipLevel = nullptr;
-		uint32 LastRequested = INDEX_NONE;
-		uint32 RefCount = 0; // Keep track of how many lower mip levels have a dependency on this one
-		uint32 PendingRequestIndex = INDEX_NONE;
-		int16 FrameIndex = INDEX_NONE;
-		int16 MipLevelIndex = INDEX_NONE;
-	};
-
 	// Streaming manager internal data for each SVT
 	struct FStreamingInfo
 	{
-		EPixelFormat FormatA;
-		EPixelFormat FormatB;
-		FVector4f FallbackValueA;
-		FVector4f FallbackValueB;
-		int32 NumMipLevelsGlobal; // Maximum number of mip levels in the sequence
-		uint32 LastRequested; // Last update index when some frame in this SVT was requested
+		uint16 SVTHandle = INDEX_NONE;
+		FName SVTName = FName();
+		EPixelFormat FormatA = PF_Unknown;
+		EPixelFormat FormatB = PF_Unknown;
+		FVector4f FallbackValueA = FVector4f();
+		FVector4f FallbackValueB = FVector4f();
 		TArray<FFrameInfo> PerFrameInfo;
-		TArray<FLRUNode> LRUNodes;
-		TArray<TIntrusiveDoubleLinkedList<FLRUNode>> PerMipLRULists; // PerMipLRULists[MipLevel]
 		TArray<FStreamingWindow> StreamingWindows;
 
+		FTileAllocator TileAllocator;
 		TUniquePtr<FTileDataTexture> TileDataTexture;
 	};
 
-	// Represents an IO request for a mip level
+	// Represents an IO request for tile(s)
 	struct FPendingRequest
 	{
-		UStreamableSparseVolumeTexture* SparseVolumeTexture = nullptr; // Do not dereference!
-		int32 FrameIndex = INDEX_NONE;
-		int32 MipLevelIndex = INDEX_NONE;
+		uint16 SVTHandle = INDEX_NONE;
+		uint16 FrameIndex = INDEX_NONE;
+		uint32 TileOffset = INDEX_NONE;
+		uint32 TileCount = 0;
 #if WITH_EDITORONLY_DATA
 		FSharedBuffer SharedBuffer;
 		enum class EState
@@ -187,11 +181,14 @@ private:
 		uint32 IssuedInFrame = 0;
 		bool bBlocking = false;
 
+		bool IsValid() const { return SVTHandle != uint16(INDEX_NONE); }
+
 		void Reset()
 		{
-			SparseVolumeTexture = nullptr;
+			SVTHandle = INDEX_NONE;
 			FrameIndex = INDEX_NONE;
-			MipLevelIndex = INDEX_NONE;
+			TileOffset = INDEX_NONE;
+			TileCount = INDEX_NONE;
 #if WITH_EDITORONLY_DATA
 			SharedBuffer.Reset();
 			State = EState::None;
@@ -231,25 +228,53 @@ private:
 	struct FNewSparseVolumeTextureInfo
 	{
 		UStreamableSparseVolumeTexture* SVT = nullptr; // Do not dereference!
+		FName SVTName = FName();
 		EPixelFormat FormatA = PF_Unknown;
 		EPixelFormat FormatB = PF_Unknown;
 		FVector4f FallbackValueA = FVector4f(0.0f, 0.0f, 0.0f, 0.0f);
 		FVector4f FallbackValueB = FVector4f(0.0f, 0.0f, 0.0f, 0.0f);
-		int32 NumMipLevelsGlobal = 0;
 		TArray<FFrameInfo> FrameInfo; // Only Resources and TextureRenderResources are initialized
 	};
 
-	// Helper struct representing a range of tiles within the cooked data corresponding to a given mip level
-	struct FMipTileReadInfo
+	// Uniquely identifies the mip level of a frame in a static or animated SparseVolumeTexture
+	struct FFrameKey
+	{
+		uint16 SVTHandle = INDEX_NONE;
+		uint16 FrameIndex = INDEX_NONE;
+
+		friend FORCEINLINE uint32 GetTypeHash(const FFrameKey& Key) { return HashCombine(GetTypeHash(Key.SVTHandle), GetTypeHash(Key.FrameIndex)); }
+		FORCEINLINE bool operator==(const FFrameKey& Other) const { return (SVTHandle == Other.SVTHandle) && (FrameIndex == Other.FrameIndex); }
+		FORCEINLINE bool operator!=(const FFrameKey& Other) const { return !(*this == Other); }
+		FORCEINLINE bool operator<(const FFrameKey& Other) const { return SVTHandle != Other.SVTHandle ? SVTHandle < Other.SVTHandle : FrameIndex < Other.FrameIndex; }
+	};
+
+	// The payload associated with a request for a given frame of a SVT
+	struct FRequestPayload
+	{
+		uint16 MipLevelMask = 0; // Bitmask of requested mip levels
+		float LowestMipFraction = 0.0f; // Setting this to a value between 0 and less than 1 signifies that only a certain percentage of pages in this mip level should be streamed
+		TStaticArray<uint8, 16> Priorities = TStaticArray<uint8, 16>(InPlace, 0); // Priority for each mip level with a set bit in MipLevelMask
+	};
+
+	struct FStreamingRequest
+	{
+		FFrameKey Key;
+		FRequestPayload Payload;
+	};
+
+	// Represents a contiguous range of tiles in a given frame of a SVT. This is the output of filtering FStreamingRequest and is used to generate FPendingRequests.
+	struct FTileRange
 	{
-		uint32 TileOffset;
-		uint32 TileCount;
-		uint32 ReadOffset;
-		uint32 ReadSize;
+		uint16 SVTHandle = INDEX_NONE;
+		uint16 FrameIndex = INDEX_NONE;
+		uint32 TileOffset = INDEX_NONE;
+		uint32 TileCount = 0;
+		uint8 Priority = 0; // Higher value means higher priority
 	};
 
-	TMap<UStreamableSparseVolumeTexture*, TUniquePtr<FStreamingInfo>> StreamingInfo; // Do not dereference the key! We just read the pointer itself.
-	TMap<FMipLevelKey, uint32> RequestsHashTable;
+	TMap<UStreamableSparseVolumeTexture*, uint16> SparseVolumeTextureToHandle; // Do not dereference the key! We just read the pointer itself.
+	TSparseArray<TUniquePtr<FStreamingInfo>> StreamingInfo; 
+	TMap<FFrameKey, FRequestPayload> RequestsHashTable;
 	TArray<FPendingRequest> PendingRequests;
 #if WITH_EDITORONLY_DATA
 	TUniquePtr<UE::DerivedData::FRequestOwner> RequestOwner;
@@ -265,30 +290,28 @@ private:
 	uint32 NextUpdateIndex = 1;
 
 	// Transient lifetime
-	TArray<FStreamingRequest> ParentRequestsToAdd;
 	TSet<FTileDataTexture*> TileDataTexturesToUpdate;
-	TSet<FFrameInfo*> InvalidatedSVTFrames; // Set of SVT frames where pages have been streamed in or out. Used in PatchPageTable().
-	TArray<FStreamingRequest> PrioritizedRequestsHeap;
-	TArray<FStreamingRequest> SelectedRequests;
-	TArray<FTileDataTask> UploadTasks; // accessed on the async thread
-	TArray<FPendingRequest*> UploadCleanupTasks; // accessed on the async thread
+	TSet<FFrameInfo*> InvalidatedSVTFrames; // Set of SVT frames where tiles have been streamed in or out. Used in PatchPageTable().
+	TArray<FTileRange> TileRangesToStream; // Output of FilterRequests(), consumed in IssueRequests().
+	TArray<FTileDataTask> UploadTasks; // Accessed on the async thread
+	TArray<int32> RequestsToCleanUp; // Accessed on the async thread. Stores indices into PendingRequests.
+	TBitArray<> ResidentPagesNew; // Used as temporary memory in PatchPageTable()
+	TBitArray<> ResidentPagesDiff; // Used as temporary memory in PatchPageTable()
 
 	void AddInternal(FRDGBuilder& GraphBuilder, FNewSparseVolumeTextureInfo&& NewSVTInfo);
 	void RemoveInternal(UStreamableSparseVolumeTexture* SparseVolumeTexture);
-	bool AddRequest(const FStreamingRequest& Request); // Returns true when the request is new or it has been updated with a higher priority
+	void AddRequest(const FStreamingRequest& Request);
 	void AsyncUpdate();
-	void AddParentRequests(); // Add requests for all parent mip levels
-	void SelectHighestPriorityRequestsAndUpdateLRU(int32 MaxSelectedRequests);
-	void IssueRequests(int32 MaxSelectedRequests);
-	void StreamOutMipLevel(FStreamingInfo* SVTInfo, FLRUNode* LRUNode);
+	void FilterRequests();
+	void IssueRequests();
 	int32 DetermineReadyRequests();
 	void InstallReadyRequests();
 	void PatchPageTable(FRDGBuilder& GraphBuilder); // Patches the page table to reflect streamed in/out pages and to ensure non-resident mip levels fall back to coarser mip level tile data
-	FStreamingInfo* FindStreamingInfo(UStreamableSparseVolumeTexture* Key); // Returns nullptr if the key can't be found
-	static FMipTileReadInfo GetMipTileReadInfo(const FFrameInfo& FrameInfo, int32 MipLevel);
+	FStreamingInfo* FindStreamingInfo(uint16 SparseVolumeTextureHandle); // Returns nullptr if the key can't be found
+	FStreamingInfo* FindStreamingInfo(UStreamableSparseVolumeTexture* SparseVolumeTexture); // Returns nullptr if the key can't be found
 
 #if WITH_EDITORONLY_DATA
-	UE::DerivedData::FCacheGetChunkRequest BuildDDCRequest(const FResources& Resources, uint64 ReadOffset, uint64 ReadSize, uint32 PendingRequestIndex);
+	UE::DerivedData::FCacheGetChunkRequest BuildDDCRequest(const FResources& Resources, uint64 ReadOffset, uint64 ReadSize, uint32 PendingMipLevelIndex);
 	void RequestDDCData(TConstArrayView<UE::DerivedData::FCacheGetChunkRequest> DDCRequests, bool bBlocking);
 #endif
 };
diff --git a/Engine/Source/Runtime/Engine/Private/SparseVolumeTexture/SparseVolumeTextureTileDataTexture.cpp b/Engine/Source/Runtime/Engine/Private/SparseVolumeTexture/SparseVolumeTextureTileDataTexture.cpp
--- a/Engine/Source/Runtime/Engine/Private/SparseVolumeTexture/SparseVolumeTextureTileDataTexture.cpp
+++ b/Engine/Source/Runtime/Engine/Private/SparseVolumeTexture/SparseVolumeTextureTileDataTexture.cpp
@@ -147,25 +147,6 @@ FTileDataTexture::FTileDataTexture(const FIntVector3& InResolutionInTiles, EPixe
 	const int64 MaxTileDataTextureResourceSize = GetMaxTileDataTextureResourceSize(MaxFormatSize);
 	check(((int64)Resolution.X * (int64)Resolution.Y * (int64)Resolution.Z * (int64)GPixelFormats[FormatA].BlockBytes) <= MaxTileDataTextureResourceSize);
 	check(((int64)Resolution.X * (int64)Resolution.Y * (int64)Resolution.Z * (int64)GPixelFormats[FormatB].BlockBytes) <= MaxTileDataTextureResourceSize);
-
-	TileCoords.SetNum(PhysicalTilesCapacity);
-
-	int32 TileCoordsIndex = 0;
-	for (int32 Z = 0; Z < ResolutionInTiles.Z; ++Z)
-	{
-		for (int32 Y = 0; Y < ResolutionInTiles.Y; ++Y)
-		{
-			for (int32 X = 0; X < ResolutionInTiles.X; ++X)
-			{
-				uint32 PackedCoord = 0;
-				PackedCoord |= (X & 0xFFu);
-				PackedCoord |= (Y & 0xFFu) << 8u;
-				PackedCoord |= (Z & 0xFFu) << 16u;
-				TileCoords[TileCoordsIndex++] = PackedCoord;
-			}
-		}
-	}
-	check(TileCoordsIndex == PhysicalTilesCapacity);
 }
 
 void FTileDataTexture::BeginReserveUpload()
diff --git a/Engine/Source/Runtime/Engine/Private/SparseVolumeTexture/SparseVolumeTextureTileDataTexture.h b/Engine/Source/Runtime/Engine/Private/SparseVolumeTexture/SparseVolumeTextureTileDataTexture.h
--- a/Engine/Source/Runtime/Engine/Private/SparseVolumeTexture/SparseVolumeTextureTileDataTexture.h
+++ b/Engine/Source/Runtime/Engine/Private/SparseVolumeTexture/SparseVolumeTextureTileDataTexture.h
@@ -27,8 +27,6 @@ public:
 		Ready, Reserving, Reserved, Uploading
 	};
 
-	static constexpr uint32 PhysicalCoordMask = (1u << 24u) - 1u; // Lower 24 bits are used for storing XYZ in 8 bit each. Upper 8 bit can be used by the caller. 
-
 	static bool ShouldUseReservedResources();
 	static int64 GetMaxTileDataTextureResourceSize(int32 InVoxelMemSize);
 	static FIntVector3 GetVolumeResolutionInTiles(int32 InNumRequiredTiles);
@@ -37,33 +35,8 @@ public:
 	// Constructor. May change the requested ResolutionInTiles (and resulting PhysicalTilesCapacity) if it exceeds hardware limits.
 	explicit FTileDataTexture(const FIntVector3& ResolutionInTiles, EPixelFormat FormatA, EPixelFormat FormatB, const FVector4f& FallbackValueA, const FVector4f& FallbackValueB);
 
-	// Allocate a tile slot in the texture. The resulting value is a packed coordinate (8 bit per component) of the allocated slot or INDEX_NONE if the allocation failed.
-	// The upper 8 bit are free to be used by the caller.
-	uint32 Allocate()
-	{
-		check(PhysicalTilesCapacity == TileCoords.Num());
-		return TileCoords.IsValidIndex(NextFreeTileCoordIndex) ? TileCoords[NextFreeTileCoordIndex++] : INDEX_NONE;
-	}
-
-	// Frees a previously allocated tile slot. The upper 8 bit (user data) are automatically cleared by this function.
-	void Free(uint32 PackedPhysicalCoord)
-	{
-		check(PackedPhysicalCoord != INDEX_NONE)
-			PackedPhysicalCoord &= PhysicalCoordMask;
-		check(PhysicalTilesCapacity == TileCoords.Num());
-		check(NextFreeTileCoordIndex > 0);
-#if DO_GUARD_SLOW
-		for (int32 i = NextFreeTileCoordIndex; i < PhysicalTilesCapacity; ++i)
-		{
-			check(TileCoords[i] != PackedPhysicalCoord);
-		}
-#endif
-		TileCoords[--NextFreeTileCoordIndex] = PackedPhysicalCoord;
-	}
-
 	EUploaderState GetUploaderState() const { return UploaderState; }
 	int32 GetTileCapacity() const { return PhysicalTilesCapacity; }
-	int32 GetNumAvailableTiles() const { return PhysicalTilesCapacity - NextFreeTileCoordIndex; } // Number of tiles available for allocation.
 	FIntVector3 GetResolutionInTiles() const { return ResolutionInTiles; }
 	TRefCountPtr<IPooledRenderTarget> GetTileDataTextureA() { return TileDataTextureA; }
 	TRefCountPtr<IPooledRenderTarget> GetTileDataTextureB() { return TileDataTextureB; }
@@ -94,8 +67,6 @@ private:
 	FVector4f FallbackValueB;
 	TRefCountPtr<IPooledRenderTarget> TileDataTextureA;
 	TRefCountPtr<IPooledRenderTarget> TileDataTextureB;
-	TArray<uint32> TileCoords;
-	int32 NextFreeTileCoordIndex = 0;
 	EUploaderState UploaderState = EUploaderState::Ready;
 	int32 NumReservedUploadTiles = 0;
 	int32 NumReservedUploadVoxelsA = 0;
