Fix warning for nullptr check after assignment of InOutProjectionData.

diff --git a/Engine/Source/Runtime/Engine/Private/SceneView.cpp b/Engine/Source/Runtime/Engine/Private/SceneView.cpp
--- a/Engine/Source/Runtime/Engine/Private/SceneView.cpp
+++ b/Engine/Source/Runtime/Engine/Private/SceneView.cpp
@@ -541,10 +541,15 @@ FVector4f CreateInvDeviceZToWorldZTransform(const FMatrix& ProjMatrix)
 
 bool FSceneViewProjectionData::UpdateOrthoNearPlane(FSceneViewProjectionData* InOutProjectionData, float& NearPlane, bool bUpdateOrthoProjectionMatrix)
 {
-	//Store the original ViewOrigin for LOD location resolving.
+	if (!InOutProjectionData)
+	{
+		return false;
+	}
+
+	//Store the original ViewOrigin for LOD location resolving regardless of if we early out.
 	InOutProjectionData->LODViewOrigin = InOutProjectionData->ViewOrigin;
 
-	if (!InOutProjectionData || NearPlane >= 0.0f)
+	if (NearPlane >= 0.0f)
 	{
 		return false;
 	}
