Added cvar to control SVT IO request size. This should mitigate the issue where very large IO requests are made and then canceled before any data ever had a chance to stream in. The cvar can be changed at runtime. The optimal value is TBD, but for now it's set to a conservative 1MiB per IO request.

#rnx

diff --git a/Engine/Source/Runtime/Engine/Private/SparseVolumeTexture/SparseVolumeTextureStreamingManager.cpp b/Engine/Source/Runtime/Engine/Private/SparseVolumeTexture/SparseVolumeTextureStreamingManager.cpp
--- a/Engine/Source/Runtime/Engine/Private/SparseVolumeTexture/SparseVolumeTextureStreamingManager.cpp
+++ b/Engine/Source/Runtime/Engine/Private/SparseVolumeTexture/SparseVolumeTextureStreamingManager.cpp
@@ -86,6 +86,14 @@ static FAutoConsoleVariableRef CVarSVTStreamingMaxPendingRequests(
 	ECVF_RenderThreadSafe | ECVF_ReadOnly
 );
 
+static int32 GSVTStreamingRequestSize = 1024;
+static FAutoConsoleVariableRef CVarSVTStreamingRequestSize(
+	TEXT("r.SparseVolumeTexture.Streaming.RequestSize"),
+	GSVTStreamingRequestSize,
+	TEXT("IO request size in KiB. The SVT streaming manager will attempt to create IO requests of roughly this size. Default: 1024 KiB"),
+	ECVF_RenderThreadSafe
+);
+
 namespace UE
 {
 namespace SVT
@@ -818,6 +826,7 @@ void FStreamingManager::FilterRequests()
 
 		FFrameInfo& FrameInfo = SVTInfo->PerFrameInfo[RequestKey.FrameIndex];
 		const FPageTopology& Topology = FrameInfo.Resources->Topology;
+		const FTileStreamingMetaData& StreamingMetaData = FrameInfo.Resources->StreamingMetaData;
 		const int32 FirstStreamingMipLevel = FrameInfo.NumMipLevels - 2;
 		const int32 LowestRequestedMip = FMath::CountTrailingZeros((uint32)RequestValue.MipLevelMask);
 
@@ -829,6 +838,8 @@ void FStreamingManager::FilterRequests()
 		// StreamingTiles must be a super set of ResidentTiles
 		check(TBitArray<>::BitwiseAND(FrameInfo.StreamingTiles, FrameInfo.ResidentTiles, EBitwiseOperatorFlags::MaxSize) == FrameInfo.ResidentTiles);
 
+		const uint32 TargetRequestSize = (uint32)FMath::Clamp(GSVTStreamingRequestSize, 1, int32(UINT32_MAX / 1024)) * 1024u;
+
 		// Creates a FTileRange for each contiguous range of set bits in RequestedTilesInCurrentPriority
 		auto MakeTileRanges = [&](uint8 Priority)
 		{
@@ -846,8 +857,19 @@ void FStreamingManager::FilterRequests()
 				}
 				++TileCount;
 
+				const bool bIsLastTile = !RequestedTilesInCurrentPriority.IsValidIndex(TileIndex + 1);
+				const bool bIsNextTileBitSet = !bIsLastTile && RequestedTilesInCurrentPriority[TileIndex + 1];
+				bool bReachedTargetRequestSize = false;
+				if (bIsNextTileBitSet)
+				{
+					// End the current batch if adding the next tile would get us past the target request size.
+					// TileDataOffsets has N+1 elements, so accessing with +2 is fine because we checked that +1 is < N.
+					const uint32 RequestSizeIncludingNextTile = StreamingMetaData.TileDataOffsets[TileIndex + 2] - StreamingMetaData.TileDataOffsets[TileOffset];
+					bReachedTargetRequestSize = RequestSizeIncludingNextTile >= TargetRequestSize;
+				}
+
 				// Is the next bit unset? If so, this tile is the end of the current range
-				if (!RequestedTilesInCurrentPriority.IsValidIndex(TileIndex + 1) || !RequestedTilesInCurrentPriority[TileIndex + 1])
+				if (bIsLastTile || !bIsNextTileBitSet || bReachedTargetRequestSize)
 				{
 					FTileRange& TileRange = TileRangesToStream.AddDefaulted_GetRef();
 					TileRange.SVTHandle = RequestKey.SVTHandle;
