Swap check for conditional check when drawing pose watches

diff --git a/Engine/Source/Runtime/Engine/Private/SkeletalDebugRendering.cpp b/Engine/Source/Runtime/Engine/Private/SkeletalDebugRendering.cpp
--- a/Engine/Source/Runtime/Engine/Private/SkeletalDebugRendering.cpp
+++ b/Engine/Source/Runtime/Engine/Private/SkeletalDebugRendering.cpp
@@ -201,25 +201,27 @@ void DrawBonesFromPoseWatch(
 
 	for (const FBoneIndexType& BoneIndex : InRequiredBones)
 	{
-		check(ParentIndices.IsValidIndex(BoneIndex));
-		const int32 ParentIndex = ParentIndices[BoneIndex];
-
-		if (ParentIndex == INDEX_NONE)
-		{
-			UseWorldTransforms[BoneIndex] = InBoneTransforms[BoneIndex] * WorldTransform;
-			UseWorldTransforms[BoneIndex].AddToTranslation(RelativeOffset);
-		}
-		else
+		if (ParentIndices.IsValidIndex(BoneIndex))
 		{
-			UseWorldTransforms[BoneIndex] = InBoneTransforms[BoneIndex] * UseWorldTransforms[ParentIndex];
-		}
+			const int32 ParentIndex = ParentIndices[BoneIndex];
 
-		if (!ViewportMaskAllowList.Contains(BoneIndex))
-		{
-			continue;
-		}
+			if (ParentIndex == INDEX_NONE)
+			{
+				UseWorldTransforms[BoneIndex] = InBoneTransforms[BoneIndex] * WorldTransform;
+				UseWorldTransforms[BoneIndex].AddToTranslation(RelativeOffset);
+			}
+			else
+			{
+				UseWorldTransforms[BoneIndex] = InBoneTransforms[BoneIndex] * UseWorldTransforms[ParentIndex];
+			}
+
+			if (!ViewportMaskAllowList.Contains(BoneIndex))
+			{
+				continue;
+			}
 
-		UseRequiredBones.Add(BoneIndex);
+			UseRequiredBones.Add(BoneIndex);
+		}
 	}
 
 	const FLinearColor BoneColor = PoseWatch.GetBoneColor();
