Prevent crash when reseting material binding

diff --git a/Engine/Source/Runtime/Engine/Private/Components/HeterogeneousVolumeComponent.cpp b/Engine/Source/Runtime/Engine/Private/Components/HeterogeneousVolumeComponent.cpp
--- a/Engine/Source/Runtime/Engine/Private/Components/HeterogeneousVolumeComponent.cpp
+++ b/Engine/Source/Runtime/Engine/Private/Components/HeterogeneousVolumeComponent.cpp
@@ -520,7 +520,7 @@ void UHeterogeneousVolumeComponent::GetUsedMaterials(TArray<UMaterialInterface*>
 void UHeterogeneousVolumeComponent::SetMaterial(int32 ElementIndex, UMaterialInterface* Material)
 {
 	Super::SetMaterial(ElementIndex, Material);
-	if (ElementIndex == 0)
+	if (Material && ElementIndex == 0)
 	{
 		MaterialInstanceDynamic = CreateOrCastToMID(Material);
 		OnSparseVolumeTextureChanged(GetSparseVolumeTexture(MaterialInstanceDynamic, 0 /*SVTParameterIndex*/));
