World Partition - Fixed typo - FWorldPartitionActorDescInstance::SetForceNonSpatiallyLoadded()

diff --git a/Engine/Source/Runtime/Engine/Private/WorldPartition/WorldPartition.cpp b/Engine/Source/Runtime/Engine/Private/WorldPartition/WorldPartition.cpp
--- a/Engine/Source/Runtime/Engine/Private/WorldPartition/WorldPartition.cpp
+++ b/Engine/Source/Runtime/Engine/Private/WorldPartition/WorldPartition.cpp
@@ -1264,7 +1264,7 @@ void UWorldPartition::OnActorDescInstanceAdded(FWorldPartitionActorDescInstance*
 		ResolvingDataLayerManager->ResolveActorDescInstanceDataLayers(NewActorDescInstance);
 	}
 
-	NewActorDescInstance->SetForceNonSpatiallyLoadded(!IsStreamingEnabledInEditor());
+	NewActorDescInstance->SetForceNonSpatiallyLoaded(!IsStreamingEnabledInEditor());
 
 	HashActorDescInstance(NewActorDescInstance);
 
@@ -1347,7 +1347,7 @@ void UWorldPartition::InitializeActorDescContainerEditorStreaming(UActorDescCont
 	TArray<FGuid> ForceLoadedActorGuids;
 	for (UActorDescContainerInstance::TIterator<> It(InActorDescContainerInstance); It; ++It)
 	{
-		It->SetForceNonSpatiallyLoadded(!bIsStreamingEnabled);
+		It->SetForceNonSpatiallyLoaded(!bIsStreamingEnabled);
 
 		if (ForceLoadedActors)
 		{
@@ -1392,7 +1392,7 @@ void UWorldPartition::OnEnableStreamingChanged()
 	for (FActorDescContainerInstanceCollection::TIterator<> Iterator(this); Iterator; ++Iterator)
 	{
 		UnhashActorDescInstance(*Iterator);
-		Iterator->SetForceNonSpatiallyLoadded(!IsStreamingEnabledInEditor());
+		Iterator->SetForceNonSpatiallyLoaded(!IsStreamingEnabledInEditor());
 		HashActorDescInstance(*Iterator);
 	}
 
diff --git a/Engine/Source/Runtime/Engine/Public/WorldPartition/WorldPartitionActorDescInstance.h b/Engine/Source/Runtime/Engine/Public/WorldPartition/WorldPartitionActorDescInstance.h
--- a/Engine/Source/Runtime/Engine/Public/WorldPartition/WorldPartitionActorDescInstance.h
+++ b/Engine/Source/Runtime/Engine/Public/WorldPartition/WorldPartitionActorDescInstance.h
@@ -105,7 +105,7 @@ public:
 	ENGINE_API virtual TWeakObjectPtr<AActor>* GetActorPtr(bool bEvenIfPendingKill = true, bool bEvenIfUnreachable = false) const;
 	ENGINE_API virtual bool IsValid() const;
 				
-	inline void SetForceNonSpatiallyLoadded(bool bForce) { bIsForcedNonSpatiallyLoaded = bForce; }
+	inline void SetForceNonSpatiallyLoaded(bool bForce) { bIsForcedNonSpatiallyLoaded = bForce; }
 	inline bool GetForceNonSpatiallyLoaded() const { return bIsForcedNonSpatiallyLoaded; }
 
 	inline void SetUnloadedReason(FText* InUnloadedReason) { UnloadedReason = InUnloadedReason; }
