PR #10697: Add GetTotalPlaybackTime to SoundSource Classes

diff --git a/Engine/Source/Runtime/AudioMixer/Private/AudioMixerSource.cpp b/Engine/Source/Runtime/AudioMixer/Private/AudioMixerSource.cpp
--- a/Engine/Source/Runtime/AudioMixer/Private/AudioMixerSource.cpp
+++ b/Engine/Source/Runtime/AudioMixer/Private/AudioMixerSource.cpp
@@ -1514,6 +1514,16 @@ namespace Audio
 		}
 	}
 
+	int64 FMixerSource::GetNumFramesPlayed() const
+	{
+		if (InitializationState == EMixerSourceInitializationState::Initialized && MixerSourceVoice != nullptr)
+		{
+			return MixerSourceVoice->GetNumFramesPlayed();
+		}
+
+		return 0;
+	}
+
 	float FMixerSource::GetEnvelopeValue() const
 	{
 		if (MixerSourceVoice)
diff --git a/Engine/Source/Runtime/AudioMixer/Private/AudioMixerSource.h b/Engine/Source/Runtime/AudioMixer/Private/AudioMixerSource.h
--- a/Engine/Source/Runtime/AudioMixer/Private/AudioMixerSource.h
+++ b/Engine/Source/Runtime/AudioMixer/Private/AudioMixerSource.h
@@ -49,6 +49,7 @@ namespace Audio
 		virtual void Pause() override;
 		virtual bool IsFinished() override;
 		virtual float GetPlaybackPercent() const override;
+		virtual int64 GetNumFramesPlayed() const override;
 		virtual float GetEnvelopeValue() const override;
 		//~ End FSoundSource Interface
 
diff --git a/Engine/Source/Runtime/Engine/Private/Audio.cpp b/Engine/Source/Runtime/Engine/Private/Audio.cpp
--- a/Engine/Source/Runtime/Engine/Private/Audio.cpp
+++ b/Engine/Source/Runtime/Engine/Private/Audio.cpp
@@ -673,6 +673,31 @@ float FSoundSource::GetPlaybackPercent() const
 
 }
 
+float FSoundSource::GetSourceSampleRate() const
+{
+	if (WaveInstance == nullptr || WaveInstance->WaveData == nullptr || WaveInstance->WaveData->bIsSourceBus)
+	{
+		return AudioDevice->GetSampleRate();
+	}
+
+	return WaveInstance->WaveData->GetSampleRateForCurrentPlatform();
+}
+
+int64 FSoundSource::GetNumFramesPlayed() const
+{
+	return NumFramesPlayed;
+}
+
+int32 FSoundSource::GetNumTotalFrames() const
+{
+	return NumTotalFrames;
+}
+
+int32 FSoundSource::GetStartFrame() const
+{
+	return StartFrame;
+}
+
 void FSoundSource::GetChannelLocations(FVector& Left, FVector&Right) const
 {
 	Left = LeftChannelSourceLocation;
diff --git a/Engine/Source/Runtime/Engine/Public/Audio.h b/Engine/Source/Runtime/Engine/Public/Audio.h
--- a/Engine/Source/Runtime/Engine/Public/Audio.h
+++ b/Engine/Source/Runtime/Engine/Public/Audio.h
@@ -719,6 +719,18 @@ public:
 	/** Returns the source's playback percent. */
 	ENGINE_API virtual float GetPlaybackPercent() const;
 
+	/** Returns the sample (frame) rate of the audio played by the sound source. */
+	ENGINE_API virtual float GetSourceSampleRate() const;
+
+	/** Returns the number of frames (Samples / NumChannels) played by the sound source. */
+	ENGINE_API virtual int64 GetNumFramesPlayed() const;
+
+	/** Returns the total number of frames of audio for the sound wave. */
+	ENGINE_API virtual int32 GetNumTotalFrames() const;
+
+	/** Returns the frame index on which the sound source began playback. */
+	ENGINE_API virtual int32 GetStartFrame() const;
+
 	/** Returns the source's envelope at the callback block rate. Only implemented in audio mixer. */
 	virtual float GetEnvelopeValue() const { return 0.0f; };
 
