Remove check for unscaled view rect non negative, assumption is not correct when accounting for PIE scene view offset.

diff --git a/Engine/Source/Runtime/Engine/Private/SceneView.cpp b/Engine/Source/Runtime/Engine/Private/SceneView.cpp
--- a/Engine/Source/Runtime/Engine/Private/SceneView.cpp
+++ b/Engine/Source/Runtime/Engine/Private/SceneView.cpp
@@ -815,11 +815,6 @@ PRAGMA_ENABLE_DEPRECATION_WARNINGS
 	, PrimaryScreenPercentageMethod(EPrimaryScreenPercentageMethod::SpatialUpscale)
 	, FeatureLevel(InitOptions.ViewFamily ? InitOptions.ViewFamily->GetFeatureLevel() : GMaxRHIFeatureLevel)
 {
-	check(UnscaledViewRect.Min.X >= 0);
-	check(UnscaledViewRect.Min.Y >= 0);
-	check(UnscaledViewRect.Width() > 0);
-	check(UnscaledViewRect.Height() > 0);
-
 	ShadowViewMatrices = ViewMatrices;
 
 	SceneViewInitOptions = FSceneViewInitOptions(InitOptions);
