[Shaders][Optimization][EditorInitiative] - Instead of calculating the Material DebugGroupName for _every_ shader we draw with, calculate it _once_ per material and cache it.

- TryGetShaders submits shader jobs for missing shaders.  If a shader was missing it would calculate a DebugGroupName for the material for each shader.
- Instead we now cache this once and store on the FMaterial.

In my test scenario I saw large speedups on the Render Thread:

- ~7x speedup to the average render frames that were heavy w/ draw command building after opening a map.
- ~4x speedup to the worst case render frames.

diff --git a/Engine/Source/Runtime/Engine/Private/Materials/MaterialShader.cpp b/Engine/Source/Runtime/Engine/Private/Materials/MaterialShader.cpp
--- a/Engine/Source/Runtime/Engine/Private/Materials/MaterialShader.cpp
+++ b/Engine/Source/Runtime/Engine/Private/Materials/MaterialShader.cpp
@@ -1403,11 +1403,6 @@ void FMaterialShaderMap::LoadFromDerivedDataCache(const FMaterial* Material, con
 	InOutShaderMap = BeginLoadFromDerivedDataCache(Material, ShaderMapId, InPlatform, TargetPlatform, InOutShaderMap, OutDDCKeyDesc)->Get();
 }
 
-FString GetMaterialDebugGroupName(EShaderPlatform ShaderPlatform, const FMaterial* Material, const FMaterialShaderMapId& ShaderMapId)
-{
-	return Material->GetUniqueAssetName(ShaderPlatform, ShaderMapId) / LexToString(Material->GetQualityLevel());
-}
-
 TSharedRef<FMaterialShaderMap::FAsyncLoadContext> FMaterialShaderMap::BeginLoadFromDerivedDataCache(const FMaterial* Material, const FMaterialShaderMapId& ShaderMapId, EShaderPlatform InPlatform, const ITargetPlatform* TargetPlatform, TRefCountPtr<FMaterialShaderMap>& InShaderMap, FString& OutDDCKeyDesc)
 {
 	TRACE_CPUPROFILER_EVENT_SCOPE(FMaterialShaderMap::BeginLoadFromDerivedDataCache);
@@ -1508,7 +1503,7 @@ TSharedRef<FMaterialShaderMap::FAsyncLoadContext> FMaterialShaderMap::BeginLoadF
 
 			if (UNLIKELY(ShouldDumpShaderDDCKeys()) || (Material->IsDefaultMaterial() && Material->GetMaterialDomain() == EMaterialDomain::MD_Surface))
 			{	
-				DumpShaderDDCKeyToFile(InPlatform, ShaderMapId.LayoutParams.WithEditorOnly(), *GetMaterialDebugGroupName(InPlatform, Material, ShaderMapId), Result->DataKey);
+				DumpShaderDDCKeyToFile(InPlatform, ShaderMapId.LayoutParams.WithEditorOnly(), *Material->GetDebugGroupName(), Result->DataKey);
 			}
 
 			if (Material->IsDefaultMaterial() && Material->GetMaterialDomain() == EMaterialDomain::MD_Surface)
@@ -1833,8 +1828,6 @@ int32 FMaterialShaderMap::SubmitCompileJobs(uint32 CompilingShaderMapId,
 	const FMaterialShaderParameters MaterialParameters(Material);
 	const FMaterialShaderMapLayout& Layout = AcquireMaterialShaderMapLayout(ShaderPlatform, LocalPermutationFlags, MaterialParameters);
 
-	const FString DebugGroupName = GetMaterialDebugGroupName(ShaderPlatform, Material, GetShaderMapId());
-
 #if ALLOW_SHADERMAP_DEBUG_DATA
 	FString DebugExtensionStr(TEXT(""));
 	FString DebugDescriptionStr(TEXT(""));
@@ -1896,7 +1889,7 @@ int32 FMaterialShaderMap::SubmitCompileJobs(uint32 CompilingShaderMapId,
 					MaterialEnvironment,
 					MeshLayout.VertexFactoryType,
 					CompileJobs,
-					DebugGroupName,
+					Material->GetDebugGroupName(),
 					DebugDescription,
 					DebugExtension
 				);
@@ -1939,7 +1932,7 @@ int32 FMaterialShaderMap::SubmitCompileJobs(uint32 CompilingShaderMapId,
 					MeshLayout.VertexFactoryType,
 					Pipeline,
 					CompileJobs,
-					DebugGroupName,
+					Material->GetDebugGroupName(),
 					DebugDescription,
 					DebugExtension);
 			}
@@ -1999,7 +1992,7 @@ int32 FMaterialShaderMap::SubmitCompileJobs(uint32 CompilingShaderMapId,
 				ShaderPlatform,
 				LocalPermutationFlags,
 				CompileJobs,
-				DebugGroupName,
+				Material->GetDebugGroupName(),
 				DebugDescription,
 				DebugExtension
 			);
@@ -2033,7 +2026,7 @@ int32 FMaterialShaderMap::SubmitCompileJobs(uint32 CompilingShaderMapId,
 					MaterialEnvironment,
 					Pipeline,
 					CompileJobs,
-					DebugGroupName,
+					Material->GetDebugGroupName(),
 					DebugDescription,
 					DebugExtension);
 			}
diff --git a/Engine/Source/Runtime/Engine/Private/Materials/MaterialShared.cpp b/Engine/Source/Runtime/Engine/Private/Materials/MaterialShared.cpp
--- a/Engine/Source/Runtime/Engine/Private/Materials/MaterialShared.cpp
+++ b/Engine/Source/Runtime/Engine/Private/Materials/MaterialShared.cpp
@@ -2765,6 +2765,9 @@ bool FMaterial::CacheShaders(const FMaterialShaderMapId& ShaderMapId, EShaderPla
 #endif
 	UE_CLOG(!ShaderMapId.IsValid(), LogMaterial, Warning, TEXT("Invalid shader map ID caching shaders for '%s', will use default material."), *GetFriendlyName());
 #if WITH_EDITOR
+
+	DebugGroupName = GetUniqueAssetName(Platform, ShaderMapId) / LexToString(GetQualityLevel());
+
 	FString DDCKeyHash;
 
 	// Just make sure that we don't already have a pending cache going on.
@@ -3171,8 +3174,6 @@ void FMaterial::CacheGivenTypes(EShaderPlatform Platform, const TArray<const FVe
 		checkf(PipelineTypes.Num() == ShaderTypes.Num(), TEXT("The size of the pipeline type array and shader type array must match.  Pass in null entries if pipelines are not used."));
 		checkf(GetGameThreadCompilingShaderMapId() != 0, TEXT("Material is not prepared to compile yet.  Please call CacheShaders first."));
 
-		const FString DebugGroupName = GetUniqueAssetName(Platform, GameThreadShaderMap->GetShaderMapId()) / LexToString(GetQualityLevel());
-
 		TArray<FShaderCommonCompileJobPtr> CompileJobs;
 		for (int i = 0; i < VFTypes.Num(); ++i)
 		{
@@ -3194,7 +3195,7 @@ void FMaterial::CacheGivenTypes(EShaderPlatform Platform, const TArray<const FVe
 					VFType,
 					PipelineType,
 					CompileJobs,
-					DebugGroupName,
+					GetDebugGroupName(),
 					nullptr,
 					nullptr);
 			}
@@ -3663,7 +3664,6 @@ bool FMaterial::TryGetShaders(const FMaterialShaderTypes& InTypes, const FVertex
 					{
 						if (!bShaderMapComplete)
 						{
-							const FString DebugGroupName = GetUniqueAssetName(ShaderPlatform, ShaderMap->GetShaderMapId()) / LexToString(GetQualityLevel());
 							if (InVertexFactoryType)
 							{
 								FMeshMaterialShaderType::BeginCompileShaderPipeline(
@@ -3678,7 +3678,7 @@ bool FMaterial::TryGetShaders(const FMaterialShaderTypes& InTypes, const FVertex
 									InVertexFactoryType, 
 									InTypes.PipelineType, 
 									CompileJobs,
-									DebugGroupName,
+									GetDebugGroupName(),
 									nullptr, 
 									nullptr);
 							}
@@ -3694,7 +3694,7 @@ bool FMaterial::TryGetShaders(const FMaterialShaderTypes& InTypes, const FVertex
 									RenderingThreadPendingCompilerEnvironment, 
 									InTypes.PipelineType, 
 									CompileJobs,
-									DebugGroupName,
+									GetDebugGroupName(),
 									nullptr, 
 									nullptr);
 							}
@@ -3752,7 +3752,6 @@ bool FMaterial::TryGetShaders(const FMaterialShaderTypes& InTypes, const FVertex
 						{
 							if (!bShaderMapComplete)
 							{
-								const FString DebugGroupName = GetUniqueAssetName(ShaderPlatform, ShaderMap->GetShaderMapId()) / LexToString(GetQualityLevel());
 								if (InVertexFactoryType)
 								{
 									ShaderType->AsMeshMaterialShaderType()->BeginCompileShader(
@@ -3766,7 +3765,7 @@ bool FMaterial::TryGetShaders(const FMaterialShaderTypes& InTypes, const FVertex
 										RenderingThreadPendingCompilerEnvironment,
 										InVertexFactoryType,
 										CompileJobs,
-										DebugGroupName,
+										GetDebugGroupName(),
 										nullptr,
 										nullptr);
 								}
@@ -3782,7 +3781,7 @@ bool FMaterial::TryGetShaders(const FMaterialShaderTypes& InTypes, const FVertex
 										ShaderPlatform,
 										PermutationFlags,
 										CompileJobs,
-										DebugGroupName,
+										GetDebugGroupName(),
 										nullptr,
 										nullptr);
 								}
diff --git a/Engine/Source/Runtime/Engine/Public/MaterialShared.h b/Engine/Source/Runtime/Engine/Public/MaterialShared.h
--- a/Engine/Source/Runtime/Engine/Public/MaterialShared.h
+++ b/Engine/Source/Runtime/Engine/Public/MaterialShared.h
@@ -2123,6 +2123,7 @@ public:
 	void SetCompileErrors(const TArray<FString>& InCompileErrors) { CompileErrors = InCompileErrors; }
 	const TArray<UMaterialExpression*>& GetErrorExpressions() const { return ObjectPtrDecay(ErrorExpressions); }
 	const FGuid& GetLegacyId() const { return Id_DEPRECATED; }
+	const FString& GetDebugGroupName() const { return DebugGroupName; }
 #endif // WITH_EDITOR
 
 	inline const FStaticFeatureLevel GetFeatureLevel() const { checkSlow(FeatureLevel != ERHIFeatureLevel::Num); return FeatureLevel; }
@@ -2451,6 +2452,8 @@ private:
 
 	TRefCountPtr<FSharedShaderCompilerEnvironment> GameThreadPendingCompilerEnvironment;
 	TRefCountPtr<FSharedShaderCompilerEnvironment> RenderingThreadPendingCompilerEnvironment;
+
+	FString DebugGroupName;
 #endif // WITH_EDITOR
 
 	/** 
