Dynamic Submixes rework

diff --git a/Engine/Source/Editor/AudioEditor/Private/SoundSubmixGraphSchema.cpp b/Engine/Source/Editor/AudioEditor/Private/SoundSubmixGraphSchema.cpp
--- a/Engine/Source/Editor/AudioEditor/Private/SoundSubmixGraphSchema.cpp
+++ b/Engine/Source/Editor/AudioEditor/Private/SoundSubmixGraphSchema.cpp
@@ -227,6 +227,12 @@ const FPinConnectionResponse USoundSubmixGraphSchema::CanCreateConnection(const
 	USoundSubmixBase* InputSubmix = CastChecked<USoundSubmixGraphNode>(OutputPin->GetOwningNode())->SoundSubmix;
 	USoundSubmixBase* OutputSubmix = CastChecked<USoundSubmixGraphNode>(InputPin->GetOwningNode())->SoundSubmix;
 
+	// Forbid connecting dynamic submixes to other submixes.
+	if (InputSubmix->IsDynamic( false /*bIncludeAncestors*/ ))
+	{
+		return FPinConnectionResponse(CONNECT_RESPONSE_DISALLOW, LOCTEXT("SubmixIsDynamic", "Submix you are trying to connect from is dynamic and shouldn't have any static parents"));
+	}
+
 	// Check to see if this is an endpoint submix.
 	if (!InputSubmix->IsA<USoundSubmixWithParentBase>())
 	{
diff --git a/Engine/Source/Runtime/AudioMixer/Private/AudioMixerDevice.cpp b/Engine/Source/Runtime/AudioMixer/Private/AudioMixerDevice.cpp
--- a/Engine/Source/Runtime/AudioMixer/Private/AudioMixerDevice.cpp
+++ b/Engine/Source/Runtime/AudioMixer/Private/AudioMixerDevice.cpp
@@ -1175,7 +1175,7 @@ namespace Audio
 				USoundSubmixBase* SubmixToLoad = *It;
 				check(SubmixToLoad);
 
-				if (!IsRequiredSubmixType(SubmixToLoad) && SubmixToLoad->bAutoRegister )
+				if (!IsRequiredSubmixType(SubmixToLoad) && !SubmixToLoad->IsDynamic( true /* bIncludeAncestors */) ) // Do not load dynamic submixes until they've been connected.
 				{
 					LoadSoundSubmix(*SubmixToLoad);
 					InitSoundfieldAndEndpointDataForSubmix(*SubmixToLoad, GetSubmixInstance(SubmixToLoad).Pin(), false);
@@ -1240,7 +1240,7 @@ namespace Audio
 			{
 				ParentSubmixInstance = GetSubmixInstance(Parent).Pin();
 			}
-			else if (SubmixWithParent->bAutoRouteToMasterSubmixWhenOrphaned)
+			else if (!SubmixWithParent->IsDynamic( true /*bIncludeAncestors*/ )) // Dynamic submixes do not auto connect.
 			{
 				// If this submix is itself the broadcast submix, set its parent to the master submix
 				if (SubmixInstance == RequiredSubmixInstances[static_cast<int32>(ERequiredSubmixes::BaseDefault)])
diff --git a/Engine/Source/Runtime/AudioMixer/Private/AudioMixerSource.cpp b/Engine/Source/Runtime/AudioMixer/Private/AudioMixerSource.cpp
--- a/Engine/Source/Runtime/AudioMixer/Private/AudioMixerSource.cpp
+++ b/Engine/Source/Runtime/AudioMixer/Private/AudioMixerSource.cpp
@@ -1834,7 +1834,7 @@ namespace Audio
 			{
 				SubmixPtr = MixerDevice->GetSubmixInstance(WaveInstance->SoundSubmix);
 			}
-			else if(WaveInstance->SoundSubmix && WaveInstance->SoundSubmix->bAutoRouteToMasterSubmixWhenOrphaned)
+			else if (!WaveInstance->bIsDynamic) // Dynamic submixes don't auto connect to master.
 			{
 				SubmixPtr = MixerDevice->GetMasterSubmix();
 			}
diff --git a/Engine/Source/Runtime/Engine/Classes/Sound/SoundSubmix.h b/Engine/Source/Runtime/Engine/Classes/Sound/SoundSubmix.h
--- a/Engine/Source/Runtime/Engine/Classes/Sound/SoundSubmix.h
+++ b/Engine/Source/Runtime/Engine/Classes/Sound/SoundSubmix.h
@@ -188,14 +188,6 @@ public:
 	UPROPERTY(config, EditAnywhere, BlueprintReadWrite, Category = AutoDisablement, meta = (EditCondition = "bAutoDisable"))
 	float AutoDisableTime = 0.01f;
 
-	// Whether to automatically register this Submix with an AudioDevice. 
-	UPROPERTY(config, EditAnywhere, BlueprintReadOnly, Category = Automatic)
-	bool bAutoRegister = true;
-
-	// Whether to automatically route to the 'MasterSubmix' aka 'MainSubmix' when the parent is null
-	UPROPERTY(config, EditAnywhere, BlueprintReadOnly, Category = Automatic)
-	bool bAutoRouteToMasterSubmixWhenOrphaned = true;
-
 	// Child submixes to this sound mix
 	UPROPERTY(VisibleAnywhere, BlueprintReadOnly, Category = SoundSubmix)
 	TArray<TObjectPtr<USoundSubmixBase>> ChildSubmixes;
@@ -225,6 +217,23 @@ public:
 		return false; 
 	}
 	
+	/** Searching upwards from this Submix to the root looking for the first Submix marked Dynamic
+	 *  If this Submix is Dynamic this will be returned.
+	**/
+	UFUNCTION(BlueprintCallable, Category = "Submix", meta = (WorldContext = "WorldContextObject"))
+	ENGINE_API virtual USoundSubmixBase* FindDynamicAncestor() 
+	{ 
+		return nullptr; 
+	}
+	
+	/** If this Submix is (or any of its parents are marked dynamic).
+	* @param	bIncludeAncestors	Whether to traverse upwards through the ancestors looking for anything with a dynamic flag.
+	**/
+	virtual bool IsDynamic(const bool bIncludeAncestors) const 
+	{ 
+		return false; 
+	}
+	
 protected:
 	//~ Begin UObject Interface.
 	ENGINE_API virtual FString GetDesc() override;
@@ -268,7 +277,7 @@ class USoundSubmixWithParentBase : public USoundSubmixBase
 public:
 
 
-	UPROPERTY(VisibleAnywhere, BlueprintReadOnly, Category = SoundSubmix)
+	UPROPERTY(VisibleAnywhere, BlueprintReadOnly, Category = SoundSubmix, meta = (EditCondition = "!bIsDynamic"))
 	TObjectPtr<USoundSubmixBase> ParentSubmix;
 
 	/**
@@ -292,9 +301,18 @@ public:
 	ENGINE_API bool DynamicConnect(FAudioDeviceHandle Handle, USoundSubmixBase* Parent);
 	ENGINE_API virtual bool DynamicDisconnect(const UObject* WorldContextObject) override;
 	ENGINE_API bool DynamicDisconnect(FAudioDeviceHandle Handle);
+	ENGINE_API virtual bool IsDynamic(const bool bIncludeAncestors) const override;
+	ENGINE_API virtual USoundSubmixBase* FindDynamicAncestor() override;
 
 protected:
 
+	/** Is Submix Dynamic. (i.e. allows connect/disconnect at runtime.)  **/
+	UPROPERTY(EditAnywhere, BlueprintReadOnly, Category = SoundSubmix, meta = (DisplayPriority=-1, EditCondition = "ParentSubmix == nullptr"))
+	uint8 bIsDynamic : 1;
+
+	// Const version. 
+	const USoundSubmixBase* FindDynamicAncestor() const;
+	
 #if WITH_EDITOR
 	ENGINE_API virtual void PostEditChangeProperty(struct FPropertyChangedEvent& PropertyChangedEvent) override;
 	ENGINE_API virtual void PostDuplicate(EDuplicateMode::Type DuplicateMode) override;
@@ -597,6 +615,8 @@ protected:
 
 namespace SubmixUtils
 {
+	ENGINE_API void ForEachStaticChildRecursive(USoundSubmixBase* StartingPoint, const TFunction<void(USoundSubmixBase*)>& Op);
+
 	ENGINE_API bool AreSubmixFormatsCompatible(const USoundSubmixBase* ChildSubmix, const USoundSubmixBase* ParentSubmix);
 
 	ENGINE_API bool FindInGraph(const USoundSubmixBase* InEntryPoint, const USoundSubmixBase* InToMatch, bool bShouldAcsend, FAudioDeviceHandle InDevice = {});
diff --git a/Engine/Source/Runtime/Engine/Private/Audio.cpp b/Engine/Source/Runtime/Engine/Private/Audio.cpp
--- a/Engine/Source/Runtime/Engine/Private/Audio.cpp
+++ b/Engine/Source/Runtime/Engine/Private/Audio.cpp
@@ -875,6 +875,7 @@ FWaveInstance::FWaveInstance(const UPTRINT InWaveInstanceHash, FActiveSound& InA
 	, bReportedSpatializationWarning(false)
 	, bIsAmbisonics(false)
 	, bIsStopping(false)
+	, bIsDynamic(false)
 	, SpatializationMethod(ESoundSpatializationAlgorithm::SPATIALIZATION_Default)
 	, SpatializationPluginSettings(nullptr)
 	, OcclusionPluginSettings(nullptr)
diff --git a/Engine/Source/Runtime/Engine/Private/SoundSubmix.cpp b/Engine/Source/Runtime/Engine/Private/SoundSubmix.cpp
--- a/Engine/Source/Runtime/Engine/Private/SoundSubmix.cpp
+++ b/Engine/Source/Runtime/Engine/Private/SoundSubmix.cpp
@@ -30,6 +30,7 @@ FAutoConsoleVariableRef CVarFixUpBrokenSubmixAssets(
 USoundSubmixWithParentBase::USoundSubmixWithParentBase(const FObjectInitializer& ObjectInitializer)
 	: Super(ObjectInitializer)
 	, ParentSubmix(nullptr)
+	, bIsDynamic(0)
 {}
 
 USoundSubmixBase::USoundSubmixBase(const FObjectInitializer& ObjectInitializer)
@@ -609,6 +610,13 @@ bool USoundSubmixWithParentBase::DynamicConnect(const UObject* WorldContextObjec
 
 bool USoundSubmixWithParentBase::DynamicConnect(FAudioDeviceHandle Handle, USoundSubmixBase* InParent)
 {
+	if (!IsDynamic(false /* bIncludeAncestors */))
+	{
+		const USoundSubmixBase* DynamicAncestor = FindDynamicAncestor();
+		UE_CLOG(DynamicAncestor, LogAudio, Warning, TEXT("Submix (DynamicConnect): Dynamic Flag not set for [%s], you need its ancestor [%s]. Call FindDynamicAncestor on this submix to find it. Ignoring..." ), *GetName(), *DynamicAncestor->GetName());
+		UE_CLOG(!DynamicAncestor, LogAudio, Warning, TEXT("Submix (DynamicConnect): Dynamic Flag not set for [%s] or any of its parents, ignoring... "), *GetName());
+		return false;
+	}
 	if (!Handle.IsValid())
 	{
 		UE_LOG(LogAudio, Warning, TEXT("Submix (DynamicConnect): No valid audio device in this world for [%s]"), *GetName());
@@ -651,11 +659,19 @@ bool USoundSubmixWithParentBase::DynamicConnect(FAudioDeviceHandle Handle, USoun
 			Handle->SetSubmixAutoDisable(Cast<USoundSubmix>(CurrentParent.Get()), false);
 		}
 
-		// Register us, and disable parents auto disable feature.
-		Handle->RegisterSoundSubmix(this, /*bInit*/ true);
+		// Register us and our children
+		SubmixUtils::ForEachStaticChildRecursive(
+			this,
+			[&Handle, Id](USoundSubmixBase* Iter)-> void
+			{
+				UE_LOG(LogAudio, Verbose, TEXT("Submix (DynamicConnect): Registering [%s] with AudioDevice [%u]"), *Iter->GetName(), Id);
+				Handle->RegisterSoundSubmix(Iter, /* bInit*/ true);
+			});
+			
+		// ... and disable parents auto disable feature.
 		Handle->SetSubmixAutoDisable(Cast<USoundSubmix>(this), CurrentParent == nullptr);
 
-		UE_LOG(LogAudio, Verbose, TEXT("Submix (DynamicConnect): Registering [%s] with AudioDevice [%u]"), *GetName(), Id);
+
 
 		return CurrentParent != nullptr;
 	}
@@ -692,6 +708,12 @@ bool USoundSubmixWithParentBase::DynamicDisconnect(FAudioDeviceHandle Handle)
 		return false;
 	}
 
+	if (!IsDynamic(false /* bIncludeAncestors */))
+	{
+		UE_LOG(LogAudio, Warning, TEXT("Submix (DynamicDisconnect): Dynamic Flag not set for [%s] ignoring."), *GetName());
+		return false;
+	}
+
 	const Audio::DeviceID Id = Handle.GetDeviceID();
 
 	TObjectPtr<USoundSubmixBase>& CurrentParent = DynamicParentSubmix.FindOrAdd(Id);
@@ -703,7 +725,12 @@ bool USoundSubmixWithParentBase::DynamicDisconnect(FAudioDeviceHandle Handle)
 		CurrentParent->DynamicChildSubmixes.FindOrAdd(Id).ChildSubmixes.Remove(this);
 		CurrentParent = nullptr;
 
-		Handle->UnregisterSoundSubmix(this, false);
+		 SubmixUtils::ForEachStaticChildRecursive(
+			this,
+			[&Handle](USoundSubmixBase* Iter)-> void
+			{
+				Handle->UnregisterSoundSubmix(Iter,  /* bReparentChildren*/ false);
+			});	
 				
 		// If we still have a valid parent static submix? Make sure that's still live and registered.
 		if (ParentSubmix)
@@ -720,6 +747,46 @@ bool USoundSubmixWithParentBase::DynamicDisconnect(FAudioDeviceHandle Handle)
 	return false;
 }
 
+bool USoundSubmixWithParentBase::IsDynamic(const bool bIncludeAncestors) const
+{
+	// If we don't care about ancestors, just return if we're dynamic.
+	if (!bIncludeAncestors)
+	{
+		return bIsDynamic;
+	}
+
+	// Find the first dynamic ancestor.
+	const USoundSubmixBase* Found = FindDynamicAncestor();
+	return Found != nullptr;
+}
+
+USoundSubmixBase* USoundSubmixWithParentBase::FindDynamicAncestor()
+{
+	const USoundSubmixBase* Found = const_cast<const USoundSubmixWithParentBase*>(this)->FindDynamicAncestor();
+	return const_cast<USoundSubmixBase*>(Found);
+}
+
+const USoundSubmixBase* USoundSubmixWithParentBase::FindDynamicAncestor() const
+{
+	// Walk up parents from here checking for dynamic flag.
+	for (const USoundSubmixWithParentBase* Current = this; Current; /* Incremented below */)
+	{
+		if (Current->bIsDynamic)
+		{
+			return Current;
+		}
+		if (const USoundSubmixWithParentBase* Parent = Cast<USoundSubmixWithParentBase>(Current->ParentSubmix))
+		{
+			Current = Parent;
+		}
+		else
+		{
+			break;
+		}
+	}
+	return nullptr;
+}
+
 #if WITH_EDITOR
 
 void USoundSubmixBase::PostDuplicate(EDuplicateMode::Type DuplicateMode)
@@ -772,8 +839,6 @@ void USoundSubmixBase::PostEditChangeProperty(FPropertyChangedEvent& PropertyCha
 						SubmixWithParent->SetParentSubmix(this);
 					}
 					
-					// Make the children follow our auto register setting.
-					ChildSubmixes[ChildIndex]->bAutoRegister = bAutoRegister;
 					ChildSubmixes[ChildIndex]->PostEditChangeProperty(PropertyChangedEvent);
 
 					break;
@@ -790,30 +855,15 @@ void USoundSubmixBase::PostEditChangeProperty(FPropertyChangedEvent& PropertyCha
 					{
 						SubmixWithParent->ParentSubmix = nullptr;
 					}
-				}
-			}
-
-			// Force the properties to be initialized for this SoundSubmix on all active audio devices
-			if (FAudioDeviceManager* AudioDeviceManager = GEngine->GetAudioDeviceManager())
-			{
-				if (bAutoRegister)
-				{
-					AudioDeviceManager->RegisterSoundSubmix(this);
-				}
-			}
-		}
 
-		// Propagate auto register to children of this submix. 
-		// We don't want them to auto register 
-		if (PropertyChangedEvent.Property->GetFName() == GET_MEMBER_NAME_CHECKED(USoundSubmixBase, bAutoRegister))
-		{
-			// All children so should follow the 
-			for (int32 ChildIndex = 0; ChildIndex < ChildSubmixes.Num(); ChildIndex++)
-			{
-				if (USoundSubmixBase* Child = ChildSubmixes[ChildIndex] )
-				{
-					Child->bAutoRegister = bAutoRegister;
-					Child->PostEditChangeProperty(PropertyChangedEvent);
+					// Force the properties to be initialized for this SoundSubmix on all active audio devices
+					if (FAudioDeviceManager* AudioDeviceManager = GEngine->GetAudioDeviceManager())
+					{
+						if (!IsDynamic(true /* bIncludeAncestors */ )) // Exclude dynamic submixes from registration
+						{
+							AudioDeviceManager->RegisterSoundSubmix(this);
+						}
+					}
 				}
 			}
 		}
@@ -1368,6 +1418,15 @@ bool SubmixUtils::FindInGraph(
 
 }
 
+void SubmixUtils::ForEachStaticChildRecursive(USoundSubmixBase* StartingPoint, const TFunction<void(USoundSubmixBase*)>& Op)
+{
+	Op(StartingPoint);
+	for (TObjectPtr<USoundSubmixBase> i : StartingPoint->ChildSubmixes)
+	{
+		ForEachStaticChildRecursive(i,Op);
+	}
+}
+
 const USoundSubmixBase* SubmixUtils::FindRoot(const USoundSubmixBase* InStartingPoint, FAudioDeviceHandle InDevice)
 {	
 	const USoundSubmixBase* HighestPoint = InStartingPoint;
diff --git a/Engine/Source/Runtime/Engine/Public/Audio.h b/Engine/Source/Runtime/Engine/Public/Audio.h
--- a/Engine/Source/Runtime/Engine/Public/Audio.h
+++ b/Engine/Source/Runtime/Engine/Public/Audio.h
@@ -325,6 +325,9 @@ public:
 	/** Whether or not this wave instance is stopping. */
 	uint32 bIsStopping:1;
 
+	/** Is this or any of the submixes above it dynamic */
+	uint32 bIsDynamic:1;
+
 	/** Which spatialization method to use to spatialize 3d sounds. */
 	ESoundSpatializationAlgorithm SpatializationMethod;
 
