[shaders] Don't show modal that requires user interaction when SCW is unattended.

#rnx

diff --git a/Engine/Source/Runtime/Engine/Private/ShaderCompiler/ShaderCompiler.cpp b/Engine/Source/Runtime/Engine/Private/ShaderCompiler/ShaderCompiler.cpp
--- a/Engine/Source/Runtime/Engine/Private/ShaderCompiler/ShaderCompiler.cpp
+++ b/Engine/Source/Runtime/Engine/Private/ShaderCompiler/ShaderCompiler.cpp
@@ -2028,7 +2028,7 @@ static void ModalErrorOrLog(const FString& Title, const FString& Text, int64 Cur
 		BadFile = FString::Printf(TEXT(" (Truncated or corrupt output file! Current file pos %lld, file size %lld)"), CurrentFilePos, ExpectedFileSize);
 	}
 
-	if (FPlatformProperties::SupportsWindowedMode())
+	if (FPlatformProperties::SupportsWindowedMode() && !FApp::IsUnattended())
 	{
 		UE_LOG(LogShaderCompilers, Error, TEXT("%s\n%s"), *Text, *BadFile);
 		if (!bModalReported.AtomicSet(true))
