Fix a small regression that broke Nanite custom stencil.

diff --git a/Engine/Source/Runtime/Engine/Private/PrimitiveUniformShaderParameters.cpp b/Engine/Source/Runtime/Engine/Private/PrimitiveUniformShaderParameters.cpp
--- a/Engine/Source/Runtime/Engine/Private/PrimitiveUniformShaderParameters.cpp
+++ b/Engine/Source/Runtime/Engine/Private/PrimitiveUniformShaderParameters.cpp
@@ -231,7 +231,11 @@ void FPrimitiveSceneShaderData::Setup(const FPrimitiveUniformShaderParameters& P
 		PrimitiveUniformShaderParameters.InstanceWPODisableDistanceSquared, 
 		PrimitiveUniformShaderParameters.NaniteRayTracingDataOffset);
 
-	OutData[32]	= FVector4f(PrimitiveUniformShaderParameters.MaxWPOExtent, PrimitiveUniformShaderParameters.CustomStencilValueAndMask, 0.0, 0.0);
+	Store4(OutData, 32,
+		PrimitiveUniformShaderParameters.MaxWPOExtent,
+		PrimitiveUniformShaderParameters.CustomStencilValueAndMask,
+		0.0f,
+		0.0f);
 
 	// Set all the custom primitive data float4. This matches the loop in SceneData.ush
 	const int32 CustomPrimitiveDataStartIndex = 33;
