Chaos - add FBodyInstance::bUseMACD flag to enable Motion-Aware Collision Detection per particle
- flag is also used by GeomrteyCollections
- wire flag through to Chaos particles in the ControlFlags structure

NOTE: this flag isn't actually being used yet. It will be connected to the MACD feature in an upcoming change.

diff --git a/Engine/Source/Programs/HeadlessChaos/Private/GeometryCollection/GeometryCollectionTestFramework.cpp b/Engine/Source/Programs/HeadlessChaos/Private/GeometryCollection/GeometryCollectionTestFramework.cpp
--- a/Engine/Source/Programs/HeadlessChaos/Private/GeometryCollection/GeometryCollectionTestFramework.cpp
+++ b/Engine/Source/Programs/HeadlessChaos/Private/GeometryCollection/GeometryCollectionTestFramework.cpp
@@ -176,6 +176,7 @@ namespace GeometryCollectionTest
 			SimulationParams.LinearDamping = 0;
 			SimulationParams.AngularDamping = 0;
 			SimulationParams.UseCCD = false;
+			SimulationParams.UseMACD = false;
 
 			Chaos::FErrorReporter ErrorReporter;
 			BuildSimulationData(ErrorReporter, *RestCollection.Get(), SimulationParams.Shared);
diff --git a/Engine/Source/Runtime/Engine/Classes/PhysicsEngine/BodyInstance.h b/Engine/Source/Runtime/Engine/Classes/PhysicsEngine/BodyInstance.h
--- a/Engine/Source/Runtime/Engine/Classes/PhysicsEngine/BodyInstance.h
+++ b/Engine/Source/Runtime/Engine/Classes/PhysicsEngine/BodyInstance.h
@@ -340,12 +340,16 @@ public:
 	UPROPERTY(EditAnywhere, Category = Physics, meta = (DisplayName = "Mode"))
 	TEnumAsByte<EDOFMode::Type> DOFMode;
 
-public:
-
 	/** If true Continuous Collision Detection (CCD) will be used for this component */
 	UPROPERTY(EditAnywhere, AdvancedDisplay, BlueprintReadOnly, Category = Collision)
 	uint8 bUseCCD : 1;
 
+private:
+	/** [EXPERIMENTAL] If true Motion-Aware Collision Detection (MACD) will be used for this component */
+	UPROPERTY(EditAnywhere, AdvancedDisplay, Category = Collision)
+	uint8 bUseMACD : 1;
+
+public:
 	/** If true ignore analytic collisions and treat objects as a general implicit surface */
 	UPROPERTY(EditAnywhere, AdvancedDisplay, BlueprintReadOnly, Category = Collision)
 	uint8 bIgnoreAnalyticCollisions : 1;
@@ -925,9 +929,19 @@ public:
 	/** Enables/disabled smoothed edge collisions */
 	ENGINE_API void SetSmoothEdgeCollisionsEnabled(bool bNewSmoothEdgeCollisions);
 
-	/** Enable/disable Continuous Collidion Detection feature */
+	/** Enable/disable Continuous Collision Detection feature */
 	ENGINE_API void SetUseCCD(bool bInUseCCD);
 
+	/** 
+	 * [EXPERIMENTAL] Enable/disable Motion-Aware Collision Detection feature. MACD attempts to take the movement of the
+	 * body into account during collisions detection to reduce the chance of objects passing through each other at moderate
+	 * speeds without the need for CCD. CCD is still required reliable collision between high-speed objects.
+	 */
+	ENGINE_API void SetUseMACD(bool bInUseMACD);
+
+	/** [EXPERIMENTAL] Whether Motion-Aware Collision Detection is enabled */
+	bool GetUseMACD() const { return bUseMACD != 0; }
+
 	/** Disable/Re-Enable this body in the solver,  when disable, the body won't be part of the simulation ( regardless if it's dynamic or kinematic ) and no collision will occur 
 	* this can be used for performance control situation for example
 	*/
diff --git a/Engine/Source/Runtime/Engine/Private/PhysicsEngine/BodyInstance.cpp b/Engine/Source/Runtime/Engine/Private/PhysicsEngine/BodyInstance.cpp
--- a/Engine/Source/Runtime/Engine/Private/PhysicsEngine/BodyInstance.cpp
+++ b/Engine/Source/Runtime/Engine/Private/PhysicsEngine/BodyInstance.cpp
@@ -336,6 +336,7 @@ FBodyInstance::FBodyInstance()
 	, SleepFamily(ESleepFamily::Normal)
 	, DOFMode(0)
 	, bUseCCD(false)
+	, bUseMACD(false)
 	, bIgnoreAnalyticCollisions(false)
 	, bNotifyRigidBodyCollision(false)
 	, bContactModification(false)
@@ -1246,6 +1247,7 @@ void FInitBodiesHelperBase::CreateActor_AssumesLocked(FBodyInstance* Instance, c
 	{
 		FPhysicsInterface::CreateActor(ActorParams, Instance->ActorHandle);
 		FPhysicsInterface::SetCcdEnabled_AssumesLocked(Instance->ActorHandle, Instance->bUseCCD);
+		FPhysicsInterface::SetMACDEnabled_AssumesLocked(Instance->ActorHandle, Instance->bUseMACD);
 		FPhysicsInterface::SetIsKinematic_AssumesLocked(Instance->ActorHandle, !Instance->ShouldInstanceSimulatingPhysics());
 
 		FPhysicsInterface::SetMaxLinearVelocity_AssumesLocked(Instance->ActorHandle, TNumericLimits<float>::Max());
@@ -2394,6 +2396,7 @@ void FBodyInstance::UpdateInstanceSimulatePhysics()
 		bInitialized = true;
 		FPhysicsInterface::SetIsKinematic_AssumesLocked(Actor, !bUseSimulate);
 		FPhysicsInterface::SetCcdEnabled_AssumesLocked(Actor, bUseCCD);
+		FPhysicsInterface::SetMACDEnabled_AssumesLocked(Actor, bUseMACD);
 
 		if(bSimulatePhysics)
 		{
@@ -3699,7 +3702,6 @@ void FBodyInstance::SetSmoothEdgeCollisionsEnabled(bool bNewSmoothEdgeCollisions
 	}
 }
 
-
 void FBodyInstance::SetUseCCD(bool bInUseCCD)
 {
 	if (bUseCCD != bInUseCCD)
@@ -3718,6 +3720,22 @@ void FBodyInstance::SetUseCCD(bool bInUseCCD)
 	}
 }
 
+void FBodyInstance::SetUseMACD(bool bInUseMACD)
+{
+	if (bUseMACD != bInUseMACD)
+	{
+		bUseMACD = bInUseMACD;
+
+		FPhysicsCommand::ExecuteWrite(ActorHandle, [this, bInUseMACD](const FPhysicsActorHandle& Actor)
+			{
+				if (FPhysicsInterface::IsValid(Actor) && FPhysicsInterface::IsRigidBody(Actor))
+				{
+					FPhysicsInterface::SetCcdEnabled_AssumesLocked(Actor, bInUseMACD);
+				}
+			});
+	}
+}
+
 void FBodyInstance::SetPhysicsDisabled(bool bSetDisabled)
 {
 	FPhysicsCommand::ExecuteWrite(ActorHandle, [&](const FPhysicsActorHandle& Actor)
diff --git a/Engine/Source/Runtime/Experimental/Chaos/Private/PhysicsProxy/GeometryCollectionPhysicsProxy.cpp b/Engine/Source/Runtime/Experimental/Chaos/Private/PhysicsProxy/GeometryCollectionPhysicsProxy.cpp
--- a/Engine/Source/Runtime/Experimental/Chaos/Private/PhysicsProxy/GeometryCollectionPhysicsProxy.cpp
+++ b/Engine/Source/Runtime/Experimental/Chaos/Private/PhysicsProxy/GeometryCollectionPhysicsProxy.cpp
@@ -1665,6 +1665,7 @@ void FGeometryCollectionPhysicsProxy::InitializeBodiesPT(Chaos::FPBDRigidsSolver
 				Handle->SetGravityEnabled(bEnableGravity);
 				Handle->SetGravityGroupIndex(Parameters.GravityGroupIndex);
 				Handle->SetCCDEnabled(Parameters.UseCCD);
+				Handle->SetMACDEnabled(Parameters.UseMACD);
 				Handle->SetOneWayInteraction(bIsOneWayInteraction);
 				Handle->SetInertiaConditioningEnabled(Parameters.UseInertiaConditioning);
 				Handle->SetLinearEtherDrag(Parameters.LinearDamping);
diff --git a/Engine/Source/Runtime/Experimental/Chaos/Public/Chaos/ParticleDirtyFlags.h b/Engine/Source/Runtime/Experimental/Chaos/Public/Chaos/ParticleDirtyFlags.h
--- a/Engine/Source/Runtime/Experimental/Chaos/Public/Chaos/ParticleDirtyFlags.h
+++ b/Engine/Source/Runtime/Experimental/Chaos/Public/Chaos/ParticleDirtyFlags.h
@@ -542,6 +542,9 @@ public:
 	bool CCDEnabled() const { return MControlFlags.GetCCDEnabled(); }
 	void SetCCDEnabled(bool bInCCDEnabled) { MControlFlags.SetCCDEnabled(bInCCDEnabled); }
 
+	bool MACDEnabled() const { return MControlFlags.GetMACDEnabled(); }
+	void SetMACDEnabled(bool bInCCDEnabled) { MControlFlags.SetMACDEnabled(bInCCDEnabled); }
+
 	bool Disabled() const { return bDisabled; }
 	void SetDisabled(bool bInDisabled) { bDisabled = bInDisabled; }
 
diff --git a/Engine/Source/Runtime/Experimental/Chaos/Public/Chaos/ParticleHandle.h b/Engine/Source/Runtime/Experimental/Chaos/Public/Chaos/ParticleHandle.h
--- a/Engine/Source/Runtime/Experimental/Chaos/Public/Chaos/ParticleHandle.h
+++ b/Engine/Source/Runtime/Experimental/Chaos/Public/Chaos/ParticleHandle.h
@@ -126,6 +126,7 @@ void PBDRigidParticleDefaultConstruct(FConcrete& Concrete, const FPBDRigidPartic
 	Concrete.SetAngularEtherDrag(0.f);
 	Concrete.SetGravityEnabled(Params.bGravityEnabled);
 	Concrete.SetCCDEnabled(Params.bCCDEnabled);
+	Concrete.SetMACDEnabled(false);
 	Concrete.SetDisabled(Params.bDisabled);
 	Concrete.SetSleepType(ESleepType::MaterialSleep);
 }
@@ -1292,6 +1293,16 @@ public:
 		PBDRigidParticles->ControlFlags(ParticleIdx).SetCCDEnabled(bEnabled);
 	}
 
+	inline bool MACDEnabled() const
+	{
+		return ControlFlags().GetMACDEnabled();
+	}
+
+	inline void SetMACDEnabled(bool bEnabled)
+	{
+		PBDRigidParticles->ControlFlags(ParticleIdx).SetMACDEnabled(bEnabled);
+	}
+
 	inline bool OneWayInteraction() const
 	{ 
 		return ControlFlags().GetOneWayInteractionEnabled();
@@ -3358,6 +3369,12 @@ public:
 		MMiscData.Modify(true, MDirtyFlags, Proxy, [bInEnabled](auto& Data) { Data.SetCCDEnabled(bInEnabled); });
 	}
 
+	bool MACDEnabled() const { return MMiscData.Read().MACDEnabled(); }
+	void SetMACDEnabled(bool bInEnabled)
+	{
+		MMiscData.Modify(true, MDirtyFlags, Proxy, [bInEnabled](auto& Data) { Data.SetMACDEnabled(bInEnabled); });
+	}
+
 	bool InertiaConditioningEnabled() const { return MMiscData.Read().InertiaConditioningEnabled(); }
 	void SetInertiaConditioningEnabled(bool bInEnabled)
 	{
diff --git a/Engine/Source/Runtime/Experimental/Chaos/Public/Chaos/RigidParticleControlFlags.h b/Engine/Source/Runtime/Experimental/Chaos/Public/Chaos/RigidParticleControlFlags.h
--- a/Engine/Source/Runtime/Experimental/Chaos/Public/Chaos/RigidParticleControlFlags.h
+++ b/Engine/Source/Runtime/Experimental/Chaos/Public/Chaos/RigidParticleControlFlags.h
@@ -44,6 +44,9 @@ namespace Chaos
 		bool GetCCDEnabled() const { return Flags.bCCDEnabled; }
 		FRigidParticleControlFlags& SetCCDEnabled(const bool bEnabled) { Flags.bCCDEnabled = bEnabled; return *this; }
 
+		bool GetMACDEnabled() const { return Flags.bMACDEnabled; }
+		FRigidParticleControlFlags& SetMACDEnabled(const bool bEnabled) { Flags.bMACDEnabled = bEnabled; return *this; }
+
 		bool GetOneWayInteractionEnabled() const { return Flags.bOneWayInteractionEnabled; }
 		FRigidParticleControlFlags& SetOneWayInteractionEnabled(const bool bEnabled) { Flags.bOneWayInteractionEnabled = bEnabled; return *this; }
 
@@ -77,6 +80,7 @@ namespace Chaos
 			FStorage bInertiaConditioningEnabled : 1;
 			FStorage GravityGroupIndex : 3;
 			FStorage bUpdateKinematicFromSimulation : 1;
+			FStorage bMACDEnabled : 1;
 			// Add new properties above this line
 			// Change FStorage typedef if we exceed the max bits
 		};
diff --git a/Engine/Source/Runtime/Experimental/Chaos/Public/GeometryCollection/GeometryCollectionSimulationCoreTypes.h b/Engine/Source/Runtime/Experimental/Chaos/Public/GeometryCollection/GeometryCollectionSimulationCoreTypes.h
--- a/Engine/Source/Runtime/Experimental/Chaos/Public/GeometryCollection/GeometryCollectionSimulationCoreTypes.h
+++ b/Engine/Source/Runtime/Experimental/Chaos/Public/GeometryCollection/GeometryCollectionSimulationCoreTypes.h
@@ -214,6 +214,7 @@ struct FSimulationParameters
 		, OneWayInteractionLevel(INDEX_NONE)
 		, UseInertiaConditioning(true)
 		, UseCCD(false)
+		, UseMACD(false)
 		, LinearDamping(0.01f)
 		, AngularDamping(0)
 		, InitialOverlapDepenetrationVelocity(-1.0f)
@@ -277,6 +278,7 @@ struct FSimulationParameters
 		, OneWayInteractionLevel(Other.OneWayInteractionLevel)
 		, UseInertiaConditioning(Other.UseInertiaConditioning)
 		, UseCCD(Other.UseCCD)
+		, UseMACD(Other.UseMACD)
 		, LinearDamping(Other.LinearDamping)
 		, AngularDamping(Other.AngularDamping)
 		, InitialOverlapDepenetrationVelocity(Other.InitialOverlapDepenetrationVelocity)
@@ -370,6 +372,7 @@ struct FSimulationParameters
 	int32 OneWayInteractionLevel;
 	bool UseInertiaConditioning;
 	bool UseCCD;
+	bool UseMACD;
 	float LinearDamping;
 	float AngularDamping;
 	float InitialOverlapDepenetrationVelocity;
diff --git a/Engine/Source/Runtime/Experimental/Chaos/Public/PhysicsProxy/SingleParticlePhysicsProxy.h b/Engine/Source/Runtime/Experimental/Chaos/Public/PhysicsProxy/SingleParticlePhysicsProxy.h
--- a/Engine/Source/Runtime/Experimental/Chaos/Public/PhysicsProxy/SingleParticlePhysicsProxy.h
+++ b/Engine/Source/Runtime/Experimental/Chaos/Public/PhysicsProxy/SingleParticlePhysicsProxy.h
@@ -445,6 +445,30 @@ public:
 			});
 	}
 
+	bool MACDEnabled() const
+	{
+		return Read([](auto* Particle)
+			{
+				if (auto Rigid = Particle->CastToRigidParticle())
+				{
+					return Rigid->MACDEnabled();
+				}
+
+				return false;
+			});
+	}
+
+	void SetMACDEnabled(const bool InCCDEnabled)
+	{
+		Write([InCCDEnabled](auto* Particle)
+			{
+				if (auto Rigid = Particle->CastToRigidParticle())
+				{
+					return Rigid->SetMACDEnabled(InCCDEnabled);
+				}
+			});
+	}
+
 	bool OneWayInteraction() const
 	{
 		return Read([](auto* Particle)
@@ -1247,6 +1271,26 @@ public:
 		return false;
 	}
 
+	void SetMACDEnabled(bool bEnabled)
+	{
+		VerifyContext();
+		if (TPBDRigidParticle<FReal, 3>*Rigid = GetParticle_LowLevel()->CastToRigidParticle())
+		{
+			Rigid->SetMACDEnabled(bEnabled);
+		}
+	}
+
+	bool MACDEnabled() const
+	{
+		VerifyContext();
+		if (const TPBDRigidParticle<FReal, 3>*Rigid = GetParticle_LowLevel()->CastToRigidParticle())
+		{
+			return Rigid->MACDEnabled();
+		}
+
+		return false;
+	}
+
 	void SetMaxLinearSpeedSq(FReal InNewSpeed)
 	{
 		VerifyContext();
diff --git a/Engine/Source/Runtime/Experimental/ChaosVisualDebugger/Private/DataWrappers/ChaosVDParticleDataWrapper.cpp b/Engine/Source/Runtime/Experimental/ChaosVisualDebugger/Private/DataWrappers/ChaosVDParticleDataWrapper.cpp
--- a/Engine/Source/Runtime/Experimental/ChaosVisualDebugger/Private/DataWrappers/ChaosVDParticleDataWrapper.cpp
+++ b/Engine/Source/Runtime/Experimental/ChaosVisualDebugger/Private/DataWrappers/ChaosVDParticleDataWrapper.cpp
@@ -13,6 +13,7 @@ bool FChaosVDFRigidParticleControlFlags::Serialize(FArchive& Ar)
 	Ar << bOneWayInteractionEnabled;
 	Ar << bInertiaConditioningEnabled;
 	Ar << GravityGroupIndex;
+	Ar << bMACDEnabled;
 
 	return !Ar.IsError();
 }
diff --git a/Engine/Source/Runtime/Experimental/ChaosVisualDebugger/Public/DataWrappers/ChaosVDParticleDataWrapper.h b/Engine/Source/Runtime/Experimental/ChaosVisualDebugger/Public/DataWrappers/ChaosVDParticleDataWrapper.h
--- a/Engine/Source/Runtime/Experimental/ChaosVisualDebugger/Public/DataWrappers/ChaosVDParticleDataWrapper.h
+++ b/Engine/Source/Runtime/Experimental/ChaosVisualDebugger/Public/DataWrappers/ChaosVDParticleDataWrapper.h
@@ -66,7 +66,8 @@ struct FChaosVDFRigidParticleControlFlags : public FChaosVDWrapperDataBase
 		  bCCDEnabled(false),
 		  bOneWayInteractionEnabled(false),
 		  bInertiaConditioningEnabled(false), 
-		  GravityGroupIndex(0)
+		  GravityGroupIndex(0),
+		  bMACDEnabled(false)
 	{
 	}
 
@@ -80,6 +81,7 @@ struct FChaosVDFRigidParticleControlFlags : public FChaosVDWrapperDataBase
 		bOneWayInteractionEnabled = Other.GetOneWayInteractionEnabled();
 		bInertiaConditioningEnabled = Other.GetInertiaConditioningEnabled();
 		GravityGroupIndex = Other.GetGravityGroupIndex();
+		bMACDEnabled = Other.GetMACDEnabled();
 
 		bHasValidData = true;
 	}
@@ -94,6 +96,8 @@ struct FChaosVDFRigidParticleControlFlags : public FChaosVDWrapperDataBase
 	bool bInertiaConditioningEnabled;
 	UPROPERTY(EditAnywhere, Category= "Particle Control Flags")
 	int32 GravityGroupIndex;
+	UPROPERTY(EditAnywhere, Category = "Particle Control Flags")
+	bool bMACDEnabled;
 };
 
 inline FArchive& operator<<(FArchive& Ar, FChaosVDFRigidParticleControlFlags& Data)
diff --git a/Engine/Source/Runtime/Experimental/GeometryCollectionEngine/Private/GeometryCollection/GeometryCollectionComponent.cpp b/Engine/Source/Runtime/Experimental/GeometryCollectionEngine/Private/GeometryCollection/GeometryCollectionComponent.cpp
--- a/Engine/Source/Runtime/Experimental/GeometryCollectionEngine/Private/GeometryCollection/GeometryCollectionComponent.cpp
+++ b/Engine/Source/Runtime/Experimental/GeometryCollectionEngine/Private/GeometryCollection/GeometryCollectionComponent.cpp
@@ -3830,6 +3830,7 @@ void UGeometryCollectionComponent::RegisterAndInitializePhysicsProxy()
 		SimulationParameters.OneWayInteractionLevel = OneWayInteractionLevel;
 		SimulationParameters.UseInertiaConditioning = BodyInstance.IsInertiaConditioningEnabled();
 		SimulationParameters.UseCCD = BodyInstance.bUseCCD;
+		SimulationParameters.UseMACD = BodyInstance.GetUseMACD();
 		SimulationParameters.LinearDamping = BodyInstance.LinearDamping;
 		SimulationParameters.AngularDamping = BodyInstance.AngularDamping;
 		SimulationParameters.InitialOverlapDepenetrationVelocity = BodyInstance.GetMaxDepenetrationVelocity();
diff --git a/Engine/Source/Runtime/PhysicsCore/Private/ChaosEngineInterface.cpp b/Engine/Source/Runtime/PhysicsCore/Private/ChaosEngineInterface.cpp
--- a/Engine/Source/Runtime/PhysicsCore/Private/ChaosEngineInterface.cpp
+++ b/Engine/Source/Runtime/PhysicsCore/Private/ChaosEngineInterface.cpp
@@ -543,6 +543,11 @@ void FChaosEngineInterface::SetCcdEnabled_AssumesLocked(const FPhysicsActorHandl
 	InActorReference->GetGameThreadAPI().SetCCDEnabled(bIsCcdEnabled);
 }
 
+void FChaosEngineInterface::SetMACDEnabled_AssumesLocked(const FPhysicsActorHandle& InActorReference, bool bIsMACDEnabled)
+{
+	InActorReference->GetGameThreadAPI().SetMACDEnabled(bIsMACDEnabled);
+}
+
 void FChaosEngineInterface::SetIgnoreAnalyticCollisions_AssumesLocked(const FPhysicsActorHandle& InActorReference,bool bIgnoreAnalyticCollisions)
 {
 	InActorReference->GetGameThreadAPI().SetIgnoreAnalyticCollisions(bIgnoreAnalyticCollisions);
diff --git a/Engine/Source/Runtime/PhysicsCore/Public/Chaos/ChaosEngineInterface.h b/Engine/Source/Runtime/PhysicsCore/Public/Chaos/ChaosEngineInterface.h
--- a/Engine/Source/Runtime/PhysicsCore/Public/Chaos/ChaosEngineInterface.h
+++ b/Engine/Source/Runtime/PhysicsCore/Public/Chaos/ChaosEngineInterface.h
@@ -406,6 +406,7 @@ public:
 
 	static PHYSICSCORE_API void SetIsKinematic_AssumesLocked(const FPhysicsActorHandle& InActorReference,bool bIsKinematic);
 	static PHYSICSCORE_API void SetCcdEnabled_AssumesLocked(const FPhysicsActorHandle& InActorReference,bool bIsCcdEnabled);
+	static PHYSICSCORE_API void SetMACDEnabled_AssumesLocked(const FPhysicsActorHandle& InActorReference, bool bIsMACDEnabled);
 	static PHYSICSCORE_API void SetIgnoreAnalyticCollisions_AssumesLocked(const FPhysicsActorHandle& InActorReference,bool bIsCcdEnabled);
 
 	static PHYSICSCORE_API FTransform GetGlobalPose_AssumesLocked(const FPhysicsActorHandle& InActorReference);
