If we recieve an IE_Repeat event but have not recieved an IE_Pressed key event then we should mark the IE_Pressed as having happened already, since it is impossible to get a repeat event of a non-analog key without an initial IE_Pressed event.

This solves an issue where, if you are holding a key down during a level transition, they keystate map would incorrectly evaluate the state of the key as being released and re-pressed every frame. This meant that, for example, if you had a "Released" trigger on an Enhanced Input mapping context (or you were listening to any released event on a legacy mapping) that it would constantly fire as you hold down the key instead of only firing on release.

This behavior can be disabled with the "Input.bAutoReconcilePressedEventsOnFirstRepeat" Cvar.

diff --git a/Engine/Source/Runtime/Engine/Private/UserInterface/PlayerInput.cpp b/Engine/Source/Runtime/Engine/Private/UserInterface/PlayerInput.cpp
--- a/Engine/Source/Runtime/Engine/Private/UserInterface/PlayerInput.cpp
+++ b/Engine/Source/Runtime/Engine/Private/UserInterface/PlayerInput.cpp
@@ -45,6 +45,11 @@ namespace UE
 		static FAutoConsoleVariableRef CVarShouldOnlyTriggerLastActionInChord(TEXT("Input.AxisEventsCanBeConsumed"),
 			AxisEventsCanBeConsumed,
 			TEXT("If true and all FKey's for a given Axis Event are consumed, then the axis delegate will not fire."));
+
+		static bool bAutoReconcilePressedEventsOnFirstRepeat = true;
+		static FAutoConsoleVariableRef CVarAutoReconcilePressedEventsOnFirstRepeat(TEXT("Input.AutoReconcilePressedEventsOnFirstRepeat"),
+			bAutoReconcilePressedEventsOnFirstRepeat,
+			TEXT("If true, then we will automatically mark a IE_Pressed event if we receive an IE_Repeat event but have not received a pressed event first.\nNote: This option will be removed in a future update."));
 	}
 }
 
@@ -327,11 +332,32 @@ bool UPlayerInput::InputKey(const FInputKeyParams& Params)
 	// Non-analog key
 	else
 	{
+		FKeyState* ExistingKeyState = KeyStateMap.Find(Params.Key);
+		
 		// first event associated with this key, add it to the map
-		FKeyState& KeyState = KeyStateMap.FindOrAdd(Params.Key);
+		FKeyState& KeyState = !ExistingKeyState ? KeyStateMap.Add(Params.Key) : *ExistingKeyState;
+
 		UWorld* World = GetWorld();
 		check(World);
 
+		const bool bIsFirstEventForKey = (ExistingKeyState == nullptr);
+
+		// If this is the first key press for us and it is a repeat, then we have missed the initial IE_Pressed event.
+		// This can happen if you are holding down a key between level transitions and the player controller gets recreated,
+		// which means that we will be using a new UPlayerInput object and the KeyState map is emptied.
+		if (UE::Input::bAutoReconcilePressedEventsOnFirstRepeat && bIsFirstEventForKey && Params.Event == IE_Repeat && KeyState.EventAccumulator[IE_Pressed].IsEmpty())
+		{
+			// Mark as having received the IE_Pressed event already, so that we can correctly evaluate the 
+			// state of the IE_Repeat event. It is impossible to get a IE_Repeat with an initial IE_Pressed somewhere
+			//
+			// Without this, the key will incorrectly evaluate as being released/pressed
+			// every frame, even if you are just holding it
+			KeyState.RawValueAccumulator.X = Params.Delta.X;
+			KeyState.EventAccumulator[IE_Pressed].Add(++EventCount);
+			KeyState.LastUpDownTransitionTime = World->GetRealTimeSeconds();
+			KeyState.SampleCountAccumulator++;
+		}
+
 		switch(Params.Event)
 		{
 		case IE_Pressed:
