[Backout] - CL31509445
Original CL Desc
-----------------------------------------------------------------
Fixing live attenuation updates.

Made the max distance value use the attenuation setting asset and not use the cached value.

diff --git a/Engine/Plugins/AudioGameplay/Source/AudioGameplay/Private/SoundHandleSubsystem.cpp b/Engine/Plugins/AudioGameplay/Source/AudioGameplay/Private/SoundHandleSubsystem.cpp
--- a/Engine/Plugins/AudioGameplay/Source/AudioGameplay/Private/SoundHandleSubsystem.cpp
+++ b/Engine/Plugins/AudioGameplay/Source/AudioGameplay/Private/SoundHandleSubsystem.cpp
@@ -84,6 +84,8 @@ Audio::EResult USoundHandleSubsystem::Play(const Audio::FSoundHandleID ID)
 				Handle->ActiveSound.AttenuationSettings = *AttenuationSettings;
 				Handle->ActiveSound.FocusData.PriorityScale = AttenuationSettings->GetFocusPriorityScale(GetAudioDeviceHandle().GetAudioDevice()->GetGlobalFocusSettings(), FocusFactor);
 			}
+
+			Handle->ActiveSound.MaxDistance = MaxDistance;
 		}
 		
 		AudioDevice->AddNewActiveSound(Handle->ActiveSound);		
diff --git a/Engine/Source/Runtime/Engine/Private/AudioDebug.cpp b/Engine/Source/Runtime/Engine/Private/AudioDebug.cpp
--- a/Engine/Source/Runtime/Engine/Private/AudioDebug.cpp
+++ b/Engine/Source/Runtime/Engine/Private/AudioDebug.cpp
@@ -1110,7 +1110,7 @@ namespace Audio
 			const FString Name = Sound->GetName();
 			const FTransform CurTransform = ActiveSound.Transform;
 			FColor TextColor = FColor::White;
-			const float CurMaxDistance = ActiveSound.GetMaxDistance();
+			const float CurMaxDistance = ActiveSound.MaxDistance;
 			float DisplayValue = 0.0f;
 			float FilterValue = 0.0f;
 			if (ActiveSoundVisualizeModeCVar == 1 || ActiveSoundVisualizeModeCVar == 2)
diff --git a/Engine/Source/Runtime/Engine/Private/AudioDevice.cpp b/Engine/Source/Runtime/Engine/Private/AudioDevice.cpp
--- a/Engine/Source/Runtime/Engine/Private/AudioDevice.cpp
+++ b/Engine/Source/Runtime/Engine/Private/AudioDevice.cpp
@@ -6397,6 +6397,8 @@ void FAudioDevice::PlaySoundAtLocation(USoundBase* Sound, UWorld* World, float V
 			NewActiveSound.FocusData.DistanceScale = AttenuationSettingsToApply->GetFocusDistanceScale(FocusSettings, FocusFactor);
 		}
 
+		NewActiveSound.MaxDistance = MaxDistance;
+
 		if (Concurrency)
 		{
 			NewActiveSound.ConcurrencySet.Add(Concurrency);
diff --git a/Engine/Source/Runtime/Engine/Private/AudioVirtualLoop.cpp b/Engine/Source/Runtime/Engine/Private/AudioVirtualLoop.cpp
--- a/Engine/Source/Runtime/Engine/Private/AudioVirtualLoop.cpp
+++ b/Engine/Source/Runtime/Engine/Private/AudioVirtualLoop.cpp
@@ -150,7 +150,7 @@ void FAudioVirtualLoop::CalculateUpdateInterval()
 	check(AudioDevice);
 
 	const float DistanceToListener = AudioDevice->GetDistanceToNearestListener(ActiveSound->Transform.GetLocation());
-	const float DistanceRatio = (DistanceToListener - ActiveSound->GetMaxDistance()) / FMath::Max(VirtualLoopsPerfDistanceCVar, 1.0f);
+	const float DistanceRatio = (DistanceToListener - ActiveSound->MaxDistance) / FMath::Max(VirtualLoopsPerfDistanceCVar, 1.0f);
 	const float DistanceRatioClamped = FMath::Clamp(DistanceRatio, 0.0f, 1.0f);
 	UpdateInterval = FMath::Lerp(VirtualLoopsUpdateRateMinCVar, VirtualLoopsUpdateRateMaxCVar, DistanceRatioClamped);
 }
@@ -215,7 +215,7 @@ bool FAudioVirtualLoop::IsInAudibleRange(const FActiveSound& InActiveSound, cons
 
 	DistanceScale = FMath::Max(DistanceScale, UE_KINDA_SMALL_NUMBER);
 	const FVector Location = InActiveSound.Transform.GetLocation();
-	return AudioDevice->LocationIsAudible(Location, InActiveSound.GetMaxDistance() / DistanceScale);
+	return AudioDevice->LocationIsAudible(Location, InActiveSound.MaxDistance / DistanceScale);
 }
 
 void FAudioVirtualLoop::UpdateFocusData(float DeltaTime)
diff --git a/Engine/Source/Runtime/Engine/Private/Components/AudioComponent.cpp b/Engine/Source/Runtime/Engine/Private/Components/AudioComponent.cpp
--- a/Engine/Source/Runtime/Engine/Private/Components/AudioComponent.cpp
+++ b/Engine/Source/Runtime/Engine/Private/Components/AudioComponent.cpp
@@ -817,6 +817,8 @@ void UAudioComponent::PlayInternal(const PlayInternalRequestData& InPlayRequestD
 	// Pass down any source buffer listener we have
 	NewActiveSound.SetSourceListener(SourceBufferListener, bShouldSourceBufferListenerZeroBuffer);
 
+	NewActiveSound.MaxDistance = MaxDistance;
+
 	// Setup the submix and bus sends that may have been set before playing
 	for (FSoundSubmixSendInfo& SubmixSendInfo : PendingSubmixSends)
 	{
@@ -866,9 +868,6 @@ void UAudioComponent::PlayInternal(const PlayInternalRequestData& InPlayRequestD
 	TArray<FAudioParameter> InstanceParamsCopy = InstanceParameters;
 	FAudioParameter::Merge(MoveTemp(InstanceParamsCopy), SoundParams);
 
-	// Make sure our new active sound has the correct audio device
-	NewActiveSound.SetAudioDevice(AudioDevice);
-
 	AudioDevice->AddNewActiveSound(NewActiveSound, MoveTemp(SoundParams));
 
 	LastSoundPlayOrder = NewActiveSound.GetPlayOrder();
