Fix MetaSound trigger parameter being retriggered on realizing virtualized sounds and unneeded parameter updates being cached while virtualized

diff --git a/Engine/Plugins/Runtime/Metasound/Source/MetasoundFrontend/Private/MetasoundParameterTransmitter.cpp b/Engine/Plugins/Runtime/Metasound/Source/MetasoundFrontend/Private/MetasoundParameterTransmitter.cpp
--- a/Engine/Plugins/Runtime/Metasound/Source/MetasoundFrontend/Private/MetasoundParameterTransmitter.cpp
+++ b/Engine/Plugins/Runtime/Metasound/Source/MetasoundFrontend/Private/MetasoundParameterTransmitter.cpp
@@ -268,8 +268,41 @@ namespace Metasound
 		AvailableParameterNames.Remove(InName);
 	}
 
+	bool FMetaSoundParameterTransmitter::SetVirtualizedParameters(TArray<FAudioParameter>&& InParameters)
+	{
+		bool bSuccess = true;
+
+		// Remove triggers
+		for (int32 ParamIndex = InParameters.Num() - 1; ParamIndex >= 0; --ParamIndex)
+		{
+			// Triggers are transient and are not applied for virtualized sounds. 
+			// If a cached value is desired, use SetBoolParameter
+			// (see comment for IAudioParameterControllerInterface::SetTriggerParameter)
+			FAudioParameter& Param = InParameters[ParamIndex];
+			if (Param.ParamType == EAudioParameterType::Trigger)
+			{
+				InParameters.RemoveAtSwap(ParamIndex, 1, EAllowShrinking::No);
+			}
+		}
+
+		if (!InParameters.IsEmpty())
+		{
+			bSuccess &= FParameterTransmitterBase::SetParameters(MoveTemp(InParameters));
+		}
+
+		InParameters.Reset();
+		return bSuccess;
+	}
+
 	bool FMetaSoundParameterTransmitter::SetParameters(TArray<FAudioParameter>&& InParameters)
 	{
+		// Don't set parameters directly if the active sound 
+		// is currently virtualized (to prevent accumulation of unneeded updates) 
+		if (bIsVirtualized)
+		{
+			return SetVirtualizedParameters(MoveTemp(InParameters));
+		}
+
 		bool bSuccess = true;
 
 		for (int32 ParamIndex = InParameters.Num() - 1; ParamIndex >= 0; --ParamIndex)
diff --git a/Engine/Plugins/Runtime/Metasound/Source/MetasoundFrontend/Public/MetasoundParameterTransmitter.h b/Engine/Plugins/Runtime/Metasound/Source/MetasoundFrontend/Public/MetasoundParameterTransmitter.h
--- a/Engine/Plugins/Runtime/Metasound/Source/MetasoundFrontend/Public/MetasoundParameterTransmitter.h
+++ b/Engine/Plugins/Runtime/Metasound/Source/MetasoundFrontend/Public/MetasoundParameterTransmitter.h
@@ -142,6 +142,13 @@ namespace Metasound
 		bool SetParameterWithLiteral(FName InParameterName, const FLiteral& InValue);
 
 	private:
+		/** Set parameters when virtualized
+		 * 
+		 *  @return true if all non trigger parameters are set on the parameter base
+		 *	(trigger parameters are removed on virtualized sounds 
+		 *	and parameters are not actually set on the MetaSound until realization)
+		 */
+		bool SetVirtualizedParameters(TArray<FAudioParameter>&& InParameters);
 
 		uint64 InstanceID;
 		FName DebugMetaSoundName;
diff --git a/Engine/Source/Runtime/AudioExtensions/Private/IAudioParameterTransmitter.cpp b/Engine/Source/Runtime/AudioExtensions/Private/IAudioParameterTransmitter.cpp
--- a/Engine/Source/Runtime/AudioExtensions/Private/IAudioParameterTransmitter.cpp
+++ b/Engine/Source/Runtime/AudioExtensions/Private/IAudioParameterTransmitter.cpp
@@ -14,6 +14,7 @@ namespace Audio
 
 	FParameterTransmitterBase::FParameterTransmitterBase(TArray<FAudioParameter>&& InDefaultParams)
 		: AudioParameters(MoveTemp(InDefaultParams))
+		, bIsVirtualized(false)
 	{
 	}
 
@@ -49,4 +50,16 @@ namespace Audio
 		FAudioParameter::Merge(MoveTemp(InParameters), AudioParameters);
 		return true;
 	}
+
+	void FParameterTransmitterBase::OnVirtualizeActiveSound()
+	{ 
+		bIsVirtualized = true;
+		ResetParameters();
+	}
+
+	 void FParameterTransmitterBase::OnRealizeVirtualizedActiveSound(TArray<FAudioParameter>&& InParameters)
+	 { 
+		 bIsVirtualized = false;
+		 SetParameters(MoveTemp(InParameters));
+	 }
 } // namespace Audio
diff --git a/Engine/Source/Runtime/AudioExtensions/Public/IAudioParameterTransmitter.h b/Engine/Source/Runtime/AudioExtensions/Public/IAudioParameterTransmitter.h
--- a/Engine/Source/Runtime/AudioExtensions/Public/IAudioParameterTransmitter.h
+++ b/Engine/Source/Runtime/AudioExtensions/Public/IAudioParameterTransmitter.h
@@ -76,6 +76,13 @@ namespace Audio
 
 			// Parameter Setters
 			virtual bool SetParameters(TArray<FAudioParameter>&& InParameters) = 0;
+
+			// Called when the active sound is virtualized  
+			virtual void OnVirtualizeActiveSound() {}
+
+			// Called when the virtualized active sound is realized 
+			// @param Parameters to set 
+			virtual void OnRealizeVirtualizedActiveSound(TArray<FAudioParameter>&& InParameters) {}
 	};
 
 	/** Base implementation for the parameter transmitter, which caches parameters
@@ -93,8 +100,11 @@ namespace Audio
 		UE_DEPRECATED(5.2, "Use ResetParameters() or OnDeleteActiveSound() instead depending on use case.")
 		AUDIOEXTENSIONS_API virtual bool Reset() override;
 		AUDIOEXTENSIONS_API virtual bool SetParameters(TArray<FAudioParameter>&& InParameters) override;
+		AUDIOEXTENSIONS_API virtual void OnVirtualizeActiveSound() override;
+		AUDIOEXTENSIONS_API virtual void OnRealizeVirtualizedActiveSound(TArray<FAudioParameter>&& InParameters);
 
 	protected:
 		TArray<FAudioParameter> AudioParameters;
+		bool bIsVirtualized;
 	};
 } // namespace Audio
diff --git a/Engine/Source/Runtime/Engine/Private/AudioVirtualLoop.cpp b/Engine/Source/Runtime/Engine/Private/AudioVirtualLoop.cpp
--- a/Engine/Source/Runtime/Engine/Private/AudioVirtualLoop.cpp
+++ b/Engine/Source/Runtime/Engine/Private/AudioVirtualLoop.cpp
@@ -125,6 +125,11 @@ bool FAudioVirtualLoop::Virtualize(const FActiveSound& InActiveSound, FAudioDevi
 	OutVirtualLoop.ActiveSound = ActiveSound;
 	OutVirtualLoop.CalculateUpdateInterval();
 
+	if (Audio::IParameterTransmitter* ParameterTransmitter = ActiveSound->GetTransmitter())
+	{
+		ParameterTransmitter->OnVirtualizeActiveSound();
+	}
+
 #if UE_AUDIO_PROFILERTRACE_ENABLED
 	const bool bChannelEnabled = UE_TRACE_CHANNELEXPR_IS_ENABLED(AudioChannel);
 	if (bChannelEnabled)
@@ -308,6 +313,21 @@ bool FAudioVirtualLoop::Update(float DeltaTime, bool bForceUpdate)
 		return false;
 	}
 
+	// Update parameters on virtual loop realization
+	if (Audio::IParameterTransmitter* ParameterTransmitter = ActiveSound->GetTransmitter())
+	{
+		uint64 AudioComponentID = ActiveSound->GetAudioComponentID();
+		if (AudioComponentID > 0)
+		{
+			if (UAudioComponent* AudioComponent = UAudioComponent::GetAudioComponentFromID(AudioComponentID))
+			{
+				// Use instance parameters which will not contain transient parameters
+				TArray<FAudioParameter> InstanceParams = AudioComponent->GetInstanceParameters();
+				ParameterTransmitter->OnRealizeVirtualizedActiveSound(MoveTemp(InstanceParams));
+			}
+		}
+	}
+
 #if UE_AUDIO_PROFILERTRACE_ENABLED
 	if (bChannelEnabled)
 	{
diff --git a/Engine/Source/Runtime/Engine/Private/SoundCue.cpp b/Engine/Source/Runtime/Engine/Private/SoundCue.cpp
--- a/Engine/Source/Runtime/Engine/Private/SoundCue.cpp
+++ b/Engine/Source/Runtime/Engine/Private/SoundCue.cpp
@@ -931,13 +931,37 @@ TArray<const TObjectPtr<UObject>*> FSoundCueParameterTransmitter::GetReferencedO
 
 bool FSoundCueParameterTransmitter::SetParameters(TArray<FAudioParameter>&& InParameters)
 {
-	TArray<FAudioParameter> TempParams = InParameters;
-	FAudioParameter::Merge(MoveTemp(TempParams), ParamsToSet);
+	auto RemoveTriggerParameters = [&]()
+	{
+		for (int32 ParamIndex = InParameters.Num() - 1; ParamIndex >= 0; --ParamIndex)
+		{
+			// Triggers are transient and are not applied for virtualized sounds. 
+			// If a cached value is desired, use SetBoolParameter
+			// (see comment for IAudioParameterControllerInterface::SetTriggerParameter)
+			FAudioParameter& Param = InParameters[ParamIndex];
+			if (Param.ParamType == EAudioParameterType::Trigger)
+			{
+				InParameters.RemoveAtSwap(ParamIndex, 1, EAllowShrinking::No);
+			}
+		}
+	};
 
-	InParameters.FilterByPredicate([](const FAudioParameter& Param)
+	if (bIsVirtualized)
 	{
-		return Param.ParamType != EAudioParameterType::Trigger;
-	});
-		
-	return Audio::FParameterTransmitterBase::SetParameters(MoveTemp(InParameters));
+		RemoveTriggerParameters();
+
+		TArray<FAudioParameter> TempParams = InParameters;
+		FAudioParameter::Merge(MoveTemp(TempParams), ParamsToSet);
+
+		return Audio::FParameterTransmitterBase::SetParameters(MoveTemp(InParameters));
+	}
+	else
+	{
+		TArray<FAudioParameter> TempParams = InParameters;
+		FAudioParameter::Merge(MoveTemp(TempParams), ParamsToSet);
+
+		RemoveTriggerParameters();
+
+		return Audio::FParameterTransmitterBase::SetParameters(MoveTemp(InParameters));
+	}
 }
