[Backout] - CL31509445
Original CL Desc
-----------------------------------------------------------------
Fixing live attenuation updates.

Made the max distance value use the attenuation setting asset and not use the cached value.

diff --git a/Engine/Source/Runtime/Engine/Private/ActiveSound.cpp b/Engine/Source/Runtime/Engine/Private/ActiveSound.cpp
--- a/Engine/Source/Runtime/Engine/Private/ActiveSound.cpp
+++ b/Engine/Source/Runtime/Engine/Private/ActiveSound.cpp
@@ -1399,28 +1399,6 @@ FWaveInstance& FActiveSound::AddWaveInstance(const UPTRINT WaveInstanceHash)
 	return *WaveInstance;
 }
 
-float FActiveSound::GetMaxDistance() const
-{
-	float CurrentMaxDistance = 0.0f;
-	float CurrentFocusFactor = 1.0f;
-	GetMaxDistanceAndFocusFactor(CurrentMaxDistance, CurrentFocusFactor);
-	return CurrentMaxDistance;
-}
-
-void FActiveSound::GetMaxDistanceAndFocusFactor(float& OutMaxDistance, float& OutFocusFactor) const
-{
-	OutMaxDistance = 0.0f;
-	OutFocusFactor = 1.0f;
-
-	if (!Sound)
-	{
-		return;
-	}
-
-	check(AudioDevice);
-	AudioDevice->GetMaxDistanceAndFocusFactor(Sound, GetWorld(), LastLocation, GetAttenuationSettings(), OutMaxDistance, OutFocusFactor);
-}
-
 void FActiveSound::ApplyRadioFilter(const FSoundParseParameters& ParseParams)
 {
 	check(AudioDevice);
@@ -1696,7 +1674,7 @@ void FActiveSound::UpdateAttenuation(float DeltaTime, FSoundParseParameters& Par
 	UpdateAttenuation(DeltaTime, ParseParams, Listener.ListenerIndex, SettingsAttenuationNode);
 }
 
-const FSoundAttenuationSettings* FActiveSound::GetAttenuationSettings(const FSoundAttenuationSettings* SettingsAttenuationNode) const
+void FActiveSound::UpdateAttenuation(float DeltaTime, FSoundParseParameters& ParseParams, int32 ListenerIndex, const FSoundAttenuationSettings* SettingsAttenuationNode)
 {
 	// We default to using the copied off "overridden" settings (or default constructed settings)
 	const FSoundAttenuationSettings* Settings = &AttenuationSettings;
@@ -1715,13 +1693,6 @@ const FSoundAttenuationSettings* FActiveSound::GetAttenuationSettings(const FSou
 			Settings = &SoundAttenuation->Attenuation;
 		}
 	}
-	return Settings;
-}
-
-void FActiveSound::UpdateAttenuation(float DeltaTime, FSoundParseParameters& ParseParams, int32 ListenerIndex, const FSoundAttenuationSettings* SettingsAttenuationNode)
-{
-	// We default to using the copied off "overridden" settings (or default constructed settings)
-	const FSoundAttenuationSettings* Settings = GetAttenuationSettings(SettingsAttenuationNode);
 
 	// Reset Focus data and recompute if necessary
 	FAttenuationFocusData FocusDataToApply;
diff --git a/Engine/Source/Runtime/Engine/Private/AudioDevice.cpp b/Engine/Source/Runtime/Engine/Private/AudioDevice.cpp
--- a/Engine/Source/Runtime/Engine/Private/AudioDevice.cpp
+++ b/Engine/Source/Runtime/Engine/Private/AudioDevice.cpp
@@ -6023,7 +6023,7 @@ bool FAudioDevice::SoundIsAudible(const FActiveSound& NewActiveSound)
 		return true;
 	}
 
-	const float ApparentMaxDistance = NewActiveSound.GetMaxDistance() * NewActiveSound.FocusData.DistanceScale;
+	const float ApparentMaxDistance = NewActiveSound.MaxDistance * NewActiveSound.FocusData.DistanceScale;
 	if (LocationIsAudible(NewActiveSound.Transform.GetLocation(), ApparentMaxDistance))
 	{
 		return true;
diff --git a/Engine/Source/Runtime/Engine/Private/AudioVirtualLoop.cpp b/Engine/Source/Runtime/Engine/Private/AudioVirtualLoop.cpp
--- a/Engine/Source/Runtime/Engine/Private/AudioVirtualLoop.cpp
+++ b/Engine/Source/Runtime/Engine/Private/AudioVirtualLoop.cpp
@@ -125,11 +125,6 @@ bool FAudioVirtualLoop::Virtualize(const FActiveSound& InActiveSound, FAudioDevi
 	OutVirtualLoop.ActiveSound = ActiveSound;
 	OutVirtualLoop.CalculateUpdateInterval();
 
-	if (Audio::IParameterTransmitter* ParameterTransmitter = ActiveSound->GetTransmitter())
-	{
-		ParameterTransmitter->OnVirtualizeActiveSound();
-	}
-
 #if UE_AUDIO_PROFILERTRACE_ENABLED
 	const bool bChannelEnabled = UE_TRACE_CHANNELEXPR_IS_ENABLED(AudioChannel);
 	if (bChannelEnabled)
@@ -313,21 +308,6 @@ bool FAudioVirtualLoop::Update(float DeltaTime, bool bForceUpdate)
 		return false;
 	}
 
-	// Update parameters on virtual loop realization
-	if (Audio::IParameterTransmitter* ParameterTransmitter = ActiveSound->GetTransmitter())
-	{
-		uint64 AudioComponentID = ActiveSound->GetAudioComponentID();
-		if (AudioComponentID > 0)
-		{
-			if (UAudioComponent* AudioComponent = UAudioComponent::GetAudioComponentFromID(AudioComponentID))
-			{
-				// Use instance parameters which will not contain transient parameters
-				TArray<FAudioParameter> InstanceParams = AudioComponent->GetInstanceParameters();
-				ParameterTransmitter->OnRealizeVirtualizedActiveSound(MoveTemp(InstanceParams));
-			}
-		}
-	}
-
 #if UE_AUDIO_PROFILERTRACE_ENABLED
 	if (bChannelEnabled)
 	{
diff --git a/Engine/Source/Runtime/Engine/Public/ActiveSound.h b/Engine/Source/Runtime/Engine/Public/ActiveSound.h
--- a/Engine/Source/Runtime/Engine/Public/ActiveSound.h
+++ b/Engine/Source/Runtime/Engine/Public/ActiveSound.h
@@ -371,15 +371,6 @@ public:
 	/** Whether or not the active sound is currently playing audible sound. */
 	bool IsPlayingAudio() const { return bIsPlayingAudio; }
 
-	/** Returns the max distance of the active sound */
-	float GetMaxDistance() const;
-
-	/** Returns both the max distance and focus factor of the active sound */
-	void GetMaxDistanceAndFocusFactor(float& OutMaxDistance, float& OutFocusFactor) const;
-
-	/** Returns the currently used attenuation settings of the active sound. Use this to account for override state, etc. Allows for dynamic updates. */
-	const FSoundAttenuationSettings* GetAttenuationSettings(const FSoundAttenuationSettings* SettingsAttenuationNode = nullptr) const;
-
 	/** Whether or not sound reference is valid and set to play when silent. */
 	ENGINE_API bool IsPlayWhenSilent() const;
 
