Adding Physics costs for Skeletal and Static meshes so that the sizes are discoverable in the Content Browser window.
#tests Editor build and testing

diff --git a/Engine/Source/Runtime/Engine/Private/SkeletalMesh.cpp b/Engine/Source/Runtime/Engine/Private/SkeletalMesh.cpp
--- a/Engine/Source/Runtime/Engine/Private/SkeletalMesh.cpp
+++ b/Engine/Source/Runtime/Engine/Private/SkeletalMesh.cpp
@@ -3750,6 +3750,25 @@ void USkeletalMesh::GetAssetRegistryTags(FAssetRegistryTagsContext Context) cons
 		NumTriangles = LODData.GetTotalFaces();
 		NumVertices = LODData.GetNumVertices();
 	}
+
+#if WITH_EDITORONLY_DATA
+	uint64 PhysicsSize = 0;
+
+	if (UBodySetup* MeshBodySetup = GetBodySetup())
+	{
+		FResourceSizeEx EstimatedSize(EResourceSizeMode::EstimatedTotal);
+		MeshBodySetup->GetResourceSizeEx(EstimatedSize);
+		PhysicsSize = EstimatedSize.GetTotalMemoryBytes();
+	}
+
+	if (UPhysicsAsset* PhysAssetSetup = GetPhysicsAsset())
+	{
+		FResourceSizeEx EstimatedSize(EResourceSizeMode::EstimatedTotal);
+		PhysAssetSetup->GetResourceSizeEx(EstimatedSize);
+		PhysicsSize += EstimatedSize.GetTotalMemoryBytes();
+	}
+	Context.AddTag(FAssetRegistryTag("PhysicsSize", FString::Printf(TEXT("%llu"), PhysicsSize), FAssetRegistryTag::TT_Numerical, FAssetRegistryTag::TD_Memory));
+#endif
 	
 	int32 NumLODs = GetLODInfoArray().Num();
 
diff --git a/Engine/Source/Runtime/Engine/Private/StaticMesh.cpp b/Engine/Source/Runtime/Engine/Private/StaticMesh.cpp
--- a/Engine/Source/Runtime/Engine/Private/StaticMesh.cpp
+++ b/Engine/Source/Runtime/Engine/Private/StaticMesh.cpp
@@ -4816,6 +4816,9 @@ void UStaticMesh::GetAssetRegistryTags(FAssetRegistryTagsContext Context) const
 	int32 NumVertices = 0;
 	int32 NumUVChannels = 0;
 	int32 NumLODs = 0;
+#if WITH_EDITORONLY_DATA
+	uint64 PhysicsSize = 0;
+#endif
 
 	if (GetRenderData() && GetRenderData()->LODResources.Num() > 0)
 	{
@@ -4832,6 +4835,12 @@ void UStaticMesh::GetAssetRegistryTags(FAssetRegistryTagsContext Context) const
 	if (GetBodySetup() != nullptr)
 	{
 		NumCollisionPrims = GetBodySetup()->AggGeom.GetElementCount();
+
+#if WITH_EDITORONLY_DATA
+		FResourceSizeEx EstimatedSize(EResourceSizeMode::EstimatedTotal);
+		GetBodySetup()->GetResourceSizeEx(EstimatedSize);
+		PhysicsSize = EstimatedSize.GetTotalMemoryBytes();
+#endif
 	}
 
 	FBoxSphereBounds Bounds(ForceInit);
@@ -4897,6 +4906,7 @@ void UStaticMesh::GetAssetRegistryTags(FAssetRegistryTagsContext Context) const
 	
 #if WITH_EDITORONLY_DATA
 	Context.AddTag(FAssetRegistryTag("HasHiResMesh", IsHiResMeshDescriptionValid() ? TEXT("True") : TEXT("False"), FAssetRegistryTag::TT_Alphabetical));
+	Context.AddTag(FAssetRegistryTag("PhysicsSize", FString::Printf(TEXT("%llu"), PhysicsSize), FAssetRegistryTag::TT_Numerical, FAssetRegistryTag::TD_Memory));
 #endif
 
 	Super::GetAssetRegistryTags(Context);
