Re-Add Helen's change (accidentally backed out)

diff --git a/Engine/Source/Runtime/Engine/Private/AudioVirtualLoop.cpp b/Engine/Source/Runtime/Engine/Private/AudioVirtualLoop.cpp
--- a/Engine/Source/Runtime/Engine/Private/AudioVirtualLoop.cpp
+++ b/Engine/Source/Runtime/Engine/Private/AudioVirtualLoop.cpp
@@ -125,6 +125,11 @@ bool FAudioVirtualLoop::Virtualize(const FActiveSound& InActiveSound, FAudioDevi
 	OutVirtualLoop.ActiveSound = ActiveSound;
 	OutVirtualLoop.CalculateUpdateInterval();
 
+	if (Audio::IParameterTransmitter* ParameterTransmitter = ActiveSound->GetTransmitter())
+	{
+		ParameterTransmitter->OnVirtualizeActiveSound();
+	}
+
 #if UE_AUDIO_PROFILERTRACE_ENABLED
 	const bool bChannelEnabled = UE_TRACE_CHANNELEXPR_IS_ENABLED(AudioChannel);
 	if (bChannelEnabled)
@@ -308,6 +313,21 @@ bool FAudioVirtualLoop::Update(float DeltaTime, bool bForceUpdate)
 		return false;
 	}
 
+	// Update parameters on virtual loop realization
+	if (Audio::IParameterTransmitter* ParameterTransmitter = ActiveSound->GetTransmitter())
+	{
+		uint64 AudioComponentID = ActiveSound->GetAudioComponentID();
+		if (AudioComponentID > 0)
+		{
+			if (UAudioComponent* AudioComponent = UAudioComponent::GetAudioComponentFromID(AudioComponentID))
+			{
+				// Use instance parameters which will not contain transient parameters
+				TArray<FAudioParameter> InstanceParams = AudioComponent->GetInstanceParameters();
+				ParameterTransmitter->OnRealizeVirtualizedActiveSound(MoveTemp(InstanceParams));
+			}
+		}
+	}
+
 #if UE_AUDIO_PROFILERTRACE_ENABLED
 	if (bChannelEnabled)
 	{
