Improve detection and resolution of pre-lwc code in custom hlsl

diff --git a/Engine/Source/Runtime/Engine/Private/Materials/MaterialExpressions.cpp b/Engine/Source/Runtime/Engine/Private/Materials/MaterialExpressions.cpp
--- a/Engine/Source/Runtime/Engine/Private/Materials/MaterialExpressions.cpp
+++ b/Engine/Source/Runtime/Engine/Private/Materials/MaterialExpressions.cpp
@@ -14168,18 +14168,6 @@ void UMaterialExpressionCustom::Serialize(FStructuredArchive::FRecord Record)
 			TEXT("PrevPreViewTranslation"),
 		};
 
-		static const TCHAR* GlobalExpressions[] =
-		{
-			TEXT("GetWorldPosition(Parameters)"),
-			TEXT("GetPrevWorldPosition(Parameters)"),
-			TEXT("GetObjectWorldPosition(Parameters)"),
-			TEXT("GetPrimitiveData(Parameters).WorldToLocal"),
-			TEXT("GetPrimitiveData(Parameters).LocalToWorld"),
-			TEXT("GetPrimitiveData(Parameters.PrimitiveId).WorldToLocal"),
-			TEXT("GetPrimitiveData(Parameters.PrimitiveId).LocalToWorld"),
-			TEXT("Parameters.AbsoluteWorldPosition"),
-		};
-
 		for (const TCHAR* Member : UniformMembers)
 		{
 			const FString ViewSearchString = FString(TEXT("View.")) + Member;
@@ -14200,9 +14188,44 @@ void UMaterialExpressionCustom::Serialize(FStructuredArchive::FRecord Record)
 			bDidUpdate = true;
 		}
 
-		for (const TCHAR* Expression : GlobalExpressions)
+		static const TCHAR* GlobalExpressionsToReplace[] =
+		{
+			TEXT("GetPrimitiveData(Parameters).WorldToLocal"),
+			TEXT("GetPrimitiveData(Parameters).LocalToWorld"),
+			TEXT("GetPrimitiveData(Parameters.PrimitiveId).WorldToLocal"),
+			TEXT("GetPrimitiveData(Parameters.PrimitiveId).LocalToWorld"),
+			TEXT("Parameters.AbsoluteWorldPosition"),
+		};
+		static const TCHAR* GlobalExpressionsReplacement[] =
+		{
+			TEXT("GetWorldToLocal(Parameters)"),
+			TEXT("GetLocalToWorld(Parameters)"),
+			TEXT("GetWorldToLocal(Parameters)"),
+			TEXT("GetLocalToWorld(Parameters)"),
+			TEXT("GetWorldPosition(Parameters)"),
+		};
+		
+		int NumGlobalExpressionsToReplace = sizeof(GlobalExpressionsToReplace)/sizeof(GlobalExpressionsToReplace[0]);
+		for (int Index = 0; Index < NumGlobalExpressionsToReplace; Index++)
+		{
+			if (Code.ReplaceInline(GlobalExpressionsToReplace[Index], GlobalExpressionsReplacement[Index], ESearchCase::CaseSensitive) > 0)
+			{
+				bDidUpdate = true;
+			}
+		}
+
+		static const TCHAR* GlobalExpressionsToDemote[] =
+		{
+			TEXT("GetWorldPosition(Parameters)"),
+			TEXT("GetPrevWorldPosition(Parameters)"),
+			TEXT("GetObjectWorldPosition(Parameters)"),
+			TEXT("GetWorldToLocal(Parameters)"),
+			TEXT("GetLocalToWorld(Parameters)")
+		};
+
+		for (const TCHAR* Expression : GlobalExpressionsToDemote)
 		{
-			const FString ReplaceString = FString::Printf(TEXT("DFDemote(%s)"), Expression);
+			const FString ReplaceString = FString::Printf(TEXT("WSDemote(%s)"), Expression);
 			if (Code.ReplaceInline(Expression, *ReplaceString, ESearchCase::CaseSensitive) > 0)
 			{
 				bDidUpdate = true;
