PR #11149: [MetaSounds] Fix attenuation override not being respected within audio interface update

diff --git a/Engine/Source/Runtime/Engine/Private/ActiveSound.cpp b/Engine/Source/Runtime/Engine/Private/ActiveSound.cpp
--- a/Engine/Source/Runtime/Engine/Private/ActiveSound.cpp
+++ b/Engine/Source/Runtime/Engine/Private/ActiveSound.cpp
@@ -667,7 +667,15 @@ void FActiveSound::UpdateInterfaceParameters(const TArray<FListener>& InListener
 	
 	if (bImplementsAttenuation)
 	{
-		const float Distance = SourceDirection.Size();
+		float Distance;
+		if (Listener.bUseAttenuationOverride)
+		{
+			Distance = (Transform.GetLocation() - Listener.AttenuationOverride).Size();
+		}
+		else
+		{
+			Distance = SourceDirection.Size();
+		}
 		ParamsToUpdate.Add({ AttenuationInterface::Inputs::Distance, Distance });
 	}
 
