PR #11463: fix log typo in level cooker: RegisterStreamableTexture

diff --git a/Engine/Source/Runtime/Engine/Private/Level.cpp b/Engine/Source/Runtime/Engine/Private/Level.cpp
--- a/Engine/Source/Runtime/Engine/Private/Level.cpp
+++ b/Engine/Source/Runtime/Engine/Private/Level.cpp
@@ -2241,7 +2241,7 @@ uint16 ULevel::RegisterStreamableTexture(const FString& InTextureName, const FGu
 	if (TextureNameIndex != TextureGuidIndex)
 	{
 		UE_CLOG(TextureNameIndex != INDEX_NONE, LogLevel, Warning, TEXT("Failed to register streamable texture Name = %s Guid = %s: An entry already exists for this Name with a different Guid = %s. Consider rebuilding texture streaming."), *InTextureName, *InTextureGuid.ToString(), *StreamingTextureGuids[TextureNameIndex].ToString());
-		UE_CLOG(TextureGuidIndex != INDEX_NONE, LogLevel, Warning, TEXT("Failed to register streamable texture Name = %s Guid = %s: An entry already exists for this Guid with a different Name = %s. Consider modifying & resaving one of these textures (will change its guid) and rebuiling texture streaming."), *InTextureName, *InTextureGuid.ToString(), *StreamingTextures[TextureGuidIndex].ToString());
+		UE_CLOG(TextureGuidIndex != INDEX_NONE, LogLevel, Warning, TEXT("Failed to register streamable texture Name = %s Guid = %s: An entry already exists for this Guid with a different Name = %s. Consider modifying & resaving one of these textures (will change its guid) and rebuilding texture streaming."), *InTextureName, *InTextureGuid.ToString(), *StreamingTextures[TextureGuidIndex].ToString());
 		return InvalidRegisteredStreamableTexture;
 	}
 	else if (TextureNameIndex != INDEX_NONE)
@@ -4694,4 +4694,4 @@ void ULevel::OnMultipleLSAsPopupDismissed()
 
 #endif // WITH_EDITOR
 
-#undef LOCTEXT_NAMESPACE
\ No newline at end of file
+#undef LOCTEXT_NAMESPACE
