Move BezierUtilities into a more appropriate place: Engine/Classes/Curves. This also fixes a circular dependency created in CL 31592109 and temporarily patched in CL 31594398

diff --git a/Engine/Plugins/AI/MassAI/Source/MassAIBehavior/Private/MassLookAtFragments.cpp b/Engine/Plugins/AI/MassAI/Source/MassAIBehavior/Private/MassLookAtFragments.cpp
--- a/Engine/Plugins/AI/MassAI/Source/MassAIBehavior/Private/MassLookAtFragments.cpp
+++ b/Engine/Plugins/AI/MassAI/Source/MassAIBehavior/Private/MassLookAtFragments.cpp
@@ -3,7 +3,7 @@
 #include "MassLookAtFragments.h"
 #include "ZoneGraphSubsystem.h"
 #include "ZoneGraphTypes.h"
-#include "BezierUtilities.h"
+#include "Curves/BezierUtilities.h"
 
 FVector FMassLookAtTrajectoryFragment::GetPointAtDistanceExtrapolated(const float DistanceAlongPath) const
 {
diff --git a/Engine/Plugins/AI/MassAI/Source/MassAIBehavior/Private/MassLookAtProcessors.cpp b/Engine/Plugins/AI/MassAI/Source/MassAIBehavior/Private/MassLookAtProcessors.cpp
--- a/Engine/Plugins/AI/MassAI/Source/MassAIBehavior/Private/MassLookAtProcessors.cpp
+++ b/Engine/Plugins/AI/MassAI/Source/MassAIBehavior/Private/MassLookAtProcessors.cpp
@@ -15,7 +15,7 @@
 #include "MassZoneGraphNavigationFragments.h"
 #include "ZoneGraphSubsystem.h"
 #include "ZoneGraphQuery.h"
-#include "BezierUtilities.h"
+#include "Curves/BezierUtilities.h"
 #include "Algo/RandomShuffle.h"
 #include "Engine/World.h"
 #include "MassLODFragments.h"
diff --git a/Engine/Plugins/Runtime/ZoneGraph/Source/ZoneGraph/Private/ZoneGraphRenderingUtilities.cpp b/Engine/Plugins/Runtime/ZoneGraph/Source/ZoneGraph/Private/ZoneGraphRenderingUtilities.cpp
--- a/Engine/Plugins/Runtime/ZoneGraph/Source/ZoneGraph/Private/ZoneGraphRenderingUtilities.cpp
+++ b/Engine/Plugins/Runtime/ZoneGraph/Source/ZoneGraph/Private/ZoneGraphRenderingUtilities.cpp
@@ -1,7 +1,7 @@
 // Copyright Epic Games, Inc. All Rights Reserved.
 
 #include "ZoneGraphRenderingUtilities.h"
-#include "BezierUtilities.h"
+#include "Curves/BezierUtilities.h"
 #include "PrimitiveSceneProxy.h"
 #include "ZoneGraphTypes.h"
 #include "ZoneGraphQuery.h"
diff --git a/Engine/Plugins/Runtime/ZoneGraph/Source/ZoneGraph/Private/ZoneShapeComponent.cpp b/Engine/Plugins/Runtime/ZoneGraph/Source/ZoneGraph/Private/ZoneShapeComponent.cpp
--- a/Engine/Plugins/Runtime/ZoneGraph/Source/ZoneGraph/Private/ZoneShapeComponent.cpp
+++ b/Engine/Plugins/Runtime/ZoneGraph/Source/ZoneGraph/Private/ZoneShapeComponent.cpp
@@ -8,7 +8,7 @@
 #include "ZoneGraphSettings.h"
 #include "ZoneGraphDelegates.h"
 #include "ZoneGraphRenderingUtilities.h"
-#include "BezierUtilities.h"
+#include "Curves/BezierUtilities.h"
 #include "ZoneGraphObjectCRC32.h"
 #include "Serialization/CustomVersion.h"
 
diff --git a/Engine/Plugins/Runtime/ZoneGraph/Source/ZoneGraph/Private/ZoneShapeUtilities.cpp b/Engine/Plugins/Runtime/ZoneGraph/Source/ZoneGraph/Private/ZoneShapeUtilities.cpp
--- a/Engine/Plugins/Runtime/ZoneGraph/Source/ZoneGraph/Private/ZoneShapeUtilities.cpp
+++ b/Engine/Plugins/Runtime/ZoneGraph/Source/ZoneGraph/Private/ZoneShapeUtilities.cpp
@@ -1,7 +1,7 @@
 // Copyright Epic Games, Inc. All Rights Reserved.
 
 #include "ZoneShapeUtilities.h"
-#include "BezierUtilities.h"
+#include "Curves/BezierUtilities.h"
 #include "Algo/Reverse.h"
 #include "HAL/IConsoleManager.h"
 #include "ZoneGraphSettings.h"
diff --git a/Engine/Plugins/Runtime/ZoneGraph/Source/ZoneGraph/ZoneGraph.Build.cs b/Engine/Plugins/Runtime/ZoneGraph/Source/ZoneGraph/ZoneGraph.Build.cs
--- a/Engine/Plugins/Runtime/ZoneGraph/Source/ZoneGraph/ZoneGraph.Build.cs
+++ b/Engine/Plugins/Runtime/ZoneGraph/Source/ZoneGraph/ZoneGraph.Build.cs
@@ -21,8 +21,6 @@ namespace UnrealBuildTool.Rules
 				}
 			);
 
-			PrivateDependencyModuleNames.Add("AIModule");
-
 			if (Target.bBuildEditor == true)
 			{
 				PrivateDependencyModuleNames.Add("EditorFramework");
diff --git a/Engine/Plugins/Runtime/ZoneGraph/Source/ZoneGraphEditor/Private/ZoneShapeComponentVisualizer.cpp b/Engine/Plugins/Runtime/ZoneGraph/Source/ZoneGraphEditor/Private/ZoneShapeComponentVisualizer.cpp
--- a/Engine/Plugins/Runtime/ZoneGraph/Source/ZoneGraphEditor/Private/ZoneShapeComponentVisualizer.cpp
+++ b/Engine/Plugins/Runtime/ZoneGraph/Source/ZoneGraphEditor/Private/ZoneShapeComponentVisualizer.cpp
@@ -25,7 +25,7 @@
 #include "ZoneShapeComponent.h"
 #include "ZoneShapeUtilities.h"
 #include "ZoneGraphRenderingUtilities.h"
-#include "BezierUtilities.h"
+#include "Curves/BezierUtilities.h"
 #include "CanvasTypes.h"
 #include "SceneManagement.h"
 
diff --git a/Engine/Plugins/Runtime/ZoneGraph/Source/ZoneGraphEditor/ZoneGraphEditor.Build.cs b/Engine/Plugins/Runtime/ZoneGraph/Source/ZoneGraphEditor/ZoneGraphEditor.Build.cs
--- a/Engine/Plugins/Runtime/ZoneGraph/Source/ZoneGraphEditor/ZoneGraphEditor.Build.cs
+++ b/Engine/Plugins/Runtime/ZoneGraph/Source/ZoneGraphEditor/ZoneGraphEditor.Build.cs
@@ -33,7 +33,6 @@ namespace UnrealBuildTool.Rules
 
 			PrivateDependencyModuleNames.AddRange(
 			new string[] {
-				"AIModule",
 				"RenderCore",
 				"GraphEditor",
 				"KismetWidgets",
diff --git a/Engine/Source/Runtime/AIModule/Private/BezierUtilities.cpp b/Engine/Source/Runtime/AIModule/Private/BezierUtilities.cpp
deleted file mode 100644
diff --git a/Engine/Source/Runtime/AIModule/Public/BezierUtilities.h b/Engine/Source/Runtime/AIModule/Public/BezierUtilities.h
deleted file mode 100644
diff --git a/Engine/Source/Runtime/Engine/Classes/Curves/BezierUtilities.h b/Engine/Source/Runtime/Engine/Classes/Curves/BezierUtilities.h
new file mode 100644
--- /dev/null
+++ b/Engine/Source/Runtime/Engine/Classes/Curves/BezierUtilities.h
@@ -0,0 +1,57 @@
+// Copyright Epic Games, Inc. All Rights Reserved.
+
+#pragma once
+
+#include "Math/Vector.h"
+#include "Math/Box.h"
+#include "Math/UnrealMathUtility.h"
+
+namespace UE
+{
+namespace CubicBezier
+{
+
+/** Evaluates cubic bezier value at specified time */
+inline FVector Eval(const FVector& P0, const FVector& P1, const FVector& P2, const FVector& P3, const float t)
+{
+	// Interpolate using De Casteljau algorithm.
+	const FVector P01 = FMath::Lerp(P0, P1, t);
+	const FVector P12 = FMath::Lerp(P1, P2, t);
+	const FVector P23 = FMath::Lerp(P2, P3, t);
+	const FVector P012 = FMath::Lerp(P01, P12, t);
+	const FVector P123 = FMath::Lerp(P12, P23, t);
+	return FMath::Lerp(P012, P123, t);
+}
+
+/** Evaluates cubic bezier derivative at specified time */
+inline FVector EvalDerivate(const FVector& P0, const FVector& P1, const FVector& P2, const FVector& P3, const float t)
+{
+	const FVector P01 = FMath::Lerp(P0, P1, t);
+	const FVector P12 = FMath::Lerp(P1, P2, t);
+	const FVector P23 = FMath::Lerp(P2, P3, t);
+	return 3.0f * (FMath::Lerp(P12, P23, t) - FMath::Lerp(P01, P12, t));
+}
+
+/** Splits cubic bezier curve at specified time, returns new points.
+  *	The new curves are (P0, Result[0], Result[1], Result[2]) and (P0, Result[2], Result[3], Result[4], P3) */
+ENGINE_API void SplitAt(const FVector& P0, const FVector& P1, const FVector& P2, const FVector& P3, const float t, FVector OutResult[5]);
+
+/** Returns bounding box of cubic bezier curve. */
+ENGINE_API FBox CalcBounds(const FVector& P0, const FVector& P1, const FVector& P2, const FVector& P3);
+
+/** Returns closest point to cubic bezier segment. The curve is split to up to Steps linear segments for approximate calculation. */
+ENGINE_API void ClosestPointApproximate(const FVector& FromPoint, const FVector& P0, const FVector& P1, const FVector& P2, const FVector& P3, FVector& OutClosestPoint, float& OutClosestT, const int Steps = 16);
+
+/** Returns closest point between a linear segment and cubic bezier segment. The curve is split to up to Steps linear segments for approximate calculation. */
+ENGINE_API void SegmentClosestPointApproximate(const FVector& SegStart, const FVector& SegEnd, const FVector& P0, const FVector& P1, const FVector& P2, const FVector& P3, FVector& OutClosestPoint, float& OutClosestT, const int Steps = 16);
+
+/** Returns approximate arc length of cubic bezier. The length is calculated based on control point hull with one subdivision. */
+ENGINE_API float ArcLengthApproximate(const FVector& P0, const FVector& P1, const FVector& P2, const FVector& P3);
+
+/** Recursively tessellate a bezier curve until the segments are within Tolerance from the curve, or until max recursion level is reached.
+  * Output will contain points describing linear segments of the flattened curve.
+  * Note that the points are appended, and that the first point is omitted so that the function can be used to flatten a bezier shape. */
+ENGINE_API void Tessellate(TArray<FVector>& Output, const FVector& P0, const FVector& P1, const FVector& P2, const FVector& P3, const float Tolerance, const int MaxLevel = 6);
+
+} // namespace CubicBezier
+} // namespace UE
diff --git a/Engine/Source/Runtime/Engine/Private/Curves/BezierUtilities.cpp b/Engine/Source/Runtime/Engine/Private/Curves/BezierUtilities.cpp
new file mode 100644
--- /dev/null
+++ b/Engine/Source/Runtime/Engine/Private/Curves/BezierUtilities.cpp
@@ -0,0 +1,210 @@
+// Copyright Epic Games, Inc. All Rights Reserved.
+
+#include "Curves/BezierUtilities.h"
+
+namespace UE
+{
+namespace CubicBezier
+{
+
+void SplitAt(const FVector& P0, const FVector& P1, const FVector& P2, const FVector& P3, const float t, FVector OutResult[5])
+{
+	// Split using De Casteljau algorithm.
+	const FVector P01 = FMath::Lerp(P0, P1, t);
+	const FVector P12 = FMath::Lerp(P1, P2, t);
+	const FVector P23 = FMath::Lerp(P2, P3, t);
+
+	OutResult[0] = P01;
+	OutResult[1] = FMath::Lerp(P01, P12, t);
+
+	OutResult[3] = FMath::Lerp(P12, P23, t);
+	OutResult[4] = P23;
+
+	OutResult[2] = FMath::Lerp(OutResult[1], OutResult[3], t);
+}
+
+FBox CalcBounds(const FVector& P0, const FVector& P1, const FVector& P2, const FVector& P3)
+{
+	FBox Result(EForceInit::ForceInit);
+
+	// Start the bounding box by end points
+	Result += P0;
+	Result += P3;
+
+	// Bezier curve fits inside the convex hull of it's control points.
+	// If control points are inside the bounds, we're done.
+	if (Result.IsInsideOrOn(P1) && Result.IsInsideOrOn(P2))
+		return Result;
+
+	// Add bezier curve inflection points in X, Y, and Z.
+	const FVector A = -3.0f * P0 + 9.0f * P1 - 9.0f * P2 + 3.0f * P3;
+	const FVector B = 6.0f * P0 - 12.0f * P1 + 6.0f * P2;
+	const FVector C = 3.0f * P1 - 3.0f * P0;
+
+	for (int32 Axis = 0; Axis < 3; Axis++)
+	{
+		float Roots[2];
+		int32 RootNum = 0;
+		if (FMath::IsNearlyZero(A[Axis]))
+		{
+			// It's linear equation.
+			if (!FMath::IsNearlyZero(B[Axis]))
+			{
+				Roots[0] = -C[Axis] / B[Axis];
+				RootNum = 1;
+			}
+		}
+		else
+		{
+			// Solve quadratic roots.
+			float Disc = B[Axis] * B[Axis] - 4.0f * C[Axis] * A[Axis];
+			if (Disc > UE_KINDA_SMALL_NUMBER)
+			{
+				Disc = FMath::Sqrt(Disc);
+				const float Denom = 1.0f / (2.0f * A[Axis]);
+				Roots[0] = (-B[Axis] + Disc) * Denom;
+				Roots[1] = (-B[Axis] - Disc) * Denom;
+				RootNum = 2;
+			}
+		}
+		for (int32 i = 0; i < RootNum; i++)
+		{
+			// Accept only roots which are inside the bounds of the curve (also excluding enpoints which are already added above).
+			const float t = Roots[i];
+			if (t > UE_KINDA_SMALL_NUMBER && t < (1.0f - UE_KINDA_SMALL_NUMBER))
+			{
+				Result += Eval(P0, P1, P2, P3, t);
+			}
+		}
+	}
+
+	return Result;
+}
+
+void ClosestPointApproximate(const FVector& FromPoint, const FVector& P0, const FVector& P1, const FVector& P2, const FVector& P3, FVector& OutClosestPoint, float& OutClosestT, const int Steps)
+{
+	float ClosestDistanceSqr = TNumericLimits<float>::Max();
+	FVector ClosestPoint = P0;
+	float ClosestT = 0.0f;
+
+	const int NumSteps = FMath::Max(1, Steps);
+	const float Step = 1.0f / (float)NumSteps;
+
+	FVector PrevPoint = P0;
+
+	for (int j = 0; j < Steps; j++)
+	{
+		const float t = (float)(j + 1) * Step;
+		const FVector CurrPoint = Eval(P0, P1, P2, P3, t);
+
+		const FVector Closest = FMath::ClosestPointOnSegment(FromPoint, PrevPoint, CurrPoint);
+
+		const float DistanceSqr = FVector::DistSquared(Closest, FromPoint);
+		if (DistanceSqr < ClosestDistanceSqr)
+		{
+			ClosestDistanceSqr = DistanceSqr;
+			ClosestPoint = Closest;
+			// The SegT is calculate by the closest point calculation above, but it does not return in!
+			const float SegLengthSqr = FVector::DistSquared(PrevPoint, CurrPoint);
+			const float SegT = SegLengthSqr > UE_KINDA_SMALL_NUMBER ? FVector::DistSquared(PrevPoint, Closest) / SegLengthSqr : 0.0f;
+			ClosestT = (t - Step) + SegT * Step;
+		}
+
+		PrevPoint = CurrPoint;
+	}
+
+	OutClosestPoint = ClosestPoint;
+	OutClosestT = ClosestT;
+}
+
+void SegmentClosestPointApproximate(const FVector& SegStart, const FVector& SegEnd, const FVector& P0, const FVector& P1, const FVector& P2, const FVector& P3, FVector& OutClosestPoint, float& OutClosestT, const int Steps)
+{
+	float ClosestDistanceSqr = TNumericLimits<float>::Max();
+	FVector ClosestPoint = P0;
+	float ClosestT = 0.0f;
+
+	const int NumSteps = FMath::Max(1, Steps);
+	const float Step = 1.0f / (float)NumSteps;
+
+	FVector PrevPoint = P0;
+
+	for (int j = 0; j < NumSteps; j++)
+	{
+		const float t = (float)(j + 1) * Step;
+		const FVector CurrPoint = Eval(P0, P1, P2, P3, t);
+
+		FVector Closest;
+		FVector SegClosest;
+		FMath::SegmentDistToSegmentSafe(PrevPoint, CurrPoint, SegStart, SegEnd, Closest, SegClosest);
+
+		const float DistanceSqr = FVector::DistSquared(Closest, SegClosest);
+		if (DistanceSqr < ClosestDistanceSqr)
+		{
+			ClosestDistanceSqr = DistanceSqr;
+			ClosestPoint = Closest;
+			// The SegT is calculate by the closest point calculation above, but it does not return in!
+			const float SegLengthSqr = FVector::DistSquared(PrevPoint, CurrPoint);
+			const float SegT = SegLengthSqr > UE_KINDA_SMALL_NUMBER ? FVector::DistSquared(PrevPoint, Closest) / SegLengthSqr : 0.0f;
+			ClosestT = (t - Step) + SegT * Step;
+		}
+
+		PrevPoint = CurrPoint;
+	}
+
+	OutClosestPoint = ClosestPoint;
+	OutClosestT = ClosestT;
+}
+
+float ArcLengthApproximate(const FVector& P0, const FVector& P1, const FVector& P2, const FVector& P3)
+{
+	FVector Result[5];
+	SplitAt(P0, P1, P2, P3, 0.5f, Result);
+	return FVector::Dist(P0, Result[0]) + FVector::Dist(Result[1], Result[2]) + FVector::Dist(Result[2], Result[3]) + FVector::Dist(Result[3], Result[4]) + FVector::Dist(Result[4], P3);
+}
+
+void TessellateRecursive(TArray<FVector>& Output, const FVector& P0, const FVector& P1, const FVector& P2, const FVector& P3, const float ToleranceSqr, int Level, const int MaxLevel)
+{
+	// Handle degenerate segment.
+	FVector Dir = P3 - P0;
+	if (Dir.IsNearlyZero())
+	{
+		Output.Add(P3);
+		return;
+	}
+
+	// If the control points are close enough to approximate a line within tolerance, stop recursing.
+	Dir = Dir.GetUnsafeNormal();
+	const FVector RelP1 = P1 - P0;
+	const FVector RelP2 = P2 - P0;
+	const FVector ProjP1 = Dir * FVector::DotProduct(Dir, RelP1);
+	const FVector ProjP2 = Dir * FVector::DotProduct(Dir, RelP2);
+	const float DistP1Sqr = FVector::DistSquared(RelP1, ProjP1);
+	const float DistP2Sqr = FVector::DistSquared(RelP2, ProjP2);
+	if (DistP1Sqr < ToleranceSqr && DistP2Sqr < ToleranceSqr)
+	{
+		Output.Add(P3);
+		return;
+	}
+
+	if (Level < MaxLevel)
+	{
+		// Split the curve in half and recurse.
+		const FVector P01 = FMath::Lerp(P0, P1, 0.5f);
+		const FVector P12 = FMath::Lerp(P1, P2, 0.5f);
+		const FVector P23 = FMath::Lerp(P2, P3, 0.5f);
+		const FVector P012 = FMath::Lerp(P01, P12, 0.5f);
+		const FVector P123 = FMath::Lerp(P12, P23, 0.5f);
+		const FVector P0123 = FMath::Lerp(P012, P123, 0.5f);
+
+		TessellateRecursive(Output, P0, P01, P012, P0123, ToleranceSqr, Level + 1, MaxLevel);
+		TessellateRecursive(Output, P0123, P123, P23, P3, ToleranceSqr, Level + 1, MaxLevel);
+	}
+}
+
+void Tessellate(TArray<FVector>& Output, const FVector& P0, const FVector& P1, const FVector& P2, const FVector& P3, const float Tolerance, const int MaxLevel)
+{
+	TessellateRecursive(Output, P0, P1, P2, P3, Tolerance * Tolerance, 0, MaxLevel);
+}
+
+} // namespace CubicBezier
+} // namespace UE
diff --git a/Engine/Source/Runtime/NavigationSystem/Private/SplineNavModifierComponent.cpp b/Engine/Source/Runtime/NavigationSystem/Private/SplineNavModifierComponent.cpp
--- a/Engine/Source/Runtime/NavigationSystem/Private/SplineNavModifierComponent.cpp
+++ b/Engine/Source/Runtime/NavigationSystem/Private/SplineNavModifierComponent.cpp
@@ -4,7 +4,7 @@
 
 #include "AI/NavigationSystemBase.h"
 #include "AI/Navigation/NavigationRelevantData.h"
-// #include "BezierUtilities.h"
+#include "Curves/BezierUtilities.h"
 #include "Components/SplineComponent.h"
 #include "VisualLogger/VisualLogger.h"
 
@@ -12,51 +12,6 @@
 
 namespace
 {
-	// Todo: temporary solution to circular dependency with BezierUtilities's new location, remove after fixing circular dependency
-	void TessellateRecursive(TArray<FVector>& Output, const FVector& P0, const FVector& P1, const FVector& P2, const FVector& P3, const float ToleranceSqr, int Level, const int MaxLevel)
-	{
-		// Handle degenerate segment.
-		FVector Dir = P3 - P0;
-		if (Dir.IsNearlyZero())
-		{
-			Output.Add(P3);
-			return;
-		}
-
-		// If the control points are close enough to approximate a line within tolerance, stop recursing.
-		Dir = Dir.GetUnsafeNormal();
-		const FVector RelP1 = P1 - P0;
-		const FVector RelP2 = P2 - P0;
-		const FVector ProjP1 = Dir * FVector::DotProduct(Dir, RelP1);
-		const FVector ProjP2 = Dir * FVector::DotProduct(Dir, RelP2);
-		const float DistP1Sqr = FVector::DistSquared(RelP1, ProjP1);
-		const float DistP2Sqr = FVector::DistSquared(RelP2, ProjP2);
-		if (DistP1Sqr < ToleranceSqr && DistP2Sqr < ToleranceSqr)
-		{
-			Output.Add(P3);
-			return;
-		}
-
-		if (Level < MaxLevel)
-		{
-			// Split the curve in half and recurse.
-			const FVector P01 = FMath::Lerp(P0, P1, 0.5f);
-			const FVector P12 = FMath::Lerp(P1, P2, 0.5f);
-			const FVector P23 = FMath::Lerp(P2, P3, 0.5f);
-			const FVector P012 = FMath::Lerp(P01, P12, 0.5f);
-			const FVector P123 = FMath::Lerp(P12, P23, 0.5f);
-			const FVector P0123 = FMath::Lerp(P012, P123, 0.5f);
-
-			TessellateRecursive(Output, P0, P01, P012, P0123, ToleranceSqr, Level + 1, MaxLevel);
-			TessellateRecursive(Output, P0123, P123, P23, P3, ToleranceSqr, Level + 1, MaxLevel);
-		}
-	}
-
-	void Tessellate(TArray<FVector>& Output, const FVector& P0, const FVector& P1, const FVector& P2, const FVector& P3, const float Tolerance, const int MaxLevel = 6)
-	{
-		TessellateRecursive(Output, P0, P1, P2, P3, Tolerance * Tolerance, 0, MaxLevel);
-	}
-
 	// Fetch the spline component from the actor
 	const USplineComponent* GetSpline(const AActor* Owner)
 	{
@@ -92,7 +47,7 @@ namespace
 				OutSubdivisions.Add(PrevSplinePoint.Position);
 
 				// Convert this segment of the spline from Hermite to Bezier and subdivide it 
-				Tessellate(OutSubdivisions,
+				UE::CubicBezier::Tessellate(OutSubdivisions,
 					PrevSplinePoint.Position,
 					PrevSplinePoint.Position + PrevSplinePoint.LeaveTangent / HermiteToBezierFactor,
 					CurrSplinePoint.Position - CurrSplinePoint.ArriveTangent / HermiteToBezierFactor,
