- Shader code will now look for ShaderCompileWorker to be in the project Binaries, and fall back to Engine/Binaries if not (this can happen when building SCW as Unique BuildEnvironment if the project overrides an SDK version)

diff --git a/Engine/Source/Runtime/Engine/Private/ShaderCompiler/ShaderCompiler.cpp b/Engine/Source/Runtime/Engine/Private/ShaderCompiler/ShaderCompiler.cpp
--- a/Engine/Source/Runtime/Engine/Private/ShaderCompiler/ShaderCompiler.cpp
+++ b/Engine/Source/Runtime/Engine/Private/ShaderCompiler/ShaderCompiler.cpp
@@ -5552,13 +5552,6 @@ FShaderCompilingManager::FShaderCompilingManager() :
 	NumExternalJobs(0),
 	AllJobs(CompileQueueSection),
 	NumSingleThreadedRunsBeforeRetry(GSingleThreadedRunsIdle),
-#if PLATFORM_MAC
-	ShaderCompileWorkerName(FPaths::EngineDir() / TEXT("Binaries/Mac/ShaderCompileWorker")),
-#elif PLATFORM_LINUX
-	ShaderCompileWorkerName(FPaths::EngineDir() / TEXT("Binaries/Linux/ShaderCompileWorker")),
-#else
-	ShaderCompileWorkerName(FPaths::EngineDir() / TEXT("Binaries/Win64/ShaderCompileWorker.exe")),
-#endif
 	SuppressedShaderPlatforms(0),
 	BuildDistributionController(nullptr),
 	bNoShaderCompilation(false),
@@ -5582,6 +5575,20 @@ FShaderCompilingManager::FShaderCompilingManager() :
 
 	WorkersBusyTime = 0;
 
+#if PLATFORM_WINDOWS
+	FString ExecutableName("ShaderCompileWorker.exe");
+#else
+	FString ExecutableName("ShaderCompileWorker");
+#endif
+
+	// first look for project-specific version
+	ShaderCompileWorkerName = FPaths::Combine(FPaths::ProjectDir(), TEXT("Binaries"), FPlatformProcess::GetBinariesSubdirectory(), ExecutableName);
+	if (!IFileManager::Get().FileExists(*ShaderCompileWorkerName))
+	{
+		// fallback to standard Engine location
+		ShaderCompileWorkerName = FPaths::Combine(FPaths::EngineDir(), TEXT("Binaries"), FPlatformProcess::GetBinariesSubdirectory(), ExecutableName);
+	}
+
 	// Threads must use absolute paths on Windows in case the current directory is changed on another thread!
 	ShaderCompileWorkerName = FPaths::ConvertRelativePathToFull(ShaderCompileWorkerName);
 
