Ensure we TickProgress() on the GT when waiting during AsyncCompilationHelper::FinishCompilation. SlowTask.EnterProgressFrame will do so and ensure proper messaging for what is being waited on. Without pumping we might not tick the work outstanding jobs are waiting on.

diff --git a/Engine/Source/Runtime/Engine/Private/AsyncCompilationHelpers.cpp b/Engine/Source/Runtime/Engine/Private/AsyncCompilationHelpers.cpp
--- a/Engine/Source/Runtime/Engine/Private/AsyncCompilationHelpers.cpp
+++ b/Engine/Source/Runtime/Engine/Private/AsyncCompilationHelpers.cpp
@@ -171,19 +171,6 @@ namespace AsyncCompilationHelpers
 
 						continue;
 					}
-					
-					// Avoid spamming task progress while waiting
-					if (!LoggedSlowTask[JobIndex])
-					{
-						LoggedSlowTask[JobIndex] = true;
-						FText Progress = FormatProgress(NumDone, Num, Job.GetName());
-						UE_LOG_REF(LogCategory, Display, TEXT("%s"), *Progress.ToString());
-
-						if (SlowTask.IsSet())
-						{
-							SlowTask->EnterProgressFrame(0.0f, Progress);
-						}
-					}
 				}
 				It.SetWord(WordJobState);
 			}
@@ -203,6 +190,24 @@ namespace AsyncCompilationHelpers
 				break;
 			}
 
+			if (SlowTask.IsSet())
+			{
+				int IncompleteJobIndex = JobsToFinish.Find(true);
+				check(IncompleteJobIndex != INDEX_NONE);
+
+				ICompilable& IncompleteJob = Getter(IncompleteJobIndex);
+				FText Progress = FormatProgress(NumDone, Num, IncompleteJob.GetName());
+
+				// Avoid spamming task progress while waiting
+				if (!LoggedSlowTask[IncompleteJobIndex])
+				{
+					UE_LOG_REF(LogCategory, Display, TEXT("%s"), *Progress.ToString());
+					LoggedSlowTask[IncompleteJobIndex] = true;
+				}
+
+				SlowTask->EnterProgressFrame(0.0f, Progress);
+			}
+
 			// Jobs are still in flight so give them some time to complete
 			FPlatformProcess::Sleep(0.016);
 		}
