Don't read global ShouldUseIris config when Iris is already initialized.

#rnx

diff --git a/Engine/Source/Runtime/Engine/Private/Net/Iris/ReplicationSystem/ActorReplicationBridge.cpp b/Engine/Source/Runtime/Engine/Private/Net/Iris/ReplicationSystem/ActorReplicationBridge.cpp
--- a/Engine/Source/Runtime/Engine/Private/Net/Iris/ReplicationSystem/ActorReplicationBridge.cpp
+++ b/Engine/Source/Runtime/Engine/Private/Net/Iris/ReplicationSystem/ActorReplicationBridge.cpp
@@ -46,16 +46,6 @@
 
 extern bool GDefaultUseSubObjectReplicationList;
 
-namespace UE::Net
-{
-
-bool ShouldUseIrisReplication(const UObject* Object)
-{
-	return ShouldUseIrisReplication();
-}
-
-}
-
 namespace UE::Net::Private
 {
 
@@ -166,11 +156,6 @@ UE::Net::FNetRefHandle UActorReplicationBridge::BeginReplication(AActor* Actor,
 {
 	using namespace UE::Net;
 
-	if (!ShouldUseIrisReplication(Actor))
-	{
-		return FNetRefHandle::GetInvalid();
-	}
-
 	if (!ensureMsgf(Actor == nullptr || !(Actor->HasAnyFlags(RF_ClassDefaultObject | RF_ArchetypeObject)), TEXT("Actor %s is a CDO or Archetype and should not be replicated."), ToCStr(GetFullNameSafe(Actor))))
 	{
 		return FNetRefHandle::GetInvalid();
@@ -354,7 +339,7 @@ UE::Net::FNetRefHandle UActorReplicationBridge::BeginReplication(FNetRefHandle O
 {
 	using namespace UE::Net;
 
-	if (!OwnerHandle.IsValid() || !ShouldUseIrisReplication(SubObject))
+	if (!OwnerHandle.IsValid())
 	{
 		return FNetRefHandle::GetInvalid();
 	}
