Fix potential ptr access issue brought up by code analysis

diff --git a/Engine/Source/Runtime/Engine/Private/MapBuildData.cpp b/Engine/Source/Runtime/Engine/Private/MapBuildData.cpp
--- a/Engine/Source/Runtime/Engine/Private/MapBuildData.cpp
+++ b/Engine/Source/Runtime/Engine/Private/MapBuildData.cpp
@@ -1170,7 +1170,7 @@ UMapBuildDataRegistry* UMapBuildDataRegistry::Get(const UActorComponent* Compone
 UMapBuildDataRegistry* UMapBuildDataRegistry::Get(const AActor* Actor)
 {
 	ULevel* OwnerLevel = Actor->GetLevel();
-	UWorld* World = OwnerLevel->GetWorld();
+	UWorld* World = OwnerLevel ? OwnerLevel->GetWorld() : nullptr;
 	UMapBuildDataRegistry* MapBuildData = nullptr;
 	
 	if (OwnerLevel && World)
