Ensure localy controlled pawns on a listen server don't have a RemoteRole of ROLE_AutonomousProxy

diff --git a/Engine/Source/Runtime/Engine/Private/Pawn.cpp b/Engine/Source/Runtime/Engine/Private/Pawn.cpp
--- a/Engine/Source/Runtime/Engine/Private/Pawn.cpp
+++ b/Engine/Source/Runtime/Engine/Private/Pawn.cpp
@@ -657,7 +657,10 @@ void APawn::PossessedBy(AController* NewController)
 		if (GetNetMode() != NM_Standalone)
 		{
 			SetReplicates(true);
-			SetAutonomousProxy(true);
+			if (!PlayerController->IsLocalController())
+			{
+				SetAutonomousProxy(true);
+			}
 		}
 	}
 	else
