Allow motion blur to be usable in Orthographic views. Results in EngineShowFlagOrthographicOverride having no entries and all features are now enabled in Ortho views. EngineShowFlagOrthographicOverride can still be used by licensees to disable specific features if needed, but this is to be improved from a usability standpoint (e.g. adding an Ortho show flag override editor UI menu).

diff --git a/Engine/Source/Runtime/Engine/Private/ShowFlags.cpp b/Engine/Source/Runtime/Engine/Private/ShowFlags.cpp
--- a/Engine/Source/Runtime/Engine/Private/ShowFlags.cpp
+++ b/Engine/Source/Runtime/Engine/Private/ShowFlags.cpp
@@ -692,10 +692,13 @@ void EngineShowFlagOverride(EShowFlagInitMode ShowFlagInitMode, EViewModeIndex V
 
 void EngineShowFlagOrthographicOverride(bool bIsPerspective, FEngineShowFlags& EngineShowFlags)
 {
-	// Disable post processing that doesn't work in ortho viewports.
 	if (!bIsPerspective)
 	{
-		EngineShowFlags.SetMotionBlur(false);
+		/**
+		 * Orthographic feature support has been improved so all features are now enabled by default.
+		 * If you wish to disable a specific feature, add the showflag disable call here.
+		 * e.g.	EngineShowFlags.SetMotionBlur(false);
+		 */
 	}
 }
 
