Add cvar to optionally bring back the functionality of not serializing a replicated, spawned actor's level, which will cause it to spawn in the persistent level on clients, regardless of its level on the server.

diff --git a/Engine/Source/Runtime/Engine/Private/PackageMapClient.cpp b/Engine/Source/Runtime/Engine/Private/PackageMapClient.cpp
--- a/Engine/Source/Runtime/Engine/Private/PackageMapClient.cpp
+++ b/Engine/Source/Runtime/Engine/Private/PackageMapClient.cpp
@@ -120,6 +120,12 @@ namespace UE::Net::Private
 			}
 		}
 	}
+
+	int32 SerializeNewActorOverrideLevel = 1;
+	static FAutoConsoleVariableRef CVarNetSerializeNewActorOverrideLevel(
+		TEXT("net.SerializeNewActorOverrideLevel"),
+		SerializeNewActorOverrideLevel,
+		TEXT("If true, servers will serialize a spawned, replicated actor's level so the client attempts to spawn it into that level too. If false, clients will spawn all these actors into the persistent level."));
 }
 
 static TAutoConsoleVariable<int32> CVarAllowAsyncLoading(
@@ -554,7 +560,12 @@ bool UPackageMapClient::SerializeNewActor(FArchive& Ar, class UActorChannel *Cha
 			{
 				Archetype = Actor->GetArchetype();
 			}
-			ActorLevel = Actor->GetLevel();
+
+			// If enabled, send the actor's level to the client. If left null, the client will spawn the actor in the persistent level.
+			if (UE::Net::Private::SerializeNewActorOverrideLevel)
+			{
+				ActorLevel = Actor->GetLevel();
+			}
 
 			check( Archetype != nullptr );
 			check( Actor->NeedsLoadForClient() );			// We have no business sending this unless the client can load
