Fix crash performing HLOD destruction in a non game world

diff --git a/Engine/Source/Runtime/Engine/Private/WorldPartition/HLOD/HLODDestruction.cpp b/Engine/Source/Runtime/Engine/Private/WorldPartition/HLOD/HLODDestruction.cpp
--- a/Engine/Source/Runtime/Engine/Private/WorldPartition/HLOD/HLODDestruction.cpp
+++ b/Engine/Source/Runtime/Engine/Private/WorldPartition/HLOD/HLODDestruction.cpp
@@ -75,12 +75,15 @@ void UWorldPartitionDestructibleInHLODSupportLibrary::DestroyInHLOD(const TScrip
 {
 	AActor* SourceActor = CastChecked<AActor>(DestructibleInHLOD.GetObject());
 
-	if (SourceActor->HasAuthority())
+	if (UWorld* World = SourceActor->GetWorld())
 	{
-		FWorldPartitionHLODDestructionTag Tag = GetDestructionTagForActor(SourceActor);
-		if (Tag.IsValid())
+		if (World->IsGameWorld() && SourceActor->HasAuthority())
 		{
-			Tag.HLODDestructionComponent->DestroyActor(Tag.ActorIndex);
+			FWorldPartitionHLODDestructionTag Tag = GetDestructionTagForActor(SourceActor);
+			if (Tag.IsValid())
+			{
+				Tag.HLODDestructionComponent->DestroyActor(Tag.ActorIndex);
+			}
 		}
 	}
 }
@@ -89,12 +92,15 @@ void UWorldPartitionDestructibleInHLODSupportLibrary::DamageInHLOD(const TScript
 {
 	AActor* SourceActor = CastChecked<AActor>(DestructibleInHLOD.GetObject());
 
-	if (SourceActor->HasAuthority())
+	if (UWorld* World = SourceActor->GetWorld())
 	{
-		FWorldPartitionHLODDestructionTag Tag = GetDestructionTagForActor(SourceActor);
-		if (Tag.IsValid())
+		if (World->IsGameWorld() && SourceActor->HasAuthority())
 		{
-			Tag.HLODDestructionComponent->DamageActor(Tag.ActorIndex, DamagePercent);
+			FWorldPartitionHLODDestructionTag Tag = GetDestructionTagForActor(SourceActor);
+			if (Tag.IsValid())
+			{
+				Tag.HLODDestructionComponent->DamageActor(Tag.ActorIndex, DamagePercent);
+			}
 		}
 	}
 }
