Add PhysicsBodies to query results

#[fyi] Benn.Gallagher

diff --git a/Engine/Source/Runtime/Engine/Classes/Engine/HitResult.h b/Engine/Source/Runtime/Engine/Classes/Engine/HitResult.h
--- a/Engine/Source/Runtime/Engine/Classes/Engine/HitResult.h
+++ b/Engine/Source/Runtime/Engine/Classes/Engine/HitResult.h
@@ -3,6 +3,7 @@
 #pragma once
 
 #include "CoreMinimal.h"
+#include "Chaos/PhysicsObject.h"
 #include "Engine/NetSerialization.h"
 #include "Engine/ActorInstanceHandle.h"
 #include "HitResult.generated.h"
@@ -129,6 +130,10 @@ struct FHitResult
 	UPROPERTY()
 	TWeakObjectPtr<UPrimitiveComponent> Component;
 
+
+	/** PhysicsObjects hit by the query. Not exposed to blueprints for the time being */
+	Chaos::FPhysicsObjectHandle PhysicsObject;
+
 	/** Name of bone we hit (for skeletal meshes). */
 	UPROPERTY()
 	FName BoneName;
diff --git a/Engine/Source/Runtime/Engine/Classes/Engine/OverlapResult.h b/Engine/Source/Runtime/Engine/Classes/Engine/OverlapResult.h
--- a/Engine/Source/Runtime/Engine/Classes/Engine/OverlapResult.h
+++ b/Engine/Source/Runtime/Engine/Classes/Engine/OverlapResult.h
@@ -3,6 +3,7 @@
 #pragma once
 
 #include "CoreMinimal.h"
+#include "Chaos/PhysicsObject.h"
 #include "Engine/ActorInstanceHandle.h"
 #include "OverlapResult.generated.h"
 
@@ -30,6 +31,9 @@ struct FOverlapResult
 	/** Utility to return the Component that was hit */
 	ENGINE_API UPrimitiveComponent* GetComponent() const;
 
+	/** PhysicsObjects hit by the query. Not exposed to blueprints for the time being */
+	Chaos::FPhysicsObjectHandle PhysicsObject;
+
 	/** Indicates if this hit was requesting a block - if false, was requesting a touch instead */
 	UPROPERTY()
 	uint32 bBlockingHit : 1;
diff --git a/Engine/Source/Runtime/Engine/Private/Collision/CollisionConversions.cpp b/Engine/Source/Runtime/Engine/Private/Collision/CollisionConversions.cpp
--- a/Engine/Source/Runtime/Engine/Private/Collision/CollisionConversions.cpp
+++ b/Engine/Source/Runtime/Engine/Private/Collision/CollisionConversions.cpp
@@ -135,6 +135,7 @@ static void SetHitResultFromShapeAndFaceIndex(const FPhysicsShape& Shape,  const
 	OutResult.ElementIndex = (uint8)ShapeIndex;
 	
 	TWeakObjectPtr<UPrimitiveComponent> OwningComponent;
+	OutResult.PhysicsObject = nullptr;
 	if(const FBodyInstance* BodyInst = GetUserData(Actor))
 	{
 		BodyInst = FPhysicsInterface::ShapeToOriginalBodyInstance(BodyInst, &Shape);
@@ -148,13 +149,13 @@ static void SetHitResultFromShapeAndFaceIndex(const FPhysicsShape& Shape,  const
 		}
 
 		OwningComponent = BodyInst->OwnerComponent;
+		OutResult.PhysicsObject = BodyInst->ActorHandle->GetPhysicsObject();
 	}
 	else
 	{
 		// Currently geom collections are registered with a primitive component user data, but maybe custom should be adapted
 		// to be more general so we can support leaf identification #BGTODO
-		void* UserData = Actor.UserData();
-		UPrimitiveComponent* PossibleOwner = FChaosUserData::Get<UPrimitiveComponent>(UserData);
+		UPrimitiveComponent* PossibleOwner = GetPrimitiveComponentFromUserData(Actor);
 
 		if(PossibleOwner)
 		{
@@ -171,6 +172,7 @@ static void SetHitResultFromShapeAndFaceIndex(const FPhysicsShape& Shape,  const
 					const FGeometryCollectionItemIndex ItemIndex = ConcreteProxy->GetItemIndexFromGTParticle_External(Actor.CastToRigidParticle());
 					OutResult.Item = ItemIndex.GetItemIndex();
 					OutResult.BoneName = ConcreteProxy->GetTransformName_External(ItemIndex);
+					OutResult.PhysicsObject = PossibleOwner->GetPhysicsObjectById(OutResult.Item);
 				}
 			}
 		}
@@ -599,7 +601,7 @@ void ConvertQueryOverlap(const FPhysicsShape& Shape, const FPhysicsActor& Actor,
 	using namespace ChaosInterface;
 
 	// Grab actor/component
-	
+	OutOverlap.PhysicsObject = nullptr;
 	// Try body instance
 	if (const FBodyInstance* BodyInst = GetUserData(Actor))
 	{
@@ -611,14 +613,14 @@ void ConvertQueryOverlap(const FPhysicsShape& Shape, const FPhysicsActor& Actor,
 			OutOverlap.OverlapObjectHandle = FActorInstanceHandle::MakeActorHandleToResolve(BodyInst->OwnerComponent, BodyInst->InstanceBodyIndex);
 			OutOverlap.Component = BodyInst->OwnerComponent; // Copying weak pointer is faster than assigning raw pointer.
 			OutOverlap.ItemIndex = OwnerComponent->bMultiBodyOverlap ? BodyInst->InstanceBodyIndex : INDEX_NONE;
+			OutOverlap.PhysicsObject = BodyInst->ActorHandle->GetPhysicsObject();
 		}
 	}
 	else
 	{
 		// Currently geom collections are registered with a primitive component user data, but maybe custom should be adapted
 		// to be more general so we can support leaf identification #BGTODO
-		void* UserData = Actor.UserData();
-		UPrimitiveComponent* PossibleOwner = FChaosUserData::Get<UPrimitiveComponent>(UserData);
+		UPrimitiveComponent* PossibleOwner = GetPrimitiveComponentFromUserData(Actor);
 
 		if(PossibleOwner)
 		{
@@ -786,6 +788,7 @@ FHitResult ConvertOverlapToHitResult(const FOverlapResult& Overlap)
 	Hit.bBlockingHit = Overlap.bBlockingHit;
 	Hit.Item = Overlap.ItemIndex;
 	Hit.Component = Overlap.Component;
+	Hit.PhysicsObject = Overlap.PhysicsObject;
 	Hit.HitObjectHandle = Overlap.OverlapObjectHandle;
 	return Hit;
 }
diff --git a/Engine/Source/Runtime/Engine/Private/PhysicsEngine/Experimental/ChaosInterfaceWrapper.cpp b/Engine/Source/Runtime/Engine/Private/PhysicsEngine/Experimental/ChaosInterfaceWrapper.cpp
--- a/Engine/Source/Runtime/Engine/Private/PhysicsEngine/Experimental/ChaosInterfaceWrapper.cpp
+++ b/Engine/Source/Runtime/Engine/Private/PhysicsEngine/Experimental/ChaosInterfaceWrapper.cpp
@@ -33,6 +33,26 @@ namespace ChaosInterface
 		return UserData ? FChaosUserData::Get<UPhysicalMaterial>(UserData) : nullptr;
 	}
 
+	UPrimitiveComponent* GetPrimitiveComponentFromUserData(const Chaos::FGeometryParticle& Actor)
+	{
+		void* UserData = Actor.UserData();
+		if (UserData)
+		{
+			UPrimitiveComponent* PrimitiveComponent = FChaosUserData::Get<UPrimitiveComponent>(UserData);
+			if (!PrimitiveComponent)
+			{
+				// Check if we appended a custom entity
+				FChaosUserEntityAppend* ChaosUserEntityAppend = FChaosUserData::Get<FChaosUserEntityAppend>(UserData);
+				if (ChaosUserEntityAppend)
+				{
+					PrimitiveComponent = FChaosUserData::Get<UPrimitiveComponent>(ChaosUserEntityAppend->ChaosUserData);
+				}
+			}
+			return PrimitiveComponent;
+		}
+		return nullptr;
+	}
+
 	FScopedSceneReadLock::FScopedSceneReadLock(FPhysScene_Chaos& SceneIn)
 		: Solver(SceneIn.GetSolver())
 	{
diff --git a/Engine/Source/Runtime/Engine/Public/Physics/Experimental/ChaosInterfaceWrapper.h b/Engine/Source/Runtime/Engine/Public/Physics/Experimental/ChaosInterfaceWrapper.h
--- a/Engine/Source/Runtime/Engine/Public/Physics/Experimental/ChaosInterfaceWrapper.h
+++ b/Engine/Source/Runtime/Engine/Public/Physics/Experimental/ChaosInterfaceWrapper.h
@@ -39,6 +39,7 @@ inline FQueryFilterData MakeQueryFilterData(const FCollisionFilterData& FilterDa
 
 FBodyInstance* GetUserData(const Chaos::FGeometryParticle& Actor);
 UPhysicalMaterial* GetUserData(const Chaos::FChaosPhysicsMaterial& Material);
+UPrimitiveComponent* GetPrimitiveComponentFromUserData(const Chaos::FGeometryParticle& Actor);
 
 }
 
