Wrapped particle system world space uniform buffer update in a mutex to fix race condition when multiple views are processing the same particle system.

#localization none
#tests Lyra

diff --git a/Engine/Source/Runtime/Engine/Private/Particles/ParticleSystemRender.cpp b/Engine/Source/Runtime/Engine/Private/Particles/ParticleSystemRender.cpp
--- a/Engine/Source/Runtime/Engine/Private/Particles/ParticleSystemRender.cpp
+++ b/Engine/Source/Runtime/Engine/Private/Particles/ParticleSystemRender.cpp
@@ -6889,6 +6889,8 @@ void FParticleSystemSceneProxy::UpdateWorldSpacePrimitiveUniformBuffer(FRHIComma
 		NewWorldSpaceUBHash = FCrc::MemCrc32(LocalCustomPrimitiveData->Data.GetData(), LocalCustomPrimitiveData->Data.Num() * LocalCustomPrimitiveData->Data.GetTypeSize());
 	}
 
+	UE::TScopeLock Lock(WorldSpacePrimitiveUniformBufferMutex);
+
 	const bool bNeedsInit = !WorldSpacePrimitiveUniformBuffer.IsInitialized();
 
 	if (bNeedsInit || (WorldSpaceUBHash != NewWorldSpaceUBHash))
diff --git a/Engine/Source/Runtime/Engine/Public/ParticleHelper.h b/Engine/Source/Runtime/Engine/Public/ParticleHelper.h
--- a/Engine/Source/Runtime/Engine/Public/ParticleHelper.h
+++ b/Engine/Source/Runtime/Engine/Public/ParticleHelper.h
@@ -2475,6 +2475,7 @@ protected:
 	/** The primitive's uniform buffer.  Mutable because it is cached state during GetDynamicMeshElements. */
 	mutable TUniformBuffer<FPrimitiveUniformShaderParameters> WorldSpacePrimitiveUniformBuffer;
 	mutable uint32 WorldSpaceUBHash = 0;
+	mutable UE::FMutex WorldSpacePrimitiveUniformBufferMutex;
 
 	/** Pool for holding FMeshBatches to reduce allocations. */
 	TIndirectArray<FMeshBatch, TInlineAllocator<4> > MeshBatchPool;
