World Partition Actor Descriptors: move actor transform out of level instance descriptor into base class so its available for things like actors with editor-only parents.

diff --git a/Engine/Source/Runtime/Core/Public/UObject/FortniteSeasonBranchObjectVersion.h b/Engine/Source/Runtime/Core/Public/UObject/FortniteSeasonBranchObjectVersion.h
--- a/Engine/Source/Runtime/Core/Public/UObject/FortniteSeasonBranchObjectVersion.h
+++ b/Engine/Source/Runtime/Core/Public/UObject/FortniteSeasonBranchObjectVersion.h
@@ -42,6 +42,8 @@ struct CORE_API FFortniteSeasonBranchObjectVersion
 		// Chaos GeometryCollection now saves levels attribute values
 		ChaosGeometryCollectionSaveLevelsAttribute,
 
+		// Serialize actor transform for actor descriptors
+		WorldPartitionActorDescActorTransformSerialization,
 
 		// -----<new versions can be added above this line>-------------------------------------------------
 		VersionPlusOne,
diff --git a/Engine/Source/Runtime/Engine/Private/WorldPartition/LevelInstance/LevelInstanceActorDesc.cpp b/Engine/Source/Runtime/Engine/Private/WorldPartition/LevelInstance/LevelInstanceActorDesc.cpp
--- a/Engine/Source/Runtime/Engine/Private/WorldPartition/LevelInstance/LevelInstanceActorDesc.cpp
+++ b/Engine/Source/Runtime/Engine/Private/WorldPartition/LevelInstance/LevelInstanceActorDesc.cpp
@@ -20,6 +20,7 @@
 #include "WorldPartition/ActorDescContainerSubsystem.h"
 #include "UObject/UE5ReleaseStreamObjectVersion.h"
 #include "UObject/UE5MainStreamObjectVersion.h"
+#include "UObject/FortniteSeasonBranchObjectVersion.h"
 #include "UObject/FortniteMainBranchObjectVersion.h"
 
 static int32 GLevelInstanceDebugForceLevelStreaming = 0;
@@ -49,7 +50,6 @@ void FLevelInstanceActorDesc::Init(const AActor* InActor)
 
 	const ILevelInstanceInterface* LevelInstance = CastChecked<ILevelInstanceInterface>(InActor);
 	WorldAsset = LevelInstance->GetWorldAsset().ToSoftObjectPath();
-	LevelInstanceTransform = InActor->GetActorTransform();
 	DesiredRuntimeBehavior = LevelInstance->GetDesiredRuntimeBehavior();
 	Filter = LevelInstance->GetFilter();
 	
@@ -80,7 +80,6 @@ bool FLevelInstanceActorDesc::Equals(const FWorldPartitionActorDesc* Other) cons
 
 		return
 			WorldAsset == LevelInstanceActorDesc->WorldAsset &&
-			LevelInstanceTransform.Equals(LevelInstanceActorDesc->LevelInstanceTransform, 0.1f) &&
 			DesiredRuntimeBehavior == LevelInstanceActorDesc->DesiredRuntimeBehavior;
 	}
 
@@ -163,7 +162,7 @@ bool FLevelInstanceActorDesc::IsChildContainerInstanceInternal() const
 FTransform FLevelInstanceActorDesc::GetChildContainerTransform() const
 {
 	FTransform LevelInstancePivotOffsetTransform = FTransform(ULevel::GetLevelInstancePivotOffsetFromPackage(GetChildContainerPackage()));
-	return LevelInstancePivotOffsetTransform * LevelInstanceTransform;
+	return LevelInstancePivotOffsetTransform * ActorTransform;
 }
 
 bool FLevelInstanceActorDesc::GetChildContainerInstance(const FWorldPartitionActorDescInstance* InActorDescInstance, FContainerInstance& OutContainerInstance) const
@@ -241,14 +240,6 @@ UActorDescContainerInstance* FLevelInstanceActorDesc::CreateChildContainerInstan
 	return ChildContainerInstance;
 }
 
-void FLevelInstanceActorDesc::TransferWorldData(const FWorldPartitionActorDesc* From)
-{
-	FWorldPartitionActorDesc::TransferWorldData(From);
-
-	FLevelInstanceActorDesc* FromLevelInstanceActorDesc = (FLevelInstanceActorDesc*)From;
-	LevelInstanceTransform = FromLevelInstanceActorDesc->LevelInstanceTransform;
-}
-
 void FLevelInstanceActorDesc::Serialize(FArchive& Ar)
 {
 	FWorldPartitionActorDesc::Serialize(Ar);
@@ -256,6 +247,7 @@ void FLevelInstanceActorDesc::Serialize(FArchive& Ar)
 	Ar.UsingCustomVersion(FUE5ReleaseStreamObjectVersion::GUID);
 	Ar.UsingCustomVersion(FUE5MainStreamObjectVersion::GUID);
 	Ar.UsingCustomVersion(FFortniteMainBranchObjectVersion::GUID);
+	Ar.UsingCustomVersion(FFortniteSeasonBranchObjectVersion::GUID);
 
 	if (Ar.CustomVer(FFortniteMainBranchObjectVersion::GUID) < FFortniteMainBranchObjectVersion::WorldPartitionActorDescSerializeSoftObjectPathSupport)
 	{
@@ -273,13 +265,13 @@ void FLevelInstanceActorDesc::Serialize(FArchive& Ar)
 	{
 		if (Ar.CustomVer(FUE5ReleaseStreamObjectVersion::GUID) < FUE5ReleaseStreamObjectVersion::LargeWorldCoordinates)
 		{
-			FTransform3f LevelInstanceTransformFlt;
-			Ar << LevelInstanceTransformFlt;
-			LevelInstanceTransform = FTransform(LevelInstanceTransformFlt);
+			FTransform3f ActorTransformFlt;
+			Ar << ActorTransformFlt;
+			ActorTransform = FTransform(ActorTransformFlt);
 		}
-		else
+		else if (Ar.CustomVer(FFortniteSeasonBranchObjectVersion::GUID) < FFortniteSeasonBranchObjectVersion::WorldPartitionActorDescActorTransformSerialization)
 		{
-			Ar << LevelInstanceTransform;
+			Ar << ActorTransform;
 		}
 	}
 
@@ -307,7 +299,7 @@ void FLevelInstanceActorDesc::Serialize(FArchive& Ar)
 				if (!IsChildContainerInstance())
 				{
 					FBox OutBounds;
-					if (ULevelInstanceSubsystem::GetLevelInstanceBoundsFromPackage(LevelInstanceTransform, GetChildContainerPackage(), OutBounds))
+					if (ULevelInstanceSubsystem::GetLevelInstanceBoundsFromPackage(ActorTransform, GetChildContainerPackage(), OutBounds))
 					{
 						OutBounds.GetCenterAndExtents(BoundsLocation, BoundsExtent);
 					}
diff --git a/Engine/Source/Runtime/Engine/Private/WorldPartition/WorldPartitionActorDesc.cpp b/Engine/Source/Runtime/Engine/Private/WorldPartition/WorldPartitionActorDesc.cpp
--- a/Engine/Source/Runtime/Engine/Private/WorldPartition/WorldPartitionActorDesc.cpp
+++ b/Engine/Source/Runtime/Engine/Private/WorldPartition/WorldPartitionActorDesc.cpp
@@ -88,6 +88,8 @@ void FWorldPartitionActorDesc::Init(const AActor* InActor)
 
 	check(Guid.IsValid());
 
+	ActorTransform = InActor->GetActorTransform();
+
 	const FBox StreamingBounds = !bIsDefaultActorDesc ? InActor->GetStreamingBounds() : FBox(ForceInit);
 	StreamingBounds.GetCenterAndExtents(BoundsLocation, BoundsExtent);
 	bIsBoundsValid = StreamingBounds.IsValid == 1;
@@ -276,6 +278,7 @@ bool FWorldPartitionActorDesc::Equals(const FWorldPartitionActorDesc* Other) con
 		ActorPath == Other->ActorPath &&
 		ActorLabel == Other->ActorLabel &&
 		bIsBoundsValid == Other->bIsBoundsValid &&
+		ActorTransform.Equals(Other->ActorTransform, 0.1f) &&
 		BoundsLocation.Equals(Other->BoundsLocation, 0.1f) &&
 		BoundsExtent.Equals(Other->BoundsExtent, 0.1f) &&
 		RuntimeGrid == Other->RuntimeGrid &&
@@ -529,6 +532,11 @@ void FWorldPartitionActorDesc::Serialize(FArchive& Ar)
 	else
 	{
 		Ar << Guid;
+
+		if (Ar.CustomVer(FFortniteSeasonBranchObjectVersion::GUID) >= FFortniteSeasonBranchObjectVersion::WorldPartitionActorDescActorTransformSerialization)
+		{
+			Ar << ActorTransform;
+		}
 	}
 
 	if(Ar.CustomVer(FUE5ReleaseStreamObjectVersion::GUID) < FUE5ReleaseStreamObjectVersion::LargeWorldCoordinates)
diff --git a/Engine/Source/Runtime/Engine/Public/WorldPartition/LevelInstance/LevelInstanceActorDesc.h b/Engine/Source/Runtime/Engine/Public/WorldPartition/LevelInstance/LevelInstanceActorDesc.h
--- a/Engine/Source/Runtime/Engine/Public/WorldPartition/LevelInstance/LevelInstanceActorDesc.h
+++ b/Engine/Source/Runtime/Engine/Public/WorldPartition/LevelInstance/LevelInstanceActorDesc.h
@@ -41,7 +41,6 @@ protected:
 	ENGINE_API virtual void Init(const FWorldPartitionActorDescInitData& DescData) override;
 	ENGINE_API virtual bool Equals(const FWorldPartitionActorDesc* Other) const override;
 	ENGINE_API virtual void TransferFrom(const FWorldPartitionActorDesc* From) override;
-	ENGINE_API virtual void TransferWorldData(const FWorldPartitionActorDesc* From) override;
 	virtual uint32 GetSizeOf() const override { return sizeof(FLevelInstanceActorDesc); }
 	ENGINE_API virtual void Serialize(FArchive& Ar) override;
 	ENGINE_API virtual void SetContainer(UActorDescContainer* InContainer) override;
@@ -51,7 +50,6 @@ protected:
 	ENGINE_API FTransform GetChildContainerTransform() const;
 
 	FSoftObjectPath WorldAsset;
-	FTransform LevelInstanceTransform;
 	ELevelInstanceRuntimeBehavior DesiredRuntimeBehavior;
 
 	TWeakObjectPtr<UActorDescContainer> ChildContainer;
diff --git a/Engine/Source/Runtime/Engine/Public/WorldPartition/WorldPartitionActorDesc.h b/Engine/Source/Runtime/Engine/Public/WorldPartition/WorldPartitionActorDesc.h
--- a/Engine/Source/Runtime/Engine/Public/WorldPartition/WorldPartitionActorDesc.h
+++ b/Engine/Source/Runtime/Engine/Public/WorldPartition/WorldPartitionActorDesc.h
@@ -142,6 +142,7 @@ public:
 	inline FName GetActorLabel() const { return ActorLabel; }
 	inline FName GetFolderPath() const { return FolderPath; }
 	inline const FGuid& GetFolderGuid() const { return FolderGuid; }
+	inline const FTransform& GetActorTransform() const { return ActorTransform; }
 
 	ENGINE_API virtual FBox GetEditorBounds() const;
 	ENGINE_API FBox GetRuntimeBounds() const;
@@ -349,6 +350,7 @@ protected:
 
 	virtual void TransferWorldData(const FWorldPartitionActorDesc* From)
 	{
+		ActorTransform = From->ActorTransform;
 		BoundsLocation = From->BoundsLocation;
 		BoundsExtent = From->BoundsExtent;
 		bIsBoundsValid = From->bIsBoundsValid;
@@ -365,6 +367,7 @@ protected:
 	FName							ActorPackage;	// Not serialized, comes from initialization data
 	FSoftObjectPath					ActorPath;		// Not serialized, comes from initialization data
 	FName							ActorLabel;
+	FTransform						ActorTransform;
 	FVector							BoundsLocation;
 	FVector							BoundsExtent;
 	FName							RuntimeGrid;
