PR #11332: Blueprint MemberReference now also looks for the "DeprecatedProperty"

diff --git a/Engine/Source/Runtime/Engine/Private/MemberReference.cpp b/Engine/Source/Runtime/Engine/Private/MemberReference.cpp
--- a/Engine/Source/Runtime/Engine/Private/MemberReference.cpp
+++ b/Engine/Source/Runtime/Engine/Private/MemberReference.cpp
@@ -497,11 +497,16 @@ TFieldType* FMemberReference::ResolveMemberImpl(UClass* SelfScope, TFieldTypeCla
 			}
 		}
 	}
-
 	// Check to see if the member has been deprecated
 	if (FProperty* Property = FFieldVariant(ReturnField).Get<FProperty>())
 	{
+#if WITH_EDITORONLY_DATA
+		// Initially this originated from python bindings, but this is useful to check for blueprints too, so that they can be upgraded.
+		static const FName NAME_DeprecatedProperty = TEXT("DeprecatedProperty");
+		bWasDeprecated = Property->HasAnyPropertyFlags(CPF_Deprecated) || Property->HasMetaData(NAME_DeprecatedProperty);
+#else
 		bWasDeprecated = Property->HasAnyPropertyFlags(CPF_Deprecated);
+#endif // WITH_EDITORONLY_DATA
 	}
 
 	return ReturnField;
@@ -617,4 +622,3 @@ FField* FMemberReference::FindRemappedField(FFieldClass* FieldClass, UClass* Ini
 	return FindRemappedFieldImpl<FField>(GLongCoreUObjectPackageName, FieldClass->GetFName(), InitialScope, InitialName,
 		bInitialScopeMustBeOwnerOfFieldForParentScopeRedirect);
 }
-
