NetDriver
* Added net.PrintNetConnections cheat that prints basic identification info about the NetConnections of a NetDriver
* Add NetDriverName=XXX or NetDriverDef=XXX to print the connections of a different NetDriver
* Default behavioris to print the GameNetDriver connections  (e.g. NetDriverDef=GameNetDriver)
* Use NetDriverDef=None to print ALL netdrivers

#rnx

diff --git a/Engine/Source/Runtime/Engine/Classes/Engine/NetDriver.h b/Engine/Source/Runtime/Engine/Classes/Engine/NetDriver.h
--- a/Engine/Source/Runtime/Engine/Classes/Engine/NetDriver.h
+++ b/Engine/Source/Runtime/Engine/Classes/Engine/NetDriver.h
@@ -1028,9 +1028,6 @@ public:
 	/** Set the NetDriver's NetDriverDefintion. */
 	void SetNetDriverDefinition(FName NewNetDriverDefinition);
 
-	UE_DEPRECATED(5.1, "Use GetNetDriverDefinition instead.")
-	FName GetNetDriverDefintion() const { return NetDriverDefinition; }
-
 	/** Get the NetDriver's NetDriverDefintion. */
 	FName GetNetDriverDefinition() const { return NetDriverDefinition; }
 
diff --git a/Engine/Source/Runtime/Engine/Private/NetDriver.cpp b/Engine/Source/Runtime/Engine/Private/NetDriver.cpp
--- a/Engine/Source/Runtime/Engine/Private/NetDriver.cpp
+++ b/Engine/Source/Runtime/Engine/Private/NetDriver.cpp
@@ -8138,12 +8138,6 @@ FAutoConsoleCommandWithWorld	DumpRelevantActorsCommand(
 	);
 #endif // NET_DEBUG_RELEVANT_ACTORS
 
-#if DO_ENABLE_NET_TEST
-
-
-
-#endif //#if DO_ENABLE_NET_TEST
-
 /**
  * Exec handler that routes online specific execs to the proper subsystem
  *
@@ -8194,3 +8188,82 @@ static bool NetDriverExec(UWorld* InWorld, const TCHAR* Cmd, FOutputDevice& Ar)
 /** Our entry point for all net driver related exec routing */
 FStaticSelfRegisteringExec NetDriverExecRegistration(NetDriverExec);
 
+
+namespace UE::Net::Private
+{
+
+FAutoConsoleCommand PrintNetConnectionInfoCommand(
+TEXT("net.PrintNetConnections"),
+TEXT("Prints information on all net connections of a NetDriver.  Defaults to the GameNetDriver. Choose a different driver via NetDriverName= or NetDriverDefinition="),
+FConsoleCommandWithArgsDelegate::CreateLambda([](const TArray< FString >& Args)
+{
+	FName NetDriverName;
+
+	// Default to print GameNetDriver if no params are set
+	FName NetDriverDef = NAME_GameNetDriver;
+
+	if (const FString* DriverNameStr = Args.FindByPredicate([](const FString& Str) { return Str.Contains(TEXT("NetDriverName=")); }))
+	{
+		FParse::Value(**DriverNameStr, TEXT("NetDriverName="), NetDriverName);
+	}
+	else if (const FString* DriverDefStr = Args.FindByPredicate([](const FString& Str) { return Str.Contains(TEXT("NetDriverDef=")); }))
+	{
+		FParse::Value(**DriverDefStr, TEXT("NetDriverDef="), NetDriverDef);
+	}
+
+	auto PrintNetConnections = [](UNetDriver* NetDriver)
+	{
+		if (NetDriver->ServerConnection)
+		{
+			UE_LOG(LogNet, Display, TEXT("Printing Server Connection for: %s (Definition=%s)"), *NetDriver->NetDriverName.ToString(), *NetDriver->GetNetDriverDefinition().ToString());
+
+			UE_LOG(LogNet, Display, TEXT("\tServerConnection: ConnectionId=%u ViewTarget=%s FullDescription=%s"), 
+				NetDriver->ServerConnection->GetConnectionId(), 
+				*GetNameSafe(NetDriver->ServerConnection->ViewTarget),
+				*NetDriver->ServerConnection->Describe()
+			);
+		}
+		else
+		{
+			UE_LOG(LogNet, Display, TEXT("Printing Client Connections (%u) for: %s (Definition=%s)"), NetDriver->ClientConnections.Num(), *NetDriver->NetDriverName.ToString(), *NetDriver->GetNetDriverDefinition().ToString());
+
+			for (UNetConnection* NetConnection : NetDriver->ClientConnections)
+			{
+				UE_LOG(LogNet, Display, TEXT("\tClientConnection: ConnectionId=%u ViewTarget=%s NetId=%s FullDescription=%s"), 
+					NetConnection->GetConnectionId(), 
+					*GetNameSafe(NetConnection->ViewTarget),
+					*NetConnection->PlayerId.ToDebugString(), 
+					*NetConnection->Describe()
+				);
+			}
+		}
+	};
+
+	for (TObjectIterator<UNetDriver> It; It; ++It)
+	{
+		if (UNetDriver* NetDriver = *It)
+		{
+			if (NetDriverName != NAME_None)
+			{
+				if (NetDriver->NetDriverName == NetDriverName)
+				{
+					PrintNetConnections(NetDriver);
+				}
+			}
+			else if (NetDriverDef != NAME_None)
+			{
+				if (NetDriver->GetNetDriverDefinition() == NetDriverDef)
+				{
+					PrintNetConnections(NetDriver);
+				}
+			}
+			// Print all NetDrivers when both names are none
+			else
+			{
+				PrintNetConnections(NetDriver);
+			}
+		}
+	}
+}));
+
+} // end namespace UE::Net::Private
