AssetCompilingManager make expected memory overrun message a Display not a Warning so it isn't logged in jira when hit

diff --git a/Engine/Source/Runtime/Engine/Private/AssetCompilingManager.cpp b/Engine/Source/Runtime/Engine/Private/AssetCompilingManager.cpp
--- a/Engine/Source/Runtime/Engine/Private/AssetCompilingManager.cpp
+++ b/Engine/Source/Runtime/Engine/Private/AssetCompilingManager.cpp
@@ -293,14 +293,19 @@ public:
 					DebugName = TEXT("No DebugName");
 				}
 
-				UE_LOGFMT_NSLOC(LogAsyncCompilation, Warning, "AsyncAssetCompilation", "MemoryLimitExceeded",
-					"AssetCompile memory estimate is greater than available, but we're running it [{TaskName}] anyway, likely to fail! "
-					"RequiredMemory = {TotalEstimatedMemory} MiB + {RequiredMemory} MiB, MemoryLimit = {MemoryLimit} MiB ",
-					("TaskName", DebugName),
-					("RequiredMemory", FString::SanitizeFloat(NewRequiredMemory / (1024 * 1024.f), 3)),
-					("TotalEstimatedMemory", FString::SanitizeFloat(TotalEstimatedMemory / (1024 * 1024.f), 3)),
-					("MemoryLimit", FString::SanitizeFloat(MemoryLimit / (1024 * 1024.f), 3))
-				);
+				if ( NewRequiredMemory != GetDefaultMemoryPerAsset() )
+				{
+					// don't bother logging if estimate is the default (4GB)
+
+					UE_LOGFMT_NSLOC(LogAsyncCompilation, Display, "AsyncAssetCompilation", "MemoryLimitExceeded",
+						"BEWARE: AssetCompile memory estimate is greater than available, but we're running it [{TaskName}] anyway! "
+						"RequiredMemory = {TotalEstimatedMemory} MiB + {RequiredMemory} MiB, MemoryLimit = {MemoryLimit} MiB ",
+						("TaskName", DebugName),
+						("RequiredMemory", FString::SanitizeFloat(NewRequiredMemory / (1024 * 1024.f), 3)),
+						("TotalEstimatedMemory", FString::SanitizeFloat(TotalEstimatedMemory / (1024 * 1024.f), 3)),
+						("MemoryLimit", FString::SanitizeFloat(MemoryLimit / (1024 * 1024.f), 3))
+					);
+				}
 
 				// @todo : ? pause the main thread? pause shader compilers? trigger a GC ?
 
