Size estimate use existing functions.

diff --git a/Engine/Source/Runtime/Engine/Private/TextureDerivedData.cpp b/Engine/Source/Runtime/Engine/Private/TextureDerivedData.cpp
--- a/Engine/Source/Runtime/Engine/Private/TextureDerivedData.cpp
+++ b/Engine/Source/Runtime/Engine/Private/TextureDerivedData.cpp
@@ -631,54 +631,19 @@ static void GetBuiltTextureSizeBytesEstimate(
 	uint64 & OutTopMipSizeBytes, uint64 & OutTotalImageSizeBytes)
 {
 	check( PixelFormat != PF_Unknown );
-	const FPixelFormatInfo & PFI = GPixelFormats[PixelFormat];
 
-	uint64 TopMipSizeBytes = 0;
-	uint64 TotalImageSizeBytes = 0;
-
-	int64 SizeX = TopMipSizeX;
-	int64 SizeY = TopMipSizeY;
-	int64 SizeZ = TopMipSizeZ;
-	
 	int64 NumMips = FImageCoreUtils::GetMipCountFromDimensions(TopMipSizeX,TopMipSizeY,TopMipSizeZ,bIsVolume);
 	check( NumMips > 0 );
 
-	// calculate bytes for linear unpadded/untiled layout :
-
-	for(int64 Mip=0;Mip<NumMips;Mip++)
-	{
-		int64 NumBlocksXY = 
-			FMath::DivideAndRoundUp<int64>( SizeX, PFI.BlockSizeX ) * 
-			FMath::DivideAndRoundUp<int64>( SizeY, PFI.BlockSizeY ); 
-		uint64 SizeBytes = NumBlocksXY * PFI.BlockBytes * SizeZ;
-
-		if ( Mip == 0 )
-		{
-			// save size of first mip
-			TopMipSizeBytes = SizeBytes;
-		}
-
-		TotalImageSizeBytes += SizeBytes;
-
-		SizeX = FMath::Max(1,SizeX>>1);
-		SizeY = FMath::Max(1,SizeY>>1);
-		if ( bIsVolume )
-		{
-			SizeZ = FMath::Max(1,SizeZ>>1);
-		}
-	}
-	
-	check( TotalImageSizeBytes > 0 );
-
-	OutTopMipSizeBytes = TopMipSizeBytes;
-	OutTotalImageSizeBytes = TotalImageSizeBytes;
-	
 	// if alpha is unknown, assume yes to be conservative about pixel size
-	bool bHasAlpha = ( BuildSettings.bKnowAlphaTransparency ) ? BuildSettings.bHasTransparentAlpha : true;
+	bool bHasAlpha = (BuildSettings.bKnowAlphaTransparency) ? BuildSettings.bHasTransparentAlpha : true;
 
 	FEncodedTextureDescription TextureDescription;
 	BuildSettings.GetEncodedTextureDescription(&TextureDescription, TextureFormat, TopMipSizeX, TopMipSizeY, TopMipSizeZ, NumMips, bHasAlpha);
-	check( TextureDescription.PixelFormat == PixelFormat );
+	check(TextureDescription.PixelFormat == PixelFormat);
+
+	uint64 TopMipSizeBytes = 0;
+	uint64 TotalImageSizeBytes = 0;
 
 	int32 LODBias = 0;
 	FEncodedTextureExtendedData ExtendedData = TextureFormat->GetExtendedDataForTexture(TextureDescription, LODBias);
@@ -696,11 +661,30 @@ static void GetBuiltTextureSizeBytesEstimate(
 		{
 			TotalImageSizeBytes += MipSize;
 		}
-		
-		OutTopMipSizeBytes = TopMipSizeBytes;
-		OutTotalImageSizeBytes = TotalImageSizeBytes;
+
+	}
+	else
+	{
+		// calculate bytes for linear unpadded/untiled layout :
+		for (int32 MipIndex = 0; MipIndex < TextureDescription.NumMips; MipIndex++)
+		{
+			if (MipIndex == 0)
+			{
+				TopMipSizeBytes = TextureDescription.GetMipSizeInBytes(0);
+				TotalImageSizeBytes = TopMipSizeBytes;
+			}
+			else
+			{
+				TotalImageSizeBytes += TextureDescription.GetMipSizeInBytes(MipIndex);
+			}
+		}
 	}
 
+	check(TotalImageSizeBytes > 0);
+
+
+	OutTopMipSizeBytes = TopMipSizeBytes;
+	OutTotalImageSizeBytes = TotalImageSizeBytes;
 }
 
 // may reduce OutSettings.MaxTextureResolution
