Cap mip memcpy now that we've gone a while without seeing this warning, suggesting our computations are correct.

diff --git a/Engine/Source/Runtime/Engine/Private/Rendering/Texture2DResource.cpp b/Engine/Source/Runtime/Engine/Private/Rendering/Texture2DResource.cpp
--- a/Engine/Source/Runtime/Engine/Private/Rendering/Texture2DResource.cpp
+++ b/Engine/Source/Runtime/Engine/Private/Rendering/Texture2DResource.cpp
@@ -286,6 +286,7 @@ void FTexture2DResource::GetData( uint32 MipIndex, void* Dest, uint32 DestPitch,
 			MipMap.SizeX,MipMap.SizeY,
 			EffectiveSize,DataSize,
 			SrcPitch,DestPitch, PixelFormat, DestSize);
+		DataSize  = DestSize; // don't overwrite memory
 	}
 
 	// for platforms that returned 0 pitch from Lock, we need to just use the bulk data directly, never do 
