Fixed triggered ensure when removing and adding a SparseVolumeTexture to the streamer. Fixed by properly removing the SVT from a recently introduced TMap which maps from the SVT uobject pointer to an internal handle.

#rnx

diff --git a/Engine/Source/Runtime/Engine/Private/SparseVolumeTexture/SparseVolumeTextureStreamingManager.cpp b/Engine/Source/Runtime/Engine/Private/SparseVolumeTexture/SparseVolumeTextureStreamingManager.cpp
--- a/Engine/Source/Runtime/Engine/Private/SparseVolumeTexture/SparseVolumeTextureStreamingManager.cpp
+++ b/Engine/Source/Runtime/Engine/Private/SparseVolumeTexture/SparseVolumeTextureStreamingManager.cpp
@@ -781,6 +781,7 @@ void FStreamingManager::RemoveInternal(UStreamableSparseVolumeTexture* SparseVol
 		}
 
 		StreamingInfo.RemoveAt(SVTInfo->SVTHandle, 1);
+		SparseVolumeTextureToHandle.Remove(SparseVolumeTexture);
 	}
 }
 
