Fixed crash when force deleting a SVT that is currently playing. Fixed by registering an editor delegate which also adds all the child uobjects (frames) of the SVT to the list of objects to delete. This in turn cleans up any remaning references to the SVT.

#rnx

diff --git a/Engine/Source/Runtime/Engine/Classes/SparseVolumeTexture/SparseVolumeTexture.h b/Engine/Source/Runtime/Engine/Classes/SparseVolumeTexture/SparseVolumeTexture.h
--- a/Engine/Source/Runtime/Engine/Classes/SparseVolumeTexture/SparseVolumeTexture.h
+++ b/Engine/Source/Runtime/Engine/Classes/SparseVolumeTexture/SparseVolumeTexture.h
@@ -535,6 +535,10 @@ public:
 	virtual const UE::SVT::FTextureRenderResources* GetTextureRenderResources() const override { return Frames.IsEmpty() ? nullptr : Frames[0]->GetTextureRenderResources(); }
 	//~ End USparseVolumeTexture Interface.
 
+#if WITH_EDITOR
+	ENGINE_API void OnAssetsAddExtraObjectsToDelete(TArray<UObject*>& ObjectsToDelete);
+#endif
+
 protected:
 
 	UPROPERTY(Export)
@@ -562,7 +566,14 @@ protected:
 	// Ensures all frames have derived data (based on the source data and the current settings like TextureAddress modes etc.) cached to DDC and are ready for rendering.
 	// Disconnects this SVT from the streaming manager, calls Cache() on all frames and finally connects to FStreamingManager again.
 	void RecacheFrames();
+
 #endif // WITH_EDITORONLY_DATA
+
+#if WITH_EDITOR
+	ENGINE_API bool ShouldRegisterDelegates();
+	ENGINE_API void RegisterEditorDelegates();
+	ENGINE_API void UnregisterEditorDelegates();
+#endif // WITH_EDITOR
 };
 
 // Represents a streamable SparseVolumeTexture asset with a single frame. Although there is only a single frame, it is still recommended to use USparseVolumeTextureFrame::GetFrameAndIssueStreamingRequest().
diff --git a/Engine/Source/Runtime/Engine/Private/SparseVolumeTexture/SparseVolumeTexture.cpp b/Engine/Source/Runtime/Engine/Private/SparseVolumeTexture/SparseVolumeTexture.cpp
--- a/Engine/Source/Runtime/Engine/Private/SparseVolumeTexture/SparseVolumeTexture.cpp
+++ b/Engine/Source/Runtime/Engine/Private/SparseVolumeTexture/SparseVolumeTexture.cpp
@@ -38,6 +38,10 @@
 #include "Serialization/EditorBulkDataReader.h"
 #include "Serialization/EditorBulkDataWriter.h"
 
+#if WITH_EDITOR
+#include "Editor.h"
+#endif
+
 #define LOCTEXT_NAMESPACE "USparseVolumeTexture"
 
 DEFINE_LOG_CATEGORY_STATIC(LogSparseVolumeTexture, Log, All);
@@ -1107,6 +1111,19 @@ void USparseVolumeTextureFrame::Cache(bool bSkipDDCAndSetResourcesToDefault)
 UStreamableSparseVolumeTexture::UStreamableSparseVolumeTexture(const FObjectInitializer& ObjectInitializer)
 	: Super(ObjectInitializer)
 {
+#if WITH_EDITOR
+	if (!IsTemplate())
+	{
+		if (HasAnyFlags(RF_NeedPostLoad) || GetOuter()->HasAnyFlags(RF_NeedPostLoad))
+		{
+			// Delegate registration is not thread-safe, so we postpone it on PostLoad when coming from loading which could be on another thread
+		}
+		else
+		{
+			RegisterEditorDelegates();
+		}
+	}
+#endif // WITH_EDITOR
 }
 
 bool UStreamableSparseVolumeTexture::BeginInitialize(int32 NumExpectedFrames)
@@ -1310,6 +1327,13 @@ void UStreamableSparseVolumeTexture::PostLoad()
 {
 	Super::PostLoad();
 
+#if WITH_EDITOR
+	if (!IsTemplate())
+	{
+		RegisterEditorDelegates();
+	}
+#endif
+
 	// Ensure that NumFrames always corresponds to the actual number of frames
 	NumFrames = GetNumFrames();
 
@@ -1318,6 +1342,7 @@ void UStreamableSparseVolumeTexture::PostLoad()
 #else
 	for (USparseVolumeTextureFrame* Frame : Frames)
 	{
+		check(Frame); // Elements in Frames should only ever be null when the SVT is being deleted
 		Frame->CreateTextureRenderResources();
 	}
 	UE::SVT::GetStreamingManager().Add_GameThread(this); // RecacheFrames() handles this in editor builds
@@ -1333,6 +1358,13 @@ void UStreamableSparseVolumeTexture::BeginDestroy()
 {
 	Super::BeginDestroy();
 	UE::SVT::GetStreamingManager().Remove_GameThread(this);
+
+#if WITH_EDITOR
+	if (!IsTemplate())
+	{
+		UnregisterEditorDelegates();
+	}
+#endif
 }
 
 void UStreamableSparseVolumeTexture::Serialize(FArchive& Ar)
@@ -1348,6 +1380,17 @@ void UStreamableSparseVolumeTexture::Serialize(FArchive& Ar)
 #if WITH_EDITOR
 void UStreamableSparseVolumeTexture::PostEditChangeProperty(FPropertyChangedEvent& PropertyChangedEvent)
 {
+	// It's possible for outside code/GC to null the elements in Frames. This happens when the SVT is deleted and all the child objects are also first deleted and their references nulled.
+	bool bInvalidFramesArray = false;
+	for (USparseVolumeTextureFrame* Frame : Frames)
+	{
+		if (!Frame)
+		{
+			bInvalidFramesArray = true;
+			break;
+		}
+	}
+
 	if (PropertyChangedEvent.GetPropertyName() == GET_MEMBER_NAME_CHECKED(UStreamableSparseVolumeTexture, AddressX)
 		|| PropertyChangedEvent.GetPropertyName() == GET_MEMBER_NAME_CHECKED(UStreamableSparseVolumeTexture, AddressY)
 		|| PropertyChangedEvent.GetPropertyName() == GET_MEMBER_NAME_CHECKED(UStreamableSparseVolumeTexture, AddressZ))
@@ -1356,13 +1399,20 @@ void UStreamableSparseVolumeTexture::PostEditChangeProperty(FPropertyChangedEven
 		NotifyMaterials();
 		for (USparseVolumeTextureFrame* Frame : Frames)
 		{
-			Frame->NotifyMaterials();
+			if (Frame)
+			{
+				Frame->NotifyMaterials();
+			}
 		}
 	}
 
 	Super::PostEditChangeProperty(PropertyChangedEvent);
 
-	RecacheFrames();
+	// Don't bother trying to recache the frame data if the Frames array is invalid. This very likely means that this object is about to be deleted.
+	if (!bInvalidFramesArray)
+	{
+		RecacheFrames();
+	}
 }
 #endif // WITH_EDITOR
 
@@ -1374,7 +1424,10 @@ void UStreamableSparseVolumeTexture::GetResourceSizeEx(FResourceSizeEx& Cumulati
 	SizeCPU += Frames.GetAllocatedSize();
 	for (USparseVolumeTextureFrame* Frame : Frames)
 	{
-		Frame->GetResourceSizeEx(CumulativeResourceSize);
+		if (Frame)
+		{
+			Frame->GetResourceSizeEx(CumulativeResourceSize);
+		}
 	}
 	CumulativeResourceSize.AddDedicatedSystemMemoryBytes(SizeCPU);
 	CumulativeResourceSize.AddDedicatedVideoMemoryBytes(SizeGPU);
@@ -1442,6 +1495,23 @@ const TArray<UAssetUserData*>* UStreamableSparseVolumeTexture::GetAssetUserDataA
 	return &ToRawPtrTArrayUnsafe(AssetUserData);
 }
 
+#if WITH_EDITOR
+void UStreamableSparseVolumeTexture::OnAssetsAddExtraObjectsToDelete(TArray<UObject*>& ObjectsToDelete)
+{
+	if (ObjectsToDelete.Contains(this))
+	{
+		// When UStreamableSparseVolumeTexture is deleted, we also want all owned USparseVolumeTextureFrame objects to be deleted.
+		for (USparseVolumeTextureFrame* Frame : Frames)
+		{
+			if (Frame)
+			{
+				ObjectsToDelete.Add(Frame);
+			}
+		}
+	}
+}
+#endif
+
 #if WITH_EDITORONLY_DATA
 void UStreamableSparseVolumeTexture::RecacheFrames()
 {
@@ -1460,6 +1530,7 @@ void UStreamableSparseVolumeTexture::RecacheFrames()
 	bool bCanceled = false;
 	for (USparseVolumeTextureFrame* Frame : Frames)
 	{
+		check(Frame); // RecacheFrames() is assumed to never be called when the Frames array is invalid (has nullptr elements). Elements may be nulled as part of deleting the SVT.
 		if (!bCanceled && RecacheTask.ShouldCancel())
 		{
 			bCanceled = true;
@@ -1485,6 +1556,28 @@ void UStreamableSparseVolumeTexture::RecacheFrames()
 }
 #endif
 
+#if WITH_EDITOR
+bool UStreamableSparseVolumeTexture::ShouldRegisterDelegates()
+{
+	return GEditor && !IsTemplate() && !IsRunningCookCommandlet();
+}
+void UStreamableSparseVolumeTexture::RegisterEditorDelegates()
+{
+	if (ShouldRegisterDelegates())
+	{
+		UnregisterEditorDelegates();
+		FEditorDelegates::OnAssetsAddExtraObjectsToDelete.AddUObject(this, &UStreamableSparseVolumeTexture::OnAssetsAddExtraObjectsToDelete);
+	}
+}
+void UStreamableSparseVolumeTexture::UnregisterEditorDelegates()
+{
+	if (ShouldRegisterDelegates())
+	{
+		FEditorDelegates::OnAssetsAddExtraObjectsToDelete.RemoveAll(this);
+	}
+}
+#endif // WITH_EDITOR
+
 ////////////////////////////////////////////////////////////////////////////////////////////////
 
 UStaticSparseVolumeTexture::UStaticSparseVolumeTexture(const FObjectInitializer& ObjectInitializer)
