- Fixed error reporting of unsupported child ClientOnly/ServerOnly data layers
- Fixed wrong message displayed for unsupported editor data layer child of a runtime parent data layer

diff --git a/Engine/Plugins/Editor/DataValidation/Source/DataValidation/Private/WorldPartitionChangelistValidator.cpp b/Engine/Plugins/Editor/DataValidation/Source/DataValidation/Private/WorldPartitionChangelistValidator.cpp
--- a/Engine/Plugins/Editor/DataValidation/Source/DataValidation/Private/WorldPartitionChangelistValidator.cpp
+++ b/Engine/Plugins/Editor/DataValidation/Source/DataValidation/Private/WorldPartitionChangelistValidator.cpp
@@ -393,19 +393,35 @@ void UWorldPartitionChangelistValidator::OnInvalidDataLayerAssetType(const UData
 	}
 }
 
-void UWorldPartitionChangelistValidator::OnDataLayerHierarchyTypeMismatch(const UDataLayerInstance* DataLayerInstance, const UDataLayerInstance* Parent)
+void UWorldPartitionChangelistValidator::OnDataLayerHierarchyTypeMismatch(const UDataLayerInstance* DataLayerInstance, const UDataLayerInstance* Parent, EDataLayerHierarchyInvalidReason Reason)
 {
 	if (Filter(DataLayerInstance)
 		|| Filter(Parent)
 		|| bSubmittingWorldDataLayers)
 	{
-		FText CurrentError = FText::Format(LOCTEXT("DataValidation.Changelist.WorldPartition.DataLayerHierarchyTypeMismatch", "Data layer {0} is of type {1} and its parent {2} is of type {3}."),
-			FText::FromString(DataLayerInstance->GetDataLayerFullName()),
-			UEnum::GetDisplayValueAsText(DataLayerInstance->GetType()),
-			FText::FromString(Parent->GetDataLayerFullName()),
-			UEnum::GetDisplayValueAsText(Parent->GetType()));
-	
-		AssetFails(CurrentAsset, CurrentError);
+		switch (Reason)
+		{
+		case EDataLayerHierarchyInvalidReason::ClientOnlyDataLayerCantBeChild:
+		case EDataLayerHierarchyInvalidReason::ServerOnlyDataLayerCantBeChild:
+			{
+				FText CurrentError = FText::Format(LOCTEXT("DataValidation.Changelist.WorldPartition.DataLayerHierarchyCantBeChild", "Data layer {0} is {1} which can't be child of parent {2}."),
+					FText::FromString(DataLayerInstance->GetDataLayerFullName()),
+					(Reason == EDataLayerHierarchyInvalidReason::ClientOnlyDataLayerCantBeChild) ? LOCTEXT("DataLayerIsClientOnly", "client-only") : LOCTEXT("DataLayerIsServerOnly", "server-only"),
+					FText::FromString(Parent->GetDataLayerFullName()));
+				AssetFails(CurrentAsset, CurrentError);
+			}
+			break;
+		case EDataLayerHierarchyInvalidReason::IncompatibleDataLayerType:
+			{
+				FText CurrentError = FText::Format(LOCTEXT("DataValidation.Changelist.WorldPartition.DataLayerHierarchyTypeMismatch", "Data layer {0} is of type {1} and its parent {2} is of type {3}."),
+					FText::FromString(DataLayerInstance->GetDataLayerFullName()),
+					UEnum::GetDisplayValueAsText(DataLayerInstance->GetType()),
+					FText::FromString(Parent->GetDataLayerFullName()),
+					UEnum::GetDisplayValueAsText(Parent->GetType()));
+				AssetFails(CurrentAsset, CurrentError);
+			}
+			break;
+		}
 	}
 }
 
diff --git a/Engine/Plugins/Editor/DataValidation/Source/DataValidation/Private/WorldPartitionChangelistValidator.h b/Engine/Plugins/Editor/DataValidation/Source/DataValidation/Private/WorldPartitionChangelistValidator.h
--- a/Engine/Plugins/Editor/DataValidation/Source/DataValidation/Private/WorldPartitionChangelistValidator.h
+++ b/Engine/Plugins/Editor/DataValidation/Source/DataValidation/Private/WorldPartitionChangelistValidator.h
@@ -43,7 +43,7 @@ protected:
 	virtual void OnInvalidReferenceRuntimeGrid(const IWorldPartitionActorDescInstanceView& ActorDescView, const IWorldPartitionActorDescInstanceView& ReferenceActorDescView) override;
 	virtual void OnInvalidReferenceDataLayerAsset(const UDataLayerInstanceWithAsset* DataLayerInstance) override;
 	virtual void OnInvalidDataLayerAssetType(const UDataLayerInstanceWithAsset* DataLayerInstance, const UDataLayerAsset* DataLayerAsset) override;
-	virtual void OnDataLayerHierarchyTypeMismatch(const UDataLayerInstance* DataLayerInstance, const UDataLayerInstance* Parent) override;
+	virtual void OnDataLayerHierarchyTypeMismatch(const UDataLayerInstance* DataLayerInstance, const UDataLayerInstance* Parent, EDataLayerHierarchyInvalidReason Reason) override;
 	virtual void OnDataLayerAssetConflict(const UDataLayerInstanceWithAsset* DataLayerInstance, const UDataLayerInstanceWithAsset* ConflictingDataLayerInstance) override;
 	virtual void OnActorNeedsResave(const IWorldPartitionActorDescInstanceView& ActorDescView) override;
 	virtual void OnLevelInstanceInvalidWorldAsset(const IWorldPartitionActorDescInstanceView& ActorDescView, FName WorldAsset, ELevelInstanceInvalidReason Reason) override;
diff --git a/Engine/Source/Editor/DataLayerEditor/Private/DataLayer/DataLayerMode.cpp b/Engine/Source/Editor/DataLayerEditor/Private/DataLayer/DataLayerMode.cpp
--- a/Engine/Source/Editor/DataLayerEditor/Private/DataLayer/DataLayerMode.cpp
+++ b/Engine/Source/Editor/DataLayerEditor/Private/DataLayer/DataLayerMode.cpp
@@ -94,6 +94,7 @@
 #include "WorldPartition/DataLayer/DataLayerInstance.h"
 #include "WorldPartition/DataLayer/DataLayerInstanceWithAsset.h"
 #include "WorldPartition/DataLayer/DataLayerManager.h"
+#include "WorldPartition/DataLayer/DataLayerUtils.h"
 #include "WorldPartition/DataLayer/ExternalDataLayerAsset.h"
 #include "WorldPartition/DataLayer/ExternalDataLayerInstance.h"
 #include "WorldPartition/DataLayer/ExternalDataLayerManager.h"
@@ -860,19 +861,12 @@ FSceneOutlinerDragValidationInfo FDataLayerMode::ValidateDataLayerAssetDrop(cons
 	if (DropTargetDataLayerWithAsset)
 	{
 		EDataLayerType ParentType = DropTargetDataLayerWithAsset->GetType();
-		auto IsParentDataLayerTypeCompatible = [ParentType](const UDataLayerAsset* InChildDataLayerAsset)
-		{
-			EDataLayerType ChildType = InChildDataLayerAsset->GetType();
-			return (ChildType != EDataLayerType::Unknown) &&
-				   (ParentType != EDataLayerType::Unknown) &&
-				   (ParentType == EDataLayerType::Editor || ChildType == EDataLayerType::Runtime);
-		};
-
 		for (const UDataLayerAsset* DataLayerAssetToDrop : DataLayerAssetsToDrop)
 		{
-			if (!IsParentDataLayerTypeCompatible(DataLayerAssetToDrop))
+			FText Reason;
+			if (!FDataLayerUtils::AreDataLayerTypesCompatible(ParentType, DataLayerAssetToDrop->GetType(), &Reason))
 			{
-				return FSceneOutlinerDragValidationInfo(ESceneOutlinerDropCompatibility::Incompatible, FText::Format(LOCTEXT("CantCreateDataLayerInstanceIncompatibleChildType", "Cannot create Data Layer Instance : {0} Data Layer cannot have {1} child Data Layers"), UEnum::GetDisplayValueAsText(ParentType), UEnum::GetDisplayValueAsText(DataLayerAssetToDrop->GetType())));
+				return FSceneOutlinerDragValidationInfo(ESceneOutlinerDropCompatibility::Incompatible, FText::Format(LOCTEXT("CantCreateDataLayerInstanceIncompatibleChildType", "Cannot create Data Layer Instance : {0}"), Reason));
 			}
 		}
 	}
diff --git a/Engine/Source/Runtime/Engine/Private/WorldPartition/DataLayer/DataLayerInstance.cpp b/Engine/Source/Runtime/Engine/Private/WorldPartition/DataLayer/DataLayerInstance.cpp
--- a/Engine/Source/Runtime/Engine/Private/WorldPartition/DataLayer/DataLayerInstance.cpp
+++ b/Engine/Source/Runtime/Engine/Private/WorldPartition/DataLayer/DataLayerInstance.cpp
@@ -437,7 +437,7 @@ bool UDataLayerInstance::CanBeChildOf(const UDataLayerInstance* InParent, FText*
 	return true;
 }
 
-bool UDataLayerInstance::IsParentDataLayerTypeCompatible(const UDataLayerInstance* InParent, FText* OutReason) const
+bool UDataLayerInstance::IsParentDataLayerTypeCompatible(const UDataLayerInstance* InParent, FText* OutReasonText, IStreamingGenerationErrorHandler::EDataLayerHierarchyInvalidReason* OutReason) const
 {
 	if (InParent == nullptr)
 	{
@@ -448,21 +448,20 @@ bool UDataLayerInstance::IsParentDataLayerTypeCompatible(const UDataLayerInstanc
 	{
 		if (OutReason)
 		{
-			*OutReason = FText::Format(LOCTEXT("ClientOrServerOnlyCantHaveParent", "{0} Data Layer cannot be a child Data Layer"), IsClientOnly() ? FText::FromString(TEXT("Client-Only")) : FText::FromString(TEXT("Server-Only")));
+			*OutReason = IsClientOnly() ? IStreamingGenerationErrorHandler::EDataLayerHierarchyInvalidReason::ClientOnlyDataLayerCantBeChild : IStreamingGenerationErrorHandler::EDataLayerHierarchyInvalidReason::ServerOnlyDataLayerCantBeChild;
+		}
+		if (OutReasonText)
+		{
+			*OutReasonText = FText::Format(LOCTEXT("ClientOrServerOnlyCantHaveParent", "{0} Data Layer cannot be a child Data Layer"), IsClientOnly() ? FText::FromString(TEXT("Client-Only")) : FText::FromString(TEXT("Server-Only")));
 		}
 		return false;
 	}
 
-	const EDataLayerType ParentType = InParent->GetType();
-	const EDataLayerType ChildType = GetType();
-
-	if ((ChildType == EDataLayerType::Unknown) ||
-		(ParentType == EDataLayerType::Unknown) || 
-		(ParentType != EDataLayerType::Editor && ChildType != EDataLayerType::Runtime))
+	if (!FDataLayerUtils::AreDataLayerTypesCompatible(InParent->GetType(), GetType(), OutReasonText))
 	{
 		if (OutReason)
 		{
-			*OutReason = FText::Format(LOCTEXT("IncompatibleChildType", "{0} Data Layer cannot have {1} child Data Layers"), UEnum::GetDisplayValueAsText(InParent->GetType()), UEnum::GetDisplayValueAsText(GetType()));
+			*OutReason = IStreamingGenerationErrorHandler::EDataLayerHierarchyInvalidReason::IncompatibleDataLayerType;
 		}
 		return false;
 	}
@@ -535,10 +534,15 @@ FText UDataLayerInstance::GetDataLayerText(const UDataLayerInstance* InDataLayer
 
 bool UDataLayerInstance::Validate(IStreamingGenerationErrorHandler* ErrorHandler) const
 {
-	if (GetParent() != nullptr && !IsParentDataLayerTypeCompatible(GetParent()))
+	if (GetParent() != nullptr)
 	{
-		ErrorHandler->OnDataLayerHierarchyTypeMismatch(this, GetParent());
-		return false;
+		FText ReasonText;
+		IStreamingGenerationErrorHandler::EDataLayerHierarchyInvalidReason Reason;
+		if (!IsParentDataLayerTypeCompatible(GetParent(), &ReasonText, &Reason))
+		{
+			ErrorHandler->OnDataLayerHierarchyTypeMismatch(this, GetParent(), Reason);
+			return false;
+		}
 	}
 
 	return true;
diff --git a/Engine/Source/Runtime/Engine/Private/WorldPartition/DataLayer/DataLayerUtils.cpp b/Engine/Source/Runtime/Engine/Private/WorldPartition/DataLayer/DataLayerUtils.cpp
--- a/Engine/Source/Runtime/Engine/Private/WorldPartition/DataLayer/DataLayerUtils.cpp
+++ b/Engine/Source/Runtime/Engine/Private/WorldPartition/DataLayer/DataLayerUtils.cpp
@@ -11,6 +11,8 @@
 #include "WorldPartition/DataLayer/DataLayerManager.h"
 #include "Algo/AllOf.h"
 
+#define LOCTEXT_NAMESPACE "DataLayer"
+
 FDataLayerInstanceNames FDataLayerUtils::ResolveDataLayerInstanceNames(const UDataLayerManager* InDataLayerManager, const FWorldPartitionActorDesc* InActorDesc, const TArray<const FWorldDataLayersActorDesc*>& InWorldDataLayersActorDescs)
 {
 	const bool bIncludeExternalDataLayer = false;
@@ -242,4 +244,21 @@ FString FDataLayerUtils::GenerateUniqueDataLayerShortName(const UDataLayerManage
 	return UniqueNewDataLayerShortName;
 }
 
+bool FDataLayerUtils::AreDataLayerTypesCompatible(EDataLayerType ParentDataLayerType, EDataLayerType ChildDataLayerType, FText* OutReason)
+{
+	if ((ChildDataLayerType == EDataLayerType::Unknown) ||
+		(ParentDataLayerType == EDataLayerType::Unknown) ||
+		(ParentDataLayerType != EDataLayerType::Editor && ChildDataLayerType != EDataLayerType::Runtime))
+	{
+		if (OutReason)
+		{
+			*OutReason = FText::Format(LOCTEXT("IncompatibleChildType", "{0} Data Layer cannot have {1} child Data Layers"), UEnum::GetDisplayValueAsText(ParentDataLayerType), UEnum::GetDisplayValueAsText(ChildDataLayerType));
+		}
+		return false;
+	}
+	return true;
+}
+
+#undef LOCTEXT_NAMESPACE
+
 #endif
\ No newline at end of file
diff --git a/Engine/Source/Runtime/Engine/Private/WorldPartition/ErrorHandling/WorldPartitionStreamingGenerationLogErrorHandler.cpp b/Engine/Source/Runtime/Engine/Private/WorldPartition/ErrorHandling/WorldPartitionStreamingGenerationLogErrorHandler.cpp
--- a/Engine/Source/Runtime/Engine/Private/WorldPartition/ErrorHandling/WorldPartitionStreamingGenerationLogErrorHandler.cpp
+++ b/Engine/Source/Runtime/Engine/Private/WorldPartition/ErrorHandling/WorldPartitionStreamingGenerationLogErrorHandler.cpp
@@ -67,9 +67,20 @@ void FStreamingGenerationLogErrorHandler::OnInvalidDataLayerAssetType(const UDat
 	UE_ASSET_LOG(LogWorldPartition, Log, DataLayerInstance, TEXT("Data Layer is not compatible with Data Layer asset %s type %s"), *DataLayerAsset->GetName(), *DataLayerAsset->GetClass()->GetName());
 }
 
-void FStreamingGenerationLogErrorHandler::OnDataLayerHierarchyTypeMismatch(const UDataLayerInstance* DataLayerInstance, const UDataLayerInstance* Parent)
+void FStreamingGenerationLogErrorHandler::OnDataLayerHierarchyTypeMismatch(const UDataLayerInstance* DataLayerInstance, const UDataLayerInstance* Parent, EDataLayerHierarchyInvalidReason Reason)
 {
-	UE_ASSET_LOG(LogWorldPartition, Log, DataLayerInstance, TEXT("Data Layer %s is of Type %s and its parent %s is of type %s"), *DataLayerInstance->GetDataLayerFullName(), *UEnum::GetValueAsString(DataLayerInstance->GetType()), *Parent->GetDataLayerFullName(), *UEnum::GetValueAsString(Parent->GetType()));
+	switch (Reason)
+	{
+	case EDataLayerHierarchyInvalidReason::ClientOnlyDataLayerCantBeChild:
+		UE_ASSET_LOG(LogWorldPartition, Log, DataLayerInstance, TEXT("Client-only Data Layer %s can't be child of parent %s"), *DataLayerInstance->GetDataLayerFullName(), *Parent->GetDataLayerFullName());
+		break;
+	case EDataLayerHierarchyInvalidReason::ServerOnlyDataLayerCantBeChild:
+		UE_ASSET_LOG(LogWorldPartition, Log, DataLayerInstance, TEXT("Server-only Data Layer %s can't be child of parent %s"), *DataLayerInstance->GetDataLayerFullName(), *Parent->GetDataLayerFullName());
+		break;
+	case EDataLayerHierarchyInvalidReason::IncompatibleDataLayerType:
+		UE_ASSET_LOG(LogWorldPartition, Log, DataLayerInstance, TEXT("Data Layer %s is of Type %s and its parent %s is of type %s"), *DataLayerInstance->GetDataLayerFullName(), *UEnum::GetValueAsString(DataLayerInstance->GetType()), *Parent->GetDataLayerFullName(), *UEnum::GetValueAsString(Parent->GetType()));
+		break;
+	}
 }
 
 void FStreamingGenerationLogErrorHandler::OnDataLayerAssetConflict(const UDataLayerInstanceWithAsset* DataLayerInstance, const UDataLayerInstanceWithAsset* ConflictingDataLayerInstance)
diff --git a/Engine/Source/Runtime/Engine/Private/WorldPartition/ErrorHandling/WorldPartitionStreamingGenerationProxyErrorHandler.cpp b/Engine/Source/Runtime/Engine/Private/WorldPartition/ErrorHandling/WorldPartitionStreamingGenerationProxyErrorHandler.cpp
--- a/Engine/Source/Runtime/Engine/Private/WorldPartition/ErrorHandling/WorldPartitionStreamingGenerationProxyErrorHandler.cpp
+++ b/Engine/Source/Runtime/Engine/Private/WorldPartition/ErrorHandling/WorldPartitionStreamingGenerationProxyErrorHandler.cpp
@@ -48,9 +48,9 @@ void FStreamingGenerationProxyErrorHandler::OnInvalidDataLayerAssetType(const UD
 	InnerErrorHandler->OnInvalidDataLayerAssetType(DataLayerInstance, DataLayerAsset);
 }
 
-void FStreamingGenerationProxyErrorHandler::OnDataLayerHierarchyTypeMismatch(const UDataLayerInstance* DataLayerInstance, const UDataLayerInstance* Parent)
+void FStreamingGenerationProxyErrorHandler::OnDataLayerHierarchyTypeMismatch(const UDataLayerInstance* DataLayerInstance, const UDataLayerInstance* Parent, EDataLayerHierarchyInvalidReason Reason)
 {
-	InnerErrorHandler->OnDataLayerHierarchyTypeMismatch(DataLayerInstance, Parent);
+	InnerErrorHandler->OnDataLayerHierarchyTypeMismatch(DataLayerInstance, Parent, Reason);
 }
 
 void FStreamingGenerationProxyErrorHandler::OnDataLayerAssetConflict(const UDataLayerInstanceWithAsset* DataLayerInstance, const UDataLayerInstanceWithAsset* ConflictingDataLayerInstance)
diff --git a/Engine/Source/Runtime/Engine/Private/WorldPartition/ErrorHandling/WorldPartitionStreamingGenerationTokenizedMessageErrorHandler.cpp b/Engine/Source/Runtime/Engine/Private/WorldPartition/ErrorHandling/WorldPartitionStreamingGenerationTokenizedMessageErrorHandler.cpp
--- a/Engine/Source/Runtime/Engine/Private/WorldPartition/ErrorHandling/WorldPartitionStreamingGenerationTokenizedMessageErrorHandler.cpp
+++ b/Engine/Source/Runtime/Engine/Private/WorldPartition/ErrorHandling/WorldPartitionStreamingGenerationTokenizedMessageErrorHandler.cpp
@@ -142,16 +142,30 @@ void ITokenizedMessageErrorHandler::OnInvalidDataLayerAssetType(const UDataLayer
 	HandleTokenizedMessage(MoveTemp(Message));
 }
 
-void ITokenizedMessageErrorHandler::OnDataLayerHierarchyTypeMismatch(const UDataLayerInstance* DataLayerInstance, const UDataLayerInstance* Parent)
+void ITokenizedMessageErrorHandler::OnDataLayerHierarchyTypeMismatch(const UDataLayerInstance* DataLayerInstance, const UDataLayerInstance* Parent, EDataLayerHierarchyInvalidReason Reason)
 {
 	TSharedRef<FTokenizedMessage> Message = FTokenizedMessage::Create(EMessageSeverity::Error);
-	Message->AddToken(FTextToken::Create(LOCTEXT("TokenMessage_DataLayers_DataLayer", "Data layer")))
-		->AddToken(FTextToken::Create(FText::FromString(DataLayerInstance->GetDataLayerShortName())))
-		->AddToken(FTextToken::Create(LOCTEXT("TokenMessage_DataLayers_IsRntime", "is Runtime but its parent data layer")))
-		->AddToken(FTextToken::Create(FText::FromString(Parent->GetDataLayerShortName())))
-		->AddToken(FTextToken::Create(LOCTEXT("TokenMessage_DataLayers_IsNot", "is not")))
-		->AddToken(FMapErrorToken::Create(TEXT("DataLayers_HierarchyTypeMisMatch_CheckForErrors")));
 
+	switch (Reason)
+	{
+	case EDataLayerHierarchyInvalidReason::ClientOnlyDataLayerCantBeChild:
+	case EDataLayerHierarchyInvalidReason::ServerOnlyDataLayerCantBeChild:
+		Message->AddToken(FTextToken::Create((Reason == EDataLayerHierarchyInvalidReason::ClientOnlyDataLayerCantBeChild) ? LOCTEXT("TokenMessage_DataLayers_ClientOnlyDataLayer", "Client-Only Data Layer") : LOCTEXT("TokenMessage_DataLayers_ServerOnlyDataLayer", "Server-Only Data Layer")))
+			->AddToken(FTextToken::Create(FText::FromString(DataLayerInstance->GetDataLayerShortName())))
+			->AddToken(FTextToken::Create(LOCTEXT("TokenMessage_DataLayers_CantBeChild", "can't be child of")))
+			->AddToken(FTextToken::Create(FText::FromString(Parent->GetDataLayerShortName())))
+			->AddToken(FMapErrorToken::Create(TEXT("DataLayers_HierarchyTypeMismatch_CheckForErrors")));
+		break;
+	case EDataLayerHierarchyInvalidReason::IncompatibleDataLayerType:
+		Message->AddToken(FTextToken::Create(LOCTEXT("TokenMessage_DataLayers_DataLayer", "Data layer")))
+			->AddToken(FTextToken::Create(FText::FromString(DataLayerInstance->GetDataLayerShortName())))
+			->AddToken(FTextToken::Create(LOCTEXT("TokenMessage_DataLayers_IsEditor", "is Editor but its parent data layer")))
+			->AddToken(FTextToken::Create(FText::FromString(Parent->GetDataLayerShortName())))
+			->AddToken(FTextToken::Create(LOCTEXT("TokenMessage_DataLayers_IsRuntime", "is Runtime")))
+			->AddToken(FMapErrorToken::Create(TEXT("DataLayers_HierarchyTypeMismatch_CheckForErrors")));
+		break;
+	}
+	
 	HandleTokenizedMessage(MoveTemp(Message));
 }
 
diff --git a/Engine/Source/Runtime/Engine/Public/WorldPartition/DataLayer/DataLayerInstance.h b/Engine/Source/Runtime/Engine/Public/WorldPartition/DataLayer/DataLayerInstance.h
--- a/Engine/Source/Runtime/Engine/Public/WorldPartition/DataLayer/DataLayerInstance.h
+++ b/Engine/Source/Runtime/Engine/Public/WorldPartition/DataLayer/DataLayerInstance.h
@@ -7,6 +7,7 @@
 #include "UObject/UObjectGlobals.h"
 #include "UObject/Object.h"
 #include "WorldPartition/DataLayer/DataLayerType.h"
+#include "WorldPartition/ErrorHandling/WorldPartitionStreamingGenerationErrorHandler.h"
 
 #include "DataLayerInstance.generated.h"
 
@@ -209,7 +210,7 @@ private:
 protected:
 #if WITH_EDITOR
 	ENGINE_API bool IsLocked(FText* OutReason) const;
-	ENGINE_API bool IsParentDataLayerTypeCompatible(const UDataLayerInstance* InParent, FText* OutReason = nullptr) const;
+	ENGINE_API bool IsParentDataLayerTypeCompatible(const UDataLayerInstance* InParent, FText* OutReasonText = nullptr, IStreamingGenerationErrorHandler::EDataLayerHierarchyInvalidReason* OutReason = nullptr) const;
 
 	virtual bool PerformAddActor(AActor* InActor) const { return false; }
 	virtual bool PerformRemoveActor(AActor* InActor) const { return false;  }
diff --git a/Engine/Source/Runtime/Engine/Public/WorldPartition/DataLayer/DataLayerUtils.h b/Engine/Source/Runtime/Engine/Public/WorldPartition/DataLayer/DataLayerUtils.h
--- a/Engine/Source/Runtime/Engine/Public/WorldPartition/DataLayer/DataLayerUtils.h
+++ b/Engine/Source/Runtime/Engine/Public/WorldPartition/DataLayer/DataLayerUtils.h
@@ -57,6 +57,8 @@ public:
 	static ENGINE_API bool SetDataLayerShortName(UDataLayerInstance* InDataLayerInstance, const FString& InNewShortName);
 
 	static ENGINE_API bool FindDataLayerByShortName(const UDataLayerManager* InDataLayerManager, const FString& InShortName, TSet<UDataLayerInstance*>& OutDataLayerInstances);
+
+	static ENGINE_API bool AreDataLayerTypesCompatible(EDataLayerType ParentDataLayerType, EDataLayerType ChildDataLayerType, FText* OutReason = nullptr);
 #endif
 
 	static FString GetSanitizedDataLayerShortName(FString InShortName)
diff --git a/Engine/Source/Runtime/Engine/Public/WorldPartition/ErrorHandling/WorldPartitionStreamingGenerationErrorHandler.h b/Engine/Source/Runtime/Engine/Public/WorldPartition/ErrorHandling/WorldPartitionStreamingGenerationErrorHandler.h
--- a/Engine/Source/Runtime/Engine/Public/WorldPartition/ErrorHandling/WorldPartitionStreamingGenerationErrorHandler.h
+++ b/Engine/Source/Runtime/Engine/Public/WorldPartition/ErrorHandling/WorldPartitionStreamingGenerationErrorHandler.h
@@ -66,10 +66,20 @@ public:
 	 */
 	virtual void OnInvalidReferenceDataLayerAsset(const UDataLayerInstanceWithAsset* DataLayerInstance) = 0;
 
+	/**
+	 * Used to identify a data layer hierarchy type mismatch error
+	 */
+	enum class EDataLayerHierarchyInvalidReason
+	{
+		ClientOnlyDataLayerCantBeChild,
+		ServerOnlyDataLayerCantBeChild,
+		IncompatibleDataLayerType
+	};
+
 	/**
 	 * Called when a data layer is not of the same type as its parent
 	 */
-	virtual void OnDataLayerHierarchyTypeMismatch(const UDataLayerInstance* DataLayerInstance, const UDataLayerInstance* Parent) = 0;
+	virtual void OnDataLayerHierarchyTypeMismatch(const UDataLayerInstance* DataLayerInstance, const UDataLayerInstance* Parent, EDataLayerHierarchyInvalidReason Reason) = 0;
 
 	/**
 	 * Called when two data layer instances share the same asset
diff --git a/Engine/Source/Runtime/Engine/Public/WorldPartition/ErrorHandling/WorldPartitionStreamingGenerationLogErrorHandler.h b/Engine/Source/Runtime/Engine/Public/WorldPartition/ErrorHandling/WorldPartitionStreamingGenerationLogErrorHandler.h
--- a/Engine/Source/Runtime/Engine/Public/WorldPartition/ErrorHandling/WorldPartitionStreamingGenerationLogErrorHandler.h
+++ b/Engine/Source/Runtime/Engine/Public/WorldPartition/ErrorHandling/WorldPartitionStreamingGenerationLogErrorHandler.h
@@ -16,7 +16,7 @@ public:
 	ENGINE_API virtual void OnInvalidReferenceRuntimeGrid(const IWorldPartitionActorDescInstanceView& ActorDescView, const IWorldPartitionActorDescInstanceView& ReferenceActorDescView) override;
 	ENGINE_API virtual void OnInvalidReferenceDataLayerAsset(const UDataLayerInstanceWithAsset* DataLayerInstance) override;
 	ENGINE_API virtual void OnInvalidDataLayerAssetType(const UDataLayerInstanceWithAsset* DataLayerInstance, const UDataLayerAsset* DataLayerAsset) override;
-	ENGINE_API virtual void OnDataLayerHierarchyTypeMismatch(const UDataLayerInstance* DataLayerInstance, const UDataLayerInstance* Parent) override;
+	ENGINE_API virtual void OnDataLayerHierarchyTypeMismatch(const UDataLayerInstance* DataLayerInstance, const UDataLayerInstance* Parent, EDataLayerHierarchyInvalidReason Reason) override;
 	ENGINE_API virtual void OnDataLayerAssetConflict(const UDataLayerInstanceWithAsset* DataLayerInstance, const UDataLayerInstanceWithAsset* ConflictingDataLayerInstance) override;
 	ENGINE_API virtual void OnActorNeedsResave(const IWorldPartitionActorDescInstanceView& ActorDescView) override;
 	ENGINE_API virtual void OnLevelInstanceInvalidWorldAsset(const IWorldPartitionActorDescInstanceView& ActorDescView, FName WorldAsset, ELevelInstanceInvalidReason Reason) override;
diff --git a/Engine/Source/Runtime/Engine/Public/WorldPartition/ErrorHandling/WorldPartitionStreamingGenerationNullErrorHandler.h b/Engine/Source/Runtime/Engine/Public/WorldPartition/ErrorHandling/WorldPartitionStreamingGenerationNullErrorHandler.h
--- a/Engine/Source/Runtime/Engine/Public/WorldPartition/ErrorHandling/WorldPartitionStreamingGenerationNullErrorHandler.h
+++ b/Engine/Source/Runtime/Engine/Public/WorldPartition/ErrorHandling/WorldPartitionStreamingGenerationNullErrorHandler.h
@@ -16,7 +16,7 @@ public:
 	virtual void OnInvalidReferenceRuntimeGrid(const IWorldPartitionActorDescInstanceView& ActorDescView, const IWorldPartitionActorDescInstanceView& ReferenceActorDescView) override {}
 	virtual void OnInvalidReferenceDataLayerAsset(const UDataLayerInstanceWithAsset* DataLayerInstance) override {}
 	virtual void OnInvalidDataLayerAssetType(const UDataLayerInstanceWithAsset* DataLayerInstance, const UDataLayerAsset* DataLayerAsset) override {}
-	virtual void OnDataLayerHierarchyTypeMismatch(const UDataLayerInstance* DataLayerInstance, const UDataLayerInstance* Parent) override {};
+	virtual void OnDataLayerHierarchyTypeMismatch(const UDataLayerInstance* DataLayerInstance, const UDataLayerInstance* Parent, EDataLayerHierarchyInvalidReason Reason) override {};
 	virtual void OnDataLayerAssetConflict(const UDataLayerInstanceWithAsset* DataLayerInstance, const UDataLayerInstanceWithAsset* ConflictingDataLayerInstance) override {};
 	virtual void OnActorNeedsResave(const IWorldPartitionActorDescInstanceView& ActorDescView) override {};
 	virtual void OnLevelInstanceInvalidWorldAsset(const IWorldPartitionActorDescInstanceView& ActorDescView, FName WorldAsset, ELevelInstanceInvalidReason Reason) override {};
diff --git a/Engine/Source/Runtime/Engine/Public/WorldPartition/ErrorHandling/WorldPartitionStreamingGenerationProxyErrorHandler.h b/Engine/Source/Runtime/Engine/Public/WorldPartition/ErrorHandling/WorldPartitionStreamingGenerationProxyErrorHandler.h
--- a/Engine/Source/Runtime/Engine/Public/WorldPartition/ErrorHandling/WorldPartitionStreamingGenerationProxyErrorHandler.h
+++ b/Engine/Source/Runtime/Engine/Public/WorldPartition/ErrorHandling/WorldPartitionStreamingGenerationProxyErrorHandler.h
@@ -20,7 +20,7 @@ public:
 	ENGINE_API virtual void OnInvalidReferenceRuntimeGrid(const IWorldPartitionActorDescInstanceView& ActorDescView, const IWorldPartitionActorDescInstanceView& ReferenceActorDescView) override;
 	ENGINE_API virtual void OnInvalidReferenceDataLayerAsset(const UDataLayerInstanceWithAsset* DataLayerInstance) override;
 	ENGINE_API virtual void OnInvalidDataLayerAssetType(const UDataLayerInstanceWithAsset* DataLayerInstance, const UDataLayerAsset* DataLayerAsset) override;
-	ENGINE_API virtual void OnDataLayerHierarchyTypeMismatch(const UDataLayerInstance* DataLayerInstance, const UDataLayerInstance* Parent) override;
+	ENGINE_API virtual void OnDataLayerHierarchyTypeMismatch(const UDataLayerInstance* DataLayerInstance, const UDataLayerInstance* Parent, EDataLayerHierarchyInvalidReason Reason) override;
 	ENGINE_API virtual void OnDataLayerAssetConflict(const UDataLayerInstanceWithAsset* DataLayerInstance, const UDataLayerInstanceWithAsset* ConflictingDataLayerInstance) override;
 	ENGINE_API virtual void OnActorNeedsResave(const IWorldPartitionActorDescInstanceView& ActorDescView) override;
 	ENGINE_API virtual void OnLevelInstanceInvalidWorldAsset(const IWorldPartitionActorDescInstanceView& ActorDescView, FName WorldAsset, ELevelInstanceInvalidReason Reason) override;
diff --git a/Engine/Source/Runtime/Engine/Public/WorldPartition/ErrorHandling/WorldPartitionStreamingGenerationTokenizedMessageErrorHandler.h b/Engine/Source/Runtime/Engine/Public/WorldPartition/ErrorHandling/WorldPartitionStreamingGenerationTokenizedMessageErrorHandler.h
--- a/Engine/Source/Runtime/Engine/Public/WorldPartition/ErrorHandling/WorldPartitionStreamingGenerationTokenizedMessageErrorHandler.h
+++ b/Engine/Source/Runtime/Engine/Public/WorldPartition/ErrorHandling/WorldPartitionStreamingGenerationTokenizedMessageErrorHandler.h
@@ -17,7 +17,7 @@ class ITokenizedMessageErrorHandler : public IStreamingGenerationErrorHandler
 	ENGINE_API virtual void OnInvalidReferenceRuntimeGrid(const IWorldPartitionActorDescInstanceView& ActorDescView, const IWorldPartitionActorDescInstanceView& ReferenceActorDescView) override;
 	ENGINE_API virtual void OnInvalidReferenceDataLayerAsset(const UDataLayerInstanceWithAsset* DataLayerInstance) override;
 	ENGINE_API virtual void OnInvalidDataLayerAssetType(const UDataLayerInstanceWithAsset* DataLayerInstance, const UDataLayerAsset* DataLayerAsset) override;
-	ENGINE_API virtual void OnDataLayerHierarchyTypeMismatch(const UDataLayerInstance* DataLayerInstance, const UDataLayerInstance* Parent) override;
+	ENGINE_API virtual void OnDataLayerHierarchyTypeMismatch(const UDataLayerInstance* DataLayerInstance, const UDataLayerInstance* Parent, EDataLayerHierarchyInvalidReason Reason) override;
 	ENGINE_API virtual void OnDataLayerAssetConflict(const UDataLayerInstanceWithAsset* DataLayerInstance, const UDataLayerInstanceWithAsset* ConflictingDataLayerInstance) override;
 	ENGINE_API virtual void OnActorNeedsResave(const IWorldPartitionActorDescInstanceView& ActorDescView) override;
 	ENGINE_API virtual void OnLevelInstanceInvalidWorldAsset(const IWorldPartitionActorDescInstanceView& ActorDescView, FName WorldAsset, ELevelInstanceInvalidReason Reason) override;
