Support Loaded runtime state for External Data Layers

diff --git a/Engine/Source/Runtime/Engine/Private/WorldPartition/WorldPartitionRuntimeCell.cpp b/Engine/Source/Runtime/Engine/Private/WorldPartition/WorldPartitionRuntimeCell.cpp
--- a/Engine/Source/Runtime/Engine/Private/WorldPartition/WorldPartitionRuntimeCell.cpp
+++ b/Engine/Source/Runtime/Engine/Private/WorldPartition/WorldPartitionRuntimeCell.cpp
@@ -164,12 +164,16 @@ EDataLayerRuntimeState UWorldPartitionRuntimeCell::GetCellEffectiveWantedState()
 			UWorldPartition* WorldPartition = OuterWorld->GetWorldPartition();
 			if (const UDataLayerManager* DataLayerManager = WorldPartition->GetDataLayerManager())
 			{
-				if (!DataLayers.HasExternalDataLayer() || DataLayerManager->IsAllDataLayerInEffectiveRuntimeState({ GetExternalDataLayer() }, EDataLayerRuntimeState::Activated))
+				// Determine the maximum runtime state the cell can have based on its External Data Layer. If none, maximum is Activated.
+				const UExternalDataLayerInstance* ExternalDatalayerInstance = GetExternalDataLayerInstance();
+				EDataLayerRuntimeState MaxEffectiveRuntimeState = ExternalDatalayerInstance ? ExternalDatalayerInstance->GetEffectiveRuntimeState() : EDataLayerRuntimeState::Activated;
+
+				if (MaxEffectiveRuntimeState > EDataLayerRuntimeState::Unloaded)
 				{
 					TArrayView<const FName> NonExternalDataLayers = DataLayers.GetNonExternalDataLayers();
 					if (NonExternalDataLayers.IsEmpty())
 					{
-						EffectiveWantedState = EDataLayerRuntimeState::Activated;
+						EffectiveWantedState = MaxEffectiveRuntimeState;
 					}
 					else
 					{
@@ -178,7 +182,7 @@ EDataLayerRuntimeState UWorldPartitionRuntimeCell::GetCellEffectiveWantedState()
 						case EWorldPartitionDataLayersLogicOperator::Or:
 							if (DataLayerManager->IsAnyDataLayerInEffectiveRuntimeState(NonExternalDataLayers, EDataLayerRuntimeState::Activated))
 							{
-								EffectiveWantedState = EDataLayerRuntimeState::Activated;
+								EffectiveWantedState = MaxEffectiveRuntimeState;
 							}
 							else if (DataLayerManager->IsAnyDataLayerInEffectiveRuntimeState(NonExternalDataLayers, EDataLayerRuntimeState::Loaded))
 							{
@@ -188,7 +192,7 @@ EDataLayerRuntimeState UWorldPartitionRuntimeCell::GetCellEffectiveWantedState()
 						case EWorldPartitionDataLayersLogicOperator::And:
 							if (DataLayerManager->IsAllDataLayerInEffectiveRuntimeState(NonExternalDataLayers, EDataLayerRuntimeState::Activated))
 							{
-								EffectiveWantedState = EDataLayerRuntimeState::Activated;
+								EffectiveWantedState = MaxEffectiveRuntimeState;
 							}
 							else if (DataLayerManager->IsAllDataLayerInEffectiveRuntimeState(NonExternalDataLayers, EDataLayerRuntimeState::Loaded))
 							{
