Add a retarget source update when the skeleton changes

diff --git a/Engine/Source/Runtime/Engine/Private/Animation/AnimSequence.cpp b/Engine/Source/Runtime/Engine/Private/Animation/AnimSequence.cpp
--- a/Engine/Source/Runtime/Engine/Private/Animation/AnimSequence.cpp
+++ b/Engine/Source/Runtime/Engine/Private/Animation/AnimSequence.cpp
@@ -4383,6 +4383,11 @@ void UAnimSequence::OnModelModified(const EAnimDataModelNotifyType& NotifyType,
 					CompressedData.ClearCompressedCurveData();
 					HandleTrackDataChanged(bWasModelReset);
 				}
+
+				if (NotifyCollector.Contains(EAnimDataModelNotifyType::SkeletonChanged))
+				{
+					UpdateRetargetSourceAsset();
+				}
 			}
 
 			break;
@@ -4462,6 +4467,7 @@ void UAnimSequence::OnModelModified(const EAnimDataModelNotifyType& NotifyType,
 			if (NotifyCollector.IsNotWithinBracket())
 			{
 				HandleTrackDataChanged();
+				UpdateRetargetSourceAsset();
 			}
 			break;
 		}
