Fix missing inclusion of HasPixelAnimation and IsTessellationEnabled in the GetBasePropertyOverridesHash of UMaterialInstance.
This fixes issues where toggling these options in a MaterialInstance has no effect.

diff --git a/Engine/Source/Runtime/Engine/Private/Materials/MaterialInstance.cpp b/Engine/Source/Runtime/Engine/Private/Materials/MaterialInstance.cpp
--- a/Engine/Source/Runtime/Engine/Private/Materials/MaterialInstance.cpp
+++ b/Engine/Source/Runtime/Engine/Private/Materials/MaterialInstance.cpp
@@ -4399,96 +4399,40 @@ void UMaterialInstance::GetBasePropertyOverridesHash(FSHAHash& OutHash)const
 
 	FSHA1 Hash;
 	bool bHasOverrides = false;
-
-	float UsedOpacityMaskClipValue = GetOpacityMaskClipValue();
-	if (FMath::Abs(UsedOpacityMaskClipValue - Mat->GetOpacityMaskClipValue()) > UE_SMALL_NUMBER)
-	{
-		const FString HashString = TEXT("bOverride_OpacityMaskClipValue");
-		Hash.UpdateWithString(*HashString, HashString.Len());
-		Hash.Update((const uint8*)&UsedOpacityMaskClipValue, sizeof(UsedOpacityMaskClipValue));
-		bHasOverrides = true;
-	}
-
-	bool bUsedCastDynamicShadowAsMasked = GetCastDynamicShadowAsMasked();
-	if ( bUsedCastDynamicShadowAsMasked != Mat->GetCastDynamicShadowAsMasked() )
-	{
-		const FString HashString = TEXT("bOverride_CastDynamicShadowAsMasked");
-		Hash.UpdateWithString(*HashString, HashString.Len());
-		Hash.Update((const uint8*)&bUsedCastDynamicShadowAsMasked, sizeof(bUsedCastDynamicShadowAsMasked));
-		bHasOverrides = true;
-	}
-
-	EBlendMode UsedBlendMode = GetBlendMode();
-	if (UsedBlendMode != Mat->GetBlendMode())
-	{
-		const FString HashString = TEXT("bOverride_BlendMode");
-		Hash.UpdateWithString(*HashString, HashString.Len());
-		Hash.Update((const uint8*)&UsedBlendMode, sizeof(UsedBlendMode));
-		bHasOverrides = true;
-	}
 	
-	FMaterialShadingModelField UsedShadingModels = GetShadingModels();
-	if (UsedShadingModels != Mat->GetShadingModels())
-	{
-		const FString HashString = TEXT("bOverride_ShadingModel");
-		Hash.UpdateWithString(*HashString, HashString.Len());
-		Hash.Update((const uint8*)&UsedShadingModels, sizeof(UsedShadingModels));
-		bHasOverrides = true;
-	}
-
-	bool bUsedIsTwoSided = IsTwoSided();
-	if (bUsedIsTwoSided != Mat->IsTwoSided())
-	{
-		const FString HashString = TEXT("bOverride_TwoSided");
-		Hash.UpdateWithString(*HashString, HashString.Len());
-		Hash.Update((uint8*)&bUsedIsTwoSided, sizeof(bUsedIsTwoSided));
-		bHasOverrides = true;
-	}
-	bool bUsedIsThinSurface = IsThinSurface();
-	if (bUsedIsThinSurface != Mat->IsThinSurface())
-	{
-		const FString HashString = TEXT("bOverride_bIsThinSurface");
-		Hash.UpdateWithString(*HashString, HashString.Len());
-		Hash.Update((uint8*)&bUsedIsThinSurface, sizeof(bUsedIsThinSurface));
-		bHasOverrides = true;
-	}
-	bool bUsedIsDitheredLODTransition = IsDitheredLODTransition();
-	if (bUsedIsDitheredLODTransition != Mat->IsDitheredLODTransition())
-	{
-		const FString HashString = TEXT("bOverride_DitheredLODTransition");
-		Hash.UpdateWithString(*HashString, HashString.Len());
-		Hash.Update((uint8*)&bUsedIsDitheredLODTransition, sizeof(bUsedIsDitheredLODTransition));
-		bHasOverrides = true;
-	}
-
-	bool bUsedIsTranslucencyWritingVelocity = IsTranslucencyWritingVelocity();
-	if (bUsedIsTranslucencyWritingVelocity != Mat->IsTranslucencyWritingVelocity())
-	{
-		const FString HashString = TEXT("bOverride_OutputTranslucentVelocity");
-		Hash.UpdateWithString(*HashString, HashString.Len());
-		Hash.Update((uint8*)&bUsedIsTranslucencyWritingVelocity, sizeof(bUsedIsTranslucencyWritingVelocity));
-		bHasOverrides = true;
-	}
-
-	FDisplacementScaling UsedDisplacementScaling = GetDisplacementScaling();
-	if (UsedDisplacementScaling != Mat->GetDisplacementScaling())
+	auto GetPropertyOverrideHash = [&](auto InstanceValue, auto MatValue, const FString &HashString)
 	{
-		const FString HashString = TEXT("bOverride_DisplacementScaling");
-		Hash.UpdateWithString(*HashString, HashString.Len());
-		Hash.Update((uint8*)&UsedDisplacementScaling.Magnitude, sizeof(UsedDisplacementScaling.Magnitude));
-		Hash.Update((uint8*)&UsedDisplacementScaling.Center, sizeof(UsedDisplacementScaling.Center));
-		bHasOverrides = true;
-	}
-
-	float UsedMaxWorldPositionOffsetDisplacement = GetMaxWorldPositionOffsetDisplacement();
-	if (FMath::Abs(UsedMaxWorldPositionOffsetDisplacement - Mat->GetMaxWorldPositionOffsetDisplacement()) > UE_SMALL_NUMBER)
-	{
-		const FString HashString = TEXT("bOverride_MaxWorldPositionOffsetDisplacement");
-		Hash.UpdateWithString(*HashString, HashString.Len());
-		Hash.Update((uint8*)&UsedMaxWorldPositionOffsetDisplacement, sizeof(UsedMaxWorldPositionOffsetDisplacement));
-		bHasOverrides = true;
-	}
-
+		bool bOverridden = false;
+		if constexpr(std::is_floating_point_v<decltype(InstanceValue)>)
+		{
+			bOverridden = !FMath::IsNearlyEqual(InstanceValue, MatValue);
+		}
+		else
+		{
+			bOverridden = InstanceValue != MatValue;
+		}
+		
+		if (bOverridden)
+		{
+			Hash.UpdateWithString(*HashString, HashString.Len());
+			Hash.Update(reinterpret_cast<const uint8*>(&InstanceValue), sizeof(InstanceValue));
+			bHasOverrides = true;
+		}
+	};
+
+	GetPropertyOverrideHash(GetOpacityMaskClipValue(), Mat->GetOpacityMaskClipValue(), TEXT("bOverride_OpacityMaskClipValue"));
+	GetPropertyOverrideHash(GetBlendMode(), Mat->GetBlendMode(), TEXT("bOverride_BlendMode"));
+	GetPropertyOverrideHash(GetShadingModels(), Mat->GetShadingModels(), TEXT("bOverride_ShadingModel"));
+	GetPropertyOverrideHash(IsTwoSided(), Mat->IsTwoSided(), TEXT("bOverride_TwoSided"));
+	GetPropertyOverrideHash(IsThinSurface(), Mat->IsThinSurface(), TEXT("bOverride_bIsThinSurface"));
+	GetPropertyOverrideHash(IsDitheredLODTransition(), Mat->IsDitheredLODTransition(), TEXT("bOverride_DitheredLODTransition"));
+	GetPropertyOverrideHash(GetCastDynamicShadowAsMasked(), Mat->GetCastDynamicShadowAsMasked(), TEXT("bOverride_CastDynamicShadowAsMasked"));
+	GetPropertyOverrideHash(IsTranslucencyWritingVelocity(), Mat->IsTranslucencyWritingVelocity(), TEXT("bOverride_OutputTranslucentVelocity"));
+	GetPropertyOverrideHash(HasPixelAnimation(), Mat->HasPixelAnimation(), TEXT("bOverride_bHasPixelAnimation"));
+	GetPropertyOverrideHash(IsTessellationEnabled(), Mat->IsTessellationEnabled(), TEXT("bOverride_bEnableTessellation"));
+	GetPropertyOverrideHash(GetDisplacementScaling(), Mat->GetDisplacementScaling(), TEXT("bOverride_DisplacementScaling"));
+	GetPropertyOverrideHash(GetMaxWorldPositionOffsetDisplacement(), Mat->GetMaxWorldPositionOffsetDisplacement(), TEXT("bOverride_MaxWorldPositionOffsetDisplacement"));
+	
 	if (bHasOverrides)
 	{
 		Hash.Final();
@@ -4500,7 +4444,7 @@ bool UMaterialInstance::HasOverridenBaseProperties()const
 {
 	const UMaterial* Material = GetMaterial_Concurrent();
 	if (Parent && Material && Material->bUsedAsSpecialEngineMaterial == false &&
-		((FMath::Abs(GetOpacityMaskClipValue() - Parent->GetOpacityMaskClipValue()) > UE_SMALL_NUMBER) ||
+		(!FMath::IsNearlyEqual(GetOpacityMaskClipValue(), Parent->GetOpacityMaskClipValue()) ||
 		(GetBlendMode() != Parent->GetBlendMode()) ||
 		(GetShadingModels() != Parent->GetShadingModels()) ||
 		(IsTwoSided() != Parent->IsTwoSided()) ||
@@ -4511,7 +4455,7 @@ bool UMaterialInstance::HasOverridenBaseProperties()const
 		(HasPixelAnimation() != Parent->HasPixelAnimation()) ||
 		(IsTessellationEnabled() != Parent->IsTessellationEnabled()) ||
 		(GetDisplacementScaling() != Parent->GetDisplacementScaling()) ||
-		(GetMaxWorldPositionOffsetDisplacement() != Parent->GetMaxWorldPositionOffsetDisplacement())
+		!FMath::IsNearlyEqual(GetMaxWorldPositionOffsetDisplacement(), Parent->GetMaxWorldPositionOffsetDisplacement())
 		))
 	{
 		return true;
