Add CSV Stats for PSO Precaching

diff --git a/Engine/Source/Runtime/Engine/Private/PSOPrecacheMaterial.cpp b/Engine/Source/Runtime/Engine/Private/PSOPrecacheMaterial.cpp
--- a/Engine/Source/Runtime/Engine/Private/PSOPrecacheMaterial.cpp
+++ b/Engine/Source/Runtime/Engine/Private/PSOPrecacheMaterial.cpp
@@ -23,6 +23,8 @@ static FAutoConsoleVariableRef CVarPSOUseBackgroundThreadForCollection(
 	ECVF_ReadOnly
 );
 
+CSV_DECLARE_CATEGORY_EXTERN(PSOPrecache);
+
 FPSOCollectorCreateManager::FPSOCollectorData FPSOCollectorCreateManager::PSOCollectors[(int32)EShadingPath::Num][FPSOCollectorCreateManager::MaxPSOCollectorCount] = {};
 
 int32 FPSOCollectorCreateManager::GetIndex(EShadingPath ShadingPath, const TCHAR* Name)
@@ -57,8 +59,11 @@ FPSOPrecacheRequestResultArray RequestPrecachePSOs(const FPSOPrecacheDataArray&
 		case FPSOPrecacheData::EType::Graphics:
 		{
 #if PSO_PRECACHING_VALIDATE
-			PSOCollectorStats::GetFullPSOPrecacheStatsCollector().AddStateToCache(PrecacheData.GraphicsPSOInitializer, PSOCollectorStats::GetPSOPrecacheHash, nullptr, PrecacheData.PSOCollectorIndex, PrecacheData.VertexFactoryType);
-#endif // PSO_PRECACHING_VALIDATE
+			if (PSOCollectorStats::GetFullPSOPrecacheStatsCollector().AddStateToCache(PrecacheData.GraphicsPSOInitializer, PSOCollectorStats::GetPSOPrecacheHash, nullptr, PrecacheData.PSOCollectorIndex, PrecacheData.VertexFactoryType))
+			{
+				CSV_CUSTOM_STAT(PSOPrecache, PrecachedGraphics, 1, ECsvCustomStatOp::Accumulate);
+			}
+#endif // PSO_PRECACHING_VALIDATE			
 
 			FPSOPrecacheRequestResult PSOPrecacheResult = PipelineStateCache::PrecacheGraphicsPipelineState(PrecacheData.GraphicsPSOInitializer);
 			if (PSOPrecacheResult.IsValid() && PrecacheData.bRequired)
@@ -70,9 +75,12 @@ FPSOPrecacheRequestResultArray RequestPrecachePSOs(const FPSOPrecacheDataArray&
 		case FPSOPrecacheData::EType::Compute:
 		{
 #if PSO_PRECACHING_VALIDATE
-			PSOCollectorStats::GetFullPSOPrecacheStatsCollector().AddStateToCache(*PrecacheData.ComputeShader, PSOCollectorStats::GetPSOPrecacheHash, nullptr, PrecacheData.PSOCollectorIndex, nullptr);
+			if (PSOCollectorStats::GetFullPSOPrecacheStatsCollector().AddStateToCache(*PrecacheData.ComputeShader, PSOCollectorStats::GetPSOPrecacheHash, nullptr, PrecacheData.PSOCollectorIndex, nullptr))
+			{
+				CSV_CUSTOM_STAT(PSOPrecache, PrecachedCompute, 1, ECsvCustomStatOp::Accumulate);
+			}
 #endif // PSO_PRECACHING_VALIDATE
-
+			
 			FPSOPrecacheRequestResult PSOPrecacheResult = PipelineStateCache::PrecacheComputePipelineState(PrecacheData.ComputeShader);
 			if (PSOPrecacheResult.IsValid() && PrecacheData.bRequired)
 			{
diff --git a/Engine/Source/Runtime/Engine/Private/PSOPrecacheValidation.cpp b/Engine/Source/Runtime/Engine/Private/PSOPrecacheValidation.cpp
--- a/Engine/Source/Runtime/Engine/Private/PSOPrecacheValidation.cpp
+++ b/Engine/Source/Runtime/Engine/Private/PSOPrecacheValidation.cpp
@@ -12,6 +12,9 @@
 
 #if PSO_PRECACHING_VALIDATE && UE_WITH_PSO_PRECACHING
 
+// CSV category for PSOPrecache
+CSV_DEFINE_CATEGORY(PSOPrecache, true);
+
 /**
 * Different IHVs and drivers can have different opinions on what subset of a PSO
 * matters for caching. We track multiple PSO subsets, ranging from shaders-only
@@ -202,7 +205,7 @@ bool PSOCollectorStats::FPrecacheStatsCollector::IsStateTracked(int32 PSOCollect
 #endif
 }
 
-EPSOPrecacheResult PSOCollectorStats::FPrecacheStatsCollector::UpdatePrecacheStats(uint64 PrecacheHash, int32 PSOCollectorIndex, const FVertexFactoryType* VertexFactoryType, bool bTracked, EPSOPrecacheResult PrecacheResult)
+bool PSOCollectorStats::FPrecacheStatsCollector::UpdatePrecacheStats(uint64 PrecacheHash, int32 PSOCollectorIndex, const FVertexFactoryType* VertexFactoryType, bool bTracked, EPSOPrecacheResult& InOutPrecacheResult)
 {
 	EPSOPrecacheResult OutResult = EPSOPrecacheResult::Unknown;
 
@@ -235,7 +238,7 @@ EPSOPrecacheResult PSOCollectorStats::FPrecacheStatsCollector::UpdatePrecacheSta
 			Stats.MissData.UpdateStats(PSOCollectorIndex, VertexFactoryType);
 			OutResult = EPSOPrecacheResult::Missed;
 		}
-		else if (PrecacheResult == EPSOPrecacheResult::Active)
+		else if (InOutPrecacheResult == EPSOPrecacheResult::Active)
 		{
 			Stats.TooLateData.UpdateStats(PSOCollectorIndex, VertexFactoryType);
 			OutResult = EPSOPrecacheResult::TooLate;
@@ -247,7 +250,8 @@ EPSOPrecacheResult PSOCollectorStats::FPrecacheStatsCollector::UpdatePrecacheSta
 		}
 	}
 
-	return OutResult;
+	InOutPrecacheResult = OutResult;
+	return bUpdateStats;
 }
 
 bool PSOCollectorStats::FPrecacheStatsCollector::IsPrecached(uint64 PrecacheStateHash)
@@ -340,17 +344,22 @@ void PSOCollectorStats::CheckFullPipelineStateInCache(
 		PrecacheStateHash = PSOCollectorStats::GetPSOPrecacheHash(Initializer);
 	}
 	
-	EPSOPrecacheResult Result = PSOCollectorStats::GetFullPSOPrecacheStatsCollector().CheckStateInCacheByHash(PrecacheStateHash, PSOPrecacheResult, PSOCollectorIndex, VFType);
-	if (IsFullPrecachingValidationEnabled() && Result == EPSOPrecacheResult::Missed)
+	bool bStatUpdated = PSOCollectorStats::GetFullPSOPrecacheStatsCollector().CheckStateInCacheByHash(PrecacheStateHash, PSOCollectorIndex, VFType, PSOPrecacheResult);
+	if (IsFullPrecachingValidationEnabled() && bStatUpdated && (PSOPrecacheResult == EPSOPrecacheResult::Missed || PSOPrecacheResult == EPSOPrecacheResult::TooLate))
 	{
 		// only report here if it's not missing with minimal PSO initializer
 		bool bMinimalPSOPrecached = bCheckMinimalPSOPrecached ? PSOCollectorStats::GetMinimalPSOPrecacheStatsCollector().IsPrecached(Initializer.StatePrecachePSOHash) : true;
 		if (bMinimalPSOPrecached)
 		{
 			const FMaterial* Material = MaterialRenderProxy ? MaterialRenderProxy->GetMaterialNoFallback(GMaxRHIFeatureLevel) : nullptr;
-			LogPSOMissInfo(Initializer, EPSOPrecacheMissType::FullPSO, Result, Material, VFType, PrimitiveSceneProxy, PSOCollectorIndex, 0);
+			LogPSOMissInfo(Initializer, EPSOPrecacheMissType::FullPSO, PSOPrecacheResult, Material, VFType, PrimitiveSceneProxy, PSOCollectorIndex, 0);
 		}
 	}
+
+	if (bStatUpdated)
+	{
+		UpdateCSVStats(PSOPrecacheResult);
+	}
 }
 
 void PSOCollectorStats::CheckComputePipelineStateInCache(
@@ -359,11 +368,48 @@ void PSOCollectorStats::CheckComputePipelineStateInCache(
 	const FMaterialRenderProxy* MaterialRenderProxy,
 	int32 PSOCollectorIndex)
 {
-	EPSOPrecacheResult Result = PSOCollectorStats::GetFullPSOPrecacheStatsCollector().CheckStateInCache(ComputeShader, PSOCollectorStats::GetPSOPrecacheHash, PSOPrecacheResult, PSOCollectorIndex, nullptr);
-	if (IsFullPrecachingValidationEnabled() && Result == EPSOPrecacheResult::Missed)
+	bool bStatsUpdated = PSOCollectorStats::GetFullPSOPrecacheStatsCollector().CheckStateInCache(ComputeShader, PSOCollectorStats::GetPSOPrecacheHash, PSOCollectorIndex, nullptr, PSOPrecacheResult);
+	if (IsFullPrecachingValidationEnabled() && bStatsUpdated && PSOPrecacheResult == EPSOPrecacheResult::Missed)
 	{
 		const FMaterial* Material = MaterialRenderProxy ? MaterialRenderProxy->GetMaterialNoFallback(GMaxRHIFeatureLevel) : nullptr;
-		LogPSOMissInfo(ComputeShader, Result, Material, PSOCollectorIndex);
+		LogPSOMissInfo(ComputeShader, PSOPrecacheResult, Material, PSOCollectorIndex);
+	}
+
+	if (bStatsUpdated)
+	{
+		UpdateCSVStats(PSOPrecacheResult);
+	}
+}
+
+void PSOCollectorStats::UpdateCSVStats(EPSOPrecacheResult PSOPrecacheResult)
+{
+	switch (PSOPrecacheResult)
+	{
+	case EPSOPrecacheResult::Complete:
+	{
+		CSV_CUSTOM_STAT(PSOPrecache, Hit, 1, ECsvCustomStatOp::Accumulate);
+		break;
+	}
+	case EPSOPrecacheResult::Missed:
+	{
+		CSV_CUSTOM_STAT(PSOPrecache, Miss, 1, ECsvCustomStatOp::Accumulate);
+		break;
+	}
+	case EPSOPrecacheResult::Active:
+	case EPSOPrecacheResult::TooLate:
+	{
+		CSV_CUSTOM_STAT(PSOPrecache, TooLate, 1, ECsvCustomStatOp::Accumulate);
+		break;
+	}
+	case EPSOPrecacheResult::Untracked:
+	{
+		CSV_CUSTOM_STAT(PSOPrecache, Untracked, 1, ECsvCustomStatOp::Accumulate);
+		break;
+	}
+	default:
+	{
+		check(false);
+	}
 	}
 }
 
@@ -908,6 +954,12 @@ void LogPSOMissInfo(
 		}
 		}
 	}
+	else if (Result == EPSOPrecacheResult::TooLate)
+	{
+		check(MissType == EPSOPrecacheMissType::FullPSO);
+		
+		// Nothing extra log yet
+	}
 	else
 	{
 		// Should have get any other results on shaders only
diff --git a/Engine/Source/Runtime/Engine/Public/PSOPrecacheValidation.h b/Engine/Source/Runtime/Engine/Public/PSOPrecacheValidation.h
--- a/Engine/Source/Runtime/Engine/Public/PSOPrecacheValidation.h
+++ b/Engine/Source/Runtime/Engine/Public/PSOPrecacheValidation.h
@@ -42,7 +42,7 @@ enum class EPSOPrecacheMissType : uint8
 };
 
 /**
- * Log PSO miss information to logg with optional detailed information on what's causing the miss compared to what's already precached
+ * Log PSO miss information to give optional detailed information on what's causing the miss compared to what's already precached
  */
 extern ENGINE_API void LogPSOMissInfo(
 	const FGraphicsPipelineStateInitializer& GraphicsPSOInitializer, 
@@ -108,6 +108,11 @@ namespace PSOCollectorStats
 		const FMaterialRenderProxy* Material,
 		int32 PSOCollectorIndex);
 
+	/*
+	 * Update the CSV stats	
+	 */
+	extern ENGINE_API void UpdateCSVStats(EPSOPrecacheResult PSOPrecacheResult);
+
 	using VertexFactoryCountTableType = Experimental::TRobinHoodHashMap<const FVertexFactoryType*, uint32>;
 	struct FShaderStateUsage
 	{
@@ -210,11 +215,11 @@ namespace PSOCollectorStats
 		const FPrecacheStats& GetStats() const { return Stats; }
 
 		template <typename TPrecacheState>
-		void AddStateToCache(const TPrecacheState& PrecacheState, uint64 HashFn(const TPrecacheState&), const FMaterial* Material, int32 PSOCollectorIndex, const FVertexFactoryType* VertexFactoryType)
+		bool AddStateToCache(const TPrecacheState& PrecacheState, uint64 HashFn(const TPrecacheState&), const FMaterial* Material, int32 PSOCollectorIndex, const FVertexFactoryType* VertexFactoryType)
 		{
 			if (!IsPrecachingValidationEnabled())
 			{
-				return;
+				return false;
 			}
 
 			uint64 PrecacheStateHash = HashFn(PrecacheState);
@@ -242,29 +247,33 @@ namespace PSOCollectorStats
 			{
 				Stats.PrecacheData.UpdateStats(PSOCollectorIndex, VertexFactoryType);
 			}
+
+			return bUpdateStats;
 		}
 
 		template <typename TPrecacheState>
-		EPSOPrecacheResult CheckStateInCache(const TPrecacheState& PrecacheState, uint64 HashFn(const TPrecacheState&), EPSOPrecacheResult PrecacheResult, int32 PSOCollectorIndex, const FVertexFactoryType* VertexFactoryType)
+		bool CheckStateInCache(const TPrecacheState& PrecacheState, uint64 HashFn(const TPrecacheState&), int32 PSOCollectorIndex, const FVertexFactoryType* VertexFactoryType, EPSOPrecacheResult& InOutPrecacheResult)
 		{
 			if (!IsPrecachingValidationEnabled())
 			{
-				return EPSOPrecacheResult::Unknown;
+				InOutPrecacheResult = EPSOPrecacheResult::Unknown;
+				return false;
 			}
 
 			uint64 PrecacheStateHash = HashFn(PrecacheState);
-			return CheckStateInCacheByHash(PrecacheStateHash, PrecacheResult, PSOCollectorIndex, VertexFactoryType);
+			return CheckStateInCacheByHash(PrecacheStateHash, PSOCollectorIndex, VertexFactoryType, InOutPrecacheResult);
 		}
 
-		EPSOPrecacheResult CheckStateInCacheByHash(const uint64 PrecacheStateHash, EPSOPrecacheResult PrecacheResult, int32 PSOCollectorIndex, const FVertexFactoryType* VertexFactoryType)
+		bool CheckStateInCacheByHash(const uint64 PrecacheStateHash, int32 PSOCollectorIndex, const FVertexFactoryType* VertexFactoryType, EPSOPrecacheResult& InOutPrecacheResult)
 		{
 			if (!IsPrecachingValidationEnabled())
 			{
-				return EPSOPrecacheResult::Unknown;
+				InOutPrecacheResult = EPSOPrecacheResult::Unknown;
+				return false;
 			}
 
 			bool bTracked = IsStateTracked(PSOCollectorIndex, VertexFactoryType);
-			return UpdatePrecacheStats(PrecacheStateHash, PSOCollectorIndex, VertexFactoryType, bTracked, PrecacheResult);
+			return UpdatePrecacheStats(PrecacheStateHash, PSOCollectorIndex, VertexFactoryType, bTracked, InOutPrecacheResult);
 		}
 
 		ENGINE_API bool IsPrecached(uint64 PrecacheStateHash);
@@ -276,7 +285,7 @@ namespace PSOCollectorStats
 
 	private:
 
-		ENGINE_API EPSOPrecacheResult UpdatePrecacheStats(uint64 PrecacheHash, int32 PSOCollectorIndex, const FVertexFactoryType* VertexFactoryType, bool bTracked, EPSOPrecacheResult InPrecacheResult);
+		ENGINE_API bool UpdatePrecacheStats(uint64 PrecacheHash, int32 PSOCollectorIndex, const FVertexFactoryType* VertexFactoryType, bool bTracked, EPSOPrecacheResult& InOutPrecacheResult);
 
 		FPrecacheStats Stats;
 
diff --git a/Engine/Source/Runtime/Renderer/Private/MeshPassProcessor.cpp b/Engine/Source/Runtime/Renderer/Private/MeshPassProcessor.cpp
--- a/Engine/Source/Runtime/Renderer/Private/MeshPassProcessor.cpp
+++ b/Engine/Source/Runtime/Renderer/Private/MeshPassProcessor.cpp
@@ -2353,9 +2353,10 @@ void PSOCollectorStats::CheckShaderOnlyStateInCache(
 	int32 PSOCollectorIndex)
 {
 	FGraphicsMinimalPipelineStateInitializer ShadersOnlyInitializer = GetShadersOnlyInitializer(Initializer);
-	EPSOPrecacheResult Result = PSOCollectorStats::GetShadersOnlyPSOPrecacheStatsCollector().CheckStateInCacheByHash(ShadersOnlyInitializer.StatePrecachePSOHash, EPSOPrecacheResult::Unknown, PSOCollectorIndex, VFType);
 
-	if (IsFullPrecachingValidationEnabled() && Result != EPSOPrecacheResult::Unknown && Result != EPSOPrecacheResult::Complete)
+	EPSOPrecacheResult Result = EPSOPrecacheResult::Unknown;
+	bool bStatsUpdated = PSOCollectorStats::GetShadersOnlyPSOPrecacheStatsCollector().CheckStateInCacheByHash(ShadersOnlyInitializer.StatePrecachePSOHash, PSOCollectorIndex, VFType, Result);
+	if (IsFullPrecachingValidationEnabled() && bStatsUpdated && Result != EPSOPrecacheResult::Unknown && Result != EPSOPrecacheResult::Complete)
 	{
 		FGraphicsPipelineStateInitializer GraphicsPSOInitializer = ShadersOnlyInitializer.AsGraphicsPipelineStateInitializer();
 		LogPSOMissInfo(GraphicsPSOInitializer, EPSOPrecacheMissType::ShadersOnly, Result, &Material, VFType, PrimitiveSceneProxy, PSOCollectorIndex, ShadersOnlyInitializer.StatePrecachePSOHash);
@@ -2370,9 +2371,10 @@ void PSOCollectorStats::CheckMinimalPipelineStateInCache(
 	int32 PSOCollectorIndex)
 {
 	FGraphicsMinimalPipelineStateInitializer PatchedMinimalInitializer = PSOCollectorStats::PatchMinimalPipelineStateToCheck(Initializer);
-	EPSOPrecacheResult Result = PSOCollectorStats::GetMinimalPSOPrecacheStatsCollector().CheckStateInCacheByHash(PatchedMinimalInitializer.StatePrecachePSOHash, EPSOPrecacheResult::Unknown, PSOCollectorIndex, VFType);
 
-	if (IsFullPrecachingValidationEnabled() && Result != EPSOPrecacheResult::Unknown && Result != EPSOPrecacheResult::Complete)
+	EPSOPrecacheResult Result = EPSOPrecacheResult::Unknown;
+	bool bStatsUpdated = PSOCollectorStats::GetMinimalPSOPrecacheStatsCollector().CheckStateInCacheByHash(PatchedMinimalInitializer.StatePrecachePSOHash, PSOCollectorIndex, VFType, Result);
+	if (IsFullPrecachingValidationEnabled() && bStatsUpdated && Result != EPSOPrecacheResult::Unknown && Result != EPSOPrecacheResult::Complete)
 	{
 		FGraphicsMinimalPipelineStateInitializer ShadersOnlyInitializer = GetShadersOnlyInitializer(Initializer);
 		bool bShaderOnlyPrecached = PSOCollectorStats::GetShadersOnlyPSOPrecacheStatsCollector().IsPrecached(ShadersOnlyInitializer.StatePrecachePSOHash);
