[Experimental] Level Instance Property Override

diff --git a/Engine/Source/Runtime/Engine/Private/LevelInstance/LevelInstanceSubsystem.cpp b/Engine/Source/Runtime/Engine/Private/LevelInstance/LevelInstanceSubsystem.cpp
--- a/Engine/Source/Runtime/Engine/Private/LevelInstance/LevelInstanceSubsystem.cpp
+++ b/Engine/Source/Runtime/Engine/Private/LevelInstance/LevelInstanceSubsystem.cpp
@@ -3,6 +3,7 @@
 #include "LevelInstance/LevelInstanceSubsystem.h"
 #include "Engine/LevelStreaming.h"
 #include "LevelInstance/LevelInstanceLevelStreaming.h"
+#include "LevelInstance/LevelInstanceSettings.h"
 #include "Misc/StringFormatArg.h"
 #include "UObject/UObjectIterator.h"
 #include "UObject/Package.h"
@@ -16,11 +17,13 @@
 #include "Settings/LevelEditorMiscSettings.h"
 #include "ActorEditorContext/ScopedActorEditorContextSetExternalDataLayerAsset.h"
 #include "LevelInstance/LevelInstanceEditorLevelStreaming.h"
+#include "LevelInstance/LevelInstanceEditorPropertyOverrideLevelStreaming.h"
 #include "LevelInstance/ILevelInstanceEditorModule.h"
 #include "LevelInstance/LevelInstanceEditorInstanceActor.h"
 #include "LevelInstance/LevelInstanceEditorObject.h"
 #include "LevelInstance/LevelInstanceEditorPivotActor.h"
 #include "LevelInstance/LevelInstanceComponent.h"
+#include "WorldPartition/LevelInstance/LevelInstanceContainerInstance.h"
 #include "WorldPartition/LevelInstance/LevelInstanceActorDesc.h"
 #include "WorldPartition/WorldPartition.h"
 #include "WorldPartition/WorldPartitionSubsystem.h"
@@ -88,9 +91,15 @@ void ULevelInstanceSubsystem::Initialize(FSubsystemCollectionBase& Collection)
 	Super::Initialize(Collection);
 
 #if WITH_EDITOR
+	// Make sure Policy is initialized
+	ULevelInstanceSettings::Get()->UpdatePropertyOverridePolicy();
+	
 	if (GEditor)
 	{
 		ILevelInstanceEditorModule& EditorModule = FModuleManager::LoadModuleChecked<ILevelInstanceEditorModule>("LevelInstanceEditor");
+		EditorModule.OnExitEditorMode().AddUObject(this, &ULevelInstanceSubsystem::OnExitEditorMode);
+		EditorModule.OnTryExitEditorMode().AddUObject(this, &ULevelInstanceSubsystem::OnTryExitEditorMode);
+		
 		FEditorDelegates::OnAssetsPreDelete.AddUObject(this, &ULevelInstanceSubsystem::OnAssetsPreDelete);
 		FEditorDelegates::PreSaveWorldWithContext.AddUObject(this, &ULevelInstanceSubsystem::OnPreSaveWorldWithContext);
 		FWorldDelegates::OnPreWorldRename.AddUObject(this, &ULevelInstanceSubsystem::OnPreWorldRename);
@@ -106,6 +115,12 @@ void ULevelInstanceSubsystem::Deinitialize()
 	FEditorDelegates::PreSaveWorldWithContext.RemoveAll(this);
 	FWorldDelegates::OnPreWorldRename.RemoveAll(this);
 	FWorldDelegates::OnWorldCleanup.RemoveAll(this);
+
+	if (ILevelInstanceEditorModule* EditorModule = FModuleManager::GetModulePtr<ILevelInstanceEditorModule>("LevelInstanceEditor"))
+	{
+		EditorModule->OnExitEditorMode().RemoveAll(this);
+		EditorModule->OnTryExitEditorMode().RemoveAll(this);
+	}
 #endif
 }
 
@@ -184,7 +199,7 @@ void ULevelInstanceSubsystem::RequestLoadLevelInstance(ILevelInstanceInterface*
 	if (LevelInstance->IsWorldAssetValid())
 	{
 #if WITH_EDITOR
-		if (!IsEditingLevelInstance(LevelInstance))
+		if (!IsEditingLevelInstance(LevelInstance) && !IsEditingLevelInstancePropertyOverrides(LevelInstance))
 #endif
 		{
 			LevelInstancesToUnload.Remove(LevelInstance->GetLevelInstanceID());
@@ -363,7 +378,11 @@ ULevel* ULevelInstanceSubsystem::GetLevelInstanceLevel(const ILevelInstanceInter
 #if WITH_EDITOR
 		if (const FLevelInstanceEdit* CurrentEdit = GetLevelInstanceEdit(LevelInstance))
 		{
-			return LevelInstanceEdit->LevelStreaming->GetLoadedLevel();
+			return CurrentEdit->LevelStreaming->GetLoadedLevel();
+		}
+		else if (const FPropertyOverrideEdit* CurrentPropertyOverrideEdit = GetLevelInstancePropertyOverrideEdit(LevelInstance))
+		{
+			return CurrentPropertyOverrideEdit->LevelStreaming->GetLoadedLevel();
 		}
 		else 
 #endif // WITH_EDITOR
@@ -449,17 +468,8 @@ void ULevelInstanceSubsystem::OnWorldCleanup(UWorld* InWorld, bool bSessionEnded
 		}
 
 		LoadedLevelInstances.Empty();
-		if (LevelInstanceEdit)
-		{
-			// If we are inside an Edit, null out streaming pointer so that destructor doesn't cause a call to RemoveLevelsFromWorld
-			LevelInstanceEdit->LevelStreaming = nullptr;
-			LevelInstanceEdit.Reset();
-
-			if (ILevelInstanceEditorModule* EditorModule = (ILevelInstanceEditorModule*)FModuleManager::Get().GetModule("LevelInstanceEditor"))
-			{
-				EditorModule->DeactivateEditorMode();
-			}
-		}	
+		LevelInstanceEdit.Reset();
+		PropertyOverrideEdit.Reset();
 	}
 }
 
@@ -478,46 +488,65 @@ void ULevelInstanceSubsystem::ForEachLevelStreaming(TFunctionRef<bool(ULevelStre
 	{
 		Operation(LevelInstanceEdit->LevelStreaming);
 	}
-}
 
-
-void ULevelInstanceSubsystem::RegisterLoadedLevelStreamingLevelInstanceEditor(ULevelStreamingLevelInstanceEditor* LevelStreaming)
-{
-	if (!bIsCreatingLevelInstance)
+	if (PropertyOverrideEdit)
 	{
-		check(!LevelInstanceEdit.IsValid());
-		ILevelInstanceInterface* LevelInstance = LevelStreaming->GetLevelInstance();
-		LevelInstanceEdit = MakeUnique<FLevelInstanceEdit>(LevelStreaming, LevelInstance);
-
-		if (ILevelInstanceEditorModule* EditorModule = FModuleManager::GetModulePtr<ILevelInstanceEditorModule>("LevelInstanceEditor"))
-		{
-			EditorModule->OnExitEditorMode().AddUObject(this, &ULevelInstanceSubsystem::OnExitEditorMode);
-			EditorModule->OnTryExitEditorMode().AddUObject(this, &ULevelInstanceSubsystem::OnTryExitEditorMode);
-		}
+		Operation(PropertyOverrideEdit->LevelStreaming);
 	}
 }
 
-void ULevelInstanceSubsystem::ResetEdit(TUniquePtr<FLevelInstanceEdit>& InLevelInstanceEdit)
+FString ULevelInstanceSubsystem::GetActorNameToSelectFromContext(const ILevelInstanceInterface* LevelInstance, const AActor* ContextActor, const FString& DefaultActorNameToSelect) const
 {
-	if (InLevelInstanceEdit)
+	FString ActorNameToSelect = DefaultActorNameToSelect;
+	if (ContextActor)
 	{
-		if (InLevelInstanceEdit == LevelInstanceEdit)
+		ActorNameToSelect = ContextActor->GetName();
+		ForEachLevelInstanceAncestorsAndSelf(ContextActor, [&ActorNameToSelect, LevelInstance](const ILevelInstanceInterface* AncestorLevelInstance)
 		{
-			if (ILevelInstanceEditorModule* EditorModule = FModuleManager::GetModulePtr<ILevelInstanceEditorModule>("LevelInstanceEditor"))
+			// stop when we hit the LevelInstance we are about to edit
+			if (AncestorLevelInstance == LevelInstance)
 			{
-				EditorModule->OnExitEditorMode().RemoveAll(this);
-				EditorModule->OnTryExitEditorMode().RemoveAll(this);
+				return false;
 			}
-		}
-		else
-		{
-			// Only case supported where we are using a tmp FLevelInstanceEdit
-			check(bIsCreatingLevelInstance);
-		}
 
-		InLevelInstanceEdit.Reset();
+			ActorNameToSelect = CastChecked<AActor>(AncestorLevelInstance)->GetName();
+			return true;
+		});
+	}
+
+	return ActorNameToSelect;
+}
+
+void ULevelInstanceSubsystem::SelectActorFromActorName(ILevelInstanceInterface* LevelInstance, const FString& ActorName) const
+{
+	// Try and select something meaningful
+	AActor* ActorToSelect = nullptr;
+	if (!ActorName.IsEmpty())
+	{
+		ActorToSelect = FindObject<AActor>(LevelInstance->GetLoadedLevel(), *ActorName);
+	}
+
+	// default to LevelInstance
+	AActor* LevelInstanceActor = CastChecked<AActor>(LevelInstance);
+	if (!ActorToSelect)
+	{
+		ActorToSelect = LevelInstanceActor;
+	}
+		
+	GEditor->SelectActor(ActorToSelect, true, true);
+}
+
+
+void ULevelInstanceSubsystem::RegisterLoadedLevelStreamingLevelInstanceEditor(ULevelStreamingLevelInstanceEditor* LevelStreaming)
+{
+	if (!bIsCreatingLevelInstance)
+	{
+		check(!LevelInstanceEdit.IsValid());
+		ILevelInstanceInterface* LevelInstance = LevelStreaming->GetLevelInstance();
+		LevelInstanceEdit = MakeUnique<FLevelInstanceEdit>(LevelStreaming, LevelInstance);
 	}
 }
+
 #endif
 
 void ULevelInstanceSubsystem::LoadLevelInstance(ILevelInstanceInterface* LevelInstance)
@@ -690,7 +719,10 @@ ILevelInstanceInterface* ULevelInstanceSubsystem::GetOwningLevelInstance(const U
 		{
 			return LevelStreamingEditor->GetLevelInstance();
 		}
-		else 
+		else if (ULevelStreamingLevelInstanceEditorPropertyOverride* LevelStreamingPropertyOverride = Cast<ULevelStreamingLevelInstanceEditorPropertyOverride>(BaseLevelStreaming))
+		{
+			return LevelStreamingPropertyOverride->GetLevelInstance();
+		}
 #endif
 		if (ULevelStreamingLevelInstance* LevelStreaming = Cast<ULevelStreamingLevelInstance>(BaseLevelStreaming))
 		{
@@ -836,41 +868,125 @@ void ULevelInstanceSubsystem::Tick()
 	if (!GetWorld()->IsGameWorld())
 	{
 		UpdateStreamingState();
+
+		// Update Editor Mode if we are the GEditor world's Subsystem
+		if (GEditor->GetEditorWorldContext().World() == GetWorld())
+		{
+			if (ILevelInstanceEditorModule* EditorModule = (ILevelInstanceEditorModule*)FModuleManager::Get().GetModule("LevelInstanceEditor"))
+			{
+				const bool bActivated = LevelInstanceEdit || PropertyOverrideEdit;
+				EditorModule->UpdateEditorMode(bActivated);
+			}
+		}
 	}
 }
 
 void ULevelInstanceSubsystem::OnExitEditorMode()
 {
-	OnExitEditorModeInternal(/*bForceExit=*/true);
+	if (LevelInstanceEdit || PropertyOverrideEdit)
+	{
+		OnExitEditorModeInternal(/*bForceExit=*/true);
+	}
 }
 
 void ULevelInstanceSubsystem::OnTryExitEditorMode()
-{	
-	if (OnExitEditorModeInternal(/*bForceExit=*/false))
+{
+	if (LevelInstanceEdit || PropertyOverrideEdit)
 	{
-		ILevelInstanceEditorModule& EditorModule = FModuleManager::GetModuleChecked<ILevelInstanceEditorModule>("LevelInstanceEditor");
-		EditorModule.DeactivateEditorMode();
+		OnExitEditorModeInternal(/*bForceExit=*/false);
 	}
 }
 
-bool ULevelInstanceSubsystem::OnExitEditorModeInternal(bool bForceExit)
+FText ULevelInstanceSubsystem::GetToolKitExitToolTip() const
 {
-	if (bIsCommittingLevelInstance || bIsCreatingLevelInstance)
+	if (ILevelInstanceInterface* PropertyOverrideInstance = GetEditingPropertyOverridesLevelInstance())
 	{
-		return false;
+		AActor* Actor = CastChecked<AActor>(PropertyOverrideInstance);
+		return FText::FromString(Actor->GetActorLabel());
+	}
+	else if (ILevelInstanceInterface* EditingLevelInstance = GetEditingLevelInstance())
+	{
+		const TSoftObjectPtr<UWorld> WorldAsset = EditingLevelInstance->GetWorldAsset();
+		return FText::FromString(WorldAsset.GetAssetName());
+	}
+
+	return FText::GetEmpty();
+}
+
+FText ULevelInstanceSubsystem::GetToolKitDisplayText() const
+{
+	if (ILevelInstanceInterface* PropertyOverrideInstance = GetEditingPropertyOverridesLevelInstance())
+	{
+		return LOCTEXT("ExitTooltip", "Exit Property Override Edit");
+	}
+	else if (ILevelInstanceInterface* EditingLevelInstance = GetEditingLevelInstance())
+	{
+		return LOCTEXT("ExitTooltip", "Exit Level Instance Edit");
+	}
+
+	return FText::GetEmpty();
+}
+
+void ULevelInstanceSubsystem::ToolKitExit()
+{
+	const bool bForceExit = false;
+	if (PropertyOverrideEdit)
+	{
+		TryCommitLevelInstancePropertyOverrideEdit(bForceExit);
+	}
+	else if (LevelInstanceEdit)
+	{
+		TryCommitLevelInstanceEdit(bForceExit);
 	}
+}
 
+bool ULevelInstanceSubsystem::TryCommitLevelInstanceEdit(bool bForceExit)
+{
 	if (LevelInstanceEdit)
 	{
 		TGuardValue<bool> CommitScope(bIsCommittingLevelInstance, true);
 		bool bDiscard = false;
-		if (PromptUserForCommit(LevelInstanceEdit.Get(), bDiscard, bForceExit))
+		if (!PromptUserForCommit(LevelInstanceEdit.Get(), bDiscard, bForceExit))
 		{
-			return CommitLevelInstanceInternal(LevelInstanceEdit, bDiscard, /*bDiscardOnFailure=*/bForceExit);
+			return false;
 		}
+
+		return CommitLevelInstanceInternal(LevelInstanceEdit, bDiscard, /*bDiscardOnFailure=*/bForceExit);
 	}
 
-	return false;
+	return true;
+}
+
+bool ULevelInstanceSubsystem::TryCommitLevelInstancePropertyOverrideEdit(bool bForceExit)
+{
+	if (PropertyOverrideEdit)
+	{
+		TGuardValue<bool> CommitScope(bIsCommittingLevelInstance, true);
+		bool bDiscard = false;
+		if (!PromptUserForCommitPropertyOverrides(PropertyOverrideEdit.Get(), bDiscard, bForceExit))
+		{
+			return false;
+		}
+
+		return CommitLevelInstancePropertyOverridesInternal(PropertyOverrideEdit, bDiscard);
+	}
+
+	return true;
+}
+
+void ULevelInstanceSubsystem::OnExitEditorModeInternal(bool bForceExit)
+{
+	if (bIsCommittingLevelInstance || bIsCreatingLevelInstance)
+	{
+		return;
+	}
+
+	if (!TryCommitLevelInstancePropertyOverrideEdit(bForceExit))
+	{
+		return;
+	}
+
+	TryCommitLevelInstanceEdit(bForceExit);
 }
 
 bool ULevelInstanceSubsystem::GetLevelInstanceBounds(const ILevelInstanceInterface* LevelInstance, FBox& OutBounds) const
@@ -1160,17 +1276,40 @@ ULevelStreamingLevelInstanceEditor* ULevelInstanceSubsystem::CreateNewStreamingL
 	return Cast<ULevelStreamingLevelInstanceEditor>(EditorLevelUtils::CreateNewStreamingLevelForWorld(*GetWorld(), CreateNewStreamingLevelParamsCopy));
 }
 
+bool ULevelInstanceSubsystem::CanCreateLevelInstanceFrom(const TArray<AActor*>& ActorsToMove, FText* OutReason)
+{
+	if (ActorsToMove.IsEmpty())
+	{
+		if (OutReason != nullptr)
+		{
+			*OutReason = LOCTEXT("CanCreateLevelInstanceFromEmptyActorArray", "Failed to create Level Instance from actor array");
+		}
+		return false;
+	}
+
+	for (AActor* ActorToMove : ActorsToMove)
+	{
+		if (!CanMoveActorToLevel(ActorToMove, OutReason))
+		{
+			return false;
+		}
+	}
+
+	return true;
+}
+
 ILevelInstanceInterface* ULevelInstanceSubsystem::CreateLevelInstanceFrom(const TArray<AActor*>& ActorsToMove, const FNewLevelInstanceParams& CreationParams)
 {
-	check(!bIsCreatingLevelInstance);
-	TGuardValue<bool> CreateLevelInstanceGuard(bIsCreatingLevelInstance, true);
-	ULevel* CurrentLevel = GetWorld()->GetCurrentLevel();
-		
-	if (ActorsToMove.Num() == 0)
+	FText Reason;
+	if (!CanCreateLevelInstanceFrom(ActorsToMove, &Reason))
 	{
-		UE_LOG(LogLevelInstance, Warning, TEXT("Failed to create Level Instance from empty actor array"));
+		UE_LOG(LogLevelInstance, Warning, TEXT("Failed to create Level Instance : %s"), *Reason.ToString());
 		return nullptr;
 	}
+
+	check(!bIsCreatingLevelInstance);
+	TGuardValue<bool> CreateLevelInstanceGuard(bIsCreatingLevelInstance, true);
+	ULevel* CurrentLevel = GetWorld()->GetCurrentLevel();
 	
 	TOptional<const UExternalDataLayerAsset*> CommonExternalDataLayerAsset;
 	FBox ActorLocationBox(ForceInit);
@@ -1186,13 +1325,6 @@ ILevelInstanceInterface* ULevelInstanceSubsystem::CreateLevelInstanceFrom(const
 		}
 		ActorLocationBox += LocalActorLocationBox;
 
-		FText Reason;
-		if (!CanMoveActorToLevel(ActorToMove, &Reason))
-		{
-			UE_LOG(LogLevelInstance, Warning, TEXT("%s"), *Reason.ToString());
-			return nullptr;
-		}
-
 		if (!CommonExternalDataLayerAsset.IsSet())
 		{
 			CommonExternalDataLayerAsset = ActorToMove->GetExternalDataLayerAsset();
@@ -1448,13 +1580,6 @@ ILevelInstanceInterface* ULevelInstanceSubsystem::CreateLevelInstanceFrom(const
 		FEditorFileUtils::PromptForCheckoutAndSave({ NewLevelInstanceActor->GetPackage() }, /*bCheckDirty*/false, /*bPromptToSave*/false);
 	}
 
-	// EditorLevelUtils::CreateNewStreamingLevelForWorld deactivates all modes. Re-activate if needed
-	if (LevelInstanceEdit)
-	{
-		ILevelInstanceEditorModule& EditorModule = FModuleManager::GetModuleChecked<ILevelInstanceEditorModule>("LevelInstanceEditor");
-		EditorModule.ActivateEditorMode();
-	}
-
 	// After commit, CurrentLevel goes back to world's PersistentLevel. Set it back to the current editing level instance (if any).
 	if (ILevelInstanceInterface* EditingLevelInstance = GetEditingLevelInstance())
 	{
@@ -1464,9 +1589,19 @@ ILevelInstanceInterface* ULevelInstanceSubsystem::CreateLevelInstanceFrom(const
 	return GetLevelInstance(NewLevelInstanceID);
 }
 
+bool ULevelInstanceSubsystem::CanBreakLevelInstance(const ILevelInstanceInterface* LevelInstance) const
+{
+	return !IsEditingLevelInstance(LevelInstance) && !IsEditingLevelInstancePropertyOverrides(LevelInstance) && !HasParentPropertyOverridesEdit(LevelInstance) && !LevelInstanceHasLevelScriptBlueprint(LevelInstance);
+}
+
 bool ULevelInstanceSubsystem::BreakLevelInstance(ILevelInstanceInterface* LevelInstance, uint32 Levels /* = 1 */,
 	TArray<AActor*>* OutMovedActors /* = nullptr */, ELevelInstanceBreakFlags Flags /* = None */)
 {
+	if (!CanBreakLevelInstance(LevelInstance))
+	{
+		return false;
+	}
+
 	const double StartTime = FPlatformTime::Seconds();
 
 	const uint32 bAvoidRelabelOnPasteSelected = GetMutableDefault<ULevelEditorMiscSettings>()->bAvoidRelabelOnPasteSelected;
@@ -1756,6 +1891,15 @@ bool ULevelInstanceSubsystem::CanMoveActorToLevel(const AActor* Actor, FText* Ou
 				return false;
 			}
 
+			if (IsEditingLevelInstancePropertyOverrides(LevelInstance))
+			{
+				if (OutReason != nullptr)
+				{
+					*OutReason = LOCTEXT("CanMoveActorLevelPropertyOverries", "Can't move Level Instance actor while it is in property override edit");
+				}
+				return false;
+			}
+
 			bool bEditingChildren = false;
 			ForEachLevelInstanceChild(LevelInstance, true, [this, &bEditingChildren](const ILevelInstanceInterface* ChildLevelInstance)
 			{
@@ -1776,6 +1920,18 @@ bool ULevelInstanceSubsystem::CanMoveActorToLevel(const AActor* Actor, FText* Ou
 				return false;
 			}
 		}
+
+		if (ILevelInstanceInterface* ParentLevelInstance = GetParentLevelInstance(Actor))
+		{
+			if (ParentLevelInstance->IsEditingPropertyOverrides())
+			{
+				if (OutReason != nullptr)
+				{
+					*OutReason = LOCTEXT("CanMoveActorToLevelParentPropertyOverrides", "Can't move actor while its parent is in property override edit");
+				}
+				return false;
+			}
+		}
 	}
 
 	return true;
@@ -1823,7 +1979,7 @@ void ULevelInstanceSubsystem::OnActorDeleted(AActor* Actor)
 				if (const FLevelInstanceEdit* ChildLevelInstanceEdit = GetLevelInstanceEdit(ChildLevelInstance))
 				{
 					check(!IsLevelInstanceEditDirty(ChildLevelInstanceEdit));
-					ResetEdit(LevelInstanceEdit);
+					LevelInstanceEdit.Reset();
 					return false;
 				}
 				return true;
@@ -1948,6 +2104,16 @@ const ULevelInstanceSubsystem::FLevelInstanceEdit* ULevelInstanceSubsystem::GetL
 	return nullptr;
 }
 
+const ULevelInstanceSubsystem::FPropertyOverrideEdit* ULevelInstanceSubsystem::GetLevelInstancePropertyOverrideEdit(const ILevelInstanceInterface* LevelInstance) const
+{
+	if (LevelInstance && PropertyOverrideEdit && PropertyOverrideEdit->GetLevelInstance() == LevelInstance)
+	{
+		return PropertyOverrideEdit.Get();
+	}
+
+	return nullptr;
+}
+
 bool ULevelInstanceSubsystem::IsEditingLevelInstanceDirty(const ILevelInstanceInterface* LevelInstance) const
 {
 	const FLevelInstanceEdit* CurrentEdit = GetLevelInstanceEdit(LevelInstance);
@@ -2004,7 +2170,35 @@ bool ULevelInstanceSubsystem::PromptUserForCommit(const FLevelInstanceEdit* InLe
 	return true;
 }
 
-bool ULevelInstanceSubsystem::CanEditLevelInstance(const ILevelInstanceInterface* LevelInstance, FText* OutReason) const
+bool ULevelInstanceSubsystem::PromptUserForCommitPropertyOverrides(const FPropertyOverrideEdit* InPropertyOverrideEdit, bool& bOutDiscard, bool bForceCommit) const
+{
+	bOutDiscard = false;
+	// Can commit no pending changes
+	if (!InPropertyOverrideEdit->IsDirty())
+	{
+		return true;
+	}
+
+	// If changes can be discarded prompt user
+	if (CanCommitLevelInstancePropertyOverrides(InPropertyOverrideEdit->GetLevelInstance(), /*bDiscardEdits=*/true))
+	{
+		// if bForceExit we can't cancel the exiting of the mode so the user needs to decide between saving or discarding
+		EAppReturnType::Type Ret = FMessageDialog::Open(
+			bForceCommit ? EAppMsgType::YesNo : EAppMsgType::YesNoCancel, LOCTEXT("CommitOrDiscardPropertyOverrideChangesMsg", "Unsaved Property override changes will get discarded. Do you want to save them now?"),
+			LOCTEXT("CommitOrDiscardPropertyOverrideChangesTitle", "Save changes?"));
+		if (Ret == EAppReturnType::Cancel && !bForceCommit)
+		{
+			return false;
+		}
+
+		bOutDiscard = (Ret != EAppReturnType::Yes);
+	}
+
+	// Can commit but can't discard changes
+	return true;
+}
+
+bool ULevelInstanceSubsystem::CanEditLevelInstanceCommon(const ILevelInstanceInterface* LevelInstance, FText* OutReason) const
 {
 	// Only allow Editing in Editor World
 	if (GetWorld()->WorldType != EWorldType::Editor)
@@ -2012,16 +2206,32 @@ bool ULevelInstanceSubsystem::CanEditLevelInstance(const ILevelInstanceInterface
 		return false;
 	}
 
-	if (LevelInstanceEdit)
+	if (IsEditingLevelInstance(LevelInstance))
 	{
-		if (LevelInstanceEdit->GetLevelInstance() == LevelInstance)
+		if (OutReason)
 		{
-			if (OutReason)
-			{
-				*OutReason = FText::Format(LOCTEXT("CanEditLevelInstanceAlreadyBeingEdited", "Level Instance already being edited ({0})."), FText::FromString(LevelInstance->GetWorldAssetPackage()));
-			}
-			return false;
+			*OutReason = FText::Format(LOCTEXT("CanEditLevelInstanceAlreadyBeingEdited", "Level Instance already being edited ({0})."), FText::FromString(LevelInstance->GetWorldAssetPackage()));
+		}
+		return false;
+	}
+
+	if (IsEditingLevelInstancePropertyOverrides(LevelInstance))
+	{
+		if (OutReason)
+		{
+			*OutReason = FText::Format(LOCTEXT("CanEditLevelInstanceAlreadyBeingEdited", "Level Instance already in property override edit ({0})."), FText::FromString(LevelInstance->GetWorldAssetPackage()));
 		}
+		return false;
+	}
+
+	return true;
+}
+
+bool ULevelInstanceSubsystem::CanEditLevelInstance(const ILevelInstanceInterface* LevelInstance, FText* OutReason) const
+{
+	if (!CanEditLevelInstanceCommon(LevelInstance, OutReason))
+	{
+		return false;
 	}
 	
 	if (LevelInstance->IsWorldAssetValid())
@@ -2071,60 +2281,135 @@ bool ULevelInstanceSubsystem::CanEditLevelInstance(const ILevelInstanceInterface
 	return true;
 }
 
-bool ULevelInstanceSubsystem::CanCommitLevelInstance(const ILevelInstanceInterface* LevelInstance, bool bDiscardEdits, FText* OutReason) const
+bool ULevelInstanceSubsystem::CanEditLevelInstancePropertyOverrides(const ILevelInstanceInterface* LevelInstance, FText* OutReason) const
 {
-	if (const FLevelInstanceEdit* CurrentEdit = GetLevelInstanceEdit(LevelInstance))
+	if (!CanEditLevelInstanceCommon(LevelInstance, OutReason))
 	{
-		return !bDiscardEdits || LevelInstanceEdit->CanDiscard(OutReason);
+		return false;
 	}
 
-	if (OutReason)
+	if (!ULevelInstanceSettings::Get()->IsPropertyOverrideEnabled())
 	{
-		*OutReason = LOCTEXT("CanCommitLevelInstanceNotEditing", "Level Instance is not currently being edited");
+		if (OutReason)
+		{
+			*OutReason = LOCTEXT("LevelInstanceNoSupportPropertyOverrides", "Level Instance property override feature is not enabled");
+		}
+		return false;
 	}
-	return false;
-}
 
-void ULevelInstanceSubsystem::EditLevelInstance(ILevelInstanceInterface* LevelInstance, TWeakObjectPtr<AActor> ContextActorPtr)
-{
-	if (EditLevelInstanceInternal(LevelInstance, ContextActorPtr, FString(), false))
+	if (!LevelInstance->SupportsPropertyOverrides())
 	{
-		ILevelInstanceEditorModule& EditorModule = FModuleManager::GetModuleChecked<ILevelInstanceEditorModule>("LevelInstanceEditor");
-		EditorModule.ActivateEditorMode();
+		if (OutReason)
+		{
+			*OutReason = LOCTEXT("LevelInstanceNoSupportPropertyOverrides", "Level Instance does not support property overrides");
+		}
+		return false;
 	}
-}
-
+		
+	const ILevelInstanceInterface* TopLevelAncestor = nullptr;
+	bool bAllWorldPartitions = true;
+	ForEachLevelInstanceAncestorsAndSelf(CastChecked<AActor>(LevelInstance), [this, &bAllWorldPartitions, &TopLevelAncestor](const ILevelInstanceInterface* AncestorOrSelf)
+	{
+		// Get the level loaded by this Level Instance
+		ULevel* LoadedLevel = GetLevelInstanceLevel(AncestorOrSelf);
+		// Check that it is World Partitioned
+		bAllWorldPartitions &= LoadedLevel && !!LoadedLevel->GetWorldPartition();
+
+		TopLevelAncestor = AncestorOrSelf;
+		return bAllWorldPartitions;
+	});
+		
+	// If any of the levels in the edit hierarchy are non-partitioned, we can't edit
+	if (!bAllWorldPartitions || !GetWorld()->IsPartitionedWorld())
+	{
+		if (OutReason)
+		{
+			*OutReason = LOCTEXT("LevelInstancePropertyOverridesWorldPartitionOnly", "Property overrides are only supported on world partition level hierarchies");
+		}
+		return false;
+	}
+
+	if (GetWorld()->GetPackage()->HasAnyPackageFlags(PKG_NewlyCreated))
+	{
+		if (OutReason)
+		{
+			*OutReason = LOCTEXT("LevelInstancePropertyOverridesNewlyCreated", "Property overrides are only supported on saved worlds. Please save world first.");
+		}
+		return false;
+	}
+
+	return true;
+}
+
+bool ULevelInstanceSubsystem::CanCommitLevelInstance(const ILevelInstanceInterface* LevelInstance, bool bDiscardEdits, FText* OutReason) const
+{
+	if (const FLevelInstanceEdit* CurrentEdit = GetLevelInstanceEdit(LevelInstance))
+	{
+		return !bDiscardEdits || CurrentEdit->CanDiscard(OutReason);
+	}
+
+	if (OutReason)
+	{
+		*OutReason = LOCTEXT("CanCommitLevelInstanceNotEditing", "Level Instance is not currently being edited");
+	}
+	return false;
+}
+
+bool ULevelInstanceSubsystem::CanCommitLevelInstancePropertyOverrides(const ILevelInstanceInterface* LevelInstance, bool bDiscardEdits, FText* OutReason) const
+{
+	if (const FPropertyOverrideEdit* CurrentEdit = GetLevelInstancePropertyOverrideEdit(LevelInstance))
+	{
+		return !bDiscardEdits || CurrentEdit->CanDiscard(OutReason);
+	}
+
+	if (OutReason)
+	{
+		*OutReason = LOCTEXT("CanCommitLevelInstanceNotEditing", "Level Instance is not currently in property override edit");
+	}
+	return false;
+}
+
+void ULevelInstanceSubsystem::EditLevelInstance(ILevelInstanceInterface* LevelInstance, TWeakObjectPtr<AActor> ContextActorPtr)
+{
+	EditLevelInstanceInternal(LevelInstance, ContextActorPtr, FString(), false);
+}
+
 bool ULevelInstanceSubsystem::EditLevelInstanceInternal(ILevelInstanceInterface* LevelInstance, TWeakObjectPtr<AActor> ContextActorPtr, const FString& InActorNameToSelect, bool bRecursive)
 {
 	check(CanEditLevelInstance(LevelInstance));
 		
-	// If there is a current edit and it is dirty, offer the user a chance to Save/Discard/Cancel
+	const FLevelInstanceID EditLevelInstanceID = LevelInstance->GetLevelInstanceID();
+
 	bool bDiscard = false;
+	bool bDiscardPropertyOverride = false;
+
+	// If there is a current property override edit and it is dirty, offer the user a chance to Save/Discard/Cancel
+	if (PropertyOverrideEdit && !PromptUserForCommitPropertyOverrides(PropertyOverrideEdit.Get(), bDiscardPropertyOverride))
+	{
+		return false;
+	}
+
+	// If there is a current edit and it is dirty, offer the user a chance to Save/Discard/Cancel
 	if (LevelInstanceEdit && !PromptUserForCommit(LevelInstanceEdit.Get(), bDiscard))
 	{
 		return false;
 	}
 
+	// Once we've promted the user and he accepted, if we do have a PropertyOverride edit, commit it first
+	if (PropertyOverrideEdit && !CommitLevelInstancePropertyOverridesInternal(PropertyOverrideEdit, bDiscardPropertyOverride))
+	{
+		return false;
+	}
+	check(!PropertyOverrideEdit);
+
+	// In case we committed some overrides
+	LevelInstance = GetLevelInstance(EditLevelInstanceID);
+		
 	FScopedSlowTask SlowTask(0, LOCTEXT("BeginEditLevelInstance", "Loading Level Instance for edit..."), !GetWorld()->IsGameWorld());
 	SlowTask.MakeDialog();
 
 	// Gather information from the context actor to try and select something meaningful after the loading
-	FString ActorNameToSelect = InActorNameToSelect;
-	if (AActor* ContextActor = ContextActorPtr.Get())
-	{
-		ActorNameToSelect = ContextActor->GetName();
-		ForEachLevelInstanceAncestorsAndSelf(ContextActor, [&ActorNameToSelect,LevelInstance](const ILevelInstanceInterface* AncestorLevelInstance)
-		{
-			// stop when we hit the LevelInstance we are about to edit
-			if (AncestorLevelInstance == LevelInstance)
-			{
-				return false;
-			}
-			
-			ActorNameToSelect = CastChecked<AActor>(AncestorLevelInstance)->GetName();
-			return true;
-		});
-	}
+	FString ActorNameToSelect = GetActorNameToSelectFromContext(LevelInstance, ContextActorPtr.Get(), InActorNameToSelect);
 
 	// Make sure selection is refreshed (Edit can have impact on details view)
 	GEditor->SelectNone(true, true);
@@ -2150,24 +2435,23 @@ bool ULevelInstanceSubsystem::EditLevelInstanceInternal(ILevelInstanceInterface*
 	{	
 		// Only support one level of recursion to commit current edit
 		check(!bRecursive);
-		FLevelInstanceID PendingEditId = LevelInstance->GetLevelInstanceID();
 
 		// Make sure to keep the top level instance actor loaded when we commit the current one
 		FWorldPartitionReference CurrentEditLevelInstanceActorRef = CurrentEditLevelInstanceActor;
 		
 		CommitLevelInstanceInternal(LevelInstanceEdit, bDiscard);
 
-		ILevelInstanceInterface* LevelInstanceToEdit = GetLevelInstance(PendingEditId);
+		ILevelInstanceInterface* LevelInstanceToEdit = GetLevelInstance(EditLevelInstanceID);
 		check(LevelInstanceToEdit);
 
 		return EditLevelInstanceInternal(LevelInstanceToEdit, nullptr, ActorNameToSelect, /*bRecursive=*/true);
 	}
 
 	// Cleanup async requests in case
-	LevelInstancesToUnload.Remove(LevelInstance->GetLevelInstanceID());
+	LevelInstancesToUnload.Remove(EditLevelInstanceID);
 	LevelInstancesToLoadOrUpdate.Remove(LevelInstance);
 	// Unload right away
-	UnloadLevelInstance(LevelInstance->GetLevelInstanceID());
+	UnloadLevelInstance(EditLevelInstanceID);
 	
 	// When editing a Level Instance, push a new empty actor editor context
 	UActorEditorContextSubsystem::Get()->PushContext();
@@ -2186,25 +2470,14 @@ bool ULevelInstanceSubsystem::EditLevelInstanceInternal(ILevelInstanceInterface*
 	check(LevelInstanceEdit->LevelStreaming == LevelStreaming);
 		
 	// Try and select something meaningful
-	AActor* ActorToSelect = nullptr;
-	if (!ActorNameToSelect.IsEmpty())
-	{		
-		ActorToSelect = FindObject<AActor>(LevelStreaming->GetLoadedLevel(), *ActorNameToSelect);
-	}
-
-	// default to LevelInstance
+	SelectActorFromActorName(LevelInstance, ActorNameToSelect);
+	
 	AActor* LevelInstanceActor = CastChecked<AActor>(LevelInstance);
-	if (!ActorToSelect)
-	{
-		ActorToSelect = LevelInstanceActor;
-	}
 	LevelInstanceActor->SetIsTemporarilyHiddenInEditor(false);
 
 	// Notify
 	LevelInstance->OnEdit();
 
-	GEditor->SelectActor(ActorToSelect, true, true);
-
 	for (const auto& Actor : LevelStreaming->LoadedLevel->Actors)
 	{
 		const bool bEditing = true;
@@ -2265,16 +2538,13 @@ void ULevelInstanceSubsystem::ResetLoadersForWorldAssetInternal(const FString& W
 
 bool ULevelInstanceSubsystem::CommitLevelInstance(ILevelInstanceInterface* LevelInstance, bool bDiscardEdits, TSet<FName>* DirtyPackages)
 {
-	check(LevelInstanceEdit.Get() == GetLevelInstanceEdit(LevelInstance));
-	check(CanCommitLevelInstance(LevelInstance));
-	if (CommitLevelInstanceInternal(LevelInstanceEdit, bDiscardEdits, /*bDiscardOnFailure=*/false, DirtyPackages))
+	if (GetEditingLevelInstance() == LevelInstance)
 	{
-		ILevelInstanceEditorModule& EditorModule = FModuleManager::GetModuleChecked<ILevelInstanceEditorModule>("LevelInstanceEditor");
-		EditorModule.DeactivateEditorMode();
-
-		return true;
+		check(LevelInstanceEdit.Get() == GetLevelInstanceEdit(LevelInstance));
+		check(CanCommitLevelInstance(LevelInstance));
+		return CommitLevelInstanceInternal(LevelInstanceEdit, bDiscardEdits, /*bDiscardOnFailure=*/false, DirtyPackages);
 	}
-
+	
 	return false;
 }
 
@@ -2292,6 +2562,14 @@ bool ULevelInstanceSubsystem::CommitLevelInstanceInternal(TUniquePtr<FLevelInsta
 		bDiscardEdits = false;
 	}
 
+	// Check if we need to Commit a Property Override Edit before
+	bool bChangesCommitted = false;
+	if (PropertyOverrideEdit)
+	{
+		bChangesCommitted |= CommitLevelInstancePropertyOverridesInternal(PropertyOverrideEdit, bDiscardEdits);
+		check(!PropertyOverrideEdit);
+	}
+		
 	// Build list of Packages to save
 	TSet<FName> PackagesToSave;
 
@@ -2308,9 +2586,15 @@ bool ULevelInstanceSubsystem::CommitLevelInstanceInternal(TUniquePtr<FLevelInsta
 	{
 		PackagesToSave.Append(*DirtyPackages);
 	}
-		
+	
+	const FString WorldAssetPackageStr = LevelInstance->GetWorldAssetPackage();
+	const FName WorldAssetPackage(*WorldAssetPackageStr);
+
+	// Backup ID on Commit in case Actor gets recreated
+	const FLevelInstanceID LevelInstanceID = LevelInstance->GetLevelInstanceID();
+
 	// Did some change get saved outside of the commit (regular saving in editor while editing)
-	bool bChangesCommitted = InLevelInstanceEdit->HasCommittedChanges();
+	bChangesCommitted |= InLevelInstanceEdit->HasCommittedChanges();
 	if (PackagesToSave.Num() > 0 && !bDiscardEdits)
 	{
 		const bool bPromptUserToSave = false;
@@ -2345,18 +2629,13 @@ bool ULevelInstanceSubsystem::CommitLevelInstanceInternal(TUniquePtr<FLevelInsta
 	// Make sure selection is refreshed (Commit can have impact on details view)
 	GEditor->SelectNone(true, true);
 
-	const FString EditPackage = LevelInstance->GetWorldAssetPackage();
-
 	// Remove from streaming level...
-	ResetEdit(InLevelInstanceEdit);
+	InLevelInstanceEdit.Reset();
 
 	if (bChangesCommitted)
 	{
-		ULevel::ScanLevelAssets(EditPackage);
+		ULevel::ScanLevelAssets(WorldAssetPackageStr);
 	}
-	
-	// Backup ID on Commit in case Actor gets recreated
-	const FLevelInstanceID LevelInstanceID = LevelInstance->GetLevelInstanceID();
 
 	// Notify (Actor might get destroyed by this call if its a packed bp)
 	LevelInstance->OnCommit(bChangesCommitted);
@@ -2365,7 +2644,8 @@ bool ULevelInstanceSubsystem::CommitLevelInstanceInternal(TUniquePtr<FLevelInsta
 	LevelInstance = GetLevelInstance(LevelInstanceID);
 
 	// Update Registered Container Bounds
-	UActorDescContainerSubsystem::GetChecked().NotifyContainerUpdated(*LevelInstance->GetWorldAssetPackage());
+	UActorDescContainerSubsystem::GetChecked().NotifyContainerUpdated(WorldAssetPackage);
+
 	
 	TArray<TPair<ULevelInstanceSubsystem*, FLevelInstanceID>> LevelInstancesToUpdate;
 	// Gather list to update
@@ -2376,7 +2656,7 @@ bool ULevelInstanceSubsystem::CommitLevelInstanceInternal(TUniquePtr<FLevelInsta
 		{
 			if (ULevelInstanceSubsystem* LevelInstanceSubsystem = CurrentWorld->GetSubsystem<ULevelInstanceSubsystem>())
 			{
-				TArray<ILevelInstanceInterface*> WorldLevelInstances = LevelInstanceSubsystem->GetLevelInstances(EditPackage);
+				TArray<ILevelInstanceInterface*> WorldLevelInstances = LevelInstanceSubsystem->GetLevelInstances(WorldAssetPackageStr);
 				for (ILevelInstanceInterface* CurrentLevelInstance : WorldLevelInstances)
 				{
 					if (CurrentLevelInstance == LevelInstance || bChangesCommitted)
@@ -2427,7 +2707,7 @@ bool ULevelInstanceSubsystem::CommitLevelInstanceInternal(TUniquePtr<FLevelInsta
 	// Send out Event if changes were committed
 	if (bChangesCommitted)
 	{
-		LevelInstanceChangedEvent.Broadcast(FName(*EditPackage));
+		LevelInstanceChangedEvent.Broadcast(WorldAssetPackage);
 	}
 
 	GEngine->BroadcastLevelActorListChanged();
@@ -2494,6 +2774,20 @@ bool ULevelInstanceSubsystem::HasParentEdit(const ILevelInstanceInterface* Level
 	return bResult;
 }
 
+bool ULevelInstanceSubsystem::HasParentPropertyOverridesEdit(const ILevelInstanceInterface* LevelInstance) const
+{
+	bool bResult = false;
+
+	const AActor* LevelInstanceActor = CastChecked<AActor>(LevelInstance);
+	ForEachLevelInstanceAncestors(LevelInstanceActor, [LevelInstance, &bResult](const ILevelInstanceInterface* Ancestor)
+	{
+		bResult = Ancestor->IsEditingPropertyOverrides();
+		return !bResult;
+	});
+
+	return bResult;
+}
+
 void ULevelInstanceSubsystem::OnCommitChild(const FLevelInstanceID& LevelInstanceID, bool bChildChanged)
 {
 	int32& ChildEditCount = ChildEdits.FindChecked(LevelInstanceID);
@@ -2520,7 +2814,7 @@ void ULevelInstanceSubsystem::OnEditChild(const FLevelInstanceID& LevelInstanceI
 	}
 }
 
-TArray<ILevelInstanceInterface*> ULevelInstanceSubsystem::GetLevelInstances(const FString& WorldAssetPackage)
+TArray<ILevelInstanceInterface*> ULevelInstanceSubsystem::GetLevelInstances(const FString& WorldAssetPackage) const
 {
 	TArray<ILevelInstanceInterface*> MatchingLevelInstances;
 	for (const auto& KeyValuePair : RegisteredLevelInstances)
@@ -2534,6 +2828,20 @@ TArray<ILevelInstanceInterface*> ULevelInstanceSubsystem::GetLevelInstances(cons
 	return MatchingLevelInstances;
 }
 
+TArray<ILevelInstanceInterface*> ULevelInstanceSubsystem::GetLevelInstances(const TSoftObjectPtr<ULevelInstancePropertyOverrideAsset>& PropertyOverrideAsset) const
+{
+	TArray<ILevelInstanceInterface*> MatchingLevelInstances;
+	for (const auto& KeyValuePair : RegisteredLevelInstances)
+	{
+		if (ULevelInstancePropertyOverrideAsset* PropertyOverride = KeyValuePair.Value->GetPropertyOverrideAsset(); PropertyOverride && PropertyOverride->GetSourceAssetPtr() == PropertyOverrideAsset)
+		{
+			MatchingLevelInstances.Add(KeyValuePair.Value);
+		}
+	}
+
+	return MatchingLevelInstances;
+}
+
 void ULevelInstanceSubsystem::ForEachLevelInstanceActorAncestors(const ULevel* Level, TFunctionRef<bool(AActor*)> Operation) const
 {
 	AActor* CurrentActor = Cast<AActor>(GetOwningLevelInstance(Level));
@@ -2607,6 +2915,408 @@ bool ULevelInstanceSubsystem::PassLevelInstanceFilter(UWorld* World, const FWorl
 	return true;
 }
 
+void ULevelInstanceSubsystem::EditLevelInstancePropertyOverrides(ILevelInstanceInterface* LevelInstance, AActor* ContextActor)
+{
+	if (!CanEditLevelInstancePropertyOverrides(LevelInstance))
+	{
+		return;
+	}
+
+	const FLevelInstanceID LevelInstanceID = LevelInstance->GetLevelInstanceID();
+
+	bool bDiscard = false;
+	bool bCommitEdit = false;
+	// Not in same hierarchy, prompt user to commit edit first
+	if (LevelInstanceEdit && !LevelInstance->HasParentEdit())
+	{
+		if (!PromptUserForCommit(LevelInstanceEdit.Get(), bDiscard))
+		{
+			return;
+		}
+		// Edit isn't our parent so commit it
+		bCommitEdit = true;
+	}
+
+	// Gather information from the context actor to try and select something meaningful after the loading
+	FString ActorNameToSelect = GetActorNameToSelectFromContext(LevelInstance, ContextActor);
+
+	bool bDiscardPropertyOverride = false;
+	if (PropertyOverrideEdit && !PromptUserForCommitPropertyOverrides(PropertyOverrideEdit.Get(), bDiscardPropertyOverride))
+	{
+		return;
+	}
+
+	if (PropertyOverrideEdit && !CommitLevelInstancePropertyOverridesInternal(PropertyOverrideEdit, bDiscardPropertyOverride))
+	{
+		return;
+	}
+
+	if (bCommitEdit && LevelInstanceEdit && !CommitLevelInstanceInternal(LevelInstanceEdit, bDiscard))
+	{
+		return;
+	}
+
+	// Make sure selection is refreshed (Edit can have impact on details view)
+	GEditor->SelectNone(true, true);
+
+	// Cleanup async requests in case
+	LevelInstancesToUnload.Remove(LevelInstanceID);
+	LevelInstancesToLoadOrUpdate.Remove(GetLevelInstance(LevelInstanceID));
+	// Unload right away
+	UnloadLevelInstance(LevelInstanceID);
+
+	// When editing a Level Instance, push a new empty actor editor context
+	UActorEditorContextSubsystem::Get()->PushContext();
+
+	// Load Edit LevelInstance level
+	ULevelStreamingLevelInstanceEditorPropertyOverride::Load(GetLevelInstance(LevelInstanceID));
+
+	// Select Context Actor
+	SelectActorFromActorName(GetLevelInstance(LevelInstanceID), ActorNameToSelect);
+}
+
+bool ULevelInstanceSubsystem::CanResetPropertyOverridesForActor(AActor* Actor) const
+{
+	// Resetting individual Actor overrides requires that we are currently in Property Override Edit
+	if (!PropertyOverrideEdit)
+	{
+		return false;
+	}
+
+	// And that the Actor's parent level instance is the current property override edit
+	ILevelInstanceInterface* ParentLevelInstance = GetParentLevelInstance(Actor);
+	if (ParentLevelInstance != PropertyOverrideEdit->GetLevelInstance())
+	{
+		return false;
+	}
+
+	return true;
+}
+
+void ULevelInstanceSubsystem::ResetPropertyOverridesForActor(AActor* Actor)
+{
+	if (!CanResetPropertyOverridesForActor(Actor))
+	{
+		return;
+	}
+	PropertyOverrideEdit->GetLevelInstance()->GetPropertyOverrideAsset()->ResetPropertyOverridesForActor(PropertyOverrideEdit->LevelStreaming, Actor);
+}
+
+bool ULevelInstanceSubsystem::IsEditingLevelInstancePropertyOverrides(const ILevelInstanceInterface* LevelInstance) const
+{
+	return PropertyOverrideEdit && PropertyOverrideEdit->GetLevelInstance() == LevelInstance;
+}
+
+bool ULevelInstanceSubsystem::CommitLevelInstancePropertyOverrides(ILevelInstanceInterface* LevelInstance, bool bDiscardEdits)
+{
+	if (GetEditingPropertyOverridesLevelInstance() == LevelInstance)
+	{
+		check(PropertyOverrideEdit.Get() == GetLevelInstancePropertyOverrideEdit(LevelInstance));
+		check(CanCommitLevelInstancePropertyOverrides(LevelInstance));
+		return CommitLevelInstancePropertyOverridesInternal(PropertyOverrideEdit, bDiscardEdits);
+	}
+
+	return false;
+}
+
+bool ULevelInstanceSubsystem::CommitLevelInstancePropertyOverridesInternal(TUniquePtr<FPropertyOverrideEdit>& InPropertyOverrideEdit, bool bDiscardEdits)
+{
+	ILevelInstanceInterface* LevelInstance = InPropertyOverrideEdit->GetLevelInstance();
+	ILevelInstanceInterface* LevelInstanceWithOverrides = GetLevelInstancePropertyOverridesEditOwner(LevelInstance);
+	TSoftObjectPtr<ULevelInstancePropertyOverrideAsset> PreviousOverrideAsset = LevelInstanceWithOverrides->GetPropertyOverrideAsset() ? LevelInstanceWithOverrides->GetPropertyOverrideAsset()->GetSourceAssetPtr() : nullptr;
+
+	bool bSaved = false;
+	if (InPropertyOverrideEdit->IsDirty() && !bDiscardEdits)
+	{
+		bSaved = InPropertyOverrideEdit->Save(LevelInstanceWithOverrides);
+	}
+
+	// Make sure selection is refreshed (Commit can have impact on details view)
+	GEditor->SelectNone(true, true);
+
+	// Restore actor editor context
+	UActorEditorContextSubsystem::Get()->PopContext();
+		
+	InPropertyOverrideEdit.Reset();
+
+	// Update Level Instance (saved or not it needs to reload)
+	LevelInstanceWithOverrides->UpdateLevelInstanceFromWorldAsset();
+
+	// If it was saved reload other affected level instances
+	if (bSaved)
+	{
+		UpdateLevelInstancesFromPropertyOverrideAsset(PreviousOverrideAsset, LevelInstanceWithOverrides->GetPropertyOverrideAsset());
+	}
+
+	BlockOnLoading();
+			
+	return true;
+}
+
+bool ULevelInstanceSubsystem::CanResetPropertyOverrides(ILevelInstanceInterface* LevelInstance) const
+{
+	// Don't allow reset of the level instances property overrides if there is a current Property Override Edit in progress
+	if (PropertyOverrideEdit)
+	{
+		return false;
+	}
+
+	if (!LevelInstance->GetPropertyOverrideAsset())
+	{
+		return false;
+	}
+		
+	// Allow reset of the overrides if this level instance is not parented or if it's parent is currently in edit
+	ILevelInstanceInterface* ParentLevelInstance = GetParentLevelInstance(CastChecked<AActor>(LevelInstance));
+	return !ParentLevelInstance || ParentLevelInstance->IsEditing();
+}
+
+void ULevelInstanceSubsystem::ResetPropertyOverrides(ILevelInstanceInterface* LevelInstance)
+{
+	if (!CanResetPropertyOverrides(LevelInstance))
+	{
+		return;
+	}
+
+	const TSoftObjectPtr<ULevelInstancePropertyOverrideAsset> AssetPath = LevelInstance->GetPropertyOverrideAsset()->GetSourceAssetPtr();
+	ULevelInstancePropertyOverrideAsset* PropertyOverrideAsset = LevelInstance->GetPropertyOverrideAsset();
+	
+	AActor* LevelInstanceActor = CastChecked<AActor>(LevelInstance);
+	const bool bWasDirty = LevelInstanceActor->GetPackage()->IsDirty();
+
+	LevelInstance->SetPropertyOverrideAsset(nullptr);
+
+	if (FEditorFileUtils::PromptForCheckoutAndSave({ LevelInstanceActor->GetPackage() }, /*bCheckDirty*/false, /*bPromptToSave*/false) != FEditorFileUtils::PR_Success)
+	{
+		LevelInstance->SetPropertyOverrideAsset(PropertyOverrideAsset);
+		LevelInstanceActor->GetPackage()->SetDirtyFlag(bWasDirty);
+		return;
+	}
+
+	LevelInstance->UpdateLevelInstanceFromWorldAsset();
+	UpdateLevelInstancesFromPropertyOverrideAsset(AssetPath, nullptr);
+	BlockOnLoading();
+	
+	// Make sure selection is refreshed (Reset can have impact on details view)
+	GEditor->SelectNone(true, true);
+}
+
+void ULevelInstanceSubsystem::UpdateLevelInstancesFromPropertyOverrideAsset(const TSoftObjectPtr<ULevelInstancePropertyOverrideAsset>& PreviousAssetPath, ULevelInstancePropertyOverrideAsset* NewAsset)
+{
+	if (PreviousAssetPath.IsNull())
+	{
+		return;
+	}
+
+	TArray<TPair<ULevelInstanceSubsystem*, FLevelInstanceID>> LevelInstancesToUpdate;
+	// Gather list to update
+	for (TObjectIterator<UWorld> It(RF_ClassDefaultObject | RF_ArchetypeObject, true); It; ++It)
+	{
+		UWorld* CurrentWorld = *It;
+		if (IsValid(CurrentWorld))
+		{
+			if (ULevelInstanceSubsystem* LevelInstanceSubsystem = CurrentWorld->GetSubsystem<ULevelInstanceSubsystem>())
+			{
+				TArray<ILevelInstanceInterface*> WorldLevelInstances = LevelInstanceSubsystem->GetLevelInstances(PreviousAssetPath);
+				for (ILevelInstanceInterface* CurrentLevelInstance : WorldLevelInstances)
+				{
+					AActor* CurrentActor = CastChecked<AActor>(CurrentLevelInstance);
+
+					// We are updating a instanced property override so we need to patchup the other instance's property override assets so that reload can be done properly
+					ULevelInstancePropertyOverrideAsset* CurrentAsset = CurrentLevelInstance->GetPropertyOverrideAsset();
+					check(CurrentAsset && CurrentAsset->GetSourceAssetPtr() == PreviousAssetPath);
+
+					// This is always the case for now but check anyways in case we do end up supporting public override assets
+					// Public assets wouldn't need to be patched up
+					if (CurrentAsset->IsInOuter(CurrentActor))
+					{
+						// If there is a new property override duplicate it into the level instance
+						ULevelInstancePropertyOverrideAsset* NewAssetCopy = nullptr;
+						if (NewAsset)
+						{	
+							NewAssetCopy = CastChecked<ULevelInstancePropertyOverrideAsset>(StaticDuplicateObject(NewAsset, CurrentActor, NewAsset->GetFName(), CurrentActor->GetFlags()));
+						}
+						// Set new duplicated asset or nullptr
+						CurrentLevelInstance->SetPropertyOverrideAsset(NewAssetCopy);
+					}
+
+					// Request update for level instance now that its property override asset has been patched up
+					LevelInstancesToUpdate.Add({ LevelInstanceSubsystem, CurrentLevelInstance->GetLevelInstanceID() });
+				}
+			}
+		}
+	}
+
+	// Do update
+	for (const auto& KeyValuePair : LevelInstancesToUpdate)
+	{
+		if (ILevelInstanceInterface* LevelInstanceToUpdate = KeyValuePair.Key->GetLevelInstance(KeyValuePair.Value))
+		{
+			LevelInstanceToUpdate->UpdateLevelInstanceFromWorldAsset();
+		}
+	}
+}
+
+ILevelInstanceInterface* ULevelInstanceSubsystem::GetEditingPropertyOverridesLevelInstance() const
+{
+	if (PropertyOverrideEdit)
+	{
+		return PropertyOverrideEdit->GetLevelInstance();
+	}
+
+	return nullptr;
+}
+
+FActorContainerID ULevelInstanceSubsystem::GetLevelInstancePropertyOverridesContext(ILevelInstanceInterface* LevelInstance) const
+{
+	// Default to Main container
+	FActorContainerID Context;
+	check(Context.IsMainContainer());
+
+	// Return the top level ContainerID for which we want to apply properties from.
+	// By default we apply everything (main container) but if an ancestor is being edited we want to apply up to the edited ancestor (same as opening that ancestor level)
+	ForEachLevelInstanceAncestors(CastChecked<AActor>(LevelInstance), [&Context](ILevelInstanceInterface* Ancestor)
+	{
+		// If Ancestor of current LevelInstanceOverrideEditOwner is being edited then we stop because we've found the override owner
+		if (Ancestor->IsEditing())
+		{
+			Context = Ancestor->GetLevelInstanceID().GetContainerID();
+			return false;
+		}
+		return true;
+	});
+
+	return Context;
+}
+
+ILevelInstanceInterface* ULevelInstanceSubsystem::GetLevelInstancePropertyOverridesEditOwner(ILevelInstanceInterface* LevelInstance) const
+{
+	ILevelInstanceInterface* LevelInstanceOverrideEditOwner = LevelInstance;
+
+	// Find where to save the overrides
+	ForEachLevelInstanceAncestorsAndSelf(CastChecked<AActor>(LevelInstance), [&LevelInstanceOverrideEditOwner](ILevelInstanceInterface* Ancestor)
+	{
+		// If Ancestor of current LevelInstanceOverrideEditOwner is being edited then we stop because we've found the override owner
+		if (Ancestor->IsEditing())
+		{
+			return false;
+		}
+
+		LevelInstanceOverrideEditOwner = Ancestor;
+		return true;
+	});
+
+	return LevelInstanceOverrideEditOwner;
+}
+
+const ILevelInstanceInterface* ULevelInstanceSubsystem::GetLevelInstancePropertyOverridesEditOwner(const ILevelInstanceInterface* LevelInstance) const
+{
+	return GetLevelInstancePropertyOverridesEditOwner(const_cast<ILevelInstanceInterface*>(LevelInstance));
+}
+
+bool ULevelInstanceSubsystem::GetLevelInstancePropertyOverridesForActor(const AActor* Actor, FActorContainerID PropertyOverrideContext, TArray<const FActorPropertyOverride*>& OutPropertyOverrides) const
+{
+	if (ILevelInstanceInterface* OwningLevelInstance = GetParentLevelInstance(Actor))
+	{
+		if (ULevel* Level = GetLevelInstanceLevel(OwningLevelInstance))
+		{
+			if (UWorldPartition* WorldPartition = Level->GetWorldPartition())
+			{
+				if (ULevelInstanceContainerInstance* LevelInstanceContainerInstance = Cast<ULevelInstanceContainerInstance>(WorldPartition->FActorDescContainerInstanceCollection::GetActorDescContainerInstance(Actor->GetActorGuid())))
+				{
+					LevelInstanceContainerInstance->GetPropertyOverridesForActor(OwningLevelInstance->GetLevelInstanceID().GetContainerID(), PropertyOverrideContext, Actor->GetActorGuid(), OutPropertyOverrides);
+					return !OutPropertyOverrides.IsEmpty();
+				}
+			}
+		}
+	}
+
+	return false;
+}
+
+ULevelInstanceSubsystem::FPropertyOverrideEdit::FPropertyOverrideEdit(ULevelStreamingLevelInstanceEditorPropertyOverride* InLevelStreaming)
+	: LevelStreaming(InLevelStreaming)
+{
+}
+
+ULevelInstanceSubsystem::FPropertyOverrideEdit::~FPropertyOverrideEdit()
+{
+	ULevelStreamingLevelInstanceEditorPropertyOverride::Unload(LevelStreaming);
+}
+
+ILevelInstanceInterface* ULevelInstanceSubsystem::FPropertyOverrideEdit::GetLevelInstance() const
+{
+	return LevelStreaming->GetLevelInstance();
+}
+
+
+void ULevelInstanceSubsystem::RegisterLoadedLevelStreamingPropertyOverride(ULevelStreamingLevelInstanceEditorPropertyOverride* LevelStreaming)
+{
+	check(!PropertyOverrideEdit.IsValid());
+	PropertyOverrideEdit = MakeUnique<FPropertyOverrideEdit>(LevelStreaming);
+}
+
+bool ULevelInstanceSubsystem::FPropertyOverrideEdit::IsDirty() const
+{
+	if (ULevel* LoadedLevel = LevelStreaming->GetLoadedLevel())
+	{
+		for (AActor* Actor : LoadedLevel->Actors)
+		{
+			if (IsValid(Actor) && Actor->GetPackage()->IsDirty())
+			{
+				return true;
+			}
+		}
+	}
+	
+	return false;
+}
+
+bool ULevelInstanceSubsystem::FPropertyOverrideEdit::Save(ILevelInstanceInterface* LevelInstanceOverrideOwner) const
+{
+	check(LevelInstanceOverrideOwner);
+			
+	AActor* OverrideOwnerActor = CastChecked<AActor>(LevelInstanceOverrideOwner);
+	check(OverrideOwnerActor->IsPackageExternal());
+	const bool bWasDirty = OverrideOwnerActor->GetPackage()->IsDirty();
+
+	// Create a unique name so that we can update other instances and be sure that their patched PropertyOverrides can have same name
+	const FGuid PropertyOverrideGuid = FGuid::NewGuid();
+	const FName PropertyOverrideName(*FString::Format(TEXT("PropertyOverride_{0}"), { PropertyOverrideGuid.ToString() }));
+	
+	ULevelInstancePropertyOverrideAsset* NewPropertyOverride = nullptr;
+	ULevelInstancePropertyOverrideAsset* ExistingPropertyOverride = LevelInstanceOverrideOwner->GetPropertyOverrideAsset();
+	if (ExistingPropertyOverride)
+	{
+		// Duplicate previous override object as it contains the overrides that we aren't currently editing that we don't want to lose
+		NewPropertyOverride = CastChecked<ULevelInstancePropertyOverrideAsset>(StaticDuplicateObject(ExistingPropertyOverride, OverrideOwnerActor, PropertyOverrideName, OverrideOwnerActor->GetFlags()));
+	}
+	else
+	{
+		NewPropertyOverride = NewObject<ULevelInstancePropertyOverrideAsset>(OverrideOwnerActor, PropertyOverrideName, OverrideOwnerActor->GetFlags());
+		NewPropertyOverride->Initialize(LevelInstanceOverrideOwner->GetWorldAsset());
+	}
+	check(NewPropertyOverride->GetWorldAsset() == LevelInstanceOverrideOwner->GetWorldAsset());
+			
+	// Serialize current edit overrides
+	NewPropertyOverride->SerializePropertyOverrides(LevelInstanceOverrideOwner, LevelStreaming);
+			
+	// Set new property override before saving
+	LevelInstanceOverrideOwner->SetPropertyOverrideAsset(NewPropertyOverride);
+
+	// This will be used when initializing the ActorDesc to distinguish between saving the LevelInstance actor normally or through Property override save
+	TGuardValue<bool> SavingGuard(NewPropertyOverride->bSavingOverrideEdit, true);
+	if (FEditorFileUtils::PromptForCheckoutAndSave({ OverrideOwnerActor->GetPackage() }, /*bCheckDirty*/false, /*bPromptToSave*/false) != FEditorFileUtils::PR_Success)
+	{
+		// Save failed, set property override to previous (could be null)
+		LevelInstanceOverrideOwner->SetPropertyOverrideAsset(ExistingPropertyOverride);
+		OverrideOwnerActor->GetPackage()->SetDirtyFlag(bWasDirty);
+		return false;
+	}	
+
+	return true;
+}
+
 #endif
 
 #undef LOCTEXT_NAMESPACE
diff --git a/Engine/Source/Runtime/Engine/Private/Tests/WorldPartition/WorldPartitionHandleTests.cpp b/Engine/Source/Runtime/Engine/Private/Tests/WorldPartition/WorldPartitionHandleTests.cpp
--- a/Engine/Source/Runtime/Engine/Private/Tests/WorldPartition/WorldPartitionHandleTests.cpp
+++ b/Engine/Source/Runtime/Engine/Private/Tests/WorldPartition/WorldPartitionHandleTests.cpp
@@ -329,7 +329,7 @@ namespace WorldPartitionTests
 					.SetPackageName(Iterator->GetActorPackage())
 					.SetActorPath(Iterator->GetActorSoftPath());
 			
-				Iterator->GetActorDesc()->SerializeTo(ActorDescInitData.SerializedData);
+				Iterator->GetActorDesc()->SerializeTo(ActorDescInitData.GetSerializedData());
 				NewActorDesc->Init(ActorDescInitData);
 
 				TestTrue(TEXT("Actor Descriptor Serialization"), NewActorDesc->Equals(Iterator->GetActorDesc()));
diff --git a/Engine/Source/Runtime/Engine/Private/WorldPartition/ActorDescContainer.cpp b/Engine/Source/Runtime/Engine/Private/WorldPartition/ActorDescContainer.cpp
--- a/Engine/Source/Runtime/Engine/Private/WorldPartition/ActorDescContainer.cpp
+++ b/Engine/Source/Runtime/Engine/Private/WorldPartition/ActorDescContainer.cpp
@@ -33,8 +33,12 @@ UActorDescContainer::UActorDescContainer(const FObjectInitializer& ObjectInitial
 void UActorDescContainer::Initialize(const FInitializeParams& InitParams)
 {
 	TRACE_CPUPROFILER_EVENT_SCOPE(UActorDescContainer::Initialize);
-
 	check(!bContainerInitialized);
+	if (InitParams.PreInitialize)
+	{
+		InitParams.PreInitialize(this);
+	}
+
 	ContainerPackageName = InitParams.PackageName;
 	if (InitParams.ExternalDataLayerAsset)
 	{
@@ -210,7 +214,7 @@ bool UActorDescContainer::HasExternalContent() const
 	return ExternalDataLayerAsset ? true : GetContentBundleGuid().IsValid();
 }
 
-bool UActorDescContainer::IsActorDescHandled(const AActor* InActor) const
+bool UActorDescContainer::IsActorDescHandled(const AActor* InActor, bool bInUseLoadedPath) const
 {
 	// Actor External Content Guid must match Container's External Content Guid to be considered
 	// AWorldDataLayers actors are an exception as they don't have an External Content Guid
@@ -218,16 +222,22 @@ bool UActorDescContainer::IsActorDescHandled(const AActor* InActor) const
 		(!HasExternalContent() && !InActor->HasExternalContent()) ||
 		(ExternalDataLayerAsset && (ExternalDataLayerAsset == InActor->GetExternalDataLayerAsset())) ||
 		(ContentBundleGuid.IsValid() && (ContentBundleGuid == InActor->GetContentBundleGuid()));
-	
+
 	if (bIsCandidateActor)
 	{
-		const FString ActorPackageName = InActor->GetPackage()->GetName();
+		const FName LoadedPackageName = InActor->GetPackage()->GetLoadedPath().GetPackageFName();
+		const FString ActorPackageName = bInUseLoadedPath && !LoadedPackageName.IsNone() ? LoadedPackageName.ToString() : InActor->GetPackage()->GetName();
 		const FString ExternalActorPath = GetExternalActorPath() / TEXT("");
 		return ActorPackageName.StartsWith(ExternalActorPath);
 	}
 	return false;
 }
 
+bool UActorDescContainer::IsActorDescHandled(const AActor* InActor) const
+{
+	return IsActorDescHandled(InActor, false);
+}
+
 void UActorDescContainer::RegisterActorDescriptor(FWorldPartitionActorDesc* ActorDesc)
 {
 	FActorDescList::AddActorDescriptor(ActorDesc);
@@ -244,9 +254,14 @@ void UActorDescContainer::UnregisterActorDescriptor(FWorldPartitionActorDesc* Ac
 	verifyf(ActorsByName.Remove(ActorDesc->GetActorName()), TEXT("Missing actor '%s' from container '%s'"), *ActorDesc->GetActorName().ToString(), *ContainerPackageName.ToString());
 }
 
+bool UActorDescContainer::ShouldHandleActorEvent(const AActor* Actor, bool bInUseLoadedPath) const
+{
+	return Actor && IsActorDescHandled(Actor, bInUseLoadedPath) && Actor->IsMainPackageActor() && Actor->GetLevel();
+}
+
 bool UActorDescContainer::ShouldHandleActorEvent(const AActor* Actor)
 {
-	return Actor && IsActorDescHandled(Actor) && Actor->IsMainPackageActor() && Actor->GetLevel();
+	return ShouldHandleActorEvent(Actor, false);
 }
 
 const FWorldPartitionActorDesc* UActorDescContainer::GetActorDescByPath(const FString& ActorPath) const
@@ -376,7 +391,7 @@ bool UActorDescContainer::RemoveActor(const FGuid& ActorGuid)
 	return false;
 }
 
-bool UActorDescContainer::ShouldRegisterDelegates()
+bool UActorDescContainer::ShouldRegisterDelegates() const
 {
 	return GEditor && !IsTemplate() && !IsRunningCookCommandlet();
 }
diff --git a/Engine/Source/Runtime/Engine/Private/WorldPartition/ActorDescContainerInstance.cpp b/Engine/Source/Runtime/Engine/Private/WorldPartition/ActorDescContainerInstance.cpp
--- a/Engine/Source/Runtime/Engine/Private/WorldPartition/ActorDescContainerInstance.cpp
+++ b/Engine/Source/Runtime/Engine/Private/WorldPartition/ActorDescContainerInstance.cpp
@@ -49,23 +49,6 @@ void UActorDescContainerInstance::UnregisterContainer()
 	Container = nullptr;
 }
 
-void UActorDescContainerInstance::OnContainerUpdated(FName ContainerPackage)
-{
-	if (ChildContainerInstances.Num())
-	{
-		TMap<FGuid, TObjectPtr<UActorDescContainerInstance>> CopyChildContainerInstances(ChildContainerInstances);
-		for (auto& [ContainerGuid, ContainerInstance] : CopyChildContainerInstances)
-		{
-			if (ContainerInstance->GetContainerPackage() == ContainerPackage)
-			{
-				FWorldPartitionActorDescInstance* ContainerDescInstance = GetActorDescInstance(ContainerGuid);
-				check(ContainerDescInstance);
-				ContainerDescInstance->UpdateChildContainerInstance();
-			}
-		}
-	}
-}
-
 void UActorDescContainerInstance::OnContainerReplaced(UActorDescContainer* InOldContainer, UActorDescContainer* InNewContainer)
 {
 	if (ChildContainerInstances.Num())
@@ -96,6 +79,7 @@ void UActorDescContainerInstance::Initialize(const FInitializeParams& InParams)
 		// It is possible to not have a ParentContainerInstance if we are in a non-WP main world
 		// We want this ContainerInstance to still not have a Main ContainerID to properly handle IsMainWorldOnly() actors
 		ContainerID = FActorContainerID(InParams.ParentContainerInstance ? InParams.ParentContainerInstance->GetContainerID() : FActorContainerID(), InParams.ContainerActorGuid);
+		ParentContainerInstance = InParams.ParentContainerInstance;
 	}
 				
 	// Only consider world if we are outered to a WorldPartition directly
@@ -150,7 +134,7 @@ void UActorDescContainerInstance::Initialize(const FInitializeParams& InParams)
 	// Create ActorDescInstances
 	for (FActorDescList::TIterator<> It(Container); It; ++It)
 	{
-		FWorldPartitionActorDescInstance ActorDescInstance(this, *It);
+		FWorldPartitionActorDescInstance ActorDescInstance = CreateActorDescInstance(*It);
 		if (bIsInstanced)
 		{
 			const FString LongActorPackageName = It->GetActorPackage().ToString();
@@ -246,7 +230,6 @@ void UActorDescContainerInstance::RegisterDelegates()
 		// Only listen to this event if we have registered Child Container Instances
 		if (bCreateChildContainerHierarchy)
 		{
-			UActorDescContainerSubsystem::GetChecked().ContainerUpdated().AddUObject(this, &UActorDescContainerInstance::OnContainerUpdated);
 			UActorDescContainerSubsystem::GetChecked().ContainerReplaced().AddUObject(this, &UActorDescContainerInstance::OnContainerReplaced);
 		}
 
@@ -279,7 +262,6 @@ void UActorDescContainerInstance::UnregisterDelegates()
 
 		if (UActorDescContainerSubsystem* ActorDescContainerSubsystem = UActorDescContainerSubsystem::Get())
 		{
-			ActorDescContainerSubsystem->ContainerUpdated().RemoveAll(this);
 			ActorDescContainerSubsystem->ContainerReplaced().RemoveAll(this);
 		}
 	}
@@ -414,7 +396,7 @@ FWorldPartitionActorDescInstance* UActorDescContainerInstance::AddActor(FWorldPa
 	// We don't support adding actors when instanced
 	check(!InstancingContext.IsSet());
 
-	return AddActorDescInstance(FWorldPartitionActorDescInstance(this, InActorDesc));
+	return AddActorDescInstance(CreateActorDescInstance(InActorDesc));
 }
 
 FWorldPartitionActorDescInstance* UActorDescContainerInstance::AddActorDescInstance(FWorldPartitionActorDescInstance&& InActorDescInstance)
@@ -513,12 +495,15 @@ void UActorDescContainerInstance::OnObjectsReplaced(const TMap<UObject*, UObject
 	{
 		if (AActor* OldActor = Cast<AActor>(OldObject))
 		{
-			if (Container->ShouldHandleActorEvent(OldActor))
+			if (Container->ShouldHandleActorEvent(OldActor, InstancingContext.GetPtrOrNull() != nullptr))
 			{
-				AActor* NewActor = Cast<AActor>(NewObject);
-				if (FWorldPartitionActorDescInstance* ActorDescInstance = GetActorDescInstance(OldActor->GetActorGuid()))
+				// Only replace if OldActor matches ActorDescInstance ActorPtr.
+				// No need to replace the ActorDescInstance ActorPtr if it isn't set, it will get refreshed on its next GetActor call. This avoids replacing pointers on the wrong ActorDescInstance for Instanced worlds
+				if (FWorldPartitionActorDescInstance* ActorDescInstance = GetActorDescInstance(OldActor->GetActorGuid()); ActorDescInstance && ActorDescInstance->ActorPtr == OldActor)
 				{
+					AActor* NewActor = Cast<AActor>(NewObject);
 					FWorldPartitionActorDescUtils::ReplaceActorDescriptorPointerFromActor(OldActor, NewActor, ActorDescInstance);
+					OnActorReplacedEvent.Broadcast(ActorDescInstance);
 				}
 			}
 		}
@@ -552,6 +537,11 @@ void UActorDescContainerInstance::OnActorDescUpdating(FWorldPartitionActorDesc*
 	TUniquePtr<FWorldPartitionActorDescInstance>* ActorDescInstance = GetActorDescInstancePtr(InActorDesc->GetGuid());
 	check(ActorDescInstance && ActorDescInstance->IsValid());
 		
+	if (bCreateChildContainerHierarchy && ActorDescInstance->Get()->IsChildContainerInstance())
+	{
+		ActorDescInstance->Get()->UnregisterChildContainerInstance();
+	}
+
 	if (UWorldPartition* WorldPartition = GetOuterWorldPartition())
 	{
 		WorldPartition->OnActorDescInstanceUpdating(ActorDescInstance->Get());
@@ -571,7 +561,7 @@ void UActorDescContainerInstance::OnActorDescUpdated(FWorldPartitionActorDesc* I
 	// Re-register container
 	if (bCreateChildContainerHierarchy && ActorDescInstance->Get()->IsChildContainerInstance())
 	{
-		ActorDescInstance->Get()->UpdateChildContainerInstance();
+		ActorDescInstance->Get()->RegisterChildContainerInstance();
 	}
 
 	if (UWorldPartition* WorldPartition = GetOuterWorldPartition())
diff --git a/Engine/Source/Runtime/Engine/Private/WorldPartition/LevelInstance/LevelInstanceActorDesc.cpp b/Engine/Source/Runtime/Engine/Private/WorldPartition/LevelInstance/LevelInstanceActorDesc.cpp
--- a/Engine/Source/Runtime/Engine/Private/WorldPartition/LevelInstance/LevelInstanceActorDesc.cpp
+++ b/Engine/Source/Runtime/Engine/Private/WorldPartition/LevelInstance/LevelInstanceActorDesc.cpp
@@ -9,6 +9,8 @@
 #include "LevelInstance/LevelInstanceInterface.h"
 #include "LevelInstance/LevelInstanceSubsystem.h"
 #include "LevelInstance/LevelInstanceComponent.h"
+#include "LevelInstance/LevelInstancePropertyOverrideAsset.h"
+#include "LevelInstance/LevelInstanceSettings.h"
 #include "Misc/PackageName.h"
 #include "WorldPartition/ActorDescContainerInstance.h"
 #include "WorldPartition/WorldPartitionActorDescInstance.h"
@@ -18,10 +20,13 @@
 #include "WorldPartition/WorldPartitionHelpers.h"
 #include "WorldPartition/WorldPartition.h"
 #include "WorldPartition/ActorDescContainerSubsystem.h"
+#include "WorldPartition/LevelInstance/LevelInstanceContainerInstance.h"
+#include "WorldPartition/LevelInstance/LevelInstancePropertyOverrideContainer.h"
 #include "UObject/UE5ReleaseStreamObjectVersion.h"
 #include "UObject/UE5MainStreamObjectVersion.h"
 #include "UObject/FortniteSeasonBranchObjectVersion.h"
 #include "UObject/FortniteMainBranchObjectVersion.h"
+#include "AssetRegistry/AssetRegistryHelpers.h"
 
 static int32 GLevelInstanceDebugForceLevelStreaming = 0;
 static FAutoConsoleVariableRef CVarForceLevelStreaming(
@@ -54,6 +59,40 @@ void FLevelInstanceActorDesc::Init(const AActor* InActor)
 	Filter = LevelInstance->GetFilter();
 	
 	bIsChildContainerInstance = IsChildContainerInstanceInternal();
+
+	if (LevelInstance->SupportsPropertyOverrides())
+	{
+		if (ULevelInstancePropertyOverrideAsset* Asset = LevelInstance->GetPropertyOverrideAsset())
+		{
+			FLevelInstancePropertyOverrideDesc* ExistingOverrideDesc = nullptr;
+			
+			// PropertyOverrides can be nested and we can only override properties one level instance at a time.
+			// So when we do save an override, we need to keep the previous actors descs that weren't part of the current property override edit
+			// So we find the existing ActorDesc and then check if it has an existing LevelInstanceOverrideActorDesc and transfer from it
+			if (UWorldPartition* WorldPartition = FWorldPartitionHelpers::GetWorldPartition(InActor))
+			{
+				if (FWorldPartitionActorDescInstance* ActorDescInstance = WorldPartition->GetActorDescInstance(InActor->GetActorGuid()))
+				{
+					if (FLevelInstanceActorDesc* ExistingActorDesc = (FLevelInstanceActorDesc*)ActorDescInstance->GetActorDesc())
+					{
+						ExistingOverrideDesc = ExistingActorDesc->GetOverrideDesc();
+					}
+				}
+			}
+
+			PropertyOverrideAsset = FSoftObjectPath(Asset);
+			PropertyOverrideDesc = MakeShared<FLevelInstancePropertyOverrideDesc>();
+			PropertyOverrideDesc->Init(Asset);
+
+			// Transfer from previous override if we are sharing the same world package
+			if (ExistingOverrideDesc && ExistingOverrideDesc->GetWorldPackage() == PropertyOverrideDesc->GetWorldPackage())
+			{
+				// We always need to transfer from here, the reason is that the PropertyOverride data is incomplete for unloaded Actors.
+				// So we take the loaded data from the last Desc and transfer it to the latest data
+				PropertyOverrideDesc->TransferNonEditedContainers(ExistingOverrideDesc);
+			}
+		}
+	}
 }
 
 void FLevelInstanceActorDesc::Init(const FWorldPartitionActorDescInitData& DescData)
@@ -69,7 +108,7 @@ void FLevelInstanceActorDesc::Init(const FWorldPartitionActorDescInitData& DescD
 
 FString FLevelInstanceActorDesc::GetChildContainerName() const
 {
-	return GetChildContainerPackage().ToString();
+	return (ULevelInstanceSettings::Get()->IsPropertyOverrideEnabled() && PropertyOverrideAsset.IsValid()) ? PropertyOverrideAsset.ToString() : GetChildContainerPackage().ToString();
 }
 
 bool FLevelInstanceActorDesc::Equals(const FWorldPartitionActorDesc* Other) const
@@ -80,6 +119,7 @@ bool FLevelInstanceActorDesc::Equals(const FWorldPartitionActorDesc* Other) cons
 
 		return
 			WorldAsset == LevelInstanceActorDesc->WorldAsset &&
+			PropertyOverrideAsset == LevelInstanceActorDesc->PropertyOverrideAsset &&
 			DesiredRuntimeBehavior == LevelInstanceActorDesc->DesiredRuntimeBehavior;
 	}
 
@@ -101,7 +141,22 @@ void FLevelInstanceActorDesc::RegisterChildContainer()
 	check(!ChildContainer.IsValid());
 	if (IsChildContainerInstance())
 	{
-		ChildContainer = UActorDescContainerSubsystem::GetChecked().RegisterContainer({ GetChildContainerName(), GetChildContainerPackage() });
+		if (!PropertyOverrideAsset.IsValid() || !ULevelInstanceSettings::Get()->IsPropertyOverrideEnabled())
+		{
+			ChildContainer = UActorDescContainerSubsystem::GetChecked().RegisterContainer({ GetChildContainerName(), GetChildContainerPackage() });
+		}
+		else
+		{
+			UActorDescContainer::FInitializeParams InitParams(PropertyOverrideAsset.ToString(), PropertyOverrideDesc->GetWorldPackage());
+			InitParams.PreInitialize = [this](UActorDescContainer* InNewContainer)
+			{
+				ULevelInstancePropertyOverrideContainer* OverrideContainer = CastChecked<ULevelInstancePropertyOverrideContainer>(InNewContainer);
+				OverrideContainer->SetPropertyOverrideDesc(PropertyOverrideDesc);
+			};
+
+			ChildContainer = UActorDescContainerSubsystem::GetChecked().RegisterContainer<ULevelInstancePropertyOverrideContainer>(InitParams);
+		}
+
 		UpdateBounds();
 	}
 }
@@ -209,6 +264,11 @@ void FLevelInstanceActorDesc::TransferFrom(const FWorldPartitionActorDesc* From)
 
 	RegisterChildContainer();
 	FromLevelInstanceActorDesc->UnregisterChildContainer();
+
+	if (ULevelInstancePropertyOverrideContainer* OverrideContainer = Cast<ULevelInstancePropertyOverrideContainer>(ChildContainer.Get()))
+	{
+		OverrideContainer->SetPropertyOverrideDesc(PropertyOverrideDesc);
+	}
 }
 
 UWorldPartition* FLevelInstanceActorDesc::GetLoadedChildWorldPartition(const FWorldPartitionActorDescInstance* InActorDescInstance) const
@@ -226,9 +286,21 @@ UWorldPartition* FLevelInstanceActorDesc::GetLoadedChildWorldPartition(const FWo
 
 UActorDescContainerInstance* FLevelInstanceActorDesc::CreateChildContainerInstance(const FWorldPartitionActorDescInstance* InActorDescInstance) const
 {
-	UActorDescContainerInstance* ContainerInstance = InActorDescInstance->GetContainerInstance();
-	UActorDescContainerInstance* ChildContainerInstance = NewObject<UActorDescContainerInstance>(ContainerInstance, NAME_None, RF_Transient);
+	// Update Actor Desc Container in case of rename
+	if (UActorDescContainer* PreviousChildContainer = ChildContainer.Get(); PreviousChildContainer->GetContainerName() != GetChildContainerName())
+	{
+		FLevelInstanceActorDesc* NonConstThis = const_cast<FLevelInstanceActorDesc*>(this);
+		NonConstThis->UnregisterChildContainer();
+		NonConstThis->RegisterChildContainer();
+	}
 	
+	// Create ChildContainerInstance
+	UActorDescContainerInstance* ContainerInstance = InActorDescInstance->GetContainerInstance();
+	ULevelInstanceContainerInstance* ChildContainerInstance = NewObject<ULevelInstanceContainerInstance>(ContainerInstance, NAME_None, RF_Transient);
+
+	// Set Override Container which might be a regular UActorDescContainer if no overrides exist on this Container Instance
+	ChildContainerInstance->SetOverrideContainerAndAsset(GetChildContainer(), nullptr);
+		
 	// When a child container instance is created we create the whole hierarchy (for generate streaming)
 	const bool bInCreateContainerInstanceHierarchy = true;
 	UActorDescContainerInstance::FInitializeParams InitParams(InActorDescInstance->GetChildContainerPackage(), bInCreateContainerInstanceHierarchy);
@@ -319,5 +391,31 @@ void FLevelInstanceActorDesc::Serialize(FArchive& Ar)
 			Ar << Filter;
 		}
 	}
+
+	if (Ar.CustomVer(FFortniteMainBranchObjectVersion::GUID) >= FFortniteMainBranchObjectVersion::LevelInstancePropertyOverrides)
+	{
+		FString PropertyOverrideAssetString = PropertyOverrideAsset.ToString();
+		Ar << PropertyOverrideAssetString;
+
+		if (Ar.IsLoading())
+		{
+			PropertyOverrideAsset = FSoftObjectPath(PropertyOverrideAssetString);
+			UAssetRegistryHelpers::FixupRedirectedAssetPath(PropertyOverrideAsset);
+		}
+
+		if (PropertyOverrideAsset.IsValid())
+		{
+			if (Ar.IsLoading())
+			{
+				PropertyOverrideDesc = MakeShared<FLevelInstancePropertyOverrideDesc>();
+			}
+		
+			FLevelInstancePropertyOverrideDesc* PropertyOverrideDescPtr = PropertyOverrideDesc.Get();
+			Ar << *PropertyOverrideDescPtr;
+
+			PropertyOverrideDescPtr->AssetPath = PropertyOverrideAsset;
+			PropertyOverrideDescPtr->PackageName = ActorPackage;
+		}
+	}
 }
 #endif
diff --git a/Engine/Source/Runtime/Engine/Private/WorldPartition/WorldPartition.cpp b/Engine/Source/Runtime/Engine/Private/WorldPartition/WorldPartition.cpp
--- a/Engine/Source/Runtime/Engine/Private/WorldPartition/WorldPartition.cpp
+++ b/Engine/Source/Runtime/Engine/Private/WorldPartition/WorldPartition.cpp
@@ -1992,7 +1992,8 @@ UActorDescContainerInstance* UWorldPartition::RegisterActorDescContainerInstance
 			}
 		};
 
-		UActorDescContainerInstance* ContainerInstanceToRegister = NewObject<UActorDescContainerInstance>(this, UActorDescContainerInstance::StaticClass(), NAME_None, RF_Transient);
+		UClass* ContainerInstanceClassPtr = ContainerInstanceClass.Get() ? ContainerInstanceClass.Get() : UActorDescContainerInstance::StaticClass();
+		UActorDescContainerInstance* ContainerInstanceToRegister = NewObject<UActorDescContainerInstance>(this, ContainerInstanceClassPtr, NAME_None, RF_Transient);
 		
 		OnActorDescContainerInstancePreInitialize.ExecuteIfBound(InitParams, ContainerInstanceToRegister);
 
diff --git a/Engine/Source/Runtime/Engine/Private/WorldPartition/WorldPartitionActorDesc.cpp b/Engine/Source/Runtime/Engine/Private/WorldPartition/WorldPartitionActorDesc.cpp
--- a/Engine/Source/Runtime/Engine/Private/WorldPartition/WorldPartitionActorDesc.cpp
+++ b/Engine/Source/Runtime/Engine/Private/WorldPartition/WorldPartitionActorDesc.cpp
@@ -20,6 +20,7 @@
 #include "WorldPartition/WorldPartitionActorDescInstanceViewInterface.h"
 #include "WorldPartition/WorldPartitionLog.h"
 #include "WorldPartition/ActorDescContainer.h"
+#include "WorldPartition/WorldPartitionActorDescInstance.h"
 #include "WorldPartition/HLOD/HLODLayer.h"
 #include "WorldPartition/DataLayer/WorldDataLayers.h"
 #include "WorldPartition/DataLayer/DataLayerManager.h"
@@ -208,30 +209,42 @@ void FWorldPartitionActorDesc::Init(const FWorldPartitionActorDescInitData& Desc
 	ActorNativeClass = DescData.NativeClass;
 	NativeClass = *DescData.NativeClass->GetPathName();
 
+	auto DeprecateClass = [this](FArchive& Archive)
+	{
+		// Call registered deprecator
+		TSubclassOf<AActor> DeprecatedClass = ActorNativeClass;
+		while (DeprecatedClass)
+		{
+			if (FActorDescDeprecator* Deprecator = Deprecators.Find(DeprecatedClass))
+			{
+				(*Deprecator)(Archive, this);
+				break;
+			}
+			DeprecatedClass = DeprecatedClass->GetSuperClass();
+		}
+	};
+
 	// Serialize actor metadata
-	FMemoryReader MetadataAr(DescData.SerializedData, true);
+	if (!DescData.IsUsingArchive())
+	{
+		FMemoryReader MetadataAr(DescData.GetSerializedData(), true);
 
-	// Serialize metadata custom versions
-	FCustomVersionContainer CustomVersions;
-	CustomVersions.Serialize(MetadataAr);
-	MetadataAr.SetCustomVersions(CustomVersions);
+		// Serialize metadata custom versions
+		FCustomVersionContainer CustomVersions;
+		CustomVersions.Serialize(MetadataAr);
+		MetadataAr.SetCustomVersions(CustomVersions);
 	
-	// Serialize metadata payload
-	FActorDescArchive ActorDescAr(MetadataAr, this);
-	ActorDescAr.Init();
-
-	Serialize(ActorDescAr);
+		// Serialize metadata payload
+		FActorDescArchive ActorDescAr(MetadataAr, this);
+		ActorDescAr.Init();
 
-	// Call registered deprecator
-	TSubclassOf<AActor> DeprecatedClass = ActorNativeClass;
-	while (DeprecatedClass)
+		Serialize(ActorDescAr);
+		DeprecateClass(MetadataAr);
+	}
+	else
 	{
-		if (FActorDescDeprecator* Deprecator = Deprecators.Find(DeprecatedClass))
-		{
-			(*Deprecator)(MetadataAr, this);
-			break;
-		}
-		DeprecatedClass = DeprecatedClass->GetSuperClass();
+		Serialize(*DescData.GetArchive());
+		DeprecateClass(*DescData.GetArchive());
 	}
 
 	Container = nullptr;
@@ -240,7 +253,7 @@ void FWorldPartitionActorDesc::Init(const FWorldPartitionActorDescInitData& Desc
 void FWorldPartitionActorDesc::Patch(const FWorldPartitionActorDescInitData& DescData, TArray<uint8>& OutData, FWorldPartitionAssetDataPatcher* InAssetDataPatcher)
 {
 	// Serialize actor metadata
-	FMemoryReader MetadataAr(DescData.SerializedData, true);
+	FMemoryReader MetadataAr(DescData.GetSerializedData(), true);
 
 	// Serialize metadata custom versions
 	FCustomVersionContainer CustomVersions;
@@ -347,14 +360,14 @@ bool FWorldPartitionActorDesc::ShouldResave(const FWorldPartitionActorDesc* Othe
 	return !bActorIsHLODRelevant && Other->bActorIsHLODRelevant;
 }
 
-void FWorldPartitionActorDesc::SerializeTo(TArray<uint8>& OutData) const
+void FWorldPartitionActorDesc::SerializeTo(TArray<uint8>& OutData, FWorldPartitionActorDesc* BaseDesc) const
 {
 	FWorldPartitionActorDesc* MutableThis = const_cast<FWorldPartitionActorDesc*>(this);
 
 	// Serialize to archive and gather custom versions
 	TArray<uint8> PayloadData;
 	FMemoryWriter PayloadAr(PayloadData, true);
-	FActorDescArchive ActorDescAr(PayloadAr, MutableThis);
+	FActorDescArchive ActorDescAr(PayloadAr, MutableThis, BaseDesc);
 	ActorDescAr.Init();
 
 	MutableThis->Serialize(ActorDescAr);
diff --git a/Engine/Source/Runtime/Engine/Private/WorldPartition/WorldPartitionActorDescArchive.cpp b/Engine/Source/Runtime/Engine/Private/WorldPartition/WorldPartitionActorDescArchive.cpp
--- a/Engine/Source/Runtime/Engine/Private/WorldPartition/WorldPartitionActorDescArchive.cpp
+++ b/Engine/Source/Runtime/Engine/Private/WorldPartition/WorldPartitionActorDescArchive.cpp
@@ -14,10 +14,11 @@
 #include "AssetRegistry/AssetData.h"
 #include "AssetRegistry/AssetRegistryHelpers.h"
 
-FActorDescArchive::FActorDescArchive(FArchive& InArchive, FWorldPartitionActorDesc* InActorDesc)
+FActorDescArchive::FActorDescArchive(FArchive& InArchive, FWorldPartitionActorDesc* InActorDesc, const FWorldPartitionActorDesc* InBaseActorDesc)
 	: FArchiveProxy(InArchive)
 	, ActorDesc(InActorDesc)
-	, bIsMissingClassDesc(false)
+	, BaseDesc(InBaseActorDesc)
+	, bIsMissingBaseDesc(false)
 {
 	check(InArchive.IsPersistent());
 
@@ -89,29 +90,32 @@ PRAGMA_ENABLE_DEPRECATION_WARNINGS
 		TryRedirectClass(ActorDesc->BaseClass);
 	}
 
-	// Get the class descriptor to do delta serialization
-	FWorldPartitionClassDescRegistry& ClassDescRegistry = FWorldPartitionClassDescRegistry::Get();
-	const FTopLevelAssetPath ClassPath = InClassPath.IsValid() ? InClassPath : (ActorDesc->BaseClass.IsValid() ? ActorDesc->BaseClass : ActorDesc->NativeClass);
-	ClassDesc = (ActorDesc->bIsDefaultActorDesc && !InClassPath.IsValid()) ? ClassDescRegistry.GetClassDescDefaultForClass(ClassPath) : ClassDescRegistry.GetClassDescDefaultForActor(ClassPath);
-
-	if (!ClassDesc)
+	// Get the class descriptor to do delta serialization if no base desc was provided
+	if (!BaseDesc)
 	{
-		if (IsLoading())
+		FWorldPartitionClassDescRegistry& ClassDescRegistry = FWorldPartitionClassDescRegistry::Get();
+		const FTopLevelAssetPath ClassPath = InClassPath.IsValid() ? InClassPath : (ActorDesc->BaseClass.IsValid() ? ActorDesc->BaseClass : ActorDesc->NativeClass);
+		BaseDesc = (ActorDesc->bIsDefaultActorDesc && !InClassPath.IsValid()) ? ClassDescRegistry.GetClassDescDefaultForClass(ClassPath) : ClassDescRegistry.GetClassDescDefaultForActor(ClassPath);
+
+		if (!BaseDesc)
 		{
-			bIsMissingClassDesc = true;
+			if (IsLoading())
+			{
+				bIsMissingBaseDesc = true;
 
-			ClassDesc = ClassDescRegistry.GetClassDescDefault(FTopLevelAssetPath(TEXT("/Script/Engine.Actor")));
-			check(ClassDesc);			
+				BaseDesc = ClassDescRegistry.GetClassDescDefault(FTopLevelAssetPath(TEXT("/Script/Engine.Actor")));
+				check(BaseDesc);
 
-			UE_LOG(LogWorldPartition, Log, TEXT("Can't find class descriptor '%s' for loading '%s', using '%s'"), *ClassPath.ToString(), *ActorDesc->GetActorSoftPath().ToString(), *ClassDesc->GetActorSoftPath().ToString());
-		}
-		else
-		{
-			UE_LOG(LogWorldPartition, Log, TEXT("Can't find class descriptor '%s' for saving '%s'"), *ClassPath.ToString(), *ActorDesc->GetActorSoftPath().ToString());
+				UE_LOG(LogWorldPartition, Log, TEXT("Can't find class descriptor '%s' for loading '%s', using '%s'"), *ClassPath.ToString(), *ActorDesc->GetActorSoftPath().ToString(), *BaseDesc->GetActorSoftPath().ToString());
+			}
+			else
+			{
+				UE_LOG(LogWorldPartition, Log, TEXT("Can't find class descriptor '%s' for saving '%s'"), *ClassPath.ToString(), *ActorDesc->GetActorSoftPath().ToString());
+			}
 		}
 	}
 
-	ClassDescSizeof = ClassDesc ? ClassDesc->GetSizeOf() : 0;
+	BaseDescSizeof = BaseDesc ? BaseDesc->GetSizeOf() : 0;
 }
 
 FArchive& FActorDescArchive::operator<<(FSoftObjectPath& Value)
diff --git a/Engine/Source/Runtime/Engine/Private/WorldPartition/WorldPartitionLevelHelper.cpp b/Engine/Source/Runtime/Engine/Private/WorldPartition/WorldPartitionLevelHelper.cpp
--- a/Engine/Source/Runtime/Engine/Private/WorldPartition/WorldPartitionLevelHelper.cpp
+++ b/Engine/Source/Runtime/Engine/Private/WorldPartition/WorldPartitionLevelHelper.cpp
@@ -22,10 +22,13 @@
 #include "WorldPartition/WorldPartitionLog.h"
 #include "WorldPartition/WorldPartitionPackageHelper.h"
 #include "WorldPartition/ContentBundle/ContentBundleEditor.h"
+#include "WorldPartition/WorldPartitionPropertyOverride.h"
 #include "GameFramework/WorldSettings.h"
 #include "LevelUtils.h"
 #include "ActorFolder.h"
 
+FUObjectAnnotationSparse<FWorldPartitionLevelHelper::FActorPropertyOverridesAnnotation, true> FWorldPartitionLevelHelper::ActorPropertyOverridesAnnotation;
+
 #endif
 
 bool FWorldPartitionResolveData::ResolveObject(UWorld* InWorld, const FSoftObjectPath& InObjectPath, UObject*& OutObject) const
@@ -98,6 +101,45 @@ void FWorldPartitionLevelHelper::PreGarbageCollect()
 	PreGCPackagesToUnload.Reset();
 }
 
+void FWorldPartitionLevelHelper::ApplyConstructionScriptPropertyOverridesFromAnnotation(AActor* InActor)
+{
+	if (IsValid(InActor))
+	{
+		FWorldPartitionLevelHelper::FActorPropertyOverridesAnnotation Annotation = FWorldPartitionLevelHelper::ActorPropertyOverridesAnnotation.GetAndRemoveAnnotation(InActor);
+		if (!Annotation.IsDefault())
+		{
+			if (InActor->GetRootComponent())
+			{
+				FTransform InverseTransform = Annotation.ContainerTransform.Inverse();
+				FLevelUtils::FApplyLevelTransformParams TransformParams(InActor->GetLevel(), InverseTransform);
+				TransformParams.Actor = InActor;
+				TransformParams.bDoPostEditMove = false;
+				TransformParams.bSetRelativeTransformDirectly = true;
+
+				FLevelUtils::ApplyLevelTransform(TransformParams);
+			}
+
+			for (const FActorPropertyOverride& ActorOverride : Annotation.ActorPropertyOverrides)
+			{
+				const bool bConstructionScriptProperties = true;
+				UWorldPartitionPropertyOverride::ApplyPropertyOverrides(&ActorOverride, InActor, bConstructionScriptProperties);
+			}
+
+			if (InActor->GetRootComponent())
+			{
+				FLevelUtils::FApplyLevelTransformParams TransformParams(InActor->GetLevel(), Annotation.ContainerTransform);
+				TransformParams.Actor = InActor;
+				TransformParams.bDoPostEditMove = false;
+				TransformParams.bSetRelativeTransformDirectly = true;
+
+				FLevelUtils::ApplyLevelTransform(TransformParams);
+				InActor->GetRootComponent()->UpdateComponentToWorld();
+				InActor->MarkComponentsRenderStateDirty();
+			}
+		}
+	}
+}
+
 void FWorldPartitionLevelHelper::AddReference(UPackage* InPackage, FPackageReferencer* InReferencer)
 {
 	check(InPackage);
@@ -210,6 +252,7 @@ void FWorldPartitionLevelHelper::MoveExternalActorsToLevel(const TArray<FWorldPa
 				if (InLevel->GetWorld()->bIsWorldInitialized)
 				{
 					Actor->RerunConstructionScripts();
+					ApplyConstructionScriptPropertyOverridesFromAnnotation(Actor);
 				}
 			}
 			else if (!InLevel->Actors.Contains(Actor))
@@ -414,10 +457,137 @@ bool FWorldPartitionLevelHelper::LoadActors(UWorld* InOuterWorld, ULevel* InDest
 		.SetLoadAsync(bInLoadAsync)
 		.SetInstancingContext(InInstancingContext);
 
-	return LoadActors(Params);
+	return LoadActors(MoveTemp(Params));
+}
+
+bool FWorldPartitionLevelHelper::LoadActorsWithPropertyOverridesInternal(FLoadActorsParams&& InParams)
+{
+	TMap<FString, FName> PropertyOverridesToLoad;
+
+	struct FLoadProgress
+	{
+		int32 NumPendingLoadRequests = 0;
+		int32 NumFailedLoadedRequests = 0;
+		
+		TMap<FSoftObjectPath, TSet<FActorContainerID>> AssetToContainerIDs;
+
+		FLoadActorsParams Params;
+		FLoadedPropertyOverrides LoadedPropertyOverrides;
+	};
+	TSharedPtr<FLoadProgress> LoadProgress = MakeShared<FLoadProgress>();
+	LoadProgress->Params = MoveTemp(InParams);
+
+	// Build up list of Property Overrides to load and an assocation between Property Override Asset Path and the overrides Owner Container ID
+	for (const FWorldPartitionRuntimeCellObjectMapping& PackageObjectMapping : LoadProgress->Params.ActorPackages)
+	{
+		for (const FWorldPartitionRuntimeCellPropertyOverride& PropertyOverride : PackageObjectMapping.PropertyOverrides)
+		{
+			FName PackageName = PropertyOverridesToLoad.FindOrAdd(PropertyOverride.AssetPath, PropertyOverride.PackageName);
+			check(PackageName == PropertyOverride.PackageName);
+			LoadProgress->AssetToContainerIDs.FindOrAdd(PropertyOverride.AssetPath).Add(PropertyOverride.OwnerContainerID);
+		}
+	}
+
+	// Nothing to load, move on to load actors
+	if (PropertyOverridesToLoad.IsEmpty())
+	{
+		return LoadActorsInternal(MoveTemp(LoadProgress->Params), MoveTemp(LoadProgress->LoadedPropertyOverrides));
+	}
+
+	LoadProgress->NumPendingLoadRequests = PropertyOverridesToLoad.Num();
+	
+	// Do Loading
+	for (auto const&[AssetPath, PackageName] : PropertyOverridesToLoad)
+	{
+		FSoftObjectPath SoftAssetPath(AssetPath);
+		
+		FLinkerInstancingContext InstancingContext;
+		InstancingContext.AddTag(ULevel::DontLoadExternalObjectsTag);
+
+		FSoftObjectPath RemappedPath = SoftAssetPath;
+
+		// Loading embedded asset
+		if (!SoftAssetPath.GetSubPathString().IsEmpty())
+		{
+			FString WorldPackageName = SoftAssetPath.GetLongPackageName();
+			FName RemappedContainerPackage = FName(*(WorldPackageName + TEXT("_LoadPropertyOverride")));
+			InstancingContext.AddPackageMapping(*WorldPackageName, RemappedContainerPackage);
+
+			const FName AssetPackageInstanceName = FName(*ULevel::GetExternalActorPackageInstanceName(RemappedContainerPackage.ToString(), PackageName.ToString()));
+
+			InstancingContext.AddPackageMapping(PackageName, AssetPackageInstanceName);
+			InstancingContext.FixupSoftObjectPath(RemappedPath);
+		}
+
+		FName RemappedPackageName = InstancingContext.RemapPackage(PackageName);
+		FName PackageToLoad = PackageName;
+
+		FLoadPackageAsyncDelegate CompletionCallback = FLoadPackageAsyncDelegate::CreateLambda([LoadProgress, SoftAssetPath, RemappedPath](const FName& LoadedPackageName, UPackage* LoadedPackage, EAsyncLoadingResult::Type Result)
+		{
+			check(LoadProgress->NumPendingLoadRequests);
+			LoadProgress->NumPendingLoadRequests--;
+			
+			if (UWorldPartitionPropertyOverride* LoadedOverride = Cast<UWorldPartitionPropertyOverride>(RemappedPath.ResolveObject()))
+			{
+				// Reference World Package and Actor Package
+				LoadProgress->Params.PackageReferencer->AddReference(LoadedOverride->GetOutermostObject()->GetPackage());
+				LoadProgress->Params.PackageReferencer->AddReference(LoadedOverride->GetPackage());
+
+				TSet<FActorContainerID>& OwnerContainerIDs = LoadProgress->AssetToContainerIDs.FindChecked(SoftAssetPath);
+				for (FActorContainerID OwnerContainerID : OwnerContainerIDs)
+				{
+					LoadProgress->LoadedPropertyOverrides.PropertyOverrides.Add(OwnerContainerID, LoadedOverride);
+				}
+			}
+
+			if (!LoadProgress->NumPendingLoadRequests)
+			{
+				LoadActorsInternal(MoveTemp(LoadProgress->Params), MoveTemp(LoadProgress->LoadedPropertyOverrides));
+			}
+		});
+
+		if (LoadProgress->Params.bLoadAsync)
+		{
+			FPackagePath PackagePath = FPackagePath::FromPackageNameChecked(PackageToLoad);
+
+			::LoadPackageAsync(PackagePath, RemappedPackageName, CompletionCallback, PKG_None, -1, 0, &InstancingContext);
+		}
+		else
+		{
+			UPackage* InstancingPackage = nullptr;
+			if (PackageName != PackageToLoad)
+			{
+				InstancingPackage = CreatePackage(*RemappedPackageName.ToString());
+			}
+
+			UPackage* Package = LoadPackage(InstancingPackage, *PackageToLoad.ToString(), LOAD_None, nullptr, &InstancingContext);
+			CompletionCallback.Execute(PackageToLoad, Package, Package ? EAsyncLoadingResult::Succeeded : EAsyncLoadingResult::Failed);
+		}
+	}
+
+	return LoadProgress->NumPendingLoadRequests == 0;
 }
 
+// deprecated
 bool FWorldPartitionLevelHelper::LoadActors(const FLoadActorsParams& InParams)
+{
+	FLoadActorsParams ParamsCopy = FLoadActorsParams()
+		.SetActorPackages(InParams.ActorPackages)
+		.SetCompletionCallback(InParams.CompletionCallback)
+		.SetDestLevel(InParams.DestLevel)
+		.SetInstancingContext(MoveTemp(InParams.InstancingContext))
+		.SetLoadAsync(InParams.bLoadAsync)
+		.SetOuterWorld(InParams.OuterWorld)
+		.SetPackageReferencer(InParams.PackageReferencer);
+	return LoadActors(MoveTemp(ParamsCopy));
+}
+
+bool FWorldPartitionLevelHelper::LoadActors(FLoadActorsParams&& InParams)
+{	
+	return FWorldPartitionLevelHelper::LoadActorsWithPropertyOverridesInternal(MoveTemp(InParams));
+}
+
+bool FWorldPartitionLevelHelper::LoadActorsInternal(FLoadActorsParams&& InParams, FLoadedPropertyOverrides&& InLoadedPropertyOverrides)
 {
 	TArray<FWorldPartitionRuntimeCellObjectMapping*> ActorPackagesToLoad;
 	TMap<FActorContainerID, FLinkerInstancingContext> LinkerInstancingContexts;
@@ -486,7 +656,7 @@ bool FWorldPartitionLevelHelper::LoadActors(const FLoadActorsParams& InParams)
 
 	for (FWorldPartitionRuntimeCellObjectMapping* PackageObjectMapping : ActorPackagesToLoad)
 	{
-		FLoadPackageAsyncDelegate CompletionCallback = FLoadPackageAsyncDelegate::CreateLambda([LoadProgress, PackageObjectMapping, PackageReferencer = InParams.PackageReferencer, OuterWorld = InParams.OuterWorld, DestLevel = InParams.DestLevel, CompletionCallback = InParams.CompletionCallback](const FName& LoadedPackageName, UPackage* LoadedPackage, EAsyncLoadingResult::Type Result)
+		FLoadPackageAsyncDelegate CompletionCallback = FLoadPackageAsyncDelegate::CreateLambda([LoadProgress, PackageObjectMapping, LoadedOverrides = InLoadedPropertyOverrides, PackageReferencer = InParams.PackageReferencer, OuterWorld = InParams.OuterWorld, DestLevel = InParams.DestLevel, CompletionCallback = InParams.CompletionCallback](const FName& LoadedPackageName, UPackage* LoadedPackage, EAsyncLoadingResult::Type Result)
 		{
 			const FName ActorName = *FPaths::GetExtension(PackageObjectMapping->Path.ToString());
 			check(LoadProgress->NumPendingLoadRequests);
@@ -553,6 +723,32 @@ bool FWorldPartitionLevelHelper::LoadActors(const FLoadActorsParams& InParams)
 						}
 					});
 					
+					// Apply Pre-ConstructionScript Properties
+					TArray<FActorPropertyOverride> ActorPropertyOverrides;
+					for (auto ActorOverrideMapping : PackageObjectMapping->PropertyOverrides)
+					{
+						if (const UWorldPartitionPropertyOverride*const* LoadedOverride = LoadedOverrides.PropertyOverrides.Find(ActorOverrideMapping.OwnerContainerID))
+						{
+							if (const FContainerPropertyOverride* ContainerOverride = (*LoadedOverride)->PropertyOverridesPerContainer.Find(ActorOverrideMapping.ContainerPath))
+							{
+								if (const FActorPropertyOverride* ActorOverride = ContainerOverride->ActorOverrides.Find(Actor->GetActorGuid()))
+								{
+									const bool bConstructionScriptProperties = false;
+									UWorldPartitionPropertyOverride::ApplyPropertyOverrides(ActorOverride, Actor, bConstructionScriptProperties);
+									
+									// Store ActorOverride for Post Construction Script apply
+									ActorPropertyOverrides.Add(*ActorOverride);
+								}
+							}
+						}
+					}
+
+					// Store annotation for Post RerunConstructionScript apply
+					if (ActorPropertyOverrides.Num() > 0)
+					{
+						ActorPropertyOverridesAnnotation.AddAnnotation(Actor, FWorldPartitionLevelHelper::FActorPropertyOverridesAnnotation(MoveTemp(ActorPropertyOverrides), PackageObjectMapping->ContainerTransform));
+					}
+															
 					FLevelUtils::FApplyLevelTransformParams TransformParams(nullptr, PackageObjectMapping->ContainerTransform);
 					TransformParams.Actor = Actor;
 					TransformParams.bDoPostEditMove = false;
diff --git a/Engine/Source/Runtime/Engine/Private/WorldPartition/WorldPartitionLevelStreamingDynamic.cpp b/Engine/Source/Runtime/Engine/Private/WorldPartition/WorldPartitionLevelStreamingDynamic.cpp
--- a/Engine/Source/Runtime/Engine/Private/WorldPartition/WorldPartitionLevelStreamingDynamic.cpp
+++ b/Engine/Source/Runtime/Engine/Private/WorldPartition/WorldPartitionLevelStreamingDynamic.cpp
@@ -17,6 +17,7 @@
 #include "ContentStreaming.h"
 #include "WorldPartition/DataLayer/ExternalDataLayerHelper.h"
 #include "WorldPartition/ContentBundle/ContentBundlePaths.h"
+#include "LevelUtils.h"
 #endif
 
 #define LOCTEXT_NAMESPACE "World"
@@ -468,7 +469,7 @@ bool UWorldPartitionLevelStreamingDynamic::IssueLoadRequests()
 			.SetLoadAsync(World->IsGameWorld())
 			.SetInstancingContext(MoveTemp(InstancingContext));
 
-		FWorldPartitionLevelHelper::LoadActors(Params);
+		FWorldPartitionLevelHelper::LoadActors(MoveTemp(Params));
 	}
 	else
 	{
@@ -478,6 +479,22 @@ bool UWorldPartitionLevelStreamingDynamic::IssueLoadRequests()
 	return bLoadRequestInProgress;
 }
 
+void UWorldPartitionLevelStreamingDynamic::OnCurrentStateChanged(ELevelStreamingState InPrevState, ELevelStreamingState InNewState)
+{
+	Super::OnCurrentStateChanged(InPrevState, InNewState);
+
+	if (GetWorld()->IsPlayInEditor())
+	{
+		if (InNewState == ELevelStreamingState::LoadedVisible)
+		{
+			for (AActor* Actor : GetLoadedLevel()->Actors)
+			{
+				FWorldPartitionLevelHelper::ApplyConstructionScriptPropertyOverridesFromAnnotation(Actor);
+			}
+		}
+	}
+}
+
 void UWorldPartitionLevelStreamingDynamic::FinalizeRuntimeLevel()
 {
 	TRACE_CPUPROFILER_EVENT_SCOPE(UWorldPartitionLevelStreamingDynamic::FinalizeRuntimeLevel);
diff --git a/Engine/Source/Runtime/Engine/Private/WorldPartition/WorldPartitionRuntimeLevelStreamingCell.cpp b/Engine/Source/Runtime/Engine/Private/WorldPartition/WorldPartitionRuntimeLevelStreamingCell.cpp
--- a/Engine/Source/Runtime/Engine/Private/WorldPartition/WorldPartitionRuntimeLevelStreamingCell.cpp
+++ b/Engine/Source/Runtime/Engine/Private/WorldPartition/WorldPartitionRuntimeLevelStreamingCell.cpp
@@ -330,7 +330,7 @@ void UWorldPartitionRuntimeLevelStreamingCell::AddActorToCell(const FStreamingGe
 		}
 	}
 
-	Packages.Emplace(
+	FWorldPartitionRuntimeCellObjectMapping ActorMapping(
 		ActorDescView.GetActorPackage(), 
 		*ActorDescView.GetActorSoftPath().ToString(), 
 		ActorDescView.GetBaseClass(),
@@ -342,6 +342,15 @@ void UWorldPartitionRuntimeLevelStreamingCell::AddActorToCell(const FStreamingGe
 		ContainerID.GetActorGuid(ActorDescView.GetGuid()),
 		false
 	);
+
+	TArray<FWorldPartitionRuntimeCellPropertyOverride> PropertyOverrides;
+	ContainerInstance->GetPropertyOverridesForActor(ContainerID, ActorDescView.GetGuid(), PropertyOverrides);
+	if (PropertyOverrides.Num())
+	{
+		ActorMapping.PropertyOverrides = MoveTemp(PropertyOverrides);
+	}
+
+	Packages.Add(MoveTemp(ActorMapping));
 }
 
 void UWorldPartitionRuntimeLevelStreamingCell::Fixup()
@@ -385,7 +394,7 @@ bool UWorldPartitionRuntimeLevelStreamingCell::OnPrepareGeneratorPackageForCook(
 			.SetLoadAsync(false)
 			.SetInstancingContext(FLinkerInstancingContext(false)); // Don't do SoftObjectPath remapping for PersistentLevel actors because references can end up in different cells
 
-		verify(FWorldPartitionLevelHelper::LoadActors(Params));
+		verify(FWorldPartitionLevelHelper::LoadActors(MoveTemp(Params)));
 
 		FWorldPartitionLevelHelper::MoveExternalActorsToLevel(Packages, OuterWorld->PersistentLevel, OutModifiedPackages);
 
@@ -454,7 +463,7 @@ bool UWorldPartitionRuntimeLevelStreamingCell::OnPopulateGeneratedPackageForCook
 			.SetLoadAsync(false)
 			.SetInstancingContext(FLinkerInstancingContext(false)); // Don't do SoftObjectPath remapping for PersistentLevel actors because references can end up in different cells
 
-		verify(FWorldPartitionLevelHelper::LoadActors(Params));
+		verify(FWorldPartitionLevelHelper::LoadActors(MoveTemp(Params)));
 
 		// Create a level and move these actors in it
 		ULevel* NewLevel = FWorldPartitionLevelHelper::CreateEmptyLevelForRuntimeCell(this, OuterWorld, LevelStreaming->GetWorldAsset().ToString(), InPackage);
diff --git a/Engine/Source/Runtime/Engine/Private/WorldPartition/WorldPartitionSubsystem.cpp b/Engine/Source/Runtime/Engine/Private/WorldPartition/WorldPartitionSubsystem.cpp
--- a/Engine/Source/Runtime/Engine/Private/WorldPartition/WorldPartitionSubsystem.cpp
+++ b/Engine/Source/Runtime/Engine/Private/WorldPartition/WorldPartitionSubsystem.cpp
@@ -36,6 +36,8 @@
 #if WITH_EDITOR
 #include "Editor.h"
 #include "LevelEditorViewport.h"
+#include "LevelUtils.h"
+#include "WorldPartition/WorldPartitionPropertyOverride.h"
 #endif
 
 #include UE_INLINE_GENERATED_CPP_BY_NAME(WorldPartitionSubsystem)
@@ -1026,6 +1028,25 @@ void UWorldPartitionSubsystem::UpdateStreamingState()
 	UWorldPartitionSubsystem::UpdateStreamingStateInternal(GetWorld());
 }
 
+void UWorldPartitionSubsystem::OnWorldComponentsUpdated(UWorld& World)
+{
+	Super::OnWorldComponentsUpdated(World);
+
+#if WITH_EDITOR
+	// While Cooking Cells will get initialized for Save and UpdateWorldComponents will get called rerunning ConstructionScripts, we then want to apply Post Construction Script Overrides on Actors
+	if (IsRunningCookCommandlet())
+	{
+		for (AActor* Actor : World.PersistentLevel->Actors)
+		{
+			if (IsValid(Actor))
+			{
+				FWorldPartitionLevelHelper::ApplyConstructionScriptPropertyOverridesFromAnnotation(Actor);
+			}
+		}
+	}
+#endif
+}
+
 bool UWorldPartitionSubsystem::IncrementalUpdateStreamingState()
 {
 	TRACE_CPUPROFILER_EVENT_SCOPE(UWorldPartitionSubsystem::IncrementalUpdateStreamingState);
diff --git a/Engine/Source/Runtime/Engine/Public/LevelInstance/ILevelInstanceEditorModule.h b/Engine/Source/Runtime/Engine/Public/LevelInstance/ILevelInstanceEditorModule.h
--- a/Engine/Source/Runtime/Engine/Public/LevelInstance/ILevelInstanceEditorModule.h
+++ b/Engine/Source/Runtime/Engine/Public/LevelInstance/ILevelInstanceEditorModule.h
@@ -3,6 +3,9 @@
 
 #include "CoreMinimal.h"
 #include "Modules/ModuleInterface.h"
+#include "UObject/UnrealType.h"
+
+class ILevelInstanceInterface;
 
 /**
  * The module holding all of the UI related pieces for LevelInstance management
@@ -12,8 +15,10 @@ class ILevelInstanceEditorModule : public IModuleInterface
 public:
 	virtual ~ILevelInstanceEditorModule() {}
 
-	virtual void ActivateEditorMode() = 0;
-	virtual void DeactivateEditorMode() = 0;
+	UE_DEPRECATED(5.5, "This method is deprecated.")
+	virtual void ActivateEditorMode() {}
+	UE_DEPRECATED(5.5, "This method is deprecated.")
+	virtual void DeactivateEditorMode() {}
 
 	virtual void BroadcastTryExitEditorMode() = 0;
 
@@ -25,4 +30,28 @@ public:
 	virtual FTryExitEditorModeEvent& OnTryExitEditorMode() = 0;
 
 	virtual bool IsEditInPlaceStreamingEnabled() const = 0;
+
+protected:
+	friend class ULevelInstanceSubsystem;
+	friend class ULevelStreamingLevelInstanceEditorPropertyOverride;
+
+	// Called by ULevelInstanceSubsystem to update if the Editor Mode should be active or not
+	virtual void UpdateEditorMode(bool bActivated) = 0;
+
+	// Policy Proxy so that ULevelStreamingLevelInstanceEditorPropertyOverride can register policies through this module without knowing about PropertyEditor module
+	class IPropertyOverridePolicy
+	{
+	public:
+		virtual UObject* GetArchetypeForObject(const UObject* Object) const = 0;
+
+		virtual bool CanEditProperty(const FEditPropertyChain& PropertyChain, const UObject* Object) const = 0;
+		virtual bool CanEditProperty(const FProperty* Property, const UObject* Object) const = 0;
+	};
+
+	virtual UObject* GetArchetype(const UObject* Object) = 0;
+
+	virtual bool IsPropertyEditConst(const FEditPropertyChain& PropertyChain, UObject* Object) = 0;
+	virtual bool IsPropertyEditConst(const FProperty* Property, UObject* Object) = 0;
+
+	virtual void SetPropertyOverridePolicy(IPropertyOverridePolicy* Policy) = 0;
 };
\ No newline at end of file
