[Experimental] Level Instance Property Override

diff --git a/Engine/Source/Runtime/Engine/Public/LevelInstance/LevelInstanceActor.h b/Engine/Source/Runtime/Engine/Public/LevelInstance/LevelInstanceActor.h
--- a/Engine/Source/Runtime/Engine/Public/LevelInstance/LevelInstanceActor.h
+++ b/Engine/Source/Runtime/Engine/Public/LevelInstance/LevelInstanceActor.h
@@ -10,6 +10,7 @@
 #include "LevelInstance/LevelInstanceInterface.h"
 #include "LevelInstance/LevelInstanceActorImpl.h"
 #include "LevelInstance/LevelInstanceActorGuid.h"
+#include "LevelInstance/LevelInstancePropertyOverrideAsset.h"
 #include "WorldPartition/WorldPartitionActorDesc.h"
 #include "LevelInstanceActor.generated.h"
 
@@ -24,8 +25,11 @@ public:
 protected:
 #if WITH_EDITORONLY_DATA
 	/** Level LevelInstance */
-	UPROPERTY(EditAnywhere, Category = Level, Meta = (NoCreate, DisplayName="Level"))
+	UPROPERTY(EditAnywhere, Category = Level, Meta = (NoCreate, DisplayName="Level", DisableLevelInstancePropertyOverride))
 	TSoftObjectPtr<UWorld> WorldAsset;
+
+	UPROPERTY(VisibleAnywhere, Category = Override, meta = (EditInline, NoResetToDefault, EditCondition="PropertyOverrides != nullptr", EditConditionHides))
+	TObjectPtr<ULevelInstancePropertyOverrideAsset> PropertyOverrides;
 #endif
 	UPROPERTY(VisibleAnywhere, Category = Default)
 	TObjectPtr<ULevelInstanceComponent> LevelInstanceComponent;
@@ -41,7 +45,7 @@ protected:
 
 public:
 #if WITH_EDITORONLY_DATA
-	UPROPERTY(EditAnywhere, Category = Level, AdvancedDisplay, Meta = (DisplayName="Level Behavior"))
+	UPROPERTY(EditAnywhere, Category = Level, AdvancedDisplay, Meta = (DisplayName="Level Behavior", DisableLevelInstancePropertyOverride))
 	ELevelInstanceRuntimeBehavior DesiredRuntimeBehavior;
 #endif
 
@@ -71,6 +75,9 @@ public:
 	virtual ELevelInstanceRuntimeBehavior GetDefaultRuntimeBehavior() const override { return ELevelInstanceRuntimeBehavior::Partitioned; }
 	ENGINE_API virtual TSubclassOf<AActor> GetEditorPivotClass() const override;
 	ENGINE_API virtual bool SupportsPartialEditorLoading() const override;
+
+	ENGINE_API virtual bool SupportsPropertyOverrides() const override;
+	ENGINE_API ULevelInstancePropertyOverrideAsset* GetPropertyOverrideAsset() const override;
 	// End ILevelInstanceInterface
 			
 	// UObject overrides
@@ -111,6 +118,7 @@ public:
 	static ENGINE_API FOnLevelInstanceActorPostLoad OnLevelInstanceActorPostLoad;
 
 private:
+	ENGINE_API virtual void SetPropertyOverrideAsset(ULevelInstancePropertyOverrideAsset* InPropertyOverrideAsset) override;
 	ENGINE_API virtual bool ShouldCookWorldAsset() const;
 #endif
 };
diff --git a/Engine/Source/Runtime/Engine/Public/LevelInstance/LevelInstanceActorImpl.h b/Engine/Source/Runtime/Engine/Public/LevelInstance/LevelInstanceActorImpl.h
--- a/Engine/Source/Runtime/Engine/Public/LevelInstance/LevelInstanceActorImpl.h
+++ b/Engine/Source/Runtime/Engine/Public/LevelInstance/LevelInstanceActorImpl.h
@@ -64,6 +64,7 @@ public:
 	ENGINE_API virtual void OnLevelInstanceLoaded();
 
 #if WITH_EDITOR
+	ENGINE_API virtual bool SupportsPropertyOverrides() const;
 	ENGINE_API virtual bool SupportsPartialEditorLoading() const;
 	ENGINE_API virtual bool ResolveSubobject(const TCHAR* SubObjectPath, UObject*& OutObject, bool bLoadIfExists);
 #endif
diff --git a/Engine/Source/Runtime/Engine/Public/LevelInstance/LevelInstanceComponent.h b/Engine/Source/Runtime/Engine/Public/LevelInstance/LevelInstanceComponent.h
--- a/Engine/Source/Runtime/Engine/Public/LevelInstance/LevelInstanceComponent.h
+++ b/Engine/Source/Runtime/Engine/Public/LevelInstance/LevelInstanceComponent.h
@@ -39,7 +39,9 @@ public:
 
 	ENGINE_API void UpdateEditorInstanceActor();
 	ENGINE_API void OnEdit();
+	ENGINE_API void OnEditOverrides();
 	ENGINE_API void OnCommit();
+	ENGINE_API void OnCommitOverrides();
 
 	const FWorldPartitionActorFilter& GetFilter() const { return IsEditFilter() ? EditFilter : Filter; }
 	ENGINE_API void SetFilter(const FWorldPartitionActorFilter& InFilter, bool bNotify = true);
@@ -57,7 +59,7 @@ private:
 #endif
 
 #if WITH_EDITORONLY_DATA
-	UPROPERTY(EditAnywhere, Category = Filter, meta=(LevelInstanceFilter))
+	UPROPERTY(EditAnywhere, Category = Filter, meta=(LevelInstanceFilter, DisableLevelInstancePropertyOverride))
 	FWorldPartitionActorFilter Filter;
 
 	UPROPERTY(EditAnywhere, Transient, Category = Filter, meta=(LevelInstanceEditFilter))
diff --git a/Engine/Source/Runtime/Engine/Public/LevelInstance/LevelInstanceInterface.h b/Engine/Source/Runtime/Engine/Public/LevelInstance/LevelInstanceInterface.h
--- a/Engine/Source/Runtime/Engine/Public/LevelInstance/LevelInstanceInterface.h
+++ b/Engine/Source/Runtime/Engine/Public/LevelInstance/LevelInstanceInterface.h
@@ -8,6 +8,7 @@
 #include "LevelInstance/LevelInstanceTypes.h"
 #include "WorldPartition/WorldPartitionActorContainerID.h"
 #include "WorldPartition/Filter/WorldPartitionActorFilter.h"
+#include "LevelInstance/LevelInstancePropertyOverrideAsset.h"
 #include "LevelInstanceInterface.generated.h"
 
 class ULevelInstanceComponent;
@@ -46,7 +47,9 @@ class ILevelInstanceInterface
 	virtual bool IsLoadingEnabled() const = 0;
 
 #if WITH_EDITOR
-
+	virtual bool SupportsPropertyOverrides() const { return false; }
+		
+	virtual ULevelInstancePropertyOverrideAsset* GetPropertyOverrideAsset() const { return nullptr; }
 
 	virtual ULevelInstanceComponent* GetLevelInstanceComponent() const = 0;
 
@@ -88,13 +91,19 @@ class ILevelInstanceInterface
 
 #if WITH_EDITOR
 	ENGINE_API virtual void OnEdit();
+	ENGINE_API virtual void OnEditPropertyOverrides();
 
 	virtual void OnEditChild() {}
 
 	ENGINE_API virtual void OnCommit(bool bChanged);
+	ENGINE_API virtual void OnCommitPropertyOverrides(bool bChanged);
 
 	virtual void OnCommitChild(bool bChanged) {}
 	
+	ENGINE_API bool IsInAnyEditMode() const { return IsEditing() || IsEditingPropertyOverrides(); }
+
+	ENGINE_API bool IsEditingPropertyOverrides() const;
+
 	ENGINE_API virtual bool IsEditing() const;
 	
 	ENGINE_API virtual bool HasChildEdit() const;
@@ -106,13 +115,21 @@ class ILevelInstanceInterface
 	ENGINE_API virtual bool HasDirtyChildren() const;
 	
 	ENGINE_API virtual bool CanEnterEdit(FText* OutReason = nullptr) const;
+
+	ENGINE_API bool CanEnterEditPropertyOverrides(FText* OutReason = nullptr) const;
 	
 	ENGINE_API virtual bool EnterEdit(AActor* ContextActor = nullptr);
+
+	ENGINE_API bool EnterEditPropertyOverrides(AActor* ContextActor = nullptr);
 	
 	ENGINE_API virtual bool CanExitEdit(bool bDiscardEdits = false, FText* OutReason = nullptr) const;
+
+	ENGINE_API bool CanExitEditPropertyOverrides(bool bDiscardEdits = false, FText* OutReason = nullptr) const;
 	
 	ENGINE_API virtual bool ExitEdit(bool bDiscardEdits = false);
-	
+
+	ENGINE_API bool ExitEditPropertyOverrides(bool bDiscardEdits = false);
+			
 	ENGINE_API virtual bool SetCurrent();
 	
 	ENGINE_API virtual bool MoveActorsTo(const TArray<AActor*>& ActorsToMove);
@@ -141,4 +158,9 @@ class ILevelInstanceInterface
 
 	ENGINE_API virtual void UpdateLevelInstanceFromWorldAsset();
 
+#if WITH_EDITOR
+private:
+	friend class ULevelInstanceSubsystem;
+	virtual void SetPropertyOverrideAsset(ULevelInstancePropertyOverrideAsset* InPropertyOverride) {}
+#endif
 };
diff --git a/Engine/Source/Runtime/Engine/Public/LevelInstance/LevelInstanceLevelStreaming.h b/Engine/Source/Runtime/Engine/Public/LevelInstance/LevelInstanceLevelStreaming.h
--- a/Engine/Source/Runtime/Engine/Public/LevelInstance/LevelInstanceLevelStreaming.h
+++ b/Engine/Source/Runtime/Engine/Public/LevelInstance/LevelInstanceLevelStreaming.h
@@ -6,7 +6,11 @@
 #include "UObject/ObjectMacros.h"
 #include "Engine/LevelStreamingDynamic.h"
 #include "LevelInstance/LevelInstanceTypes.h"
+
+#if WITH_EDITOR
 #include "WorldPartition/ActorDescContainerInstance.h"
+#endif
+
 #include "LevelInstanceLevelStreaming.generated.h"
 
 class ILevelInstanceInterface;
@@ -40,12 +44,24 @@ protected:
 	const FLevelInstanceID& GetLevelInstanceID() const { return LevelInstanceID; }
 private:
 #if WITH_EDITOR
-	ENGINE_API void ResetLevelInstanceLoaders();
+	void ResetLevelInstanceLoaders();
+	
 	ENGINE_API virtual void OnLoadedActorsAddedToLevelPreEvent(const TArray<AActor*>& InActors);
 	ENGINE_API virtual void OnLoadedActorsAddedToLevelPostEvent(const TArray<AActor*>& InActors);
 	ENGINE_API virtual void OnLoadedActorsRemovedFromLevelPostEvent(const TArray<AActor*>& InActors);
-	ENGINE_API void OnLevelStreamingStateChanged(UWorld* InWorld, const ULevelStreaming* InLevelStreaming, ULevel* InLevelIfLoaded, ELevelStreamingState InPrevState, ELevelStreamingState InNewState);
+	ENGINE_API virtual void OnCurrentStateChanged(ELevelStreamingState InPrevState, ELevelStreamingState InNewState) override;
 	ENGINE_API void OnPreInitializeContainerInstance(UActorDescContainerInstance::FInitializeParams& InInitParams, UActorDescContainerInstance* InContainerInstance);
+	void InitializeActors(const TArray<AActor*>& InActors);
+	void OnActorReplacedEvent(FWorldPartitionActorDescInstance* InActorDescInstance);
+
+	enum EApplyPropertieOverrideType
+	{
+		PreConstruction,
+		PostConstruction,
+		PreAndPostConstruction
+	};
+
+	void ApplyPropertyOverrides(const TArray<AActor*>& InActor, bool bInAlreadyAppliedTransformOnActors, EApplyPropertieOverrideType InApplyPropertyOverrideType);
 
 	bool IsEditorWorldMode() const;
 
diff --git a/Engine/Source/Runtime/Engine/Public/LevelInstance/LevelInstanceSubsystem.h b/Engine/Source/Runtime/Engine/Public/LevelInstance/LevelInstanceSubsystem.h
--- a/Engine/Source/Runtime/Engine/Public/LevelInstance/LevelInstanceSubsystem.h
+++ b/Engine/Source/Runtime/Engine/Public/LevelInstance/LevelInstanceSubsystem.h
@@ -11,9 +11,11 @@
 #include "WorldPartition/Filter/WorldPartitionActorFilter.h"
 #include "WorldPartition/WorldPartitionHandle.h"
 #include "WorldPartition/WorldPartitionActorContainerID.h"
+#include "LevelInstance/LevelInstancePropertyOverrideAsset.h"
 
 #if WITH_EDITOR
 #include "EditorLevelUtils.h"
+#include "Internationalization/Text.h"
 #endif
 
 #include "LevelInstanceSubsystem.generated.h"
@@ -22,8 +24,10 @@ class ILevelInstanceInterface;
 class ULevelInstanceEditorObject;
 class ULevelStreamingLevelInstance;
 class ULevelStreamingLevelInstanceEditor;
+class ULevelStreamingLevelInstanceEditorPropertyOverride;
 class UWorldPartitionSubsystem;
 class UBlueprint;
+struct FActorPropertyOverride;
 
 enum class ELevelInstanceBreakFlags : uint8
 {
@@ -96,7 +100,6 @@ public:
 	ENGINE_API void Tick();
 	ENGINE_API void OnExitEditorMode();
 	ENGINE_API void OnTryExitEditorMode();
-	ENGINE_API bool OnExitEditorModeInternal(bool bForceExit);
 
 	UE_DEPRECATED(5.3, "Use FPackedLevelActorUtils::PackAllLoadedActors")
 	void PackAllLoadedActors() {}
@@ -124,9 +127,11 @@ public:
 	ENGINE_API bool SetCurrent(ILevelInstanceInterface* LevelInstance) const;
 	ENGINE_API bool IsCurrent(const ILevelInstanceInterface* LevelInstance) const;
 	ENGINE_API ILevelInstanceInterface* CreateLevelInstanceFrom(const TArray<AActor*>& ActorsToMove, const FNewLevelInstanceParams& CreationParams);
+	ENGINE_API bool CanCreateLevelInstanceFrom(const TArray<AActor*>& ActorsToMove, FText* OutReason = nullptr);
 	ENGINE_API bool MoveActorsToLevel(const TArray<AActor*>& ActorsToRemove, ULevel* DestinationLevel, TArray<AActor*>* OutActors = nullptr) const;
 	ENGINE_API bool MoveActorsTo(ILevelInstanceInterface* LevelInstance, const TArray<AActor*>& ActorsToMove, TArray<AActor*>* OutActors = nullptr);
 	ENGINE_API bool BreakLevelInstance(ILevelInstanceInterface* LevelInstance, uint32 Levels = 1, TArray<AActor*>* OutMovedActors = nullptr, ELevelInstanceBreakFlags Flags = ELevelInstanceBreakFlags::None);
+	ENGINE_API bool CanBreakLevelInstance(const ILevelInstanceInterface* LevelInstance) const;
 
 	ENGINE_API bool CanMoveActorToLevel(const AActor* Actor, FText* OutReason = nullptr) const;
 	ENGINE_API void OnActorDeleted(AActor* Actor);
@@ -139,9 +144,9 @@ public:
 		
 	ENGINE_API bool HasChildEdit(const ILevelInstanceInterface* LevelInstance) const;
 	ENGINE_API bool HasParentEdit(const ILevelInstanceInterface* LevelInstance) const;
-
-	ENGINE_API TArray<ILevelInstanceInterface*> GetLevelInstances(const FString& WorldAssetPackage);
 	
+	ENGINE_API TArray<ILevelInstanceInterface*> GetLevelInstances(const FString& WorldAssetPackage) const;
+		
 	// Returns the upper chain of level instance actors for the specified level starting with the level instance referencing the level
 	ENGINE_API void ForEachLevelInstanceActorAncestors(const ULevel* Level, TFunctionRef<bool(AActor*)> Operation) const;
 	ENGINE_API TArray<AActor*> GetParentLevelInstanceActors(const ULevel* Level) const;
@@ -160,6 +165,25 @@ public:
 
 	ENGINE_API bool PassLevelInstanceFilter(UWorld* World, const FWorldPartitionHandle& Actor) const;
 
+	ENGINE_API bool IsEditingLevelInstancePropertyOverrides(const ILevelInstanceInterface* LevelInstance) const;
+private:
+	friend class FLevelInstanceEditorModule;
+	friend struct FLevelInstanceMenuUtils;
+	friend class ILevelInstanceInterface;
+	friend class ULevelInstanceEditorMode;
+
+	// EditPropertyOverrides
+	ENGINE_API TArray<ILevelInstanceInterface*> GetLevelInstances(const TSoftObjectPtr<ULevelInstancePropertyOverrideAsset>& PropertyOverrideAsset) const;
+	ENGINE_API bool HasParentPropertyOverridesEdit(const ILevelInstanceInterface* LevelInstance) const;
+	ENGINE_API ILevelInstanceInterface* GetEditingPropertyOverridesLevelInstance() const;
+	ENGINE_API bool CanCommitLevelInstancePropertyOverrides(const ILevelInstanceInterface* LevelInstance, bool bDiscardEdits = false, FText* OutReason = nullptr) const;
+	ENGINE_API bool CanEditLevelInstancePropertyOverrides(const ILevelInstanceInterface* LevelInstance, FText* OutReason = nullptr) const;
+	ENGINE_API void EditLevelInstancePropertyOverrides(ILevelInstanceInterface* LevelInstance, AActor* ContextActor);
+	ENGINE_API bool CommitLevelInstancePropertyOverrides(ILevelInstanceInterface* LevelInstance, bool bDiscardEdits = false);
+	ENGINE_API bool CanResetPropertyOverridesForActor(AActor* Actor) const;
+	ENGINE_API void ResetPropertyOverridesForActor(AActor* Actor);
+	ENGINE_API bool CanResetPropertyOverrides(ILevelInstanceInterface* LevelInstance) const;
+	ENGINE_API void ResetPropertyOverrides(ILevelInstanceInterface* LevelInstance);
 #endif
 
 private:
@@ -173,8 +197,6 @@ private:
 
 #if WITH_EDITOR
 	ENGINE_API ULevelStreamingLevelInstanceEditor* CreateNewStreamingLevelForWorld(UWorld& InWorld, const EditorLevelUtils::FCreateNewStreamingLevelForWorldParams& InParams);
-	
-	ENGINE_API FWorldPartitionActorFilter GetLevelInstanceFilterInternal(const FString& LevelPackage, TSet<FString>& VisitedPackages) const;
 
 	ENGINE_API void ResetLoadersForWorldAssetInternal(const FString& WorldAsset);
 	ENGINE_API void OnAssetsPreDelete(const TArray<UObject*>& Objects);
@@ -183,7 +205,7 @@ private:
 	void OnWorldCleanup(UWorld* InWorld, bool bSessionEnded, bool bCleanupResources);
 
 	ENGINE_API void RegisterLoadedLevelStreamingLevelInstanceEditor(ULevelStreamingLevelInstanceEditor* LevelStreaming);
-
+	
 	ENGINE_API void OnEditChild(const FLevelInstanceID& LevelInstanceID);
 	ENGINE_API void OnCommitChild(const FLevelInstanceID& LevelInstanceID, bool bChildChanged);
 
@@ -202,7 +224,7 @@ private:
 		TObjectPtr<AActor> LevelInstanceActor;
 
 		FLevelInstanceEdit(ULevelStreamingLevelInstanceEditor* InLevelStreaming, ILevelInstanceInterface* InLevelInstance);
-		virtual ~FLevelInstanceEdit();
+		~FLevelInstanceEdit();
 
 		UWorld* GetEditWorld() const;
 		ILevelInstanceInterface* GetLevelInstance() const;
@@ -215,16 +237,53 @@ private:
 		void MarkCommittedChanges();
 	};
 
-	ENGINE_API void ResetEdit(TUniquePtr<FLevelInstanceEdit>& InLevelInstanceEdit);
-	ENGINE_API bool EditLevelInstanceInternal(ILevelInstanceInterface* LevelInstance, TWeakObjectPtr<AActor> ContextActorPtr, const FString& InActorNameToSelect, bool bRecursive);
-	ENGINE_API bool CommitLevelInstanceInternal(TUniquePtr<FLevelInstanceEdit>& InLevelInstanceEdit, bool bDiscardEdits = false, bool bDiscardOnFailure = false, TSet<FName>* DirtyPackages = nullptr);
+	class FPropertyOverrideEdit
+	{
+	public:
+		FPropertyOverrideEdit(ULevelStreamingLevelInstanceEditorPropertyOverride* InLevelStreaming);
+		~FPropertyOverrideEdit();
+
+		ULevelStreamingLevelInstanceEditorPropertyOverride* LevelStreaming;
+		
+		ILevelInstanceInterface* GetLevelInstance() const;
+		bool CanDiscard(FText* OutReason = nullptr) const { return true; }
+		bool IsDirty() const;
+		bool Save(ILevelInstanceInterface* InLevelInstanceOverrideOwner) const;
+	};
+		
+	ENGINE_API FActorContainerID GetLevelInstancePropertyOverridesContext(ILevelInstanceInterface* LevelInstance) const;
+	ENGINE_API bool GetLevelInstancePropertyOverridesForActor(const AActor* Actor, FActorContainerID PropertyOverrideContext, TArray<const FActorPropertyOverride*>& OutPropertyOverrides) const;
 	
+	ENGINE_API ILevelInstanceInterface* GetLevelInstancePropertyOverridesEditOwner(ILevelInstanceInterface* LevelInstance) const;
+	ENGINE_API const ILevelInstanceInterface* GetLevelInstancePropertyOverridesEditOwner(const ILevelInstanceInterface* LevelInstance) const;
+	
+	ENGINE_API bool EditLevelInstanceInternal(ILevelInstanceInterface* LevelInstance, TWeakObjectPtr<AActor> ContextActorPtr, const FString& InActorNameToSelect, bool bRecursive);
+	ENGINE_API bool CommitLevelInstanceInternal(TUniquePtr<FLevelInstanceEdit>& LevelInstanceEdit, bool bDiscardEdits = false, bool bDiscardOnFailure = false, TSet<FName>* DirtyPackages = nullptr);
+	ENGINE_API bool CommitLevelInstancePropertyOverridesInternal(TUniquePtr<FPropertyOverrideEdit>& InPropertyOverrideEdit, bool bDiscardEdits);
+	bool CanEditLevelInstanceCommon(const ILevelInstanceInterface* LevelInstance, FText* OutReason = nullptr) const;
+
+	ENGINE_API FText GetToolKitExitToolTip() const;
+	ENGINE_API FText GetToolKitDisplayText() const;
+	ENGINE_API void ToolKitExit();
+
+	void OnExitEditorModeInternal(bool bForceExit);
+	bool TryCommitLevelInstanceEdit(bool bForceExit);
+	bool TryCommitLevelInstancePropertyOverrideEdit(bool bForceExit);
+
 	ENGINE_API const FLevelInstanceEdit* GetLevelInstanceEdit(const ILevelInstanceInterface* LevelInstance) const;
 	ENGINE_API bool IsLevelInstanceEditDirty(const FLevelInstanceEdit* LevelInstanceEdit) const;
 	ENGINE_API bool PromptUserForCommit(const FLevelInstanceEdit* InLevelInstanceEdit, bool& bOutDiscard, bool bForceCommit = false) const;
+	ENGINE_API bool PromptUserForCommitPropertyOverrides(const FPropertyOverrideEdit* InPropertyOverrideEdit, bool& bOutDiscard, bool bForceCommit = false) const;
+
+	ENGINE_API const FPropertyOverrideEdit* GetLevelInstancePropertyOverrideEdit(const ILevelInstanceInterface* LevelInstance) const;
+	ENGINE_API void RegisterLoadedLevelStreamingPropertyOverride(ULevelStreamingLevelInstanceEditorPropertyOverride* LevelStreaming);
+	void UpdateLevelInstancesFromPropertyOverrideAsset(const TSoftObjectPtr<ULevelInstancePropertyOverrideAsset>& PreviousAssetPath, ULevelInstancePropertyOverrideAsset* NewAsset);
 
 	void ForEachLevelStreaming(TFunctionRef<bool(ULevelStreaming*)> Operation) const;
 
+	FString GetActorNameToSelectFromContext(const ILevelInstanceInterface* LevelInstance, const AActor* ContextActor, const FString& DefaultName = FString()) const;
+	void SelectActorFromActorName(ILevelInstanceInterface* LevelInstance, const FString& ActorName) const;
+
 	struct FLevelsToRemoveScope
 	{
 		FLevelsToRemoveScope(ULevelInstanceSubsystem* InOwner);
@@ -241,8 +300,11 @@ private:
 
 private:
 #endif
-	friend ULevelStreamingLevelInstance;
-	friend ULevelStreamingLevelInstanceEditor;
+	friend class ULevelStreamingLevelInstance;
+	friend class ULevelStreamingLevelInstanceEditor;
+	friend class ULevelStreamingLevelInstanceEditorPropertyOverride;
+	friend class ULevelInstancePropertyOverrideAsset;
+	friend class FLevelInstanceEditorModeToolkit;
 
 #if WITH_EDITOR
 	bool bIsCreatingLevelInstance;
@@ -267,6 +329,7 @@ private:
 	TUniquePtr<FLevelsToRemoveScope> LevelsToRemoveScope;
 
 	TUniquePtr<FLevelInstanceEdit> LevelInstanceEdit;
+	TUniquePtr<FPropertyOverrideEdit> PropertyOverrideEdit;
 
 	TMap<FLevelInstanceID, int32> ChildEdits;
 
diff --git a/Engine/Source/Runtime/Engine/Public/PackedLevelActor/PackedLevelActor.h b/Engine/Source/Runtime/Engine/Public/PackedLevelActor/PackedLevelActor.h
--- a/Engine/Source/Runtime/Engine/Public/PackedLevelActor/PackedLevelActor.h
+++ b/Engine/Source/Runtime/Engine/Public/PackedLevelActor/PackedLevelActor.h
@@ -40,6 +40,7 @@ public:
 	ENGINE_API bool ShouldLoadForPacking() const;
 	// When Loading a APackedLevelActor it needs to be fully loaded for packing.
 	virtual bool SupportsPartialEditorLoading() const override { return false; }
+	virtual bool SupportsPropertyOverrides() const override { return false; }
 
 	ENGINE_API virtual TUniquePtr<class FWorldPartitionActorDesc> CreateClassActorDesc() const override;
 	
diff --git a/Engine/Source/Runtime/Engine/Public/WorldPartition/ActorDescContainer.h b/Engine/Source/Runtime/Engine/Public/WorldPartition/ActorDescContainer.h
--- a/Engine/Source/Runtime/Engine/Public/WorldPartition/ActorDescContainer.h
+++ b/Engine/Source/Runtime/Engine/Public/WorldPartition/ActorDescContainer.h
@@ -21,6 +21,7 @@ class UActorDescContainer : public UObject, public FActorDescList
 #if WITH_EDITOR
 	friend struct FWorldPartitionHandleUtils;
 	friend class FWorldPartitionActorDesc;
+	friend class UActorDescContainerInstance;
 
 	using FNameActorDescMap = TMap<FName, TUniquePtr<FWorldPartitionActorDesc>*>;
 
@@ -56,12 +57,15 @@ public:
 		/** The associated External Data Layer Asset */
 		const UExternalDataLayerAsset* ExternalDataLayerAsset = nullptr;
 
+		/* Custom pre-init function that is called before calling Initialize on the new container */
+		TUniqueFunction<void(UActorDescContainer*)> PreInitialize;
+
 		/* Custom filter function used to filter actors descriptors. */
 		TUniqueFunction<bool(const FWorldPartitionActorDesc*)> FilterActorDesc;
 	};
 
-	ENGINE_API void Initialize(const FInitializeParams& InitParams);
-	ENGINE_API void Uninitialize();
+	ENGINE_API virtual void Initialize(const FInitializeParams& InitParams);
+	ENGINE_API virtual void Uninitialize();
 
 	bool IsInitialized() const { return bContainerInitialized; }
 
@@ -115,9 +119,9 @@ public:
 
 	ENGINE_API bool ShouldHandleActorEvent(const AActor* Actor);
 
-	ENGINE_API const FWorldPartitionActorDesc* GetActorDescByPath(const FString& ActorPath) const;
-	ENGINE_API const FWorldPartitionActorDesc* GetActorDescByPath(const FSoftObjectPath& ActorPath) const;
-	ENGINE_API const FWorldPartitionActorDesc* GetActorDescByName(FName ActorName) const;
+	virtual ENGINE_API const FWorldPartitionActorDesc* GetActorDescByPath(const FString& ActorPath) const;
+	virtual ENGINE_API const FWorldPartitionActorDesc* GetActorDescByPath(const FSoftObjectPath& ActorPath) const;
+	virtual ENGINE_API const FWorldPartitionActorDesc* GetActorDescByName(FName ActorName) const;
 
 	bool bContainerInitialized;
 
@@ -162,14 +166,18 @@ protected:
 	ENGINE_API virtual void BeginDestroy() override;
 	//~ End UObject Interface
 
+	ENGINE_API virtual bool ShouldRegisterDelegates() const;
+
+	ENGINE_API bool ShouldHandleActorEvent(const AActor* Actor, bool bInUseLoadedPath) const;
+	ENGINE_API bool IsActorDescHandled(const AActor* InActor, bool bInUseLoadedPath) const;
 private:
 	// GetWorld() should never be called on an ActorDescContainer to avoid any confusion as it can be used as a template
 	UWorld* GetWorld() const override { return nullptr; }
 
-	ENGINE_API bool ShouldRegisterDelegates();
 	ENGINE_API void RegisterEditorDelegates();
 	ENGINE_API void UnregisterEditorDelegates();
 
+protected:
 	TObjectPtr<const UExternalDataLayerAsset> ExternalDataLayerAsset;
 #endif
 };
diff --git a/Engine/Source/Runtime/Engine/Public/WorldPartition/ActorDescContainerInstance.h b/Engine/Source/Runtime/Engine/Public/WorldPartition/ActorDescContainerInstance.h
--- a/Engine/Source/Runtime/Engine/Public/WorldPartition/ActorDescContainerInstance.h
+++ b/Engine/Source/Runtime/Engine/Public/WorldPartition/ActorDescContainerInstance.h
@@ -97,6 +97,9 @@ public:
 	DECLARE_EVENT_OneParam(UActorDescContainerInstance, FActorDescInstanceUpdatedEvent, FWorldPartitionActorDescInstance*);
 	FActorDescInstanceUpdatedEvent OnActorDescInstanceUpdatedEvent;
 
+	DECLARE_EVENT_OneParam(UActorDescContainerInstance, FActorReplacedEvent, FWorldPartitionActorDescInstance*);
+	FActorReplacedEvent OnActorReplacedEvent;
+
 	ENGINE_API virtual void Initialize(const FInitializeParams& InParams);
 	ENGINE_API bool IsInitialized() const { return bIsInitialized; }
 	ENGINE_API virtual void Uninitialize();
@@ -116,6 +119,8 @@ public:
 	ENGINE_API FString GetExternalActorPath() const;
 	ENGINE_API FString GetExternalObjectPath() const;
 		
+	ENGINE_API virtual void GetPropertyOverridesForActor(const FActorContainerID& InContainerID, const FGuid& InActorGuid, TArray<FWorldPartitionRuntimeCellPropertyOverride>& OutPropertyOverrides) const {}
+
 	ENGINE_API TUniquePtr<FWorldPartitionActorDescInstance>* GetActorDescInstancePtr(const FGuid& InActorGuid) const;
 	ENGINE_API FWorldPartitionActorDescInstance* GetActorDescInstance(const FGuid& InActorGuid) const;
 	ENGINE_API FWorldPartitionActorDescInstance& GetActorDescInstanceChecked(const FGuid& InActorGuid) const;
@@ -135,6 +140,7 @@ public:
 
 	ENGINE_API const TMap<FGuid, TObjectPtr<UActorDescContainerInstance>>& GetChildContainerInstances() const { return ChildContainerInstances; }
 		
+	const UActorDescContainerInstance* GetParentContainerInstance() const { return ParentContainerInstance.Get(); }
 protected:
 	virtual void RegisterContainer(const FInitializeParams& InParams);
 	virtual void UnregisterContainer();
@@ -143,8 +149,8 @@ protected:
 	virtual FWorldPartitionActorDesc* GetActorDesc(const FGuid& InActorGuid) const;
 	virtual FWorldPartitionActorDesc* GetActorDescChecked(const FGuid& InActorGuid) const;
 
+	virtual FWorldPartitionActorDescInstance CreateActorDescInstance(FWorldPartitionActorDesc* InActorDesc) { return FWorldPartitionActorDescInstance(this, InActorDesc); }
 private:
-	void OnContainerUpdated(FName ContainerPackage);
 	void OnContainerReplaced(UActorDescContainer* InOldContainer, UActorDescContainer* InNewContainer);
 	void SetContainerPackage(FName InContainerPackageName);
 		
@@ -171,7 +177,7 @@ private:
 	friend class UGameFeatureActionConvertContentBundleWorldPartitionBuilder;
 #endif
 
-private:
+protected:
 #if WITH_EDITORONLY_DATA
 	FSoftObjectPath												WorldContainerPath;
 	FSoftObjectPath												SourceWorldContainerPath;
@@ -180,6 +186,7 @@ private:
 	FActorContainerID											ContainerID;
 	FGuid														ContainerActorGuid;
 	TOptional<FTransform>										Transform;
+	TWeakObjectPtr<const UActorDescContainerInstance>			ParentContainerInstance;
 
 	UPROPERTY(Transient)
 	TObjectPtr<UActorDescContainer>								Container;
diff --git a/Engine/Source/Runtime/Engine/Public/WorldPartition/ActorDescContainerInstanceCollection.h b/Engine/Source/Runtime/Engine/Public/WorldPartition/ActorDescContainerInstanceCollection.h
--- a/Engine/Source/Runtime/Engine/Public/WorldPartition/ActorDescContainerInstanceCollection.h
+++ b/Engine/Source/Runtime/Engine/Public/WorldPartition/ActorDescContainerInstanceCollection.h
@@ -63,6 +63,9 @@ public:
 	DECLARE_EVENT_OneParam(UWorldPartition, FActorDescInstanceRemovedEvent, FWorldPartitionActorDescInstance*);
 	FActorDescInstanceRemovedEvent OnActorDescInstanceRemovedEvent;
 
+	DECLARE_EVENT_OneParam(UWorldPartition, FActorReplacedEvent, FWorldPartitionActorDescInstance*);
+	FActorReplacedEvent OnActorReplacedEvent;
+
 	void ForEachActorDescContainerInstanceBreakable(TFunctionRef<bool(ActorDescContPtrType)> Func, bool bRecursive = false);
 	void ForEachActorDescContainerInstanceBreakable(TFunctionRef<bool(ActorDescContPtrType)> Func, bool bRecursive = false) const;
 
@@ -80,6 +83,7 @@ private:
 
 	void OnActorDescInstanceAdded(FWorldPartitionActorDescInstance* ActorDescInstance);
 	void OnActorDescInstanceRemoved(FWorldPartitionActorDescInstance* ActorDescInstance);
+	void OnActorReplaced(FWorldPartitionActorDescInstance* ActorDescInstance);
 
 public:
 	template<bool bConst, class ActorType>
@@ -297,6 +301,7 @@ void TActorDescContainerInstanceCollection<ActorDescContPtrType>::RegisterDelega
 {
 	ConstCast(Container)->OnActorDescInstanceAddedEvent.AddRaw(this, &TActorDescContainerInstanceCollection<ActorDescContPtrType>::OnActorDescInstanceAdded);
 	ConstCast(Container)->OnActorDescInstanceRemovedEvent.AddRaw(this, &TActorDescContainerInstanceCollection<ActorDescContPtrType>::OnActorDescInstanceRemoved);
+	ConstCast(Container)->OnActorReplacedEvent.AddRaw(this, &TActorDescContainerInstanceCollection<ActorDescContPtrType>::OnActorReplaced);
 }
 
 template<class ActorDescContPtrType>
@@ -304,6 +309,7 @@ void TActorDescContainerInstanceCollection<ActorDescContPtrType>::UnregisterDele
 {
 	ConstCast(Container)->OnActorDescInstanceAddedEvent.RemoveAll(this);
 	ConstCast(Container)->OnActorDescInstanceRemovedEvent.RemoveAll(this);
+	ConstCast(Container)->OnActorReplacedEvent.RemoveAll(this);
 }
 
 template<class ActorDescContPtrType>
@@ -504,6 +510,12 @@ void TActorDescContainerInstanceCollection<ActorDescContPtrType>::OnActorDescIns
 	OnActorDescInstanceRemovedEvent.Broadcast(ActorDescInstance);
 }
 
+template<class ActorDescContPtrType>
+void TActorDescContainerInstanceCollection<ActorDescContPtrType>::OnActorReplaced(FWorldPartitionActorDescInstance* ActorDescInstance)
+{
+	OnActorReplacedEvent.Broadcast(ActorDescInstance);
+}
+
 #endif 
 
 using FActorDescContainerInstanceCollection = TActorDescContainerInstanceCollection<TObjectPtr<UActorDescContainerInstance>>;
\ No newline at end of file
diff --git a/Engine/Source/Runtime/Engine/Public/WorldPartition/ActorDescList.h b/Engine/Source/Runtime/Engine/Public/WorldPartition/ActorDescList.h
--- a/Engine/Source/Runtime/Engine/Public/WorldPartition/ActorDescList.h
+++ b/Engine/Source/Runtime/Engine/Public/WorldPartition/ActorDescList.h
@@ -29,20 +29,22 @@ class TActorDescList
 		}
 	};
 
+protected:
 	using FGuidActorDescMap = TMap<FGuid, TUniquePtr<DescType>*, FDefaultSetAllocator, FActorGuidKeyFuncs>;
 	using FActorDescArray = TChunkedArray<TUniquePtr<DescType>>;
 
 public:
-	TActorDescList() {}
+	TActorDescList() : bIsProxy(false) {}
 	virtual ~TActorDescList() {}
 
 	// Non-copyable
 	TActorDescList(const TActorDescList&) = delete;
 	TActorDescList& operator=(const TActorDescList&) = delete;
 		
-	bool IsEmpty() const { return ActorsByGuid.Num() == 0; }
+	bool IsEmpty() const { return GetActorsByGuid().Num() == 0; }
 	ENGINE_API void Empty()
 	{
+		check(!bIsProxy || ActorsByGuid.IsEmpty());
 		ActorsByGuid.Empty();
 		ActorDescList.Empty();
 	}
@@ -61,8 +63,9 @@ public:
 
 	public:
 		TBaseIterator(ListType InActorDescList, UClass* InActorClass)
-			: ActorsIterator(InActorDescList->ActorsByGuid)
+			: ActorsIterator(InActorDescList->GetActorsByGuid())
 			, ActorClass(InActorClass)
+			, bIsProxy(InActorDescList->bIsProxy)
 		{
 			check(ActorClass->IsNative());
 			check(ActorClass->IsChildOf(ActorType::StaticClass()));
@@ -89,6 +92,7 @@ public:
 		 */
 		FORCEINLINE void RemoveCurrent()
 		{
+			check(!bIsProxy);
 			ActorsIterator.RemoveCurrent();
 		}
 
@@ -147,6 +151,7 @@ public:
 
 		IteratorType ActorsIterator;
 		UClass* ActorClass;
+		bool bIsProxy;
 	};
 
 	template <class ActorType = AActor>
@@ -173,8 +178,17 @@ public:
 
 	void AddActorDescriptor(DescType* ActorDesc);
 	void RemoveActorDescriptor(DescType* ActorDesc);
-
+		
+	virtual FGuidActorDescMap& GetActorsByGuid() const { return GetProxyActorsByGuid(); }
 protected:
+	void SetIsProxy()
+	{
+		check(ActorsByGuid.IsEmpty());
+		bIsProxy = true;
+	}
+	// Allow Child class to Proxy to another list
+	virtual FGuidActorDescMap& GetProxyActorsByGuid() const { return const_cast<FGuidActorDescMap&>(ActorsByGuid); }
+
 	TUniquePtr<DescType>* GetActorDescriptor(const FGuid& ActorGuid);
 	const TUniquePtr<DescType>* GetActorDescriptor(const FGuid& ActorGuid) const;
 
@@ -183,6 +197,7 @@ protected:
 
 	FActorDescArray ActorDescList;
 	FGuidActorDescMap ActorsByGuid;
+	bool bIsProxy;
 #endif
 };
 
@@ -190,6 +205,7 @@ protected:
 template<class DescType>
 void TActorDescList<DescType>::AddActorDescriptor(DescType* ActorDesc)
 {
+	check(!bIsProxy);
 	check(ActorDesc);
 	checkf(!ActorsByGuid.Contains(ActorDesc->GetGuid()), TEXT("Duplicated actor descriptor guid '%s' detected: `%s`"), *ActorDesc->GetGuid().ToString(), *ActorDesc->GetActorName().ToString());
 
@@ -200,6 +216,7 @@ void TActorDescList<DescType>::AddActorDescriptor(DescType* ActorDesc)
 template<class DescType>
 void TActorDescList<DescType>::RemoveActorDescriptor(DescType* ActorDesc)
 {
+	check(!bIsProxy);
 	check(ActorDesc);
 	verify(ActorsByGuid.Remove(ActorDesc->GetGuid()));
 }
@@ -207,25 +224,25 @@ void TActorDescList<DescType>::RemoveActorDescriptor(DescType* ActorDesc)
 template<class DescType>
 TUniquePtr<DescType>* TActorDescList<DescType>::GetActorDescriptor(const FGuid& ActorGuid)
 {
-	return ActorsByGuid.FindRef(ActorGuid);
+	return GetProxyActorsByGuid().FindRef(ActorGuid);
 }
 
 template<class DescType>
 const TUniquePtr<DescType>* TActorDescList<DescType>::GetActorDescriptor(const FGuid& ActorGuid) const
 {
-	return ActorsByGuid.FindRef(ActorGuid);
+	return GetProxyActorsByGuid().FindRef(ActorGuid);
 }
 
 template<class DescType>
 TUniquePtr<DescType>* TActorDescList<DescType>::GetActorDescriptorChecked(const FGuid& ActorGuid)
 {
-	return ActorsByGuid.FindChecked(ActorGuid);
+	return GetProxyActorsByGuid().FindChecked(ActorGuid);
 }
 
 template<class DescType>
 const TUniquePtr<DescType>* TActorDescList<DescType>::GetActorDescriptorChecked(const FGuid& ActorGuid) const
 {
-	return ActorsByGuid.FindChecked(ActorGuid);
+	return GetProxyActorsByGuid().FindChecked(ActorGuid);
 }
 
 // FActorDescList supports subclassing through ActorType's and FWorldPartitionActorDescType will return  the proper iterator value type
@@ -240,13 +257,13 @@ class FActorDescList : public TActorDescList<FWorldPartitionActorDesc>
 {
 #if WITH_EDITOR
 public:
-	ENGINE_API FWorldPartitionActorDesc* GetActorDesc(const FGuid& Guid);
-	ENGINE_API const FWorldPartitionActorDesc* GetActorDesc(const FGuid& Guid) const;
+	ENGINE_API virtual FWorldPartitionActorDesc* GetActorDesc(const FGuid& Guid);
+	ENGINE_API virtual const FWorldPartitionActorDesc* GetActorDesc(const FGuid& Guid) const;
 
-	ENGINE_API FWorldPartitionActorDesc& GetActorDescChecked(const FGuid& Guid);
-	ENGINE_API const FWorldPartitionActorDesc& GetActorDescChecked(const FGuid& Guid) const;
+	ENGINE_API virtual FWorldPartitionActorDesc& GetActorDescChecked(const FGuid& Guid);
+	ENGINE_API virtual const FWorldPartitionActorDesc& GetActorDescChecked(const FGuid& Guid) const;
 
-	int32 GetActorDescCount() const { return ActorsByGuid.Num(); }
+	int32 GetActorDescCount() const { return GetProxyActorsByGuid().Num(); }
 
 	ENGINE_API FWorldPartitionActorDesc* AddActor(const AActor* InActor);
 #endif
diff --git a/Engine/Source/Runtime/Engine/Public/WorldPartition/LevelInstance/LevelInstanceActorDesc.h b/Engine/Source/Runtime/Engine/Public/WorldPartition/LevelInstance/LevelInstanceActorDesc.h
--- a/Engine/Source/Runtime/Engine/Public/WorldPartition/LevelInstance/LevelInstanceActorDesc.h
+++ b/Engine/Source/Runtime/Engine/Public/WorldPartition/LevelInstance/LevelInstanceActorDesc.h
@@ -9,6 +9,7 @@
 #include "UObject/ObjectPtr.h"
 #include "Containers/Map.h"
 #include "WorldPartition/WorldPartitionActorDesc.h"
+#include "WorldPartition/LevelInstance/LevelInstancePropertyOverrideDesc.h"
 #include "WorldPartition/Filter/WorldPartitionActorFilter.h"
 
 class ULevelInstanceSubsystem;
@@ -31,6 +32,7 @@ public:
 
 	virtual EWorldPartitionActorFilterType GetChildContainerFilterType() const { return IsChildContainerInstance() ? EWorldPartitionActorFilterType::Loading : EWorldPartitionActorFilterType::None; }
 	virtual FName GetChildContainerPackage() const override { return WorldAsset.GetLongPackageFName(); }
+	virtual FString GetChildContainerName() const override;
 	virtual const FWorldPartitionActorFilter* GetChildContainerFilter() const override { return &Filter; }
 		
 	ENGINE_API virtual void CheckForErrors(const IWorldPartitionActorDescInstanceView* InActorDescView, IStreamingGenerationErrorHandler* ErrorHandler) const override;
@@ -48,6 +50,7 @@ protected:
 	ENGINE_API virtual bool GetChildContainerInstance(const FWorldPartitionActorDescInstance* InActorDescInstance, FContainerInstance& OutContainerInstance) const override;
 	ENGINE_API virtual UWorldPartition* GetLoadedChildWorldPartition(const FWorldPartitionActorDescInstance* InActorDescInstance) const override;
 	ENGINE_API FTransform GetChildContainerTransform() const;
+	FLevelInstancePropertyOverrideDesc* GetOverrideDesc() const { return PropertyOverrideDesc.Get(); }
 
 	FSoftObjectPath WorldAsset;
 	ELevelInstanceRuntimeBehavior DesiredRuntimeBehavior;
@@ -57,11 +60,12 @@ protected:
 	FWorldPartitionActorFilter Filter;
 	bool bIsChildContainerInstance;
 
+	FSoftObjectPath PropertyOverrideAsset;
+	TSharedPtr<FLevelInstancePropertyOverrideDesc> PropertyOverrideDesc;
 private:
 	ENGINE_API bool IsChildContainerInstanceInternal() const;
 	ENGINE_API void RegisterChildContainer();
 	ENGINE_API void UnregisterChildContainer();
 	ENGINE_API void UpdateBounds();
-	FString GetChildContainerName() const;
 };
 #endif
diff --git a/Engine/Source/Runtime/Engine/Public/WorldPartition/WorldPartition.h b/Engine/Source/Runtime/Engine/Public/WorldPartition/WorldPartition.h
--- a/Engine/Source/Runtime/Engine/Public/WorldPartition/WorldPartition.h
+++ b/Engine/Source/Runtime/Engine/Public/WorldPartition/WorldPartition.h
@@ -523,6 +523,12 @@ public:
 #if WITH_EDITOR
 	UActorDescContainerInstance* GetActorDescContainerInstance() const { return ActorDescContainerInstance; }
 
+	void SetContainerInstanceClass(TSubclassOf<UActorDescContainerInstance> InContainerInstanceClass)
+	{
+		check(!IsInitialized());
+		ContainerInstanceClass = InContainerInstanceClass;
+	}
+
 	UE_DEPRECATED(5.4, "Use ForEachActorDescContainerInstanceBreakable.")
 	void ForEachActorDescContainerBreakable(TFunctionRef<bool(UActorDescContainer*)> Func) {}
 	UE_DEPRECATED(5.4, "Use ForEachActorDescContainerInstanceBreakable.")
@@ -602,6 +608,8 @@ private:
 	UPROPERTY(Transient)
 	TObjectPtr<UActorDescContainerInstance> ActorDescContainerInstance;
 
+	UPROPERTY(Transient)
+	TSubclassOf<UActorDescContainerInstance> ContainerInstanceClass;
 public:
 	TOptional<bool> bOverrideEnableStreamingInEditor;
 
diff --git a/Engine/Source/Runtime/Engine/Public/WorldPartition/WorldPartitionActorContainerID.h b/Engine/Source/Runtime/Engine/Public/WorldPartition/WorldPartitionActorContainerID.h
--- a/Engine/Source/Runtime/Engine/Public/WorldPartition/WorldPartitionActorContainerID.h
+++ b/Engine/Source/Runtime/Engine/Public/WorldPartition/WorldPartitionActorContainerID.h
@@ -4,6 +4,41 @@
 #include "Hash/CityHash.h"
 #include "WorldPartitionActorContainerID.generated.h"
 
+USTRUCT()
+struct FActorContainerPath
+{
+	GENERATED_BODY()
+
+public:
+	UPROPERTY()
+	TArray<FGuid> ContainerGuids;
+
+	friend uint32 GetTypeHash(const FActorContainerPath& InActorContainerPath)
+	{
+		if (InActorContainerPath.ContainerGuids.Num() > 0)
+		{
+			FGuid Guid = InActorContainerPath.ContainerGuids[0];
+			for (int32 i = 1; i < InActorContainerPath.ContainerGuids.Num(); ++i)
+			{
+				Guid = FGuid::Combine(Guid, InActorContainerPath.ContainerGuids[i]);
+			}
+
+			return GetTypeHash(Guid);
+		}
+
+		return 0;
+	}
+
+	friend FArchive& operator<<(FArchive& Ar, FActorContainerPath& InActorContainerPath)
+	{
+		Ar << InActorContainerPath.ContainerGuids;
+		return Ar;
+	}
+
+	bool operator==(const FActorContainerPath& InOther) const = default;
+	bool operator!=(const FActorContainerPath& InOther) const = default;
+};
+
 USTRUCT()
 struct FActorContainerID
 {
@@ -24,6 +59,15 @@ struct FActorContainerID
 		: Guid(FGuid::Combine(InParent.Guid, InActorGuid))
 	{}
 
+	FActorContainerID(const FActorContainerID& InParent, const FActorContainerPath& InPath)
+		: Guid(InParent.Guid)
+	{
+		for (const FGuid& ContainerGuid : InPath.ContainerGuids)
+		{
+			Guid = FGuid::Combine(Guid, ContainerGuid);
+		}
+	}
+
 	void operator=(const FActorContainerID& InOther)
 	{
 		Guid = InOther.Guid;
diff --git a/Engine/Source/Runtime/Engine/Public/WorldPartition/WorldPartitionActorDesc.h b/Engine/Source/Runtime/Engine/Public/WorldPartition/WorldPartitionActorDesc.h
--- a/Engine/Source/Runtime/Engine/Public/WorldPartition/WorldPartitionActorDesc.h
+++ b/Engine/Source/Runtime/Engine/Public/WorldPartition/WorldPartitionActorDesc.h
@@ -12,21 +12,61 @@
 #include "UObject/WeakObjectPtr.h"
 #include "Templates/SubclassOf.h"
 #include "Misc/Guid.h"
+#include "Misc/TVariant.h"
 #include "WorldPartition/WorldPartitionActorDescType.h"
 #include "WorldPartition/WorldPartitionActorContainerID.h"
 #include "WorldPartition/Filter/WorldPartitionActorFilter.h"
 
+class FActorDescArchive;
+
 // Struct used to create actor descriptor
 struct FWorldPartitionActorDescInitData
 {
+	FWorldPartitionActorDescInitData()
+		: DataSource(TInPlaceType<TArray<uint8>>(), TArray<uint8>())
+	{
+	}
+
+	FWorldPartitionActorDescInitData(FActorDescArchive* InArchive)
+		: DataSource(TInPlaceType<FActorDescArchive*>(), InArchive)
+	{
+	}
+
 	UClass* NativeClass;
 	FName PackageName;
 	FSoftObjectPath ActorPath;
-	TArray<uint8> SerializedData;
 
+	TArray<uint8>& GetSerializedData()
+	{
+		check(DataSource.IsType<TArray<uint8>>());
+		return DataSource.Get<TArray<uint8>>();
+	}
+
+	const TArray<uint8>& GetSerializedData() const
+	{
+		check(DataSource.IsType<TArray<uint8>>());
+		return DataSource.Get<TArray<uint8>>();
+	}
+
+	FActorDescArchive* GetArchive() const
+	{
+		check(DataSource.IsType<FActorDescArchive*>());
+		return DataSource.Get<FActorDescArchive*>();
+	}
+
+	bool IsUsingArchive() const
+	{
+		return DataSource.IsType<FActorDescArchive*>();
+	}
+		
 	FWorldPartitionActorDescInitData& SetNativeClass(UClass* InNativeClass) { NativeClass = InNativeClass; return *this; }
 	FWorldPartitionActorDescInitData& SetPackageName(FName InPackageName) { PackageName = InPackageName; return *this; }
 	FWorldPartitionActorDescInitData& SetActorPath(const FSoftObjectPath& InActorPath) { ActorPath = InActorPath; return *this; }
+
+private:
+	// Provide SerializedData or an already initialized Archive
+	TVariant<TArray<uint8>, FActorDescArchive*> DataSource;
+
 };
 
 struct FWorldPartitionAssetDataPatcher
@@ -95,6 +135,7 @@ class FWorldPartitionActorDesc
 	friend class FActorDescArchive;
 	friend class FWorldPartitionActorDescInstance;
 	friend class FStreamingGenerationUnsavedDirtyActorDescInstance;
+	friend class FPropertyOverrideUtils;
 	template<class U> friend class TActorDescContainerCollection;
 
 public:
@@ -161,6 +202,7 @@ public:
 		
 	virtual bool IsChildContainerInstance() const { return false; }
 	virtual FName GetChildContainerPackage() const { return NAME_None; }
+	virtual FString GetChildContainerName() const { return FString(); }
 	virtual EWorldPartitionActorFilterType GetChildContainerFilterType() const { return EWorldPartitionActorFilterType::None; }
 	virtual const FWorldPartitionActorFilter* GetChildContainerFilter() const { return nullptr; }
 	virtual UActorDescContainer* GetChildContainer() const { return nullptr; }
@@ -328,7 +370,7 @@ public:
 	 */
 	ENGINE_API virtual bool ShouldResave(const FWorldPartitionActorDesc* Other) const;
 
-	ENGINE_API void SerializeTo(TArray<uint8>& OutData) const;
+	ENGINE_API void SerializeTo(TArray<uint8>& OutData, FWorldPartitionActorDesc* BaseDesc = nullptr) const;
 
 	using FActorDescDeprecator = TFunction<void(FArchive&, FWorldPartitionActorDesc*)>;
 	static ENGINE_API void RegisterActorDescDeprecator(TSubclassOf<AActor> ActorClass, const FActorDescDeprecator& Deprecator);
diff --git a/Engine/Source/Runtime/Engine/Public/WorldPartition/WorldPartitionActorDescArchive.h b/Engine/Source/Runtime/Engine/Public/WorldPartition/WorldPartitionActorDescArchive.h
--- a/Engine/Source/Runtime/Engine/Public/WorldPartition/WorldPartitionActorDescArchive.h
+++ b/Engine/Source/Runtime/Engine/Public/WorldPartition/WorldPartitionActorDescArchive.h
@@ -13,7 +13,7 @@ struct FWorldPartitionAssetDataPatcher;
 class FActorDescArchive : public FArchiveProxy
 {
 public:
-	FActorDescArchive(FArchive& InArchive, FWorldPartitionActorDesc* InActorDesc);
+	FActorDescArchive(FArchive& InArchive, FWorldPartitionActorDesc* InActorDesc, const FWorldPartitionActorDesc* InBaseActorDesc = nullptr);
 
 	void Init(const FTopLevelAssetPath InClassPath = FTopLevelAssetPath());
 
@@ -46,23 +46,23 @@ public:
 		{
 			FActorDescArchive& ActorDescAr = (FActorDescArchive&)Ar;
 			
-			check(ActorDescAr.ClassDesc || Ar.IsSaving());
+			check(ActorDescAr.BaseDesc || Ar.IsSaving());
 			check(ActorDescAr.ActorDesc);
 
 			Ar.UsingCustomVersion(FFortniteMainBranchObjectVersion::GUID);
 
-			auto GetClassDefaultValue = [&V, &ActorDescAr]() -> const DestPropertyType*
+			auto GetBaseDefaultValue = [&V, &ActorDescAr]() -> const DestPropertyType*
 			{
 				const UPTRINT PropertyOffset = (UPTRINT)&V.Value - *(UPTRINT*)&ActorDescAr.ActorDesc;
 
-				if (PropertyOffset < ActorDescAr.ClassDescSizeof)
+				if (PropertyOffset < ActorDescAr.BaseDescSizeof)
 				{
-					check((PropertyOffset + sizeof(V.Value)) <= ActorDescAr.ClassDescSizeof);
-					const DestPropertyType* RefValue = (const DestPropertyType*)(*(UPTRINT*)&ActorDescAr.ClassDesc + PropertyOffset);
+					check((PropertyOffset + sizeof(V.Value)) <= ActorDescAr.BaseDescSizeof);
+					const DestPropertyType* RefValue = (const DestPropertyType*)(*(UPTRINT*)&ActorDescAr.BaseDesc + PropertyOffset);
 					return RefValue;
 				}
 
-				check(ActorDescAr.bIsMissingClassDesc);
+				check(ActorDescAr.bIsMissingBaseDesc);
 				return nullptr;
 			};
 
@@ -70,15 +70,15 @@ public:
 
 			if (Ar.CustomVer(FFortniteMainBranchObjectVersion::GUID) >= FFortniteMainBranchObjectVersion::WorldPartitionActorClassDescSerialize)
 			{
-				if (Ar.IsSaving() && ActorDescAr.ClassDesc)
+				if (Ar.IsSaving() && ActorDescAr.BaseDesc)
 				{
 					if constexpr (std::is_same_v<DestPropertyType, SourcePropertyType>)
 					{
 						// When saving, we expect the class descriptor to be the exact type as what we are serializing.
-						const DestPropertyType* ClassDefaultValue = GetClassDefaultValue();
-						check(ClassDefaultValue);
+						const DestPropertyType* BaseDefaultValue = GetBaseDefaultValue();
+						check(BaseDefaultValue);
 
-						bSerialize = (V.Value != *ClassDefaultValue) ? 1 : 0;
+						bSerialize = (V.Value != *BaseDefaultValue) ? 1 : 0;
 					}
 				}
 
@@ -103,9 +103,9 @@ public:
 			else if (Ar.IsLoading())
 			{
 				// When loading, we need to handle a different class descriptor in case of missing classes, etc.
-				if (const DestPropertyType* ClassDefaultValue = GetClassDefaultValue())
+				if (const DestPropertyType* BaseDefaultValue = GetBaseDefaultValue())
 				{
-					V.Value = *ClassDefaultValue;
+					V.Value = *BaseDefaultValue;
 				}
 			}
 
@@ -117,9 +117,9 @@ public:
 	};
 
 	FWorldPartitionActorDesc* ActorDesc;
-	const FWorldPartitionActorDesc* ClassDesc;
-	uint32 ClassDescSizeof;
-	bool bIsMissingClassDesc;
+	const FWorldPartitionActorDesc* BaseDesc;
+	uint32 BaseDescSizeof;
+	bool bIsMissingBaseDesc;
 };
 
 class FActorDescArchivePatcher : public FActorDescArchive
diff --git a/Engine/Source/Runtime/Engine/Public/WorldPartition/WorldPartitionActorDescInstance.h b/Engine/Source/Runtime/Engine/Public/WorldPartition/WorldPartitionActorDescInstance.h
--- a/Engine/Source/Runtime/Engine/Public/WorldPartition/WorldPartitionActorDescInstance.h
+++ b/Engine/Source/Runtime/Engine/Public/WorldPartition/WorldPartitionActorDescInstance.h
@@ -78,6 +78,7 @@ public:
 
 	virtual bool IsChildContainerInstance() const { return ChildContainerInstance || ActorDesc->IsChildContainerInstance(); }
 	virtual FName GetChildContainerPackage() const { return ActorDesc->GetChildContainerPackage(); }
+	virtual FString GetChildContainerName() const { return ActorDesc->GetChildContainerName(); }
 	virtual EWorldPartitionActorFilterType GetChildContainerFilterType() const { return ActorDesc->GetChildContainerFilterType(); }
 	virtual const FWorldPartitionActorFilter* GetChildContainerFilter() const { return ActorDesc->GetChildContainerFilter(); }
 	virtual bool GetChildContainerInstance(FWorldPartitionActorDesc::FContainerInstance& OutContainerInstance) const { return ActorDesc->GetChildContainerInstance(this, OutContainerInstance); }
diff --git a/Engine/Source/Runtime/Engine/Public/WorldPartition/WorldPartitionActorDescUtils.h b/Engine/Source/Runtime/Engine/Public/WorldPartition/WorldPartitionActorDescUtils.h
--- a/Engine/Source/Runtime/Engine/Public/WorldPartition/WorldPartitionActorDescUtils.h
+++ b/Engine/Source/Runtime/Engine/Public/WorldPartition/WorldPartitionActorDescUtils.h
@@ -31,6 +31,13 @@ struct FWorldPartitionActorDescUtils
 	 */
 	static ENGINE_API UClass* GetActorNativeClassFromAssetData(const FAssetData& InAssetData);
 
+	/**
+	 * Retrieve the actor's native class from the string.
+	 * @param InClassPath	The class path to resolve.
+	 * @return				The actor's native class.
+	 */
+	static ENGINE_API UClass* GetActorNativeClassFromString(const FString& InClassPath);
+
 	/** 
 	 * Creates a valid actor descriptor from the provided asset data.
 	 * @param InAssetData	The asset data to look into.
diff --git a/Engine/Source/Runtime/Engine/Public/WorldPartition/WorldPartitionLevelHelper.h b/Engine/Source/Runtime/Engine/Public/WorldPartition/WorldPartitionLevelHelper.h
--- a/Engine/Source/Runtime/Engine/Public/WorldPartition/WorldPartitionLevelHelper.h
+++ b/Engine/Source/Runtime/Engine/Public/WorldPartition/WorldPartitionLevelHelper.h
@@ -11,7 +11,9 @@
 
 #include "Engine/World.h"
 #include "UObject/LinkerInstancingContext.h"
+#include "UObject/UObjectAnnotation.h"
 #include "WorldPartition/WorldPartitionRuntimeCell.h"
+#include "WorldPartition/WorldPartitionPropertyOverride.h"
 
 class FWorldPartitionPackageHelper;
 class UWorldPartition;
@@ -70,7 +72,7 @@ public:
 	};
 
 	static bool LoadActors(const FLoadActorsParams& InParams);
-	
+
 	static FSoftObjectPath RemapActorPath(const FActorContainerID& InContainerID, const FString& SourceWorldName, const FSoftObjectPath& InActorPath);
 
 private:
@@ -87,6 +89,15 @@ private:
 	friend class FContentBundleEditor;
 	static bool RemapLevelCellPathInContentBundle(ULevel* Level, const class FContentBundleEditor* ContentBundleEditor, const UWorldPartitionRuntimeCell* Cell);
 
+	struct FLoadedPropertyOverrides
+	{
+		TMap<FActorContainerID, const UWorldPartitionPropertyOverride*> PropertyOverrides;
+	};
+
+	static bool LoadActors(FLoadActorsParams&& InParams);
+	static bool LoadActorsInternal(FLoadActorsParams&& InParams, FLoadedPropertyOverrides&& InLoadedPropertyOverrides);
+	static bool LoadActorsWithPropertyOverridesInternal(FLoadActorsParams&& InParams);
+
 	struct FPackageReference
 	{
 		TSet<FPackageReferencer*> Referencers;
@@ -98,6 +109,30 @@ private:
 	TMap<FName, FPackageReference> PackageReferences;
 
 	TSet<TWeakObjectPtr<UPackage>> PreGCPackagesToUnload;
+
+	friend class UWorldPartitionLevelStreamingDynamic;
+	friend class UWorldPartitionSubsystem;
+	friend class UWorldPartitionRuntimeLevelStreamingCell;
+	// Cache of Property Overrides to apply after ReRunConstructionScript
+	// In PIE this will  be done when Streaming state changes to Visible on the UWorldPartitionLevelStreamingDynamic
+	// In Cook this will be done on save of the Level cell
+	class FActorPropertyOverridesAnnotation
+	{
+	public:
+		FActorPropertyOverridesAnnotation() {}
+		FActorPropertyOverridesAnnotation(TArray<FActorPropertyOverride>&& InActorPropertyOverrides, const FTransform& InContainerTransform) : ActorPropertyOverrides(MoveTemp(InActorPropertyOverrides)), ContainerTransform(InContainerTransform) {}
+
+		TArray<FActorPropertyOverride> ActorPropertyOverrides;
+		FTransform ContainerTransform;
+		FORCEINLINE bool IsDefault()
+		{
+			return ActorPropertyOverrides.IsEmpty();
+		}
+	};
+	static FUObjectAnnotationSparse<FActorPropertyOverridesAnnotation, true> ActorPropertyOverridesAnnotation;
+
+	// Apply Existing Property Override annotation to Actor
+	static void ApplyConstructionScriptPropertyOverridesFromAnnotation(AActor* InActor);
 #endif
 };
 
diff --git a/Engine/Source/Runtime/Engine/Public/WorldPartition/WorldPartitionLevelStreamingDynamic.h b/Engine/Source/Runtime/Engine/Public/WorldPartition/WorldPartitionLevelStreamingDynamic.h
--- a/Engine/Source/Runtime/Engine/Public/WorldPartition/WorldPartitionLevelStreamingDynamic.h
+++ b/Engine/Source/Runtime/Engine/Public/WorldPartition/WorldPartitionLevelStreamingDynamic.h
@@ -68,6 +68,7 @@ private:
 	ENGINE_API bool IssueLoadRequests();
 	ENGINE_API void FinalizeRuntimeLevel();
 	ENGINE_API void OnCleanupLevel();
+	virtual void OnCurrentStateChanged(ELevelStreamingState InPrevState, ELevelStreamingState InNewState) override;
 
 	ENGINE_API void Initialize(UWorld* OuterWorld, const TArray<FWorldPartitionRuntimeCellObjectMapping>& InPackages);
 
diff --git a/Engine/Source/Runtime/Engine/Public/WorldPartition/WorldPartitionRuntimeCell.h b/Engine/Source/Runtime/Engine/Public/WorldPartition/WorldPartitionRuntimeCell.h
--- a/Engine/Source/Runtime/Engine/Public/WorldPartition/WorldPartitionRuntimeCell.h
+++ b/Engine/Source/Runtime/Engine/Public/WorldPartition/WorldPartitionRuntimeCell.h
@@ -34,6 +34,38 @@ enum class EWorldPartitionRuntimeCellVisualizeMode
 	StreamingStatus
 };
 
+USTRUCT()
+struct FWorldPartitionRuntimeCellPropertyOverride
+{
+	GENERATED_USTRUCT_BODY()
+
+	FWorldPartitionRuntimeCellPropertyOverride()
+	{}
+
+	FWorldPartitionRuntimeCellPropertyOverride(const FActorContainerID& InOwnerContainerID, const FString& InAssetPath, FName InPackageName, const FActorContainerPath& InContainerPath)
+#if WITH_EDITORONLY_DATA
+		: OwnerContainerID(InOwnerContainerID)
+		, AssetPath(InAssetPath)
+		, PackageName(InPackageName)
+		, ContainerPath(InContainerPath)
+#endif
+	{}
+
+#if WITH_EDITORONLY_DATA
+	UPROPERTY()
+	FActorContainerID OwnerContainerID;
+
+	UPROPERTY()
+	FString AssetPath;
+
+	UPROPERTY()
+	FName PackageName;
+
+	UPROPERTY()
+	FActorContainerPath ContainerPath;
+#endif
+};
+
 USTRUCT()
 struct FWorldPartitionRuntimeCellObjectMapping
 {
@@ -130,6 +162,9 @@ struct FWorldPartitionRuntimeCellObjectMapping
 
 	UPROPERTY()
 	bool bIsEditorOnly;
+		
+	UPROPERTY()
+	TArray<FWorldPartitionRuntimeCellPropertyOverride> PropertyOverrides;
 #endif
 };
 
diff --git a/Engine/Source/Runtime/Engine/Public/WorldPartition/WorldPartitionSettings.h b/Engine/Source/Runtime/Engine/Public/WorldPartition/WorldPartitionSettings.h
--- a/Engine/Source/Runtime/Engine/Public/WorldPartition/WorldPartitionSettings.h
+++ b/Engine/Source/Runtime/Engine/Public/WorldPartition/WorldPartitionSettings.h
@@ -5,6 +5,8 @@
 #include "WorldPartition/WorldPartition.h"
 #include "WorldPartitionSettings.generated.h"
 
+class UWorldPartitionPropertyOverridePolicy;
+
 UCLASS(config = Engine, defaultconfig, DisplayName = "World Partition")
 class ENGINE_API UWorldPartitionSettings : public UDeveloperSettings
 {
@@ -22,7 +24,17 @@ public:
 	TSubclassOf<UWorldPartitionEditorHash> GetEditorHashDefaultClass() const { return EditorHashDefaultClass; }
 	TSubclassOf<UWorldPartitionRuntimeHash> GetRuntimeHashDefaultClass() const { return RuntimeHashDefaultClass; }
 
+	UWorldPartitionPropertyOverridePolicy* GetPropertyOverridePolicy() const
+	{
+		return PropertyOverridePolicy;
+	}
 protected:
+	friend class ULevelInstanceSettings;
+	void SetPropertyOverridePolicy(UWorldPartitionPropertyOverridePolicy* InPropertyOverridePolicy)
+	{
+		PropertyOverridePolicy = InPropertyOverridePolicy;
+	}
+	
 	/** Set the default logical operator for actor data layers activation for new maps */
 	UPROPERTY(EditAnywhere, Config, Category = WorldPartition)
 	EWorldPartitionDataLayersLogicOperator NewMapsDataLayersLogicOperator = EWorldPartitionDataLayersLogicOperator::Or;
@@ -42,4 +54,7 @@ protected:
 	/** Set the default runtime hash class to use for new maps  */
 	UPROPERTY(EditAnywhere, Config, NoClear, Category = WorldPartition)
 	TSubclassOf<UWorldPartitionRuntimeHash> RuntimeHashDefaultClass;
+
+	UPROPERTY(Transient)
+	TObjectPtr<UWorldPartitionPropertyOverridePolicy> PropertyOverridePolicy;
 };
\ No newline at end of file
diff --git a/Engine/Source/Runtime/Engine/Public/WorldPartition/WorldPartitionSubsystem.h b/Engine/Source/Runtime/Engine/Public/WorldPartition/WorldPartitionSubsystem.h
--- a/Engine/Source/Runtime/Engine/Public/WorldPartition/WorldPartitionSubsystem.h
+++ b/Engine/Source/Runtime/Engine/Public/WorldPartition/WorldPartitionSubsystem.h
@@ -61,6 +61,7 @@ public:
 
 	//~ Begin UWorldSubsystem Interface.
 	ENGINE_API virtual void UpdateStreamingState() override;
+	ENGINE_API virtual void OnWorldComponentsUpdated(UWorld& World) override;
 	//~ End UWorldSubsystem Interface.
 
 	//~ Begin FTickableGameObject
