Downgrade an error to a warning when encountering an actor cluster with an invalid target grid.

diff --git a/Engine/Source/Runtime/Engine/Private/WorldPartition/RuntimeSpatialHash/RuntimeSpatialHashHLOD.cpp b/Engine/Source/Runtime/Engine/Private/WorldPartition/RuntimeSpatialHash/RuntimeSpatialHashHLOD.cpp
--- a/Engine/Source/Runtime/Engine/Private/WorldPartition/RuntimeSpatialHash/RuntimeSpatialHashHLOD.cpp
+++ b/Engine/Source/Runtime/Engine/Private/WorldPartition/RuntimeSpatialHash/RuntimeSpatialHashHLOD.cpp
@@ -487,7 +487,7 @@ bool UWorldPartitionRuntimeSpatialHash::SetupHLODActors(const IStreamingGenerati
 		if (!FoundIndex)
 		{
 			//@todo_ow should this be done upstream?
-			UE_LOG(LogWorldPartitionRuntimeSpatialHashHLOD, Error, TEXT("Invalid partition grid '%s' referenced by actor cluster"), *ActorSetInstance.RuntimeGrid.ToString());
+			UE_LOG(LogWorldPartitionRuntimeSpatialHashHLOD, Warning, TEXT("Invalid partition grid '%s' referenced by actor cluster"), *ActorSetInstance.RuntimeGrid.ToString());
 		}
 
 		int32 GridIndex = FoundIndex ? *FoundIndex : 0;
diff --git a/Engine/Source/Runtime/Engine/Private/WorldPartition/WorldPartitionRuntimeSpatialHash.cpp b/Engine/Source/Runtime/Engine/Private/WorldPartition/WorldPartitionRuntimeSpatialHash.cpp
--- a/Engine/Source/Runtime/Engine/Private/WorldPartition/WorldPartitionRuntimeSpatialHash.cpp
+++ b/Engine/Source/Runtime/Engine/Private/WorldPartition/WorldPartitionRuntimeSpatialHash.cpp
@@ -1377,7 +1377,7 @@ bool UWorldPartitionRuntimeSpatialHash::GenerateStreaming(UWorldPartitionStreami
 		if (!FoundIndex)
 		{
 			//@todo_ow should this be done upstream?
-			UE_LOG(LogWorldPartition, Error, TEXT("Invalid partition grid '%s' referenced by actor cluster"), *ActorSetInstance.RuntimeGrid.ToString());
+			UE_LOG(LogWorldPartition, Warning, TEXT("Invalid partition grid '%s' referenced by actor cluster"), *ActorSetInstance.RuntimeGrid.ToString());
 		}
 
 		int32 GridIndex = FoundIndex ? *FoundIndex : 0;
