Added support of save as & duplicate worlds using External Data Layers

diff --git a/Engine/Plugins/Runtime/GameFeatures/Source/GameFeatures/Private/GameFeatureData.cpp b/Engine/Plugins/Runtime/GameFeatures/Source/GameFeatures/Private/GameFeatureData.cpp
--- a/Engine/Plugins/Runtime/GameFeatures/Source/GameFeatures/Private/GameFeatureData.cpp
+++ b/Engine/Plugins/Runtime/GameFeatures/Source/GameFeatures/Private/GameFeatureData.cpp
@@ -581,12 +581,6 @@ FName UGameFeatureData::GetContentBundleGuidsAssetRegistryTagPrivate()
 	return ContentBundlesTag;
 }
 
-FName UGameFeatureData::GetExternalDataLayerUIDsAssetRegistryTag()
-{
-	static const FName ExternalDataLayerUIDsTag("ExternalDataLayerUIDs");
-	return ExternalDataLayerUIDsTag;
-}
-
 void UGameFeatureData::GetContentBundleGuids(const FAssetData& Asset, TArray<FGuid>& OutContentBundleGuids)
 {
 	FString ContentBundleGuidsStr;
@@ -605,24 +599,6 @@ void UGameFeatureData::GetContentBundleGuids(const FAssetData& Asset, TArray<FGu
 	}
 }
 
-void UGameFeatureData::GetExternalDataLayerUIDs(const FAssetData& Asset, TArray<FExternalDataLayerUID>& OutExternalDataLayerUIDs)
-{
-	FString ExternalDataLayerUIDsStr;
-	if (Asset.GetTagValue(GetExternalDataLayerUIDsAssetRegistryTag(), ExternalDataLayerUIDsStr))
-	{
-		TArray<FString> ExternalDataLayerUIDStrArray;
-		ExternalDataLayerUIDsStr.ParseIntoArray(ExternalDataLayerUIDStrArray, TEXT(","));
-		for (const FString& ExternalDataLayerUIDStr : ExternalDataLayerUIDStrArray)
-		{
-			FExternalDataLayerUID ExternalDataLayerUID;
-			if (FExternalDataLayerUID::Parse(ExternalDataLayerUIDStr, ExternalDataLayerUID))
-			{
-				OutExternalDataLayerUIDs.Add(ExternalDataLayerUID);
-			}
-		}
-	}
-}
-
 void UGameFeatureData::GetDependencyDirectoriesFromAssetData(const FAssetData& AssetData, TArray<FString>& OutDependencyDirectories)
 {
 	const FString MountPoint = FPackageName::GetPackageMountPoint(AssetData.PackagePath.ToString()).ToString();
@@ -640,7 +616,7 @@ void UGameFeatureData::GetDependencyDirectoriesFromAssetData(const FAssetData& A
 	}
 
 	TArray<FExternalDataLayerUID> ExternalDataLayerUIDs;
-	GetExternalDataLayerUIDs(AssetData, ExternalDataLayerUIDs);
+	FExternalDataLayerHelper::GetExternalDataLayerUIDs(AssetData, ExternalDataLayerUIDs);
 	for (const FExternalDataLayerUID& ExternalDataLayerUID : ExternalDataLayerUIDs)
 	{
 		FString ExternalDataLayerRootPath;
@@ -691,11 +667,7 @@ void UGameFeatureData::GetAssetRegistryTags(FAssetRegistryTagsContext Context) c
 		Context.AddTag(FAssetRegistryTag(GetContentBundleGuidsAssetRegistryTagPrivate(), ContentBundleGuidsStr, FAssetRegistryTag::TT_Hidden));
 	}
 
-	if (ExternalDataLayerUIDs.Num() > 0)
-	{
-		FString ExternalDataLayerUIDsStr = FString::JoinBy(ExternalDataLayerUIDs, TEXT(","), [&](const FExternalDataLayerUID& ExternalDataLayerUID) { return ExternalDataLayerUID.ToString(); });
-		Context.AddTag(FAssetRegistryTag(GetExternalDataLayerUIDsAssetRegistryTag(), ExternalDataLayerUIDsStr, FAssetRegistryTag::TT_Hidden));
-	}
+	FExternalDataLayerHelper::AddAssetRegistryTags(Context, ExternalDataLayerUIDs);
 }
 #endif
 
diff --git a/Engine/Plugins/Runtime/GameFeatures/Source/GameFeatures/Public/GameFeatureData.h b/Engine/Plugins/Runtime/GameFeatures/Source/GameFeatures/Public/GameFeatureData.h
--- a/Engine/Plugins/Runtime/GameFeatures/Source/GameFeatures/Public/GameFeatureData.h
+++ b/Engine/Plugins/Runtime/GameFeatures/Source/GameFeatures/Public/GameFeatureData.h
@@ -84,8 +84,6 @@ public:
 private:
 #if WITH_EDITOR
 	static void GetContentBundleGuids(const FAssetData& Asset, TArray<FGuid>& OutContentBundleGuids);
-	static void GetExternalDataLayerUIDs(const FAssetData& Asset, TArray<FExternalDataLayerUID>& OutExternalDataLayerUIDs);
-	static FName GetExternalDataLayerUIDsAssetRegistryTag();
 	static FName GetContentBundleGuidsAssetRegistryTagPrivate();
 
 	UFUNCTION()
diff --git a/Engine/Plugins/Runtime/GameFeatures/Source/GameFeaturesEditor/Private/GameFeatureActionConvertContentBundleWorldPartitionBuilder.cpp b/Engine/Plugins/Runtime/GameFeatures/Source/GameFeaturesEditor/Private/GameFeatureActionConvertContentBundleWorldPartitionBuilder.cpp
--- a/Engine/Plugins/Runtime/GameFeatures/Source/GameFeaturesEditor/Private/GameFeatureActionConvertContentBundleWorldPartitionBuilder.cpp
+++ b/Engine/Plugins/Runtime/GameFeatures/Source/GameFeaturesEditor/Private/GameFeatureActionConvertContentBundleWorldPartitionBuilder.cpp
@@ -271,13 +271,13 @@ bool UGameFeatureActionConvertContentBundleWorldPartitionBuilder::RunInternal(UW
 			}, false);
 
 			// Create a new external package and assign it to the actor
-			const FString BaseDir = ULevel::GetExternalActorsPath(ExternalDataLayerManager->GetExternalDataLayerLevelRootPath(ExternalDataLayerAsset));
-			UPackage* NewActorPackage = ULevel::CreateActorPackage(BaseDir, Actor->GetLevel()->GetActorPackagingScheme(), Actor->GetName());
+			ULevel* DestinationLevel = Actor->GetLevel();
+			UPackage* NewActorPackage = ULevel::CreateActorPackage(DestinationLevel->GetPackage(), DestinationLevel->GetActorPackagingScheme(), Actor->GetName(), ExternalDataLayerAsset);
 			Actor->SetPackageExternal(true, true, NewActorPackage);
 
 			// Validation
 			check(NewActorPackage == Actor->GetExternalPackage());
-			check(NewActorPackage->GetName() == ExternalDataLayerManager->GetActorPackageName(ExternalDataLayerAsset, Actor->GetLevel(), Actor->GetName()));
+			check(NewActorPackage->GetName() == ExternalDataLayerManager->GetActorPackageName(ExternalDataLayerAsset, DestinationLevel, Actor->GetName()));
 			check(NewActorPackage->GetName() != OldActorPackageName);
 
 			// Move dependant objects into the new actor package
diff --git a/Engine/Source/Runtime/Engine/Classes/Engine/Level.h b/Engine/Source/Runtime/Engine/Classes/Engine/Level.h
--- a/Engine/Source/Runtime/Engine/Classes/Engine/Level.h
+++ b/Engine/Source/Runtime/Engine/Classes/Engine/Level.h
@@ -1115,17 +1115,22 @@ public:
 #if WITH_EDITOR
 	/**
 	 * Get the package name for this actor
+	 * @param InLevelPackage the package to get the external actors package name of
+	 * @param InActorPackagingScheme the packaging scheme to use
 	 * @param InActorPath the fully qualified actor path, in the format: 'Outermost.Outer.Name'
+	 * @param InMountPointContext an optional context object used to determine the mount point of the package
 	 * @return the package name
 	 */
-	static ENGINE_API FString GetActorPackageName(UPackage* InLevelPackage, EActorPackagingScheme ActorPackagingScheme, const FString& InActorPath);
+	static ENGINE_API FString GetActorPackageName(UPackage* InLevelPackage, EActorPackagingScheme InActorPackagingScheme, const FString& InActorPath, const UObject* InMountPointContext = nullptr);
 
 	/**
 	 * Get the package name for this actor
+	 * @param InBaseDir the base directory used when building the actor package name
+	 * @param InActorPackagingScheme the packaging scheme to use
 	 * @param InActorPath the fully qualified actor path, in the format: 'Outermost.Outer.Name'
 	 * @return the package name
 	 */
-	static ENGINE_API FString GetActorPackageName(const FString& InBaseDir, EActorPackagingScheme ActorPackagingScheme, const FString& InActorPath);
+	static ENGINE_API FString GetActorPackageName(const FString& InBaseDir, EActorPackagingScheme InActorPackagingScheme, const FString& InActorPath);
 
 	/**
 	 * Get the folder containing the external actors for this level path
@@ -1135,6 +1140,14 @@ public:
 	 */
 	static ENGINE_API FString GetExternalActorsPath(const FString& InLevelPackageName, const FString& InPackageShortName = FString());
 
+	/**
+	 * Get the folders containing the external actors for this level path, including actor folders of registered plugins for this level 
+	 * @param InLevelPackageName The package name to get the external actors path of
+	 * @param InPackageShortName Optional short name to use instead of the package short name
+	 * @return the folder
+	 */
+	static ENGINE_API TArray<FString> GetExternalActorsPaths(const FString& InLevelPackageName, const FString& InPackageShortName = FString());
+
 	/**
 	 * Get the folder containing the external actors for this level
 	 * @param InLevelPackage The package to get the external actors path of
@@ -1240,9 +1253,10 @@ public:
 	 * @param InLevelPackage the level package used when building the actor package name
 	 * @param InActorPackagingScheme the packaging scheme to use
 	 * @param InActorPath the fully qualified actor path, in the format: 'Outermost.Outer.Name'
+	 * @param InMountPointContext an optional context object used to determine the mount point of the package
 	 * @return the created package
 	 */
-	static ENGINE_API UPackage* CreateActorPackage(UPackage* InLevelPackage, EActorPackagingScheme InActorPackagingScheme, const FString& InActorPath);
+	static ENGINE_API UPackage* CreateActorPackage(UPackage* InLevelPackage, EActorPackagingScheme InActorPackagingScheme, const FString& InActorPath, const UObject* InMountPointContext = nullptr);
 
 	/**
 	 * Create an package for this actor
@@ -1348,6 +1362,17 @@ public:
 	/** meant to be called only from editor, calculating and storing static geometry to be used with off-line and/or on-line navigation building */
 	ENGINE_API void RebuildStaticNavigableGeometry();
 
+	DECLARE_DELEGATE_ThreeParams(FLevelExternalActorsPathsProviderDelegate, const FString&, const FString&, TArray<FString>&);
+	/** Registers a level external actor paths provider */
+	static ENGINE_API FDelegateHandle RegisterLevelExternalActorsPathsProvider(const FLevelExternalActorsPathsProviderDelegate& Provider);
+	/** Unregisters a level external actor paths provider */
+	static ENGINE_API void UnregisterLevelExternalActorsPathsProvider(const FDelegateHandle& ProviderDelegateHandle);
+
+	DECLARE_DELEGATE_RetVal_ThreeParams(bool, FLevelMountPointResolverDelegate, const FString&, const UObject*, FString&);
+	/** Registers a level mount point resolver */
+	static ENGINE_API FDelegateHandle RegisterLevelMountPointResolver(const FLevelMountPointResolverDelegate& Resolver);
+	/** Unregisters a level mount point resolver */
+	static ENGINE_API void UnregisterLevelMountPointResolver(const FDelegateHandle& ResolverDelegateHandle);
 #endif
 
 private:
@@ -1355,6 +1380,12 @@ private:
 #if WITH_EDITOR
 	bool IncrementalRunConstructionScripts(bool bProcessAllActors);
 	TOptional<bool> bCachedHasStaticMeshCompilationPending;
+
+	/** Array of registered delegates used by GetExternalActorsPaths. */
+	static TArray<FLevelExternalActorsPathsProviderDelegate> LevelExternalActorsPathsProviders;
+
+	/** Array of registered delegates used by GetExternalActorsPaths. */
+	static TArray<FLevelMountPointResolverDelegate> LevelMountPointResolvers;
 private:
 	/**
 	 * Potentially defer the running of an actor's construction script on load
diff --git a/Engine/Source/Runtime/Engine/Private/ActorEditor.cpp b/Engine/Source/Runtime/Engine/Private/ActorEditor.cpp
--- a/Engine/Source/Runtime/Engine/Private/ActorEditor.cpp
+++ b/Engine/Source/Runtime/Engine/Private/ActorEditor.cpp
@@ -1085,10 +1085,10 @@ void AActor::SetPackageExternal(bool bExternal, bool bShouldDirty, UPackage* Act
 	UPackage* LevelPackage = GetLevel()->GetPackage(); 
 	if (bExternal)
 	{
-		UPackage* NewActorPackage = ActorExternalPackage ? ActorExternalPackage : ULevel::CreateActorPackage(LevelPackage, GetLevel()->GetActorPackagingScheme(), GetPathName());
+		UPackage* NewActorPackage = ActorExternalPackage ? ActorExternalPackage : ULevel::CreateActorPackage(LevelPackage, GetLevel()->GetActorPackagingScheme(), GetPathName(), this);
 		SetExternalPackage(NewActorPackage);
 	}
-	else 
+	else
 	{
 		UPackage* ActorPackage = GetExternalPackage();
 		// Detach the linker exports so it doesn't resolve to this actor anymore
diff --git a/Engine/Source/Runtime/Engine/Private/Level.cpp b/Engine/Source/Runtime/Engine/Private/Level.cpp
--- a/Engine/Source/Runtime/Engine/Private/Level.cpp
+++ b/Engine/Source/Runtime/Engine/Private/Level.cpp
@@ -113,6 +113,10 @@ static FAutoConsoleVariableRef CVarRouteActorInitializationWorkUnitWeighting(
 );
 
 #if WITH_EDITOR
+
+TArray<ULevel::FLevelExternalActorsPathsProviderDelegate> ULevel::LevelExternalActorsPathsProviders;
+TArray<ULevel::FLevelMountPointResolverDelegate> ULevel::LevelMountPointResolvers;
+
 void FLevelActorFoldersHelper::SetUseActorFolders(ULevel* InLevel, bool bInEnabled)
 {
 	InLevel->SetUseActorFoldersInternal(bInEnabled);
@@ -1394,7 +1398,7 @@ void ULevel::PreDuplicate(FObjectDuplicationParameters& DupParams)
 				{
 					Path = Path.Replace(*ReplaceFrom, *ReplaceTo);
 				}
-				UPackage* DupPackage = Object->IsA<AActor>() ? CreateActorPackage(DstPackage, GetActorPackagingScheme(), Path) : FExternalPackageHelper::CreateExternalPackage(DstPackage, Path);
+				UPackage* DupPackage = Object->IsA<AActor>() ? ULevel::CreateActorPackage(DstPackage, GetActorPackagingScheme(), Path, Object) : FExternalPackageHelper::CreateExternalPackage(DstPackage, Path);
 				DupPackage->MarkAsFullyLoaded();
 				DupPackage->MarkPackageDirty();
 				DupParams.DuplicationSeed.Add(Package, DupPackage);
@@ -3530,9 +3534,28 @@ bool ULevel::HasAnyActorsOfType(UClass *SearchType)
 }
 
 #if WITH_EDITOR
-FString ULevel::GetActorPackageName(UPackage* InLevelPackage, EActorPackagingScheme ActorPackagingScheme, const FString& InActorPath)
+FString ULevel::GetActorPackageName(UPackage* InLevelPackage, EActorPackagingScheme ActorPackagingScheme, const FString& InActorPath, const UObject* InLevelMountPointContext)
 {
-	return GetActorPackageName(GetExternalActorsPath(InLevelPackage), ActorPackagingScheme, InActorPath);
+	check(InLevelPackage);
+	const FString LevelPackageName = InLevelPackage->GetName();
+
+	TOptional<FString> ResolvedLevelMountPoint;
+	if (InLevelMountPointContext)
+	{
+		for (const FLevelMountPointResolverDelegate& Resolver : ULevel::LevelMountPointResolvers)
+		{
+			FString Result;
+			if (Resolver.Execute(LevelPackageName, InLevelMountPointContext, Result) && !Result.IsEmpty())
+			{
+				ResolvedLevelMountPoint = Result;
+				break;
+			}
+		}
+	}
+	
+	// If ResolvedLevelMountPoint is not set, fallback on LevelPackageName
+	const FString LevelRootPath = ResolvedLevelMountPoint.Get(LevelPackageName);
+	return ULevel::GetActorPackageName(ULevel::GetExternalActorsPath(LevelRootPath), ActorPackagingScheme, InActorPath);
 }
 
 FString ULevel::GetActorPackageName(const FString& InBaseDir, EActorPackagingScheme ActorPackagingScheme, const FString& InActorPath)
@@ -3575,6 +3598,44 @@ FString ULevel::GetActorPackageName(const FString& InBaseDir, EActorPackagingSch
 	return ActorPackageName.ToString();
 }
 
+TArray<FString> ULevel::GetExternalActorsPaths(const FString& InLevelPackageName, const FString& InPackageShortName)
+{
+	TArray<FString> Result;
+	Result.Add(ULevel::GetExternalActorsPath(InLevelPackageName, InPackageShortName));
+	for (const FLevelExternalActorsPathsProviderDelegate& Provider : ULevel::LevelExternalActorsPathsProviders)
+	{
+		Provider.ExecuteIfBound(InLevelPackageName, InPackageShortName, Result);
+	}
+	return Result;
+}
+
+FDelegateHandle ULevel::RegisterLevelExternalActorsPathsProvider(const FLevelExternalActorsPathsProviderDelegate& Provider)
+{
+	LevelExternalActorsPathsProviders.Add(Provider);
+	return LevelExternalActorsPathsProviders.Last().GetHandle();
+}
+
+void ULevel::UnregisterLevelExternalActorsPathsProvider(const FDelegateHandle& ProviderDelegateHandle)
+{
+	LevelExternalActorsPathsProviders.RemoveAll([ProviderDelegateHandle](const FLevelExternalActorsPathsProviderDelegate& Delegate)
+	{
+		return Delegate.GetHandle() == ProviderDelegateHandle;
+	});
+}
+
+FDelegateHandle ULevel::RegisterLevelMountPointResolver(const FLevelMountPointResolverDelegate& Resolver)
+{
+	LevelMountPointResolvers.Add(Resolver);
+	return LevelMountPointResolvers.Last().GetHandle();
+}
+
+void ULevel::UnregisterLevelMountPointResolver(const FDelegateHandle& ResolverDelegateHandle)
+{
+	LevelMountPointResolvers.RemoveAll([ResolverDelegateHandle](const FLevelMountPointResolverDelegate& Delegate)
+	{
+		return Delegate.GetHandle() == ResolverDelegateHandle;
+	});
+}
 FString ULevel::GetExternalActorsPath(const FString& InLevelPackageName, const FString& InPackageShortName)
 {
 	// Strip the temp prefix if found
@@ -3616,7 +3677,7 @@ FString ULevel::GetExternalActorsPath(UPackage* InLevelPackage, const FString& I
 
 	// We can't use the Package->FileName here because it might be a duplicated a package
 	// We can't use the package short name directly in some cases either (PIE, instanced load) as it may contain pie prefix or not reflect the real actor location
-	return GetExternalActorsPath(InLevelPackage->GetName(), InPackageShortName);
+	return ULevel::GetExternalActorsPath(InLevelPackage->GetName(), InPackageShortName);
 }
 
 EActorPackagingScheme ULevel::GetActorPackagingSchemeFromActorPackageName(const FStringView InActorPackageName)
@@ -3667,7 +3728,7 @@ void ULevel::SetUseExternalActors(bool bEnable)
 TArray<FString> ULevel::GetExternalObjectsPaths(const FString& InLevelPackageName, const FString& InPackageShortName)
 {
 	TArray<FString> Paths;
-	Paths.Add(GetExternalActorsPath(InLevelPackageName, InPackageShortName));
+	Paths.Append(ULevel::GetExternalActorsPaths(InLevelPackageName, InPackageShortName));
 	Paths.Add(FExternalPackageHelper::GetExternalObjectsPath(InLevelPackageName, InPackageShortName));
 	return Paths;
 }
@@ -4020,16 +4081,23 @@ TArray<FString> ULevel::GetOnDiskExternalActorPackages(const FString& ExternalAc
 TArray<FString> ULevel::GetOnDiskExternalActorPackages(bool bTryUsingPackageLoadedPath) const
 {
 	UWorld* World = GetTypedOuter<UWorld>();
-	FString ExternalActorsPath;
+	FString LevelPackageName;
+	FString LevelPackageShortName;
 	if (bTryUsingPackageLoadedPath && !World->GetPackage()->GetLoadedPath().IsEmpty())
 	{
-		ExternalActorsPath = ULevel::GetExternalActorsPath(World->GetPackage()->GetLoadedPath().GetPackageName());
+		LevelPackageName = World->GetPackage()->GetLoadedPath().GetPackageName();
 	}
 	else
 	{
-		ExternalActorsPath = ULevel::GetExternalActorsPath(World->GetPackage(), (World->OriginalWorldName == NAME_None) ? World->GetName() : World->OriginalWorldName.ToString());
+		LevelPackageName = World->GetPackage()->GetName();
+		LevelPackageShortName = !World->OriginalWorldName.IsNone() ? World->OriginalWorldName.ToString() : World->GetName();
+	}
+	TArray<FString> ActorPackagePaths;
+	for (const FString& ExternalActorsPath : ULevel::GetExternalActorsPaths(LevelPackageName, LevelPackageShortName))
+	{
+		ActorPackagePaths.Append(GetOnDiskExternalActorPackages(ExternalActorsPath));
 	}
-	return GetOnDiskExternalActorPackages(ExternalActorsPath);
+	return ActorPackagePaths;
 }
 
 TArray<UPackage*> ULevel::GetLoadedExternalObjectPackages() const
@@ -4060,8 +4128,9 @@ TArray<UPackage*> ULevel::GetLoadedExternalObjectPackages() const
 	};
 
 	TSet<FString> ExternalObjectsPathSet;
-	ExternalObjectsPathSet.Add(SanitizeExternalPath(ULevel::GetExternalActorsPath(World->GetPackage(), (World->OriginalWorldName == NAME_None) ? World->GetName() : World->OriginalWorldName.ToString())));
-	ExternalObjectsPathSet.Add(SanitizeExternalPath(FExternalPackageHelper::GetExternalObjectsPath(World->GetPackage(), (World->OriginalWorldName == NAME_None) ? World->GetName() : World->OriginalWorldName.ToString())));
+	const FString LevelPackageShortName = !World->OriginalWorldName.IsNone() ? World->OriginalWorldName.ToString() : World->GetName();
+	ExternalObjectsPathSet.Add(SanitizeExternalPath(ULevel::GetExternalActorsPath(World->GetPackage(), LevelPackageShortName)));
+	ExternalObjectsPathSet.Add(SanitizeExternalPath(FExternalPackageHelper::GetExternalObjectsPath(World->GetPackage(), LevelPackageShortName)));
 
 	if (UWorldPartition* WorldPartition = GetWorldPartition())
 	{
@@ -4132,9 +4201,9 @@ static UPackage* CreateActorPackageInternal(const FString& InPackageName, const
 	return ActorPackage;
 };
 
-UPackage* ULevel::CreateActorPackage(UPackage* InLevelPackage, EActorPackagingScheme InActorPackagingScheme, const FString& InActorPath)
+UPackage* ULevel::CreateActorPackage(UPackage* InLevelPackage, EActorPackagingScheme InActorPackagingScheme, const FString& InActorPath, const UObject* InTargetContextObject)
 {
-	const FString PackageName = GetActorPackageName(InLevelPackage, InActorPackagingScheme, InActorPath);
+	const FString PackageName = ULevel::GetActorPackageName(InLevelPackage, InActorPackagingScheme, InActorPath, InTargetContextObject);
 	UPackage* ActorPackage = CreateActorPackageInternal(PackageName, InActorPath);
 	// Should be prevented upstream but we propagate the flag to prevent issues in asset enumeration
 	if (!ensureMsgf(!(InLevelPackage->GetPackageFlags() & PKG_PlayInEditor), TEXT("Actor packages should not be created on PlayInEditor levels")))
@@ -4146,7 +4215,7 @@ UPackage* ULevel::CreateActorPackage(UPackage* InLevelPackage, EActorPackagingSc
 
 UPackage* ULevel::CreateActorPackage(const FString& InBaseDir, EActorPackagingScheme InActorPackagingScheme, const FString& InActorPath)
 {
-	const FString PackageName = GetActorPackageName(InBaseDir, InActorPackagingScheme, InActorPath);
+	const FString PackageName = ULevel::GetActorPackageName(InBaseDir, InActorPackagingScheme, InActorPath);
 	return CreateActorPackageInternal(PackageName, InActorPath);
 }
 
diff --git a/Engine/Source/Runtime/Engine/Private/WorldPartition/DataLayer/ExternalDataLayerEngineSubsystem.cpp b/Engine/Source/Runtime/Engine/Private/WorldPartition/DataLayer/ExternalDataLayerEngineSubsystem.cpp
--- a/Engine/Source/Runtime/Engine/Private/WorldPartition/DataLayer/ExternalDataLayerEngineSubsystem.cpp
+++ b/Engine/Source/Runtime/Engine/Private/WorldPartition/DataLayer/ExternalDataLayerEngineSubsystem.cpp
@@ -8,6 +8,8 @@
 #include "Algo/Find.h"
 #if WITH_EDITOR
 #include "Editor.h"
+#include "WorldPartition/DataLayer/ExternalDataLayerHelper.h"
+#include "WorldPartition/DataLayer/DataLayerInstanceProviderInterface.h"
 #endif
 
 #define LOCTEXT_NAMESPACE "ExternalDataLayerEngineSubsystem"
@@ -66,12 +68,16 @@ TStatId UExternalDataLayerEngineSubsystem::GetStatId() const
 void UExternalDataLayerEngineSubsystem::Initialize(FSubsystemCollectionBase& Collection)
 {
 	Super::Initialize(Collection);
+	LevelExternalActorsPathsProviderDelegateHandle = ULevel::RegisterLevelExternalActorsPathsProvider(ULevel::FLevelExternalActorsPathsProviderDelegate::CreateUObject(this, &UExternalDataLayerEngineSubsystem::OnGetLevelExternalActorsPaths));
+	LevelMountPointResolverDelegateHandle = ULevel::RegisterLevelMountPointResolver(ULevel::FLevelMountPointResolverDelegate::CreateUObject(this, &UExternalDataLayerEngineSubsystem::OnResolveLevelMountPoint));
 	FEditorDelegates::OnAssetsPreDelete.AddUObject(this, &UExternalDataLayerEngineSubsystem::OnAssetsPreDelete);
 }
 
 void UExternalDataLayerEngineSubsystem::Deinitialize()
 {
 	Super::Deinitialize();
+	ULevel::UnregisterLevelExternalActorsPathsProvider(LevelExternalActorsPathsProviderDelegateHandle);
+	ULevel::UnregisterLevelMountPointResolver(LevelMountPointResolverDelegateHandle);
 	FEditorDelegates::OnAssetsPreDelete.RemoveAll(this);
 }
 
@@ -98,6 +104,39 @@ void UExternalDataLayerEngineSubsystem::OnAssetsPreDelete(const TArray<UObject*>
 	}
 }
 
+void UExternalDataLayerEngineSubsystem::OnGetLevelExternalActorsPaths(const FString& InLevelPackageName, const FString& InPackageShortName, TArray<FString>& OutExternalActorsPaths)
+{
+	FExternalDataLayerHelper::ForEachExternalDataLayerLevelPackagePath(InLevelPackageName, [&OutExternalActorsPaths, &InPackageShortName](const FString& InEDLLevelPackagePath)
+	{
+		const FString EDLExternalActorsPath = ULevel::GetExternalActorsPath(InEDLLevelPackagePath, InPackageShortName);
+		OutExternalActorsPaths.Add(EDLExternalActorsPath);
+	});
+}
+
+bool UExternalDataLayerEngineSubsystem::OnResolveLevelMountPoint(const FString& InLevelPackageName, const UObject* InLevelMountPointContext, FString& OutResolvedLevelMountPoint)
+{
+	const UExternalDataLayerAsset* ExternalDataLayerAssetContext = nullptr;
+	if (InLevelMountPointContext)
+	{
+		ExternalDataLayerAssetContext = Cast<UExternalDataLayerAsset>(InLevelMountPointContext);
+		if (!ExternalDataLayerAssetContext && InLevelMountPointContext->Implements<UDataLayerInstanceProvider>())
+		{
+			ExternalDataLayerAssetContext = CastChecked<IDataLayerInstanceProvider>(InLevelMountPointContext)->GetRootExternalDataLayerAsset();
+		}
+		if (!ExternalDataLayerAssetContext && InLevelMountPointContext->IsA<AActor>())
+		{
+			ExternalDataLayerAssetContext = CastChecked<AActor>(InLevelMountPointContext)->GetExternalDataLayerAsset();
+		}
+	}
+
+	if (ExternalDataLayerAssetContext)
+	{
+		OutResolvedLevelMountPoint = FExternalDataLayerHelper::GetExternalDataLayerLevelRootPath(ExternalDataLayerAssetContext, InLevelPackageName);
+		return true;
+	}
+	return false;
+}
+
 UWorld* UExternalDataLayerEngineSubsystem::GetTickableGameObjectWorld() const
 {
 	return GEditor->GetEditorWorldContext().World();
diff --git a/Engine/Source/Runtime/Engine/Private/WorldPartition/DataLayer/ExternalDataLayerHelper.cpp b/Engine/Source/Runtime/Engine/Private/WorldPartition/DataLayer/ExternalDataLayerHelper.cpp
--- a/Engine/Source/Runtime/Engine/Private/WorldPartition/DataLayer/ExternalDataLayerHelper.cpp
+++ b/Engine/Source/Runtime/Engine/Private/WorldPartition/DataLayer/ExternalDataLayerHelper.cpp
@@ -5,6 +5,13 @@
 #include "Misc/PackageName.h"
 #include "Misc/Paths.h"
 
+#if WITH_EDITOR
+#include "AssetRegistry/ARFilter.h"
+#include "ExternalPackageHelper.h"
+#include "UObject/Object.h"
+#include "UObject/AssetRegistryTagsContext.h"
+#endif
+
 FString FExternalDataLayerHelper::GetExternalStreamingObjectPackageName(const UExternalDataLayerAsset* InExternalDataLayerAsset)
 {
 	check(InExternalDataLayerAsset);
@@ -37,8 +44,13 @@ FString FExternalDataLayerHelper::GetExternalDataLayerLevelRootPath(const UExter
 {
 	check(InExternalDataLayerAsset);
 	check(InExternalDataLayerAsset->GetUID().IsValid());
+	return FExternalDataLayerHelper::GetExternalDataLayerLevelRootPath(FPackageName::GetPackageMountPoint(InExternalDataLayerAsset->GetPackage()->GetName()).ToString(), InExternalDataLayerAsset->GetUID(), InLevelPackagePath);
+}
+
+FString FExternalDataLayerHelper::GetExternalDataLayerLevelRootPath(const FString& InExternalDataLayerMountPoint, const FExternalDataLayerUID& InExternalDataLayerUID, const FString& InLevelPackagePath)
+{
 	FString ExternalDataLayerRootPath;
-	verify(BuildExternalDataLayerRootPath(FPackageName::GetPackageMountPoint(InExternalDataLayerAsset->GetPackage()->GetName()).ToString(), InExternalDataLayerAsset->GetUID(), ExternalDataLayerRootPath));
+	verify(BuildExternalDataLayerRootPath(InExternalDataLayerMountPoint, InExternalDataLayerUID, ExternalDataLayerRootPath));
 	TStringBuilderWithBuffer<TCHAR, NAME_SIZE> Builder;
 	Builder += ExternalDataLayerRootPath;
 	Builder += TEXT("/");
@@ -50,6 +62,69 @@ FString FExternalDataLayerHelper::GetExternalDataLayerLevelRootPath(const UExter
 
 #if WITH_EDITOR
 
+static FName GetExternalDataLayerUIDsAssetRegistryTag()
+{
+	static const FName ExternalDataLayerUIDsTag("ExternalDataLayerUIDs");
+	return ExternalDataLayerUIDsTag;
+}
+
+void FExternalDataLayerHelper::AddAssetRegistryTags(FAssetRegistryTagsContext OutContext, const TArray<FExternalDataLayerUID>& InExternalDataLayerUIDs)
+{
+	if (InExternalDataLayerUIDs.Num() > 0)
+	{
+		FString ExternalDataLayerUIDsStr = FString::JoinBy(InExternalDataLayerUIDs, TEXT(","), [&](const FExternalDataLayerUID& ExternalDataLayerUID) { return ExternalDataLayerUID.ToString(); });
+		OutContext.AddTag(UObject::FAssetRegistryTag(GetExternalDataLayerUIDsAssetRegistryTag(), ExternalDataLayerUIDsStr, UObject::FAssetRegistryTag::TT_Hidden));
+	}
+}
+
+void FExternalDataLayerHelper::GetExternalDataLayerUIDs(const FAssetData& Asset, TArray<FExternalDataLayerUID>& OutExternalDataLayerUIDs)
+{
+	FString ExternalDataLayerUIDsStr;
+	if (Asset.GetTagValue(GetExternalDataLayerUIDsAssetRegistryTag(), ExternalDataLayerUIDsStr))
+	{
+		TArray<FString> ExternalDataLayerUIDStrArray;
+		ExternalDataLayerUIDsStr.ParseIntoArray(ExternalDataLayerUIDStrArray, TEXT(","));
+		for (const FString& ExternalDataLayerUIDStr : ExternalDataLayerUIDStrArray)
+		{
+			FExternalDataLayerUID ExternalDataLayerUID;
+			if (FExternalDataLayerUID::Parse(ExternalDataLayerUIDStr, ExternalDataLayerUID))
+			{
+				OutExternalDataLayerUIDs.Add(ExternalDataLayerUID);
+			}
+		}
+	}
+}
+
+void FExternalDataLayerHelper::ForEachExternalDataLayerLevelPackagePath(const FString& InLevelPackageName, TFunctionRef<void(const FString&)> Func)
+{
+	UClass* GameFeatureDataClass = FindObject<UClass>(nullptr, TEXT("/Script/GameFeatures.GameFeatureData"));
+	if (GameFeatureDataClass)
+	{
+		FARFilter Filter;
+		Filter.bIncludeOnlyOnDiskAssets = false;
+		Filter.ClassPaths = { GameFeatureDataClass->GetClassPathName() };
+		Filter.bRecursivePaths = true;
+		TArray<FAssetData> AssetsData;
+		FExternalPackageHelper::GetSortedAssets(Filter, AssetsData);
+
+		for (const FAssetData& AssetData : AssetsData)
+		{
+			const FString MountPoint = FPackageName::GetPackageMountPoint(AssetData.PackagePath.ToString()).ToString();
+
+			TArray<FExternalDataLayerUID> ExternalDataLayerUIDs;
+			FExternalDataLayerHelper::GetExternalDataLayerUIDs(AssetData, ExternalDataLayerUIDs);
+			for (const FExternalDataLayerUID& ExternalDataLayerUID : ExternalDataLayerUIDs)
+			{
+				if (ExternalDataLayerUID.IsValid())
+				{
+					FString LevelPackageEDLPath = FExternalDataLayerHelper::GetExternalDataLayerLevelRootPath(MountPoint, ExternalDataLayerUID, InLevelPackageName);
+					Func(LevelPackageEDLPath);
+				}
+			}
+		}
+	}
+}
+
 bool FExternalDataLayerHelper::IsExternalDataLayerPath(FStringView InExternalDataLayerPath, FExternalDataLayerUID* OutExternalDataLayerUID)
 {
 	int32 ExternalDataLayerFolderIdx = UE::String::FindFirst(InExternalDataLayerPath, GetExternalDataLayerFolder(), ESearchCase::IgnoreCase);
diff --git a/Engine/Source/Runtime/Engine/Private/WorldPartition/DataLayer/ExternalDataLayerManager.cpp b/Engine/Source/Runtime/Engine/Private/WorldPartition/DataLayer/ExternalDataLayerManager.cpp
--- a/Engine/Source/Runtime/Engine/Private/WorldPartition/DataLayer/ExternalDataLayerManager.cpp
+++ b/Engine/Source/Runtime/Engine/Private/WorldPartition/DataLayer/ExternalDataLayerManager.cpp
@@ -907,8 +907,7 @@ AWorldDataLayers* UExternalDataLayerManager::GetWorldDataLayers(const UExternalD
 
 	if (bAllowCreate)
 	{
-		const FString BaseDir = ULevel::GetExternalActorsPath(GetExternalDataLayerLevelRootPath(InExternalDataLayerAsset));
-		ActorPackage = ULevel::CreateActorPackage(BaseDir, DestinationLevel->GetActorPackagingScheme(), WorldDataLayersName);
+		ActorPackage = ULevel::CreateActorPackage(DestinationLevel->GetPackage(), DestinationLevel->GetActorPackagingScheme(), WorldDataLayersName, InExternalDataLayerAsset);
 		FActorSpawnParameters SpawnParameters;
 		SpawnParameters.Name = *WorldDataLayersName;
 		SpawnParameters.OverrideLevel = DestinationLevel;
diff --git a/Engine/Source/Runtime/Engine/Public/WorldPartition/DataLayer/ExternalDataLayerEngineSubsystem.h b/Engine/Source/Runtime/Engine/Public/WorldPartition/DataLayer/ExternalDataLayerEngineSubsystem.h
--- a/Engine/Source/Runtime/Engine/Public/WorldPartition/DataLayer/ExternalDataLayerEngineSubsystem.h
+++ b/Engine/Source/Runtime/Engine/Public/WorldPartition/DataLayer/ExternalDataLayerEngineSubsystem.h
@@ -55,6 +55,8 @@ public:
 private:
 #if WITH_EDITOR
 	void OnAssetsPreDelete(const TArray<UObject*>& Objects);
+	void OnGetLevelExternalActorsPaths(const FString& InLevelPackageName, const FString& InPackageShortName, TArray<FString>& OutExternalActorsPaths);
+	bool OnResolveLevelMountPoint(const FString& InLevelPackageName, const UObject* InLevelMountPointContext, FString& OutResolvedLevelMountPoint);
 #endif
 	ENGINE_API bool CanWorldInjectExternalDataLayerAsset(const UWorld* InWorld, const UExternalDataLayerAsset* InExternalDataLayerAsset, FText* OutFailureReason = nullptr) const;
 
@@ -83,6 +85,9 @@ private:
 
 	using FForcedExternalDataLayerInjectionKey = TPair<const UWorld*, const UExternalDataLayerAsset*>;
 	TSet<FForcedExternalDataLayerInjectionKey> ForcedAllowInjection;
+
+	FDelegateHandle LevelExternalActorsPathsProviderDelegateHandle;
+	FDelegateHandle LevelMountPointResolverDelegateHandle;
 #endif
 
 	friend class UExternalDataLayerManager;
diff --git a/Engine/Source/Runtime/Engine/Public/WorldPartition/DataLayer/ExternalDataLayerHelper.h b/Engine/Source/Runtime/Engine/Public/WorldPartition/DataLayer/ExternalDataLayerHelper.h
--- a/Engine/Source/Runtime/Engine/Public/WorldPartition/DataLayer/ExternalDataLayerHelper.h
+++ b/Engine/Source/Runtime/Engine/Public/WorldPartition/DataLayer/ExternalDataLayerHelper.h
@@ -3,9 +3,12 @@
 
 #include "CoreFwd.h"
 #include "Containers/StringView.h"
+#include "Templates/Function.h"
 
+class FAssetRegistryTagsContext;
 class UExternalDataLayerAsset;
 struct FExternalDataLayerUID;
+struct FAssetData;
 
 class FExternalDataLayerHelper
 {
@@ -21,16 +24,30 @@ public:
 	 */
 	ENGINE_API static bool BuildExternalDataLayerRootPath(const FString& InEDLMountPoint, const FExternalDataLayerUID& InExternalDataLayerUID, FString& OutExternalDataLayerRootPath);
 
-	/** Return true if succeeds building the external data layer level root path (OutExternalDataLayerLevelRootPath) using the provided mount point, EDL UID and the level package path.
+	/** Return true if succeeds building the external data layer level root path (OutExternalDataLayerLevelRootPath) using the ExternalDataLayerAsset and the level package path.
 	  * Format is /{MountPoint}/{ExternalDataLayerFolder}/{EDL_UID}/{LevelPath}
 	  */
 	ENGINE_API static FString GetExternalDataLayerLevelRootPath(const UExternalDataLayerAsset* InExternalDataLayerAsset, const FString& InLevelPackagePath);
 
+	/** Return true if succeeds building the external data layer level root path (OutExternalDataLayerLevelRootPath) using the provided mount point, EDL UID and the level package path.
+	  * Format is /{MountPoint}/{ExternalDataLayerFolder}/{EDL_UID}/{LevelPath}
+	  */
+	ENGINE_API static FString GetExternalDataLayerLevelRootPath(const FString& InEDLMountPoint, const FExternalDataLayerUID& InExternalDataLayerUID, const FString& InLevelPackagePath);
+
 #if WITH_EDITOR
 	/** Returns whether the provided path respects the format <start_path>/{ExternalDataLayerFolder}/{EDL_UID}/<end_path>. 
 	 * If true, fills OutExternalDataLayerUID when provided.
 	 */
 	ENGINE_API static bool IsExternalDataLayerPath(FStringView InExternalDataLayerPath, FExternalDataLayerUID* OutExternalDataLayerUID = nullptr);
+
+	/** Fills Asset Registry Tags Context with provided External Data Layer UIDs. */
+	ENGINE_API static void AddAssetRegistryTags(FAssetRegistryTagsContext OutContext, const TArray<FExternalDataLayerUID>& InExternalDataLayerUIDs);
+
+	/** Retrieves External Data Layers UIDs from provided Asset. */
+	ENGINE_API static void GetExternalDataLayerUIDs(const FAssetData& Asset, TArray<FExternalDataLayerUID>& OutExternalDataLayerUIDs);
+
+	/** Iterates through all possible External Data Layer Level Package Paths using Asset Registry. */
+	ENGINE_API static void ForEachExternalDataLayerLevelPackagePath(const FString& InLevelPackageName, TFunctionRef<void(const FString&)> Func);
 #endif
 
 private:
