Fixed crash that occuring because the HISM proxy is recreated while the component is in a half-setup state during ApplyComponentInstanceData.

diff --git a/Engine/Source/Runtime/Engine/Private/HierarchicalInstancedStaticMesh.cpp b/Engine/Source/Runtime/Engine/Private/HierarchicalInstancedStaticMesh.cpp
--- a/Engine/Source/Runtime/Engine/Private/HierarchicalInstancedStaticMesh.cpp
+++ b/Engine/Source/Runtime/Engine/Private/HierarchicalInstancedStaticMesh.cpp
@@ -2978,6 +2978,9 @@ FPrimitiveSceneProxy* UHierarchicalInstancedStaticMeshComponent::CreateSceneProx
 
 	// Verify that the mesh is valid & that we have instances before creating a proxy.
 	const bool bAreMeshAndInstancesValid =
+#if WITH_EDITOR
+		bIsInstanceDataApplyCompleted && 
+#endif
 		// Make sure we have instances, or an update with instances on the way, or already built instances (for the external data (landscape grass) mode).
 		(GetNumInstances() > 0 ||  PrimitiveInstanceDataManager.GetMaxInstanceIndex() > 0 || NumBuiltRenderInstances > 0) &&
 		// Make sure we have an actual static mesh.
