World Partition: add support for editor-only attach parents by applying their missing parent transforms.

diff --git a/Engine/Plugins/Editor/WorldPartitionHLODUtilities/Source/Private/WorldPartition/HLOD/Utilities/WorldPartitionHLODUtilities.cpp b/Engine/Plugins/Editor/WorldPartitionHLODUtilities/Source/Private/WorldPartition/HLOD/Utilities/WorldPartitionHLODUtilities.cpp
--- a/Engine/Plugins/Editor/WorldPartitionHLODUtilities/Source/Private/WorldPartition/HLOD/Utilities/WorldPartitionHLODUtilities.cpp
+++ b/Engine/Plugins/Editor/WorldPartitionHLODUtilities/Source/Private/WorldPartition/HLOD/Utilities/WorldPartitionHLODUtilities.cpp
@@ -102,6 +102,7 @@ void AddSourceActor(const FStreamingGenerationActorDescView& ActorDescView, cons
 			ActorDescView.GetNativeClass(),
 			ContainerID,
 			ContainerTransform,
+			ActorDescView.GetEditorOnlyParentTransform(),
 			ContainerPackage,
 			WorldPackageName,
 			ActorInstanceGuid,
@@ -130,6 +131,7 @@ void AddSourceActor(const FStreamingGenerationActorDescView& ActorDescView, cons
 					ReferenceActorDesc.GetNativeClass(),
 					ContainerID,
 					ContainerTransform,
+					FTransform::Identity,
 					ContainerPackage,
 					WorldPackageName,
 					EditorRefInstanceGuid,
diff --git a/Engine/Source/Runtime/Engine/Private/WorldPartition/WorldPartitionLevelHelper.cpp b/Engine/Source/Runtime/Engine/Private/WorldPartition/WorldPartitionLevelHelper.cpp
--- a/Engine/Source/Runtime/Engine/Private/WorldPartition/WorldPartitionLevelHelper.cpp
+++ b/Engine/Source/Runtime/Engine/Private/WorldPartition/WorldPartitionLevelHelper.cpp
@@ -722,7 +722,7 @@ bool FWorldPartitionLevelHelper::LoadActorsInternal(FLoadActorsParams&& InParams
 							ChildActor->UObject::Rename(*FString::Printf(TEXT("%s_%s"), *ChildActor->GetName(), *PackageObjectMapping->ContainerID.ToShortString()), DestLevel, REN_NonTransactional | REN_ForceNoResetLoaders | REN_DoNotDirty | REN_DontCreateRedirectors);
 						}
 					});
-					
+
 					// Apply Pre-ConstructionScript Properties
 					TArray<FActorPropertyOverride> ActorPropertyOverrides;
 					for (auto ActorOverrideMapping : PackageObjectMapping->PropertyOverrides)
@@ -748,8 +748,8 @@ bool FWorldPartitionLevelHelper::LoadActorsInternal(FLoadActorsParams&& InParams
 					{
 						ActorPropertyOverridesAnnotation.AddAnnotation(Actor, FWorldPartitionLevelHelper::FActorPropertyOverridesAnnotation(MoveTemp(ActorPropertyOverrides), PackageObjectMapping->ContainerTransform));
 					}
-															
-					FLevelUtils::FApplyLevelTransformParams TransformParams(nullptr, PackageObjectMapping->ContainerTransform);
+
+					FLevelUtils::FApplyLevelTransformParams TransformParams(nullptr, PackageObjectMapping->ContainerTransform * PackageObjectMapping->EditorOnlyParentTransform);
 					TransformParams.Actor = Actor;
 					TransformParams.bDoPostEditMove = false;
 					FLevelUtils::ApplyLevelTransform(TransformParams);
@@ -780,6 +780,13 @@ bool FWorldPartitionLevelHelper::LoadActorsInternal(FLoadActorsParams&& InParams
 						ObjectResolver->SetWorldPartitionResolveData(FWorldPartitionResolveData(PackageObjectMapping->ContainerID, FTopLevelAssetPath(SourceWorldPath)));
 					}
 				}
+				else if (!PackageObjectMapping->EditorOnlyParentTransform.Equals(FTransform::Identity))
+				{
+					FLevelUtils::FApplyLevelTransformParams TransformParams(nullptr, PackageObjectMapping->EditorOnlyParentTransform);
+					TransformParams.Actor = Actor;
+					TransformParams.bDoPostEditMove = false;
+					FLevelUtils::ApplyLevelTransform(TransformParams);
+				}
 
 				if (DestLevel)
 				{
diff --git a/Engine/Source/Runtime/Engine/Private/WorldPartition/WorldPartitionRuntimeLevelStreamingCell.cpp b/Engine/Source/Runtime/Engine/Private/WorldPartition/WorldPartitionRuntimeLevelStreamingCell.cpp
--- a/Engine/Source/Runtime/Engine/Private/WorldPartition/WorldPartitionRuntimeLevelStreamingCell.cpp
+++ b/Engine/Source/Runtime/Engine/Private/WorldPartition/WorldPartitionRuntimeLevelStreamingCell.cpp
@@ -294,40 +294,41 @@ void UWorldPartitionRuntimeLevelStreamingCell::AddActorToCell(const FStreamingGe
 
 	for (const FGuid& EditorReferenceGuid : ActorDescView.GetEditorReferences())
 	{
-		// Special case where ActorDescView has invalid references : Use InvalidReference information as the Actor Guid isn't 
-		// necessarily in the ContainerInstance.
-		FStreamingGenerationActorDescView::FInvalidReference InvalidRef;
-		if (ActorDescView.IsInvalidReference(EditorReferenceGuid, &InvalidRef))
+		FName ReferencePackage;
+		FName ReferencePath;
+		FTopLevelAssetPath ReferenceBaseClass;
+		FTopLevelAssetPath ReferenceNativeClass;
+
+		// Special case where the actor descriptor view has an invalid references, use invalid reference information as the actor guid isn't necessarily in the container instance.
+		if (const FStreamingGenerationActorDescView::FInvalidReference* InvalidReference = ActorDescView.GetInvalidReference(EditorReferenceGuid))
 		{
-			Packages.Emplace(
-				InvalidRef.ActorPackage,
-				*InvalidRef.ActorSoftPath.ToString(),
-				InvalidRef.BaseClass,
-				InvalidRef.NativeClass,
-				ContainerID,
-				ContainerTransform,
-				ContainerPackage,
-				GetWorld()->GetPackage()->GetFName(),
-				ContainerID.GetActorGuid(EditorReferenceGuid),
-				true
-			);
+			ReferencePackage = InvalidReference->ActorPackage;
+			ReferencePath = *InvalidReference->ActorSoftPath.ToString();
+			ReferenceBaseClass = InvalidReference->BaseClass;
+			ReferenceNativeClass = InvalidReference->NativeClass;
 		}
 		else
 		{
 			const FWorldPartitionActorDescInstance& ReferenceActorDesc = ContainerInstance->GetActorDescInstanceChecked(EditorReferenceGuid);
-			Packages.Emplace(
-				ReferenceActorDesc.GetActorPackage(),
-				*ReferenceActorDesc.GetActorSoftPath().ToString(),
-				ReferenceActorDesc.GetBaseClass(),
-				ReferenceActorDesc.GetNativeClass(),
-				ContainerID,
-				ContainerTransform,
-				ContainerPackage,
-				GetWorld()->GetPackage()->GetFName(),
-				ContainerID.GetActorGuid(EditorReferenceGuid),
-				true
-			);
+			ReferencePackage = ReferenceActorDesc.GetActorPackage();
+			ReferencePath = *ReferenceActorDesc.GetActorSoftPath().ToString();
+			ReferenceBaseClass = ReferenceActorDesc.GetBaseClass();
+			ReferenceNativeClass = ReferenceActorDesc.GetNativeClass();
 		}
+		
+		Packages.Emplace(
+			ReferencePackage,
+			ReferencePath,
+			ReferenceBaseClass,
+			ReferenceNativeClass,
+			ContainerID,
+			ContainerTransform,
+			FTransform::Identity,
+			ContainerPackage,
+			GetWorld()->GetPackage()->GetFName(),
+			ContainerID.GetActorGuid(EditorReferenceGuid),
+			true
+		);
 	}
 
 	FWorldPartitionRuntimeCellObjectMapping ActorMapping(
@@ -337,6 +338,7 @@ void UWorldPartitionRuntimeLevelStreamingCell::AddActorToCell(const FStreamingGe
 		ActorDescView.GetNativeClass(),
 		ContainerID,
 		ContainerTransform,
+		ActorDescView.GetEditorOnlyParentTransform(),
 		ContainerPackage, 
 		GetWorld()->GetPackage()->GetFName(), 
 		ContainerID.GetActorGuid(ActorDescView.GetGuid()),
diff --git a/Engine/Source/Runtime/Engine/Private/WorldPartition/WorldPartitionStreamingGeneration.cpp b/Engine/Source/Runtime/Engine/Private/WorldPartition/WorldPartitionStreamingGeneration.cpp
--- a/Engine/Source/Runtime/Engine/Private/WorldPartition/WorldPartitionStreamingGeneration.cpp
+++ b/Engine/Source/Runtime/Engine/Private/WorldPartition/WorldPartitionStreamingGeneration.cpp
@@ -255,10 +255,19 @@ const TArray<FGuid>& FStreamingGenerationActorDescView::GetEditorReferences() co
 void FStreamingGenerationActorDescView::SetParentView(const FStreamingGenerationActorDescView* InParentView)
 {
 	check(!ParentView);
+	check(!EditorOnlyParentTransform.IsSet());
 	check(GetParentActor().IsValid());
 	ParentView = InParentView;
 }
 
+void FStreamingGenerationActorDescView::SetEditorOnlyParentTransform(const FTransform& InEditorOnlyParentTransform)
+{
+	check(!ParentView);
+	check(!EditorOnlyParentTransform.IsSet());
+	check(GetParentActor().IsValid());
+	EditorOnlyParentTransform = InEditorOnlyParentTransform;
+}
+
 void FStreamingGenerationActorDescView::SetDataLayerInstanceNames(const FDataLayerInstanceNames& InDataLayerInstanceNames)
 {
 	check(!Super::HasResolvedDataLayerInstanceNames());
@@ -270,17 +279,9 @@ void FStreamingGenerationActorDescView::SetDataLayerInstanceNames(const FDataLay
 	}
 }
 
-bool FStreamingGenerationActorDescView::IsInvalidReference(const FGuid& InGuid, FInvalidReference* OutInvalidReference) const
+const FStreamingGenerationActorDescView::FInvalidReference* FStreamingGenerationActorDescView::GetInvalidReference(const FGuid& InGuid) const
 {
-	if (const FInvalidReference* InvalidReference = ForcedInvalidReference.Find(InGuid))
-	{
-		if (OutInvalidReference)
-		{
-			*OutInvalidReference = *InvalidReference;
-		}
-		return true;
-	}
-	return false;
+	return ForcedInvalidReference.Find(InGuid);
 }
 
 void FStreamingGenerationActorDescView::AddForcedInvalidReference(const FStreamingGenerationActorDescView* ReferenceView)
@@ -586,7 +587,7 @@ class FWorldPartitionStreamingGenerator
 		TUniqueObj<FStreamingGenerationActorDescViewMap> ActorDescViewMap;
 
 		/** Set of editor-only actors that are not part of the actor descriptor views */
-		TSet<FGuid> EditorOnlyActorDescMap;
+		TSet<FGuid> EditorOnlyActorDescSet;
 
 		/** List of actor descriptor views that are containers (mainly level instances) */
 		TArray<FStreamingGenerationActorDescView> ContainerCollectionInstanceViews;
@@ -656,10 +657,16 @@ class FWorldPartitionStreamingGenerator
 			return InstanceData;
 		}
 
+		/** Per instance data */
 		FPerInstanceData InstanceData;
 		TSet<FPerInstanceData> UniquePerInstanceData;
 		TMap<FGuid, FSetElementId> PerInstanceData;
+
+		/** Map of actor descriptor mutators */
 		TMap<FGuid, FActorDescViewMutator> ActorDescViewMutators;
+
+		/** Map of editor-only parent actor transforms */
+		TMap<FGuid, FTransform> EditorOnlyParentActorTransforms;
 	};
 
 	void ResolveRuntimeSpatiallyLoaded(FStreamingGenerationActorDescView& ActorDescView)
@@ -707,16 +714,31 @@ class FWorldPartitionStreamingGenerator
 		}
 	}
 
-	void ResolveParentView(FStreamingGenerationActorDescView& ActorDescView, const FStreamingGenerationActorDescViewMap& ActorDescViewMap)
+	void ResolveParentView(FStreamingGenerationActorDescView& ActorDescView, const FStreamingGenerationActorDescViewMap& ActorDescViewMap, const TSet<FGuid>& EditorOnlyActorDescSet, const TMap<FGuid, FTransform>& EditorOnlyParentActorTransforms)
 	{
 		if (FGuid ParentGuid = ActorDescView.GetParentActor(); ParentGuid.IsValid())
 		{
-			ActorDescView.SetParentView(ActorDescViewMap.FindByGuid(ParentGuid));
+			if (const FStreamingGenerationActorDescView* ParentView = ActorDescViewMap.FindByGuid(ParentGuid))
+			{
+				ActorDescView.SetParentView(ParentView);
+			}
+			else if (const FTransform* EditorOnlyParentActorTransform = EditorOnlyParentActorTransforms.Find(ParentGuid))
+			{
+				ActorDescView.SetEditorOnlyParentTransform(*EditorOnlyParentActorTransform);
+			}
 		}
 	}
 
-	void CreateActorDescViewMap(const FStreamingGenerationContainerInstanceCollection& InActorDescCollection, FStreamingGenerationActorDescViewMap& OutActorDescViewMap, TSet<FGuid>& OutEditorOnlyActorDescMap, const FActorContainerID& InContainerID, TArray<FStreamingGenerationActorDescView>& OutContainerInstances, TArray<TUniquePtr<FStreamingGenerationUnsavedDirtyActorDescInstance>>& OutUnsavedDirtyInstances)
+	void CreateActorDescViewMap(FContainerCollectionInstanceDescriptor& InContainerCollectionInstanceDescriptor)
 	{
+		const FStreamingGenerationContainerInstanceCollection& InActorDescCollection = *InContainerCollectionInstanceDescriptor.ContainerInstanceCollection;
+		FStreamingGenerationActorDescViewMap& OutActorDescViewMap = InContainerCollectionInstanceDescriptor.ActorDescViewMap.Get();
+		TSet<FGuid>& OutEditorOnlyActorDescSet = InContainerCollectionInstanceDescriptor.EditorOnlyActorDescSet;
+		const FActorContainerID& InContainerID = InContainerCollectionInstanceDescriptor.ID;
+		TArray<FStreamingGenerationActorDescView>& OutContainerInstances = InContainerCollectionInstanceDescriptor.ContainerCollectionInstanceViews;
+		TArray<TUniquePtr<FStreamingGenerationUnsavedDirtyActorDescInstance>>& OutUnsavedDirtyInstances = InContainerCollectionInstanceDescriptor.UnsavedDirtyInstances;
+		TMap<FGuid, FTransform>& EditorOnlyParentActorTransforms = InContainerCollectionInstanceDescriptor.EditorOnlyParentActorTransforms;
+
 		// Should we handle unsaved or newly created actors?
 		const bool bShouldHandleUnsavedActors = bHandleUnsavedActors && InContainerID.IsMainContainer();
 
@@ -797,7 +819,31 @@ class FWorldPartitionStreamingGenerator
 			}
 			else
 			{
-				OutEditorOnlyActorDescMap.Add(Iterator->GetGuid());
+				OutEditorOnlyActorDescSet.Add(Iterator->GetGuid());
+			}
+		}
+
+		// Register transforms from editor-only parents as the childs won't be properly offset if they are not present
+		for (FStreamingGenerationContainerInstanceCollection::TConstIterator<> Iterator(&InActorDescCollection); Iterator; ++Iterator)
+		{
+			if (FGuid ParentGuid = Iterator->GetParentActor(); ParentGuid.IsValid())
+			{
+				if (OutEditorOnlyActorDescSet.Contains(ParentGuid) && !EditorOnlyParentActorTransforms.Contains(ParentGuid))
+				{
+					const FWorldPartitionActorDescInstance* ParentActorDescInstance = InActorDescCollection.GetActorDescInstance(ParentGuid);
+					FTransform ParentActorDescInstanceTransform = ParentActorDescInstance->GetActorTransform();
+					
+					// Dirty actors
+					if (AActor* Actor = ParentActorDescInstance->GetActor())
+					{
+						if (bShouldHandleUnsavedActors && (bIsTempContainerPackage || Actor->GetPackage()->IsDirty()))
+						{
+							ParentActorDescInstanceTransform = Actor->GetActorTransform();
+						}
+					}
+
+					EditorOnlyParentActorTransforms.Add(ParentGuid, ParentActorDescInstanceTransform);
+				}
 			}
 		}
 
@@ -863,8 +909,9 @@ class FWorldPartitionStreamingGenerator
 				check(!ContainerCollectionInstanceDescriptorsMap.Contains(ContainerID));
 
 				FContainerCollectionInstanceDescriptor& ContainerCollectionInstanceDescriptor = ContainerCollectionInstanceDescriptorsMap.Add(ContainerID, MoveTemp(InContainerCollectionInstanceDescriptor));
+				
 				// Gather actor descriptor views for this container
-				CreateActorDescViewMap(*ContainerCollectionInstanceDescriptor.ContainerInstanceCollection, ContainerCollectionInstanceDescriptor.ActorDescViewMap.Get(), ContainerCollectionInstanceDescriptor.EditorOnlyActorDescMap, ContainerCollectionInstanceDescriptor.ID, ContainerCollectionInstanceDescriptor.ContainerCollectionInstanceViews, ContainerCollectionInstanceDescriptor.UnsavedDirtyInstances);
+				CreateActorDescViewMap(ContainerCollectionInstanceDescriptor);
 
 				// Resolve actor descriptor views before validation
 				ResolveContainerDescriptor(ContainerCollectionInstanceDescriptor);
@@ -892,7 +939,6 @@ class FWorldPartitionStreamingGenerator
 				ContainerCollectionInstanceViews.Append(ContainerCollectionInstanceDescriptor.ContainerCollectionInstanceViews);
 			}
 
-
 			// Parse actor containers
 			for (const FStreamingGenerationActorDescView& ContainerCollectionInstanceView : ContainerCollectionInstanceViews)
 			{
@@ -1021,7 +1067,7 @@ class FWorldPartitionStreamingGenerator
 			ResolveRuntimeGrid(ActorDescView);
 			ResolveRuntimeDataLayers(ActorDescView, ContainerCollectionInstanceDescriptor.ActorDescViewMap.Get());
 			ResolveHLODLayer(ActorDescView, WorldPartitionContext ? FSoftObjectPath(WorldPartitionContext->GetDefaultHLODLayer()) : FSoftObjectPath());
-			ResolveParentView(ActorDescView, ContainerCollectionInstanceDescriptor.ActorDescViewMap.Get());
+			ResolveParentView(ActorDescView, ContainerCollectionInstanceDescriptor.ActorDescViewMap.Get(), ContainerCollectionInstanceDescriptor.EditorOnlyActorDescSet, ContainerCollectionInstanceDescriptor.EditorOnlyParentActorTransforms);
 		};
 
 		ContainerCollectionInstanceDescriptor.ActorDescViewMap->ForEachActorDescView([this, &ResolveActorDescView](FStreamingGenerationActorDescView& ActorDescView)
@@ -1289,7 +1335,7 @@ class FWorldPartitionStreamingGenerator
 					}
 					else
 					{
-						if (!ContainerCollectionInstanceDescriptor.EditorOnlyActorDescMap.Contains(Info.ReferenceGuid))
+						if (!ContainerCollectionInstanceDescriptor.EditorOnlyActorDescSet.Contains(Info.ReferenceGuid))
 						{
 							if (PassType == EPassType::ErrorReporting)
 							{
@@ -1341,7 +1387,7 @@ class FWorldPartitionStreamingGenerator
 						RuntimeReferences.Add(ReferenceGuid);
 					}
 				}
-				else if (ContainerCollectionInstanceDescriptor.EditorOnlyActorDescMap.Contains(ReferenceGuid))
+				else if (ContainerCollectionInstanceDescriptor.EditorOnlyActorDescSet.Contains(ReferenceGuid))
 				{
 					EditorReferences.Add(ReferenceGuid);
 				}
diff --git a/Engine/Source/Runtime/Engine/Public/WorldPartition/WorldPartitionActorDescInstance.h b/Engine/Source/Runtime/Engine/Public/WorldPartition/WorldPartitionActorDescInstance.h
--- a/Engine/Source/Runtime/Engine/Public/WorldPartition/WorldPartitionActorDescInstance.h
+++ b/Engine/Source/Runtime/Engine/Public/WorldPartition/WorldPartitionActorDescInstance.h
@@ -60,6 +60,7 @@ public:
 	ENGINE_API virtual FName GetActorName() const;
 	virtual FName GetFolderPath() const { return ActorDesc->GetFolderPath(); }
 	virtual const FGuid& GetFolderGuid() const { return ActorDesc->GetFolderGuid(); }
+	virtual const FTransform& GetActorTransform() const { return ActorDesc->GetActorTransform(); }
 
 	ENGINE_API virtual FBox GetEditorBounds() const;
 	ENGINE_API virtual FBox GetRuntimeBounds() const;
diff --git a/Engine/Source/Runtime/Engine/Public/WorldPartition/WorldPartitionRuntimeCell.h b/Engine/Source/Runtime/Engine/Public/WorldPartition/WorldPartitionRuntimeCell.h
--- a/Engine/Source/Runtime/Engine/Public/WorldPartition/WorldPartitionRuntimeCell.h
+++ b/Engine/Source/Runtime/Engine/Public/WorldPartition/WorldPartitionRuntimeCell.h
@@ -78,7 +78,19 @@ struct FWorldPartitionRuntimeCellObjectMapping
 #endif
 	{}
 
-	FWorldPartitionRuntimeCellObjectMapping(FName InPackage, FName InPath, const FTopLevelAssetPath& InBaseClass, const FTopLevelAssetPath& InNativeClass, const FActorContainerID& InContainerID, const FTransform& InContainerTransform, FName InContainerPackage, FName InWorldPackage, const FGuid& InActorInstanceGuid, bool bInIsEditorOnly)
+	FWorldPartitionRuntimeCellObjectMapping(
+		FName InPackage, 
+		FName InPath, 
+		const FTopLevelAssetPath& InBaseClass, 
+		const FTopLevelAssetPath& InNativeClass, 
+		const FActorContainerID& InContainerID, 
+		const FTransform& InContainerTransform, 
+		const FTransform& InEditorOnlyParentTransform, 
+		FName InContainerPackage, 
+		FName InWorldPackage, 
+		const FGuid& InActorInstanceGuid, 
+		bool bInIsEditorOnly
+	)
 #if WITH_EDITORONLY_DATA
 		: Package(InPackage)
 		, Path(InPath)
@@ -86,6 +98,7 @@ struct FWorldPartitionRuntimeCellObjectMapping
 		, NativeClass(InNativeClass)
 		, ContainerID(InContainerID)
 		, ContainerTransform(InContainerTransform)
+		, EditorOnlyParentTransform(InEditorOnlyParentTransform)
 		, ContainerPackage(InContainerPackage)
 		, WorldPackage(InWorldPackage)
 		, ActorInstanceGuid(InActorInstanceGuid)
@@ -130,7 +143,13 @@ struct FWorldPartitionRuntimeCellObjectMapping
 	 */
 	UPROPERTY()
 	FTransform ContainerTransform;
-		
+
+	/** 
+	 * Transform of the owning actor if editor only
+	 */
+	UPROPERTY()
+	FTransform EditorOnlyParentTransform;
+
 	/**
 	 * Package of the owning container instance
 	 */
diff --git a/Engine/Source/Runtime/Engine/Public/WorldPartition/WorldPartitionStreamingGeneration.h b/Engine/Source/Runtime/Engine/Public/WorldPartition/WorldPartitionStreamingGeneration.h
--- a/Engine/Source/Runtime/Engine/Public/WorldPartition/WorldPartitionStreamingGeneration.h
+++ b/Engine/Source/Runtime/Engine/Public/WorldPartition/WorldPartitionStreamingGeneration.h
@@ -68,7 +68,7 @@ public:
 	ENGINE_API const FDataLayerInstanceNames& GetRuntimeDataLayerInstanceNames() const;
 
 	ENGINE_API void AddForcedInvalidReference(const FStreamingGenerationActorDescView* ReferenceView);
-	ENGINE_API bool IsInvalidReference(const FGuid& InGuid, FInvalidReference* OutInvalidReference = nullptr) const;
+	ENGINE_API const FInvalidReference* GetInvalidReference(const FGuid& InGuid) const;
 	ENGINE_API void SetForcedNonSpatiallyLoaded();
 	ENGINE_API void SetForcedNoRuntimeGrid();
 	ENGINE_API void SetForcedNoDataLayers();
@@ -77,6 +77,8 @@ public:
 	ENGINE_API void SetEditorReferences(const TArray<FGuid>& InEditorReferences);
 	ENGINE_API void SetDataLayerInstanceNames(const FDataLayerInstanceNames& InDataLayerInstanceNames);
 	ENGINE_API void SetParentView(const FStreamingGenerationActorDescView* InParentView);
+	ENGINE_API void SetEditorOnlyParentTransform(const FTransform& InEditorOnlyParentTransform);
+	ENGINE_API FTransform GetEditorOnlyParentTransform() const { return EditorOnlyParentTransform.IsSet() ? EditorOnlyParentTransform.GetValue() : FTransform::Identity; }
 
 	ENGINE_API void SetForcedNoHLODLayer();
 	ENGINE_API void SetRuntimeHLODLayer(const FSoftObjectPath& InHLODLayer);
@@ -112,6 +114,7 @@ private:
 	TOptional<TArray<FGuid>> RuntimeReferences;
 	TOptional<FSoftObjectPath> RuntimedHLODLayer;
 	TArray<FGuid> EditorReferences;
+	TOptional<FTransform> EditorOnlyParentTransform;
 };
 
 class FStreamingGenerationActorDescViewMap
