Fix for FOverlapResult not getting its PhysicsObject field populated when the body comes from an entity

diff --git a/Engine/Source/Runtime/Engine/Private/Collision/CollisionConversions.cpp b/Engine/Source/Runtime/Engine/Private/Collision/CollisionConversions.cpp
--- a/Engine/Source/Runtime/Engine/Private/Collision/CollisionConversions.cpp
+++ b/Engine/Source/Runtime/Engine/Private/Collision/CollisionConversions.cpp
@@ -149,7 +149,7 @@ static void SetHitResultFromShapeAndFaceIndex(const FPhysicsShape& Shape,  const
 		}
 
 		OwningComponent = BodyInst->OwnerComponent;
-		OutResult.PhysicsObject = BodyInst->ActorHandle->GetPhysicsObject();
+		OutResult.PhysicsObject = BodyInst->ActorHandle ? BodyInst->ActorHandle->GetPhysicsObject() : nullptr;
 	}
 	else
 	{
@@ -613,8 +613,8 @@ void ConvertQueryOverlap(const FPhysicsShape& Shape, const FPhysicsActor& Actor,
 			OutOverlap.OverlapObjectHandle = FActorInstanceHandle::MakeActorHandleToResolve(BodyInst->OwnerComponent, BodyInst->InstanceBodyIndex);
 			OutOverlap.Component = BodyInst->OwnerComponent; // Copying weak pointer is faster than assigning raw pointer.
 			OutOverlap.ItemIndex = OwnerComponent->bMultiBodyOverlap ? BodyInst->InstanceBodyIndex : INDEX_NONE;
-			OutOverlap.PhysicsObject = BodyInst->ActorHandle->GetPhysicsObject();
 		}
+		OutOverlap.PhysicsObject = BodyInst->ActorHandle ? BodyInst->ActorHandle->GetPhysicsObject() : nullptr;
 	}
 	else
 	{
