Enable per-shader DDC (including use of remote caches) when cooking.

diff --git a/Engine/Source/Runtime/Engine/Private/ShaderCompiler/ShaderCompiler.cpp b/Engine/Source/Runtime/Engine/Private/ShaderCompiler/ShaderCompiler.cpp
--- a/Engine/Source/Runtime/Engine/Private/ShaderCompiler/ShaderCompiler.cpp
+++ b/Engine/Source/Runtime/Engine/Private/ShaderCompiler/ShaderCompiler.cpp
@@ -149,7 +149,13 @@ static TAutoConsoleVariable<bool> CVarJobCacheDDC(
 static TAutoConsoleVariable<bool> CVarJobCacheDDCPolicy(
 	TEXT("r.ShaderCompiler.JobCacheDDCEnableRemotePolicy"),
 	false,
-	TEXT("If true, individual shader jobs will be cached to remote/shared DDC instances; if false they will only cache to DDC instances on the local machine.\n"),
+	TEXT("If true, individual shader jobs will be cached to remote/shared DDC instances in all operation modes; if false they will only cache to DDC instances on the local machine.\n"),
+	ECVF_ReadOnly);
+
+static TAutoConsoleVariable<bool> CVarJobCacheDDCCookPolicy(
+	TEXT("r.ShaderCompiler.JobCacheDDCCookEnableRemotePolicy"),
+	false,
+	TEXT("If true, individual shader jobs will be cached to remote/shared DDC instances in all cook commandlet only; if false they will only cache to DDC instances on the local machine.\n"),
 	ECVF_ReadOnly);
 
 static TAutoConsoleVariable<bool> CVarDebugDumpWorkerInputs(
@@ -228,7 +234,7 @@ bool AreShaderErrorsFatal()
 
 static bool IsShaderJobCacheDDCRemotePolicyEnabled()
 {
-	return CVarJobCacheDDCPolicy.GetValueOnAnyThread();
+	return CVarJobCacheDDCPolicy.GetValueOnAnyThread() || (IsRunningCookCommandlet() && CVarJobCacheDDCCookPolicy.GetValueOnAnyThread());
 }
 
 
