[nDisplay] Fixed typo that does not affect native engine behaviour but affects overriden method used in nDisplay to sync replication in the cluster

diff --git a/Engine/Source/Runtime/Engine/Private/NetConnection.cpp b/Engine/Source/Runtime/Engine/Private/NetConnection.cpp
--- a/Engine/Source/Runtime/Engine/Private/NetConnection.cpp
+++ b/Engine/Source/Runtime/Engine/Private/NetConnection.cpp
@@ -3164,7 +3164,7 @@ void UNetConnection::DispatchPacket( FBitReader& Reader, int32 PacketId, bool& b
 			FInBunch Bunch( this );
 			int32 IncomingStartPos		= Reader.GetPosBits();
 			uint8 bControl				= Reader.ReadBit();
-			Bunch.PacketId				= InPacketId;
+			Bunch.PacketId				= PacketId;
 			Bunch.bOpen					= bControl ? Reader.ReadBit() : 0;
 			Bunch.bClose				= bControl ? Reader.ReadBit() : 0;
 		
@@ -3280,7 +3280,7 @@ void UNetConnection::DispatchPacket( FBitReader& Reader, int32 PacketId, bool& b
 			else if ( Bunch.bPartial )
 			{
 				// If this is an unreliable partial bunch, we simply use packet sequence since we already have it
-				Bunch.ChSequence = InPacketId;
+				Bunch.ChSequence = PacketId;
 			}
 			else
 			{
@@ -3715,7 +3715,7 @@ void UNetConnection::DispatchPacket( FBitReader& Reader, int32 PacketId, bool& b
 				// Disconnect if we received a corrupted packet from the client (eg server crash attempt).
 				if (bIsServer)
 				{
-					UE_LOG(LogNetTraffic, Error, TEXT("Received corrupted packet data with SequenceId: %d from client %s. Disconnecting."), InPacketId, *LowLevelGetRemoteAddress());
+					UE_LOG(LogNetTraffic, Error, TEXT("Received corrupted packet data with SequenceId: %d from client %s. Disconnecting."), PacketId, *LowLevelGetRemoteAddress());
 
 					Close(AddToAndConsumeChainResultPtr(Bunch.ExtendedError, ENetCloseResult::CorruptData));
 
@@ -3723,7 +3723,7 @@ void UNetConnection::DispatchPacket( FBitReader& Reader, int32 PacketId, bool& b
 				}
 				else
 				{
-					UE_LOG(LogNetTraffic, Error, TEXT("Received corrupted packet data with SequenceId: %d from server %s"), InPacketId, *LowLevelGetRemoteAddress());
+					UE_LOG(LogNetTraffic, Error, TEXT("Received corrupted packet data with SequenceId: %d from server %s"), PacketId, *LowLevelGetRemoteAddress());
 #if UE_WITH_IRIS
 					// If Iris reports errors they are unrecoverable.
 					if (Driver->GetReplicationSystem() != nullptr)
