Remove invalid check (for case where WP Level Instance gets loaded in non WP world)

#rnx

diff --git a/Engine/Source/Runtime/Engine/Private/WorldPartition/LevelInstance/LevelInstanceContainerInstance.cpp b/Engine/Source/Runtime/Engine/Private/WorldPartition/LevelInstance/LevelInstanceContainerInstance.cpp
--- a/Engine/Source/Runtime/Engine/Private/WorldPartition/LevelInstance/LevelInstanceContainerInstance.cpp
+++ b/Engine/Source/Runtime/Engine/Private/WorldPartition/LevelInstance/LevelInstanceContainerInstance.cpp
@@ -16,9 +16,6 @@ void ULevelInstanceContainerInstance::Initialize(const FInitializeParams& InPara
 {
 	Super::Initialize(InParams);
 
-	// LevelInstance Container should always have a parent container instance
-	check(GetParentContainerInstance());
-
 	// Add References to parent Container Instance(s). This prevents them from being unloaded before this Container instance is uninitialized (can happen on a map change where Worlds are unloaded in random order through CleanupWorld)
 	UActorDescContainerInstance* CurrentParentContainerInstance = const_cast<UActorDescContainerInstance*>(GetParentContainerInstance());
 	while (CurrentParentContainerInstance)
