World Partition HLOD Stats: make MemoryDiskSizeBytes stat transient and gathered on-demand, instead of updating it during loading of actor descriptors, to avoid a disk operation to gather the disk size.

diff --git a/Engine/Plugins/Editor/WorldPartitionHLODUtilities/Source/Private/WorldPartition/HLOD/Utilities/WorldPartitionHLODUtilities.cpp b/Engine/Plugins/Editor/WorldPartitionHLODUtilities/Source/Private/WorldPartition/HLOD/Utilities/WorldPartitionHLODUtilities.cpp
--- a/Engine/Plugins/Editor/WorldPartitionHLODUtilities/Source/Private/WorldPartition/HLOD/Utilities/WorldPartitionHLODUtilities.cpp
+++ b/Engine/Plugins/Editor/WorldPartitionHLODUtilities/Source/Private/WorldPartition/HLOD/Utilities/WorldPartitionHLODUtilities.cpp
@@ -605,7 +605,6 @@ void GatherOutputStats(AWorldPartitionHLOD* InHLODActor)
 	// Memory stats
 	InHLODActor->SetStat(FWorldPartitionHLODStats::MemoryMeshResourceSizeBytes, MeshResourceSize);
 	InHLODActor->SetStat(FWorldPartitionHLODStats::MemoryTexturesResourceSizeBytes, TexturesResourceSize);
-	InHLODActor->SetStat(FWorldPartitionHLODStats::MemoryDiskSizeBytes, 0); // Clear disk size as it is unknown at this point. It will be assigned at package loading
 }
 
 // Iterate over the source actors and retrieve HLOD relevant components using GetHLODRelevantComponents()
diff --git a/Engine/Source/Runtime/Engine/Private/WorldPartition/HLOD/HLODActor.cpp b/Engine/Source/Runtime/Engine/Private/WorldPartition/HLOD/HLODActor.cpp
--- a/Engine/Source/Runtime/Engine/Private/WorldPartition/HLOD/HLODActor.cpp
+++ b/Engine/Source/Runtime/Engine/Private/WorldPartition/HLOD/HLODActor.cpp
@@ -18,6 +18,7 @@
 #if WITH_EDITOR
 #include "Components/StaticMeshComponent.h"
 #include "Editor.h"
+#include "HAL/FileManager.h"
 #include "Modules/ModuleManager.h"
 #include "Misc/ArchiveMD5.h"
 #include "PhysicsEngine/BodySetup.h"
@@ -31,10 +32,8 @@
 #include "WorldPartition/HLOD/HLODSourceActorsFromCell.h"
 #endif
 
-
 DEFINE_LOG_CATEGORY(LogHLODHash);
 
-
 static int32 GWorldPartitionHLODForceDisableShadows = 0;
 static FAutoConsoleVariableRef CVarWorldPartitionHLODForceDisableShadows(
 	TEXT("wp.Runtime.HLOD.ForceDisableShadows"),
@@ -249,12 +248,6 @@ void AWorldPartitionHLOD::PostLoad()
 			check(SourceCellGuid.IsValid());
 		}
 	}
-
-	if (GetPackage()->GetPIEInstanceID() == INDEX_NONE)
-	{
-		// Update the disk size stat on load, as we can't really know it when saving
-		HLODStats.Add(FWorldPartitionHLODStats::MemoryDiskSizeBytes, FHLODActorDesc::GetPackageSize(this));
-	}
 #endif
 }
 
@@ -347,10 +340,8 @@ void AWorldPartitionHLOD::PreRegisterAllComponents()
 }
 
 #if WITH_EDITOR
-
 void AWorldPartitionHLOD::RerunConstructionScripts()
-{
-}
+{}
 
 void AWorldPartitionHLOD::OnWorldCleanup(UWorld* InWorld, bool bSessionEnded, bool bCleanupResources)
 {
@@ -471,6 +462,16 @@ FBox AWorldPartitionHLOD::GetStreamingBounds() const
 	return HLODBounds;
 }
 
+int64 AWorldPartitionHLOD::GetStat(FName InStatName) const
+{
+	if (InStatName == FWorldPartitionHLODStats::MemoryDiskSizeBytes)
+	{
+		const FString PackageFileName = GetPackage()->GetLoadedPath().GetLocalFullPath();
+		return IFileManager::Get().FileSize(*PackageFileName);
+	}
+	return HLODStats.FindRef(InStatName);
+}
+
 uint32 AWorldPartitionHLOD::GetHLODHash() const
 {
 	return HLODHash;
@@ -496,4 +497,4 @@ void AWorldPartitionHLOD::BuildHLOD(bool bForceBuild)
 	BuildActorTextureStreamingData(this, EMaterialQualityLevel::High, GMaxRHIFeatureLevel);
 }
 
-#endif // WITH_EDITOR
+#endif // WITH_EDITOR
\ No newline at end of file
diff --git a/Engine/Source/Runtime/Engine/Private/WorldPartition/HLOD/HLODActorDesc.cpp b/Engine/Source/Runtime/Engine/Private/WorldPartition/HLOD/HLODActorDesc.cpp
--- a/Engine/Source/Runtime/Engine/Private/WorldPartition/HLOD/HLODActorDesc.cpp
+++ b/Engine/Source/Runtime/Engine/Private/WorldPartition/HLOD/HLODActorDesc.cpp
@@ -20,10 +20,19 @@
 
 FHLODActorDesc::FHLODActorDesc()
 	: EditorBounds(ForceInit)
+{}
+
+int64 FHLODActorDesc::GetStat(FName InStatName) const
 {
+	if (InStatName == FWorldPartitionHLODStats::MemoryDiskSizeBytes)
+	{
+		FString PackageFileName = GetActorPackage().ToString();
+		FPackageName::TryConvertLongPackageNameToFilename(GetActorPackage().ToString(), PackageFileName, FPackageName::GetAssetPackageExtension());
+		return IFileManager::Get().FileSize(*PackageFileName);
+	}
+	return HLODStats.FindRef(InStatName);
 }
 
-
 void FHLODActorDesc::Init(const AActor* InActor)
 {
 	FWorldPartitionActorDesc::Init(InActor);
@@ -143,12 +152,6 @@ void FHLODActorDesc::Serialize(FArchive& Ar)
 			}
 
 			Ar << HLODStats;
-
-			// Update package size stat on load
-			if (Ar.IsLoading())
-			{
-				HLODStats.Add(FWorldPartitionHLODStats::MemoryDiskSizeBytes, GetPackageSize());
-			}
 		}
 
 		if (Ar.CustomVer(FFortniteMainBranchObjectVersion::GUID) >= FFortniteMainBranchObjectVersion::WorldPartitionHLODActorDescSerializeSourceHLODLayer)
@@ -173,32 +176,13 @@ bool FHLODActorDesc::Equals(const FWorldPartitionActorDesc* Other) const
 	{
 		const FHLODActorDesc& HLODActorDesc = *(FHLODActorDesc*)Other;
 		return SourceHLODLayer == HLODActorDesc.SourceHLODLayer &&
-			   HLODStats.OrderIndependentCompareEqual(HLODActorDesc.GetStats()) &&
+			   HLODStats.OrderIndependentCompareEqual(HLODActorDesc.HLODStats) &&
 			   CompareUnsortedArrays(ChildHLODActors, HLODActorDesc.ChildHLODActors) &&
 			   EditorBounds == HLODActorDesc.EditorBounds;
 	}
 	return false;
 }
 
-static int64 GetPackageSize(const FString& InPackageFileName)
-{
-	const int64 PackageSize = IFileManager::Get().FileSize(*InPackageFileName);
-	return PackageSize;
-}
-
-int64 FHLODActorDesc::GetPackageSize() const
-{
-	FString PackageFileName = GetActorPackage().ToString();
-	FPackageName::TryConvertLongPackageNameToFilename(GetActorPackage().ToString(), PackageFileName, FPackageName::GetAssetPackageExtension());
-	return ::GetPackageSize(PackageFileName);
-}
-
-int64 FHLODActorDesc::GetPackageSize(const AWorldPartitionHLOD* InHLODActor)
-{
-	const FString PackageFileName = InHLODActor->GetPackage()->GetLoadedPath().GetLocalFullPath();
-	return ::GetPackageSize(PackageFileName);
-}
-
 bool FHLODActorDesc::IsRuntimeRelevant(const FWorldPartitionActorDescInstance* InActorDescInstance) const
 {
 	return !InActorDescInstance->GetForceNonSpatiallyLoaded();
diff --git a/Engine/Source/Runtime/Engine/Public/WorldPartition/HLOD/HLODActor.h b/Engine/Source/Runtime/Engine/Public/WorldPartition/HLOD/HLODActor.h
--- a/Engine/Source/Runtime/Engine/Public/WorldPartition/HLOD/HLODActor.h
+++ b/Engine/Source/Runtime/Engine/Public/WorldPartition/HLOD/HLODActor.h
@@ -25,7 +25,8 @@ class AWorldPartitionHLOD : public AActor
 {
 	GENERATED_UCLASS_BODY()
 
-	typedef TMap<FName, int64>	FStats;
+	friend class FHLODActorDesc;
+	typedef TMap<FName, int64> FStats;
 
 public:
 	ENGINE_API void SetVisibility(bool bInVisible);
@@ -58,10 +59,13 @@ public:
 	ENGINE_API void BuildHLOD(bool bForceBuild = false);
 	ENGINE_API uint32 GetHLODHash() const;
 
-	const FStats& GetStats() const { return HLODStats; }
-	int64 GetStat(FName InStatName) const { return HLODStats.FindRef(InStatName); }
+	ENGINE_API int64 GetStat(FName InStatName) const;
 	void SetStat(FName InStatName, int64 InStatValue) { HLODStats.Add(InStatName, InStatValue); }
 	void ResetStats() { HLODStats.Reset(); }
+
+private:
+	const FStats& GetStats() const { return HLODStats; }
+
 #endif // WITH_EDITOR
 
 protected:
diff --git a/Engine/Source/Runtime/Engine/Public/WorldPartition/HLOD/HLODActorDesc.h b/Engine/Source/Runtime/Engine/Public/WorldPartition/HLOD/HLODActorDesc.h
--- a/Engine/Source/Runtime/Engine/Public/WorldPartition/HLOD/HLODActorDesc.h
+++ b/Engine/Source/Runtime/Engine/Public/WorldPartition/HLOD/HLODActorDesc.h
@@ -28,10 +28,7 @@ public:
 	inline const TArray<FGuid>& GetChildHLODActors() const { return ChildHLODActors; }
 	inline const FTopLevelAssetPath& GetSourceHLODLayer() const { return SourceHLODLayer; }
 	inline const FStats& GetStats() const { return HLODStats; }
-	inline int64 GetStat(FName InStatName) const { return HLODStats.FindRef(InStatName); }
-
-	ENGINE_API int64 GetPackageSize() const;
-	static ENGINE_API int64 GetPackageSize(const AWorldPartitionHLOD* InHLODActor);
+	int64 GetStat(FName InStatName) const;
 
 	//~ Begin FWorldPartitionActorDesc Interface.
 	virtual bool ShouldValidateRuntimeGrid() const override { return false; }
