[HLOD] fixed crashes when creating WorldPartitioned levels in a project that does not have the "World Partition HLOD Utilities" plugin enabled.

diff --git a/Engine/Source/Runtime/Engine/Private/WorldPartition/HLOD/HLODActor.cpp b/Engine/Source/Runtime/Engine/Private/WorldPartition/HLOD/HLODActor.cpp
--- a/Engine/Source/Runtime/Engine/Private/WorldPartition/HLOD/HLODActor.cpp
+++ b/Engine/Source/Runtime/Engine/Private/WorldPartition/HLOD/HLODActor.cpp
@@ -479,8 +479,8 @@ uint32 AWorldPartitionHLOD::GetHLODHash() const
 
 void AWorldPartitionHLOD::BuildHLOD(bool bForceBuild)
 {
-	IWorldPartitionHLODUtilities* WPHLODUtilities = FModuleManager::Get().LoadModuleChecked<IWorldPartitionHLODUtilitiesModule>("WorldPartitionHLODUtilities").GetUtilities();
-	if (WPHLODUtilities)
+	IWorldPartitionHLODUtilitiesModule* WPHLODUtilitiesModule = FModuleManager::Get().LoadModulePtr<IWorldPartitionHLODUtilitiesModule>("WorldPartitionHLODUtilities");
+	if (IWorldPartitionHLODUtilities* WPHLODUtilities = WPHLODUtilitiesModule != nullptr ? WPHLODUtilitiesModule->GetUtilities() : nullptr)
 	{
 		if (bForceBuild)
 		{
diff --git a/Engine/Source/Runtime/Engine/Private/WorldPartition/HLOD/HLODLayer.cpp b/Engine/Source/Runtime/Engine/Private/WorldPartition/HLOD/HLODLayer.cpp
--- a/Engine/Source/Runtime/Engine/Private/WorldPartition/HLOD/HLODLayer.cpp
+++ b/Engine/Source/Runtime/Engine/Private/WorldPartition/HLOD/HLODLayer.cpp
@@ -33,8 +33,8 @@ UHLODLayer::UHLODLayer(const FObjectInitializer& ObjectInitializer)
 #if WITH_EDITOR
 bool UHLODLayer::DoesRequireWarmup() const
 {
-	IWorldPartitionHLODUtilities* WPHLODUtilities = FModuleManager::Get().LoadModuleChecked<IWorldPartitionHLODUtilitiesModule>("WorldPartitionHLODUtilities").GetUtilities();
-	if (WPHLODUtilities)
+	IWorldPartitionHLODUtilitiesModule* WPHLODUtilitiesModule = FModuleManager::Get().LoadModulePtr<IWorldPartitionHLODUtilitiesModule>("WorldPartitionHLODUtilities");
+	if (IWorldPartitionHLODUtilities* WPHLODUtilities = WPHLODUtilitiesModule != nullptr ? WPHLODUtilitiesModule->GetUtilities() : nullptr)
 	{
 		return WPHLODUtilities->GetHLODBuilderClass(this)->GetDefaultObject<UHLODBuilder>()->RequiresWarmup();
 	}
@@ -105,8 +105,8 @@ void UHLODLayer::PostLoad()
 {
 	Super::PostLoad();
 
-	IWorldPartitionHLODUtilities* WPHLODUtilities = FModuleManager::Get().LoadModuleChecked<IWorldPartitionHLODUtilitiesModule>("WorldPartitionHLODUtilities").GetUtilities();
-	if (WPHLODUtilities)
+	IWorldPartitionHLODUtilitiesModule* WPHLODUtilitiesModule = FModuleManager::Get().LoadModulePtr<IWorldPartitionHLODUtilitiesModule>("WorldPartitionHLODUtilities");
+	if (IWorldPartitionHLODUtilities* WPHLODUtilities = WPHLODUtilitiesModule != nullptr ? WPHLODUtilitiesModule->GetUtilities() : nullptr)
 	{
 		const UClass* BuilderClass = WPHLODUtilities->GetHLODBuilderClass(this);
 		const UClass* BuilderSettingsClass = BuilderClass ? BuilderClass->GetDefaultObject<UHLODBuilder>()->GetSettingsClass() : nullptr;
@@ -140,12 +140,16 @@ void UHLODLayer::PostEditChangeProperty(FPropertyChangedEvent& PropertyChangedEv
 
 	const FName PropertyName = PropertyChangedEvent.Property ? PropertyChangedEvent.Property->GetFName() : FName();
 
-	IWorldPartitionHLODUtilities* WPHLODUtilities = FModuleManager::Get().LoadModuleChecked<IWorldPartitionHLODUtilitiesModule>("WorldPartitionHLODUtilities").GetUtilities();
 
 	if (PropertyName == GET_MEMBER_NAME_CHECKED(UHLODLayer, LayerType) ||
 		PropertyName == GET_MEMBER_NAME_CHECKED(UHLODLayer, HLODBuilderClass))
 	{
-		HLODBuilderSettings = WPHLODUtilities->CreateHLODBuilderSettings(this);
+		IWorldPartitionHLODUtilitiesModule* WPHLODUtilitiesModule = FModuleManager::Get().LoadModulePtr<IWorldPartitionHLODUtilitiesModule>("WorldPartitionHLODUtilities");
+		IWorldPartitionHLODUtilities* WPHLODUtilities = WPHLODUtilitiesModule != nullptr ? WPHLODUtilitiesModule->GetUtilities() : nullptr;
+		if (WPHLODUtilities != nullptr)
+		{
+			HLODBuilderSettings = WPHLODUtilities->CreateHLODBuilderSettings(this);
+		}
 	}
 	else if (PropertyName == GET_MEMBER_NAME_CHECKED(UHLODLayer, ParentLayer))
 	{
diff --git a/Engine/Source/Runtime/Engine/Private/WorldPartition/RuntimeHashSet/WorldPartitionRuntimeHashSetHLODGeneration.cpp b/Engine/Source/Runtime/Engine/Private/WorldPartition/RuntimeHashSet/WorldPartitionRuntimeHashSetHLODGeneration.cpp
--- a/Engine/Source/Runtime/Engine/Private/WorldPartition/RuntimeHashSet/WorldPartitionRuntimeHashSetHLODGeneration.cpp
+++ b/Engine/Source/Runtime/Engine/Private/WorldPartition/RuntimeHashSet/WorldPartitionRuntimeHashSetHLODGeneration.cpp
@@ -161,8 +161,13 @@ bool UWorldPartitionRuntimeHashSet::SupportsHLODs() const
 
 bool UWorldPartitionRuntimeHashSet::SetupHLODActors(const IStreamingGenerationContext* StreamingGenerationContext, const UWorldPartition::FSetupHLODActorsParams& Params) const
 {
-	IWorldPartitionHLODUtilities* WPHLODUtilities = FModuleManager::Get().LoadModuleChecked<IWorldPartitionHLODUtilitiesModule>("WorldPartitionHLODUtilities").GetUtilities();
-	check(WPHLODUtilities);
+	IWorldPartitionHLODUtilitiesModule* WPHLODUtilitiesModule = FModuleManager::Get().LoadModulePtr<IWorldPartitionHLODUtilitiesModule>("WorldPartitionHLODUtilities");
+	IWorldPartitionHLODUtilities* WPHLODUtilities = WPHLODUtilitiesModule != nullptr ? WPHLODUtilitiesModule->GetUtilities() : nullptr;
+	if (WPHLODUtilities == nullptr)
+	{
+		UE_LOG(LogWorldPartition, Error, TEXT("%hs requires plugin 'World Partition HLOD Utilities'."), __FUNCTION__);
+		return false;
+	}
 
 	UWorldPartition* WorldPartition = GetOuterUWorldPartition();
 	const UDataLayerManager* DataLayerManager = WorldPartition->GetDataLayerManager();
diff --git a/Engine/Source/Runtime/Engine/Private/WorldPartition/RuntimeSpatialHash/RuntimeSpatialHashHLOD.cpp b/Engine/Source/Runtime/Engine/Private/WorldPartition/RuntimeSpatialHash/RuntimeSpatialHashHLOD.cpp
--- a/Engine/Source/Runtime/Engine/Private/WorldPartition/RuntimeSpatialHash/RuntimeSpatialHashHLOD.cpp
+++ b/Engine/Source/Runtime/Engine/Private/WorldPartition/RuntimeSpatialHash/RuntimeSpatialHashHLOD.cpp
@@ -113,6 +113,14 @@ static void GameTick(UWorld* InWorld)
 
 static TArray<FGuid> GenerateHLODActorsForGrid(UWorldPartition* WorldPartition, const IStreamingGenerationContext* StreamingGenerationContext, const UWorldPartition::FSetupHLODActorsParams& Params, const FSpatialHashRuntimeGrid& RuntimeGrid, uint32 HLODLevel, FHLODCreationContext& Context, ISourceControlHelper* SourceControlHelper, const TArray<const IStreamingGenerationContext::FActorSetInstance*>& ActorSetInstances, TArray<FName>& NewActors, TArray<FName>& DirtyActors, const FSpatialHashSettings& Settings)
 {
+	IWorldPartitionHLODUtilitiesModule* WPHLODUtilitiesModule = FModuleManager::Get().LoadModulePtr<IWorldPartitionHLODUtilitiesModule>("WorldPartitionHLODUtilities");
+	IWorldPartitionHLODUtilities* WPHLODUtilities = WPHLODUtilitiesModule != nullptr ? WPHLODUtilitiesModule->GetUtilities() : nullptr;
+	if (WPHLODUtilities == nullptr)
+	{
+		UE_LOG(LogWorldPartition, Error, TEXT("%hs requires plugin 'World Partition HLOD Utilities'."), __FUNCTION__);
+		return {};
+	}
+
 	const FBox WorldBounds = StreamingGenerationContext->GetWorldBounds();
 	const FSquare2DGridHelper PartitionedActors = GetPartitionedActors(WorldBounds, RuntimeGrid, ActorSetInstances, Settings);
 	const FSquare2DGridHelper::FGridLevel::FGridCell& AlwaysLoadedCell = PartitionedActors.GetAlwaysLoadedCell();
@@ -207,7 +215,6 @@ static TArray<FGuid> GenerateHLODActorsForGrid(UWorldPartition* WorldPartition,
 				CreationParams.ContentBundleGuid = GridCellDataChunk.GetContentBundleID();
 				CreationParams.DataLayerInstances = GridCellDataChunk.GetDataLayers();
 
-				IWorldPartitionHLODUtilities* WPHLODUtilities = FModuleManager::Get().LoadModuleChecked<IWorldPartitionHLODUtilitiesModule>("WorldPartitionHLODUtilities").GetUtilities();
 				TArray<AWorldPartitionHLOD*> CellHLODActors = WPHLODUtilities->CreateHLODActors(Context, CreationParams, ActorInstances);
 				if (!CellHLODActors.IsEmpty())
 				{
