AutomationTest inequality asserts and macros, fixes nullptr crash in NotEqualCaseInsensitive, adds missing overloads.

diff --git a/Engine/Source/Runtime/Core/Private/Misc/AutomationTest.cpp b/Engine/Source/Runtime/Core/Private/Misc/AutomationTest.cpp
--- a/Engine/Source/Runtime/Core/Private/Misc/AutomationTest.cpp
+++ b/Engine/Source/Runtime/Core/Private/Misc/AutomationTest.cpp
@@ -1695,7 +1695,7 @@ bool FAutomationTestBase::TestEqual(const TCHAR* What, const float Actual, const
 {
 	if (!FMath::IsNearlyEqual(Actual, Expected, Tolerance))
 	{
-		AddError(FString::Printf(TEXT("Expected '%s' to be %f, but it was %f within tolerance %f."), What, Expected, Actual, Tolerance), 1);
+		AddError(FString::Printf(TEXT("Expected '%s' to be %f, but it was %f and outside tolerance %f."), What, Expected, Actual, Tolerance), 1);
 		return false;
 	}
 	return true;
@@ -1705,7 +1705,7 @@ bool FAutomationTestBase::TestEqual(const TCHAR* What, const double Actual, cons
 {
 	if (!FMath::IsNearlyEqual(Actual, Expected, Tolerance))
 	{
-		AddError(FString::Printf(TEXT("Expected '%s' to be %f, but it was %f within tolerance %f."), What, Expected, Actual, Tolerance), 1);
+		AddError(FString::Printf(TEXT("Expected '%s' to be %f, but it was %f and outside tolerance %f."), What, Expected, Actual, Tolerance), 1);
 		return false;
 	}
 	return true;
@@ -1715,7 +1715,7 @@ bool FAutomationTestBase::TestEqual(const TCHAR* What, const FVector Actual, con
 {
 	if (!Expected.Equals(Actual, Tolerance))
 	{
-		AddError(FString::Printf(TEXT("Expected '%s' to be %s, but it was %s within tolerance %f."), What, *Expected.ToString(), *Actual.ToString(), Tolerance), 1);
+		AddError(FString::Printf(TEXT("Expected '%s' to be %s, but it was %s and outside tolerance %f."), What, *Expected.ToString(), *Actual.ToString(), Tolerance), 1);
 		return false;
 	}
 	return true;
@@ -1725,7 +1725,7 @@ bool FAutomationTestBase::TestEqual(const TCHAR* What, const FTransform Actual,
 {
 	if (!Expected.Equals(Actual, Tolerance))
 	{
-		AddError(FString::Printf(TEXT("Expected '%s' to be %s, but it was %s within tolerance %f."), What, *Expected.ToString(), *Actual.ToString(), Tolerance), 1);
+		AddError(FString::Printf(TEXT("Expected '%s' to be %s, but it was %s and outside tolerance %f."), What, *Expected.ToString(), *Actual.ToString(), Tolerance), 1);
 		return false;
 	}
 	return true;
@@ -1735,7 +1735,7 @@ bool FAutomationTestBase::TestEqual(const TCHAR* What, const FRotator Actual, co
 {
 	if (!Expected.Equals(Actual, Tolerance))
 	{
-		AddError(FString::Printf(TEXT("Expected '%s' to be %s, but it was %s within tolerance %f."), What, *Expected.ToString(), *Actual.ToString(), Tolerance), 1);
+		AddError(FString::Printf(TEXT("Expected '%s' to be %s, but it was %s and outside tolerance %f."), What, *Expected.ToString(), *Actual.ToString(), Tolerance), 1);
 		return false;
 	}
 	return true;
@@ -1771,6 +1771,26 @@ bool FAutomationTestBase::TestEqual(const TCHAR* What, const TCHAR* Actual, cons
 	return true;
 }
 
+bool FAutomationTestBase::TestNotEqual(const TCHAR* What, const float Actual, const float Expected, float Tolerance)
+{
+	if (FMath::IsNearlyEqual(Actual, Expected, Tolerance))
+	{
+		AddError(FString::Printf(TEXT("Expected '%s' to be unequal to %f, but it was %f and within tolerance %f."), What, Expected, Actual, Tolerance), 1);
+		return false;
+	}
+	return true;
+}
+
+bool FAutomationTestBase::TestNotEqual(const TCHAR* What, const double Actual, const double Expected, double Tolerance)
+{
+	if (FMath::IsNearlyEqual(Actual, Expected, Tolerance))
+	{
+		AddError(FString::Printf(TEXT("Expected '%s' to be unequal to %f, but it was %f and within tolerance %f."), What, Expected, Actual, Tolerance), 1);
+		return false;
+	}
+	return true;
+}
+
 bool FAutomationTestBase::TestEqualInsensitive(const TCHAR* What, const TCHAR* Actual, const TCHAR* Expected)
 {
 	if (FCString::Stricmp(Actual, Expected) != 0)
@@ -1783,10 +1803,21 @@ bool FAutomationTestBase::TestEqualInsensitive(const TCHAR* What, const TCHAR* A
 
 bool FAutomationTestBase::TestNotEqualInsensitive(const TCHAR* What, const TCHAR* Actual, const TCHAR* Expected)
 {
-	if (FCString::Stricmp(Actual, Expected) == 0)
+	if (Actual && Expected)
 	{
-		AddError(FString::Printf(TEXT("Expected '%s' to differ from \"%s\", but it was \"%s\"."), What, Expected, Actual), 1);
-		return false;
+		if (FCString::Stricmp(Actual, Expected) == 0)
+		{
+			AddError(FString::Printf(TEXT("Expected '%s' to differ from \"%s\", but it was \"%s\"."), What, Expected, Actual), 1);
+			return false;
+		}
+	}
+	else // null exists
+	{ 
+		if (Actual == Expected)
+		{
+			AddError(FString::Printf(TEXT("Expected '%s' to differ but both values were unexpectedly null"), What), 1);
+			return false;
+		}
 	}
 	return true;
 }
@@ -1816,6 +1847,244 @@ bool FAutomationTestBase::TestNearlyEqual(const TCHAR* What, const FRotator Actu
 	return TestEqual(What, Actual, Expected, Tolerance);
 }
 
+bool FAutomationTestBase::TestLessThan(const TCHAR* What, const int32 Actual, const int32 Expected)
+{
+	if (Actual < Expected)
+	{
+		return true;
+	}
+	AddError(FString::Printf(TEXT("Expected '%s' to be less than %d, but it was %d."), What, Expected, Actual), 1);
+	return false;
+}
+
+bool FAutomationTestBase::TestLessThan(const TCHAR* What, const int64 Actual, const int64 Expected)
+{
+	if (Actual < Expected)
+	{
+		return true;
+	}	
+	AddError(FString::Printf(TEXT("Expected '%s' to be less than %" PRId64 ", but it was %" PRId64 "."), What, Expected, Actual), 1);
+	return false;
+}
+
+bool FAutomationTestBase::TestGreaterThan(const TCHAR* What, const int32 Actual, const int32 Expected)
+{
+	if (Actual > Expected)
+	{
+		return true;
+	}
+	AddError(FString::Printf(TEXT("Expected '%s' to be greater than %d, but it was %d."), What, Expected, Actual), 1);
+	return false;
+}
+
+bool FAutomationTestBase::TestGreaterThan(const TCHAR* What, const int64 Actual, const int64 Expected)
+{
+	if (Actual > Expected)
+	{
+		return true;
+	}
+	AddError(FString::Printf(TEXT("Expected '%s' to be greater than %" PRId64 ", but it was %" PRId64 "."), What, Expected, Actual), 1);
+	return false;
+}
+
+bool FAutomationTestBase::TestLessEqual(const TCHAR* What, const int32 Actual, const int32 Expected)
+{
+	if (Actual <= Expected)
+	{
+		return true;
+	}
+	AddError(FString::Printf(TEXT("Expected '%s' to be less than or equal to %d, but it was %d."), What, Expected, Actual), 1);
+	return false;
+}
+
+bool FAutomationTestBase::TestLessEqual(const TCHAR* What, const int64 Actual, const int64 Expected)
+{
+	if (Actual <= Expected)
+	{
+		return true;
+	}
+	AddError(FString::Printf(TEXT("Expected '%s' to be less than or equal to %" PRId64 ", but it was %" PRId64 "."), What, Expected, Actual), 1);
+	return false;
+}
+
+bool FAutomationTestBase::TestGreaterEqual(const TCHAR* What, const int32 Actual, const int32 Expected)
+{
+	if (Actual >= Expected)
+	{
+		return true;
+	}
+	AddError(FString::Printf(TEXT("Expected '%s' to be greater than or equal to %d, but it was %d."), What, Expected, Actual), 1);
+	return false;
+}
+
+bool FAutomationTestBase::TestGreaterEqual(const TCHAR* What, const int64 Actual, const int64 Expected)
+{
+	if (Actual >= Expected)
+	{
+		return true;
+	}
+	AddError(FString::Printf(TEXT("Expected '%s' to be greater than or equal to %" PRId64 ", but it was %" PRId64 "."), What, Expected, Actual), 1);
+	return false;
+}
+
+#if PLATFORM_64BITS
+bool FAutomationTestBase::TestLessThan(const TCHAR* What, const SIZE_T Actual, const SIZE_T Expected)
+{
+	if (Actual < Expected)
+	{
+		return true;
+	}
+	AddError(FString::Printf(TEXT("Expected '%s' to be less than %" PRIuPTR ", but it was %" PRIuPTR "."), What, Expected, Actual), 1);
+	return false;
+}
+
+bool FAutomationTestBase::TestGreaterThan(const TCHAR* What, const SIZE_T Actual, const SIZE_T Expected)
+{
+	if (Actual > Expected)
+	{
+		return true;
+	}
+	AddError(FString::Printf(TEXT("Expected '%s' to be greater than %" PRIuPTR ", but it was %" PRIuPTR "."), What, Expected, Actual), 1);
+	return false;
+}
+
+bool FAutomationTestBase::TestLessEqual(const TCHAR* What, const SIZE_T Actual, const SIZE_T Expected)
+{
+	if (Actual <= Expected)
+	{
+		return true;
+	}
+	AddError(FString::Printf(TEXT("Expected '%s' to be less than or equal to %" PRIuPTR ", but it was %" PRIuPTR "."), What, Expected, Actual), 1);
+	return false;
+}
+
+bool FAutomationTestBase::TestGreaterEqual(const TCHAR* What, const SIZE_T Actual, const SIZE_T Expected)
+{
+	if (Actual >= Expected)
+	{
+		return true;
+	}
+	AddError(FString::Printf(TEXT("Expected '%s' to be greater than or equal to %" PRIuPTR ", but it was %" PRIuPTR "."), What, Expected, Actual), 1);
+	return false;
+}
+#endif // PLATFORM_64BITS
+
+bool FAutomationTestBase::TestLessThan(const TCHAR* What, const float Actual, const float Expected, float Tolerance)
+{
+	if (FMath::IsNearlyEqual(Actual, Expected, Tolerance))
+	{
+		AddError(FString::Printf(TEXT("Expected '%s' to be less than %f, but it was %f and within equality tolerance %f."), What, Expected, Actual, Tolerance), 1);
+		return false;
+	}
+	if (Actual < Expected)
+	{
+		return true;
+	}
+	AddError(FString::Printf(TEXT("Expected '%s' to be less than %f, but it was %f and outside equality tolerance %f."), What, Expected, Actual, Tolerance), 1);
+	return false;
+}
+
+bool FAutomationTestBase::TestLessThan(const TCHAR* What, const double Actual, const double Expected, double Tolerance)
+{
+	if (FMath::IsNearlyEqual(Actual, Expected, Tolerance))
+	{
+		AddError(FString::Printf(TEXT("Expected '%s' to be less than %f, but it was %f and within equality tolerance %f."), What, Expected, Actual, Tolerance), 1);
+		return false;
+	}
+	if (Actual < Expected)
+	{
+		return true;
+	}
+	AddError(FString::Printf(TEXT("Expected '%s' to be less than %f, but it was %f and outside equality tolerance %f."), What, Expected, Actual, Tolerance), 1);
+	return false;
+}
+
+bool FAutomationTestBase::TestGreaterThan(const TCHAR* What, const float Actual, const float Expected, float Tolerance)
+{
+	if (FMath::IsNearlyEqual(Actual, Expected, Tolerance))
+	{
+		AddError(FString::Printf(TEXT("Expected '%s' to be less than %f, but it was %f and within equality tolerance %f."), What, Expected, Actual, Tolerance), 1);
+		return false;
+	}
+	if (Actual > Expected)
+	{
+		return true;
+	}
+	AddError(FString::Printf(TEXT("Expected '%s' to be less than %f, but it was %f and outside equality tolerance %f."), What, Expected, Actual, Tolerance), 1);
+	return false;
+}
+
+bool FAutomationTestBase::TestGreaterThan(const TCHAR* What, const double Actual, const double Expected, double Tolerance)
+{
+	if (FMath::IsNearlyEqual(Actual, Expected, Tolerance))
+	{
+		AddError(FString::Printf(TEXT("Expected '%s' to be less than %f, but it was %f and within equality tolerance %f."), What, Expected, Actual, Tolerance), 1);
+		return false;
+	}
+	if (Actual > Expected)
+	{
+		return true;
+	}
+	AddError(FString::Printf(TEXT("Expected '%s' to be less than %f, but it was %f and outside equality tolerance %f."), What, Expected, Actual, Tolerance), 1);
+	return false;
+}
+
+bool FAutomationTestBase::TestLessEqual(const TCHAR* What, const float Actual, const float Expected, float Tolerance)
+{
+	if (Actual < Expected)
+	{
+		return true;
+	}
+	if (FMath::IsNearlyEqual(Actual, Expected, Tolerance))
+	{
+		return true;
+	}
+	AddError(FString::Printf(TEXT("Expected '%s' to be less than or equal to %f, but it was %f and outside equality tolerance %f."), What, Expected, Actual, Tolerance), 1);
+	return false;
+}
+
+bool FAutomationTestBase::TestLessEqual(const TCHAR* What, const double Actual, const double Expected, double Tolerance)
+{
+	if (Actual < Expected)
+	{
+		return true;
+	}
+	if (FMath::IsNearlyEqual(Actual, Expected, Tolerance))
+	{
+		return true;
+	}
+	AddError(FString::Printf(TEXT("Expected '%s' to be less than or equal to %f, but it was %f and outside equality tolerance %f."), What, Expected, Actual, Tolerance), 1);
+	return false;
+}
+
+bool FAutomationTestBase::TestGreaterEqual(const TCHAR* What, const float Actual, const float Expected, float Tolerance)
+{
+	if (Actual > Expected)
+	{
+		return true;
+	}
+	if (FMath::IsNearlyEqual(Actual, Expected, Tolerance))
+	{
+		return true;
+	}
+	AddError(FString::Printf(TEXT("Expected '%s' to be greater than or equal to %f, but it was %f and outside equality tolerance %f."), What, Expected, Actual, Tolerance), 1);
+	return false;
+}
+
+bool FAutomationTestBase::TestGreaterEqual(const TCHAR* What, const double Actual, const double Expected, double Tolerance)
+{
+	if (Actual > Expected)
+	{
+		return true;
+	}
+	if (FMath::IsNearlyEqual(Actual, Expected, Tolerance))
+	{
+		return true;
+	}
+	AddError(FString::Printf(TEXT("Expected '%s' to be greater than or equal to %f, but it was %f and outside equality tolerance %f."), What, Expected, Actual, Tolerance), 1);
+	return false;
+}
+
 bool FAutomationTestBase::TestFalse(const TCHAR* What, bool Value)
 {
 	if (Value)
diff --git a/Engine/Source/Runtime/Core/Public/Misc/AutomationTest.h b/Engine/Source/Runtime/Core/Public/Misc/AutomationTest.h
--- a/Engine/Source/Runtime/Core/Public/Misc/AutomationTest.h
+++ b/Engine/Source/Runtime/Core/Public/Misc/AutomationTest.h
@@ -1916,6 +1916,66 @@ public:
 		return TestEqualInsensitive(*What, *Actual, *Expected);
 	}
 
+	bool TestEqualInsensitive(const TCHAR* What, const FString& Actual, const TCHAR* Expected)
+	{
+		return TestEqualInsensitive(What, *Actual, Expected);
+	}
+
+	bool TestEqualInsensitive(const FString& What, const FString& Actual, const TCHAR* Expected)
+	{
+		return TestEqualInsensitive(*What, *Actual, Expected);
+	}
+
+	bool TestEqualInsensitive(const TCHAR* What, const TCHAR* Actual, const FString& Expected)
+	{
+		return TestEqualInsensitive(What, Actual, *Expected);
+	}
+
+	bool TestEqualInsensitive(const FString& What, const TCHAR* Actual, const FString& Expected)
+	{
+		return TestEqualInsensitive(*What, Actual, *Expected);
+	}
+
+	bool TestEqualInsensitive(const TCHAR* What, const FString& Actual, const FString& Expected)
+	{
+		return TestEqualInsensitive(What, *Actual, *Expected);
+	}
+
+	bool TestEqualInsensitive(const FString& What, const FString& Actual, const FString& Expected)
+	{
+		return TestEqualInsensitive(*What, *Actual, *Expected);
+	}
+
+	bool TestNotEqualInsensitive(const TCHAR* What, const FString& Actual, const TCHAR* Expected)
+	{
+		return TestNotEqualInsensitive(What, *Actual, Expected);
+	}
+
+	bool TestNotEqualInsensitive(const FString& What, const FString& Actual, const TCHAR* Expected)
+	{
+		return TestNotEqualInsensitive(*What, *Actual, Expected);
+	}
+
+	bool TestNotEqualInsensitive(const TCHAR* What, const TCHAR* Actual, const FString& Expected)
+	{
+		return TestNotEqualInsensitive(What, Actual, *Expected);
+	}
+
+	bool TestNotEqualInsensitive(const FString& What, const TCHAR* Actual, const FString& Expected)
+	{
+		return TestNotEqualInsensitive(*What, Actual, *Expected);
+	}
+
+	bool TestNotEqualInsensitive(const TCHAR* What, const FString& Actual, const FString& Expected)
+	{
+		return TestNotEqualInsensitive(What, *Actual, *Expected);
+	}
+
+	bool TestNotEqualInsensitive(const FString& What, const FString& Actual, const FString& Expected)
+	{
+		return TestNotEqualInsensitive(*What, *Actual, *Expected);
+	}
+
 	/**
 	 * Logs an error if the two values are not equal.
 	 *
@@ -1973,6 +2033,108 @@ public:
 		return TestNearlyEqual(*What, Actual, Expected, Tolerance);
 	}
 
+	CORE_API bool TestLessThan(const TCHAR* What, const int32 Actual, const int32 Expected);
+	CORE_API bool TestLessThan(const TCHAR* What, const int64 Actual, const int64 Expected);
+	CORE_API bool TestGreaterThan(const TCHAR* What, const int32 Actual, const int32 Expected);
+	CORE_API bool TestGreaterThan(const TCHAR* What, const int64 Actual, const int64 Expected);
+	CORE_API bool TestLessEqual(const TCHAR* What, const int32 Actual, const int32 Expected);
+	CORE_API bool TestLessEqual(const TCHAR* What, const int64 Actual, const int64 Expected);
+	CORE_API bool TestGreaterEqual(const TCHAR* What, const int32 Actual, const int32 Expected);
+	CORE_API bool TestGreaterEqual(const TCHAR* What, const int64 Actual, const int64 Expected);
+#if PLATFORM_64BITS
+	CORE_API bool TestLessThan(const TCHAR* What, const SIZE_T Actual, const SIZE_T Expected);
+	CORE_API bool TestGreaterThan(const TCHAR* What, const SIZE_T Actual, const SIZE_T Expected);
+	CORE_API bool TestLessEqual(const TCHAR* What, const SIZE_T Actual, const SIZE_T Expected);
+	CORE_API bool TestGreaterEqual(const TCHAR* What, const SIZE_T Actual, const SIZE_T Expected);
+#endif
+	CORE_API bool TestLessThan(const TCHAR* What, const float Actual, const float Expected, float Tolerance = UE_KINDA_SMALL_NUMBER);
+	CORE_API bool TestLessThan(const TCHAR* What, const double Actual, const double Expected, double Tolerance = UE_KINDA_SMALL_NUMBER);
+	CORE_API bool TestGreaterThan(const TCHAR* What, const float Actual, const float Expected, float Tolerance = UE_KINDA_SMALL_NUMBER);
+	CORE_API bool TestGreaterThan(const TCHAR* What, const double Actual, const double Expected, double Tolerance = UE_KINDA_SMALL_NUMBER);
+	CORE_API bool TestLessEqual(const TCHAR* What, const float Actual, const float Expected, float Tolerance = UE_KINDA_SMALL_NUMBER);
+	CORE_API bool TestLessEqual(const TCHAR* What, const double Actual, const double Expected, double Tolerance = UE_KINDA_SMALL_NUMBER);
+	CORE_API bool TestGreaterEqual(const TCHAR* What, const float Actual, const float Expected, float Tolerance = UE_KINDA_SMALL_NUMBER);
+	CORE_API bool TestGreaterEqual(const TCHAR* What, const double Actual, const double Expected, double Tolerance = UE_KINDA_SMALL_NUMBER);
+
+	bool TestLessThan(const FString& What, const int32 Actual, const int32 Expected)
+	{
+		return TestLessThan(*What, Actual, Expected);
+	}
+
+	bool TestLessThan(const FString& What, const int64 Actual, const int64 Expected)
+	{
+		return TestLessThan(*What, Actual, Expected);
+	}
+
+	bool TestGreaterThan(const FString& What, const int32 Actual, const int32 Expected)
+	{
+		return TestGreaterThan(*What, Actual, Expected);
+	}
+
+	bool TestGreaterThan(const FString& What, const int64 Actual, const int64 Expected)
+	{
+		return TestGreaterThan(*What, Actual, Expected);
+	}
+
+	bool TestLessEqual(const FString& What, const int32 Actual, const int32 Expected)
+	{
+		return TestLessEqual(*What, Actual, Expected);
+	}
+
+	bool TestLessEqual(const FString& What, const int64 Actual, const int64 Expected)
+	{
+		return TestLessEqual(*What, Actual, Expected);
+	}
+
+	bool TestGreaterEqual(const FString& What, const int32 Actual, const int32 Expected)
+	{
+		return TestGreaterEqual(*What, Actual, Expected);
+	}
+
+	bool TestGreaterEqual(const FString& What, const int64 Actual, const int64 Expected)
+	{
+		return TestGreaterEqual(*What, Actual, Expected);
+	}
+
+	bool TestLessThan(const FString& What, const float Actual, const float Expected, float Tolerance = UE_KINDA_SMALL_NUMBER)
+	{
+		return TestLessThan(*What, Actual, Expected, Tolerance);
+	}
+
+	bool TestLessThan(const FString& What, const double Actual, const double Expected, double Tolerance = UE_KINDA_SMALL_NUMBER)
+	{
+		return TestLessThan(*What, Actual, Expected, Tolerance);
+	}
+
+	bool TestGreaterThan(const FString& What, const float Actual, const float Expected, float Tolerance = UE_KINDA_SMALL_NUMBER)
+	{
+		return TestGreaterThan(*What, Actual, Expected, Tolerance);
+	}
+
+	bool TestGreaterThan(const FString& What, const double Actual, const double Expected, double Tolerance = UE_KINDA_SMALL_NUMBER)
+	{
+		return TestGreaterThan(*What, Actual, Expected, Tolerance);
+	}
+
+	bool TestLessEqual(const FString& What, const float Actual, const float Expected, float Tolerance = UE_KINDA_SMALL_NUMBER)
+	{
+		return TestLessEqual(*What, Actual, Expected, Tolerance);
+	}
+
+	bool TestLessEqual(const FString& What, const double Actual, const double Expected, double Tolerance = UE_KINDA_SMALL_NUMBER)
+	{
+		return TestLessEqual(*What, Actual, Expected, Tolerance);
+	}
+
+	bool TestGreaterEqual(const FString& What, const float Actual, const float Expected, float Tolerance = UE_KINDA_SMALL_NUMBER)
+	{
+		return TestGreaterEqual(*What, Actual, Expected, Tolerance);
+	}
+
+	bool TestGreaterEqual(const FString& What, const double Actual, const double Expected, double Tolerance = UE_KINDA_SMALL_NUMBER)
+	{
+		return TestGreaterEqual(*What, Actual, Expected, Tolerance);
+	}
 
 	/**
 	 * Logs an error if the specified Boolean value is not false.
@@ -1980,7 +2142,7 @@ public:
 	 * @param What - Description text for the test.
 	 * @param Value - The value to test.
 	 *
-	 * @see TestFalse
+	 * @see TestTrue
 	 */
 	CORE_API bool TestFalse(const TCHAR* What, bool Value);
 
@@ -2039,6 +2201,19 @@ public:
 		return TestNotEqual(*Description, Actual, Expected);
 	}
 
+	CORE_API bool TestNotEqual(const TCHAR* What, const float Actual, const float Expected, float Tolerance = UE_KINDA_SMALL_NUMBER);
+	CORE_API bool TestNotEqual(const TCHAR* What, const double Actual, const double Expected, double Tolerance = UE_KINDA_SMALL_NUMBER);
+
+	bool TestNotEqual(const FString& What, const float Actual, const float Expected, float Tolerance = UE_KINDA_SMALL_NUMBER)
+	{
+		return TestNotEqual(*What, Actual, Expected, Tolerance);
+	}
+
+	bool TestNotEqual(const FString& What, const double Actual, const double Expected, double Tolerance = UE_KINDA_SMALL_NUMBER)
+	{
+		return TestNotEqual(*What, Actual, Expected, Tolerance);
+	}
+
 	/**
 	 * Logs an error if the specified pointer is NULL.
 	 *
@@ -3871,11 +4046,12 @@ public: \
 
 
 /**
- * Macros to make it easy to test state with one-liners: they will run the appropriate
- * test method and, if the test fail, with execute `return false;`, which (if placed in
+ * Macros for early exit one-liners: they will run the appropriate
+ * test method and on failure also execute `return false;`, which (if placed in
  * the main test case method) will stop the test immediately.
  *
  * The error logging is already handled by the test method being called.
+ * EXPR variants automatically generate a `What` description string.
  *
  * As a result, you can easily test things that, if wrong, would potentially crash the test:
  *
@@ -3915,6 +4091,12 @@ public: \
 		return false;\
 	}
 
+#define UTEST_NEARLY_EQUAL(What, Actual, Expected, Tolerance)\
+	if (!TestNearlyEqual(What, Actual, Expected, Tolerance))\
+	{\
+		return false;\
+	}
+
 #define UTEST_NEARLY_EQUAL_EXPR(Actual, Expected, Tolerance)\
 	if (!TestNearlyEqual(TEXT(#Actual), Actual, Expected, Tolerance))\
 	{\
@@ -3957,6 +4139,102 @@ public: \
 		return false;\
 	}
 
+#define UTEST_LESS(What, Actual, Expected)\
+	if (!TestLessThan(What, Actual, Expected))\
+	{\
+		return false;\
+	}
+
+#define UTEST_LESS_EXPR(Actual, Expected)\
+	if (!TestLessThan(TEXT(#Actual), Actual, Expected))\
+	{\
+		return false;\
+	}
+
+#define UTEST_LESS_TOLERANCE(What, Actual, Expected, Tolerance)\
+	if (!TestLessThan(What, Actual, Expected, Tolerance))\
+	{\
+		return false;\
+	}
+
+#define UTEST_LESS_TOLERANCE_EXPR(Actual, Expected, Tolerance)\
+	if (!TestLessThan(TEXT(#Actual), Actual, Expected, Tolerance))\
+	{\
+		return false;\
+	}
+
+#define UTEST_GREATER(What, Actual, Expected)\
+	if (!TestGreaterThan(What, Actual, Expected))\
+	{\
+		return false;\
+	}
+
+#define UTEST_GREATER_EXPR(Actual, Expected)\
+	if (!TestGreaterThan(TEXT(#Actual), Actual, Expected))\
+	{\
+		return false;\
+	}
+
+#define UTEST_GREATER_TOLERANCE(What, Actual, Expected, Tolerance)\
+	if (!TestGreaterThan(What, Actual, Expected, Tolerance))\
+	{\
+		return false;\
+	}
+
+#define UTEST_GREATER_TOLERANCE_EXPR(Actual, Expected, Tolerance)\
+	if (!TestGreaterThan(TEXT(#Actual), Actual, Expected, Tolerance))\
+	{\
+		return false;\
+	}
+
+#define UTEST_LESS_EQUAL(What, Actual, Expected)\
+	if (!TestLessEqual(What, Actual, Expected))\
+	{\
+		return false;\
+	}
+
+#define UTEST_LESS_EQUAL_EXPR(Actual, Expected)\
+	if (!TestLessEqual(TEXT(#Actual), Actual, Expected))\
+	{\
+		return false;\
+	}
+
+#define UTEST_LESS_EQUAL_TOLERANCE(What, Actual, Expected, Tolerance)\
+	if (!TestLessEqual(What, Actual, Expected, Tolerance))\
+	{\
+		return false;\
+	}
+
+#define UTEST_LESS_EQUAL_TOLERANCE_EXPR(Actual, Expected, Tolerance)\
+	if (!TestLessEqual(TEXT(#Actual), Actual, Expected, Tolerance))\
+	{\
+		return false;\
+	}
+
+#define UTEST_GREATER_EQUAL(What, Actual, Expected)\
+	if (!TestGreaterEqual(What, Actual, Expected))\
+	{\
+		return false;\
+	}
+
+#define UTEST_GREATER_EQUAL_EXPR(Actual, Expected)\
+	if (!TestGreaterEqual(TEXT(#Actual), Actual, Expected))\
+	{\
+		return false;\
+	}
+
+#define UTEST_GREATER_EQUAL_TOLERANCE(What, Actual, Expected, Tolerance)\
+	if (!TestGreaterEqual(What, Actual, Expected, Tolerance))\
+	{\
+		return false;\
+	}
+
+#define UTEST_GREATER_EQUAL_TOLERANCE_EXPR(Actual, Expected, Tolerance)\
+	if (!TestGreaterEqual(TEXT(#Actual), Actual, Expected, Tolerance))\
+	{\
+		return false;\
+	}
+
 #define UTEST_SAME(What, Actual, Expected)\
 	if (!TestSame(What, Actual, Expected))\
 	{\
diff --git a/Engine/Source/Runtime/Engine/Private/Tests/EngineAutomationTests.cpp b/Engine/Source/Runtime/Engine/Private/Tests/EngineAutomationTests.cpp
--- a/Engine/Source/Runtime/Engine/Private/Tests/EngineAutomationTests.cpp
+++ b/Engine/Source/Runtime/Engine/Private/Tests/EngineAutomationTests.cpp
@@ -834,6 +834,307 @@ bool FAutomationTestNearlyEqualFTransformNegative::RunTest(const FString& Parame
 	return true;
 }
 
+IMPLEMENT_SIMPLE_AUTOMATION_TEST(FAutomationTestInequalityBool, "TestFramework.Validation.TestInequalityBool", EAutomationTestFlags::ApplicationContextMask | EAutomationTestFlags::EngineFilter);
+bool FAutomationTestInequalityBool::RunTest(const FString& Parameters)
+{
+	TestTrue("True constant", true);
+	TestTrue("True int", 1);
+	TestFalse("False constant", false);
+	TestFalse("False int", 0);
+
+	return true;
+}
+
+IMPLEMENT_SIMPLE_AUTOMATION_TEST(FAutomationTestInequalityPointer, "TestFramework.Validation.TestInequalityPointer", EAutomationTestFlags::ApplicationContextMask | EAutomationTestFlags::EngineFilter);
+bool FAutomationTestInequalityPointer::RunTest(const FString& Parameters)
+{
+	int32 StackValue = 42;
+	int32* StackPointer = &StackValue;
+	int32* SameStackPointer = &StackValue;
+	int32 OtherStackValue = 42;
+	int32* OtherStackPointer = &OtherStackValue;
+	UWorld* World = UWorld::CreateWorld(EWorldType::Game, false);
+	UWorld* SameWorld = World;
+	UWorld* OtherWorld = UWorld::CreateWorld(EWorldType::Game, false);
+	
+	TestSame("Identity stack primitive", StackPointer, StackPointer);
+	TestSame("Identity world object", World, World);
+	/*
+	* Disabled due to UE-207121
+	* TestSame("Same stack primitive", SameStackPointer, StackPointer);
+	* TestSame("Same world object", SameWorld, World);
+	* TestNotSame("Other stack primitive", OtherStackPointer, StackPointer);
+	* TestNotSame("Other world object", OtherWorld, World);
+	*/
+	TestNotNull("Stack primitive not null", StackPointer);
+	TestNotNull("Constructed World object not null", World);
+	TestNull("Nullptr", nullptr);
+
+	World->DestroyWorld(false);
+	OtherWorld->DestroyWorld(false);
+
+	return true;
+}
+
+IMPLEMENT_SIMPLE_AUTOMATION_TEST(FAutomationTestInequalityInt32, "TestFramework.Validation.TestInequalityInt32", EAutomationTestFlags::ApplicationContextMask | EAutomationTestFlags::EngineFilter);
+bool FAutomationTestInequalityInt32::RunTest(const FString& Parameters)
+{
+	int32 Expected(42);
+	int32 Identical(Expected);
+	int32 Zero(0);
+	int32 LargePositive(2048);
+	int32 LargeNegative(-2048);
+	int32 MaxPositive(INT32_MAX);
+	int32 MaxNegative(INT32_MIN);
+
+	TestEqual(TEXT("Identity equal"), Identical, Expected);
+	TestNotEqual(TEXT("Zero unequal"), Zero, Expected);
+	TestNotEqual(TEXT("Positive unequal"), LargePositive, Expected);
+	TestNotEqual(TEXT("Negative unequal"), LargeNegative, Expected);
+	TestNotEqual(TEXT("Max unequal"), MaxPositive, Expected);
+	TestNotEqual(TEXT("Min unequal"), MaxNegative, Expected);
+	TestLessEqual(TEXT("Identity LE"), Identical, Expected);
+	TestLessEqual(TEXT("Less LE"), LargeNegative, Expected);
+	TestLessThan(TEXT("Min LE"), MaxNegative, Expected);
+	TestLessThan(TEXT("Less than"), LargeNegative, Expected);
+	TestLessThan(TEXT("Min less than"), MaxNegative, Expected);
+	TestGreaterEqual(TEXT("Identity GE"), Identical, Expected);
+	TestGreaterEqual(TEXT("Less GE"), LargePositive, Expected);
+	TestGreaterEqual(TEXT("Max GE"), MaxPositive, Expected);
+	TestGreaterThan(TEXT("Greater than"), LargePositive, Expected);
+	TestGreaterThan(TEXT("Max greater than"), MaxPositive, Expected);
+
+	return true;
+}
+
+IMPLEMENT_SIMPLE_AUTOMATION_TEST(FAutomationTestInequalityInt64, "TestFramework.Validation.TestInequalityInt64", EAutomationTestFlags::ApplicationContextMask | EAutomationTestFlags::EngineFilter);
+bool FAutomationTestInequalityInt64::RunTest(const FString& Parameters)
+{
+	int64 Expected(42);
+	int64 Identical(Expected);
+	int64 Zero(0);
+	int64 LargePositive(2048);
+	int64 LargeNegative(-2048);
+	int64 MaxPositive(INT64_MAX);
+	int64 MaxNegative(INT64_MIN);
+
+	TestEqual(TEXT("Identity equal"), Identical, Expected);
+	TestNotEqual(TEXT("Zero unequal"), Zero, Expected);
+	TestNotEqual(TEXT("Positive unequal"), LargePositive, Expected);
+	TestNotEqual(TEXT("Negative unequal"), LargeNegative, Expected);
+	TestNotEqual(TEXT("Max unequal"), MaxPositive, Expected);
+	TestNotEqual(TEXT("Min unequal"), MaxNegative, Expected);
+	TestLessEqual(TEXT("Identity LE"), Identical, Expected);
+	TestLessEqual(TEXT("Less LE"), LargeNegative, Expected);
+	TestLessThan(TEXT("Min LE"), MaxNegative, Expected);
+	TestLessThan(TEXT("Less than"), LargeNegative, Expected);
+	TestLessThan(TEXT("Min less than"), MaxNegative, Expected);
+	TestGreaterEqual(TEXT("Identity GE"), Identical, Expected);
+	TestGreaterEqual(TEXT("Less GE"), LargePositive, Expected);
+	TestGreaterEqual(TEXT("Max GE"), MaxPositive, Expected);
+	TestGreaterThan(TEXT("Greater than"), LargePositive, Expected);
+	TestGreaterThan(TEXT("Max greater than"), MaxPositive, Expected);
+
+	return true;
+}
+
+#if PLATFORM_64BITS
+IMPLEMENT_SIMPLE_AUTOMATION_TEST(FAutomationTestInequalitySizeT, "TestFramework.Validation.TestInequalitySizeT", EAutomationTestFlags::ApplicationContextMask | EAutomationTestFlags::EngineFilter);
+bool FAutomationTestInequalitySizeT::RunTest(const FString& Parameters)
+{
+	SIZE_T Expected(42);
+	SIZE_T Identical(Expected);
+	SIZE_T Zero(0);
+	SIZE_T LargePositive(2048);
+	SIZE_T SmallPositive(17);
+
+	TestEqual(TEXT("Identity equal"), Identical, Expected);
+	TestNotEqual(TEXT("Zero unequal"), Zero, Expected);
+	TestNotEqual(TEXT("Positive unequal"), LargePositive, Expected);
+	TestLessEqual(TEXT("Identity LE"), Identical, Expected);
+	TestLessEqual(TEXT("Less LE"), SmallPositive, Expected);
+	TestLessThan(TEXT("Less than"), SmallPositive, Expected);
+	TestGreaterEqual(TEXT("Identity GE"), Identical, Expected);
+	TestGreaterEqual(TEXT("Less GE"), LargePositive, Expected);
+	TestGreaterThan(TEXT("Greater than"), LargePositive, Expected);
+#ifdef SIZE_T_MAX
+	SIZE_T MaxPositive(SIZE_T_MAX);
+	TestNotEqual(TEXT("Max unequal"), MaxPositive, Expected);
+	TestGreaterEqual(TEXT("Max GE"), MaxPositive, Expected);
+	TestGreaterThan(TEXT("Max greater than"), MaxPositive, Expected);
+#endif //SIZE_T_MAX
+
+	return true;
+}
+#endif // PLATFORM_64BITS
+
+IMPLEMENT_SIMPLE_AUTOMATION_TEST(FAutomationTestInequalityFloat, "TestFramework.Validation.TestInequalityFloat", EAutomationTestFlags::ApplicationContextMask | EAutomationTestFlags::EngineFilter);
+bool FAutomationTestInequalityFloat::RunTest(const FString& Parameters)
+{
+	float Expected(42);
+	float Identical(Expected);
+	float Zero(0);
+	float LargePositive(2048);
+	float LargeNegative(-2048);
+	float LargeTolerance(1);
+	float SmallTolerance(UE_SMALL_NUMBER);
+	float KindaClosePositive(Expected + (UE_KINDA_SMALL_NUMBER * 0.9));
+	float KindaCloseNegative(Expected - (UE_KINDA_SMALL_NUMBER * 0.9));
+	float ExtremelyClosePositive(Expected + (UE_SMALL_NUMBER * 0.9));
+	float ExtremelyCloseNegative(Expected - (UE_SMALL_NUMBER * 0.9));
+
+	TestNearlyEqual(TEXT("Identity equal"), Identical, Expected);
+	TestNearlyEqual(TEXT("Identity equal low tolerance"), Identical, Expected, LargeTolerance);
+	TestNearlyEqual(TEXT("Identity equal high tolerance"), Identical, Expected, SmallTolerance);
+	TestNearlyEqual(TEXT("Nearby positive equal"), KindaClosePositive, Expected);
+	TestNearlyEqual(TEXT("Nearby positive equal low tolerance"), KindaClosePositive, Expected, LargeTolerance);
+	TestNearlyEqual(TEXT("Nearby positive equal high tolerance"), ExtremelyClosePositive, Expected, SmallTolerance);
+	TestNearlyEqual(TEXT("Nearby negative equal"), KindaCloseNegative, Expected);
+	TestNearlyEqual(TEXT("Nearby negative equal low tolerance"), KindaCloseNegative, Expected, LargeTolerance);
+	TestNearlyEqual(TEXT("Nearby negative equal high tolerance"), ExtremelyCloseNegative, Expected, SmallTolerance);
+
+	TestEqual(TEXT("Identity equal (forwards to TestNearlyEqual)"), Identical, Expected);
+	TestEqual(TEXT("Identity equal low tolerance (forwards to TestNearlyEqual)"), Identical, Expected, LargeTolerance);
+	TestEqual(TEXT("Identity equal high tolerance (forwards to TestNearlyEqual)"), Identical, Expected, SmallTolerance);
+	TestEqual(TEXT("Nearby positive equal (forwards to TestNearlyEqual)"), KindaClosePositive, Expected);
+	TestEqual(TEXT("Nearby positive equal low tolerance (forwards to TestNearlyEqual)"), KindaClosePositive, Expected, LargeTolerance);
+	TestEqual(TEXT("Nearby positive equal high tolerance (forwards to TestNearlyEqual)"), ExtremelyClosePositive, Expected, SmallTolerance);
+	TestEqual(TEXT("Nearby negative equal (forwards to TestNearlyEqual)"), KindaCloseNegative, Expected);
+	TestEqual(TEXT("Nearby negative equal low tolerance (forwards to TestNearlyEqual)"), KindaCloseNegative, Expected, LargeTolerance);
+	TestEqual(TEXT("Nearby negative equal high tolerance (forwards to TestNearlyEqual)"), ExtremelyCloseNegative, Expected, SmallTolerance);
+
+	TestNotEqual(TEXT("Zero unequal"), Zero, Expected);
+	TestNotEqual(TEXT("Positive unequal"), LargePositive, Expected);
+	TestNotEqual(TEXT("Nearby positive unequal due to high tolerance"), KindaClosePositive, Expected, SmallTolerance);
+	TestNotEqual(TEXT("Negative unequal"), LargeNegative, Expected);
+	TestNotEqual(TEXT("Nearby negative unequal due to high tolerance"), KindaClosePositive, Expected, SmallTolerance);
+	TestNotEqual(TEXT("Max unequal"), FLT_MAX, Expected);
+	TestNotEqual(TEXT("Min unequal"), FLT_MIN, Expected);
+
+	TestLessEqual(TEXT("Identity LE"), Identical, Expected);
+	TestLessEqual(TEXT("Identity LE low tolerance"), Identical, Expected, LargeTolerance);
+	TestLessEqual(TEXT("Identity LE high tolerance"), Identical, Expected, SmallTolerance);
+	TestLessEqual(TEXT("Less LE"), LargeNegative, Expected);
+	TestLessThan(TEXT("Min LE"), FLT_MIN, Expected);
+	TestLessThan(TEXT("Less than"), LargeNegative, Expected);
+	TestLessThan(TEXT("Min less than"), FLT_MIN, Expected);
+
+	TestGreaterEqual(TEXT("Identity GE"), Identical, Expected);
+	TestGreaterEqual(TEXT("Identity GE low tolerance"), Identical, Expected, LargeTolerance);
+	TestGreaterEqual(TEXT("Identity GE high tolerance"), Identical, Expected, SmallTolerance);
+	TestGreaterEqual(TEXT("Less GE"), LargePositive, Expected);
+	TestGreaterEqual(TEXT("Max GE"), FLT_MAX, Expected);
+	TestGreaterThan(TEXT("Greater than"), LargePositive, Expected);
+	TestGreaterThan(TEXT("Max greater than"), FLT_MAX, Expected);
+
+	return true;
+}
+
+IMPLEMENT_SIMPLE_AUTOMATION_TEST(FAutomationTestInequalityDouble, "TestFramework.Validation.TestInequalityDouble", EAutomationTestFlags::ApplicationContextMask | EAutomationTestFlags::EngineFilter);
+bool FAutomationTestInequalityDouble::RunTest(const FString& Parameters)
+{
+	double Expected(42);
+	double Identical(Expected);
+	double Zero(0);
+	double LargePositive(2048);
+	double LargeNegative(-2048);
+	double SmallTolerance(UE_SMALL_NUMBER);
+	double LargeTolerance(1);
+	double KindaClosePositive(Expected + (UE_KINDA_SMALL_NUMBER * 0.9));
+	double KindaCloseNegative(Expected - (UE_KINDA_SMALL_NUMBER * 0.9));
+	double ExtremelyClosePositive(Expected + (UE_SMALL_NUMBER * 0.9));
+	double ExtremelyCloseNegative(Expected - (UE_SMALL_NUMBER * 0.9));
+
+	TestNearlyEqual(TEXT("Identity equal"), Identical, Expected);
+	TestNearlyEqual(TEXT("Identity equal low tolerance"), Identical, Expected, LargeTolerance);
+	TestNearlyEqual(TEXT("Identity equal high tolerance"), Identical, Expected, SmallTolerance);
+	TestNearlyEqual(TEXT("Nearby positive equal"), KindaClosePositive, Expected);
+	TestNearlyEqual(TEXT("Nearby positive equal low tolerance"), KindaClosePositive, Expected, LargeTolerance);
+	TestNearlyEqual(TEXT("Nearby positive equal high tolerance"), ExtremelyClosePositive, Expected, SmallTolerance);
+	TestNearlyEqual(TEXT("Nearby negative equal"), KindaCloseNegative, Expected);
+	TestNearlyEqual(TEXT("Nearby negative equal low tolerance"), KindaCloseNegative, Expected, LargeTolerance);
+	TestNearlyEqual(TEXT("Nearby negative equal high tolerance"), ExtremelyCloseNegative, Expected, SmallTolerance);
+
+	TestEqual(TEXT("Identity equal (forwards to TestNearlyEqual)"), Identical, Expected);
+	TestEqual(TEXT("Identity equal low tolerance (forwards to TestNearlyEqual)"), Identical, Expected, LargeTolerance);
+	TestEqual(TEXT("Identity equal high tolerance (forwards to TestNearlyEqual)"), Identical, Expected, SmallTolerance);
+	TestEqual(TEXT("Nearby positive equal (forwards to TestNearlyEqual)"), KindaClosePositive, Expected);
+	TestEqual(TEXT("Nearby positive equal low tolerance (forwards to TestNearlyEqual)"), KindaClosePositive, Expected, LargeTolerance);
+	TestEqual(TEXT("Nearby positive equal high tolerance (forwards to TestNearlyEqual)"), ExtremelyClosePositive, Expected, SmallTolerance);
+	TestEqual(TEXT("Nearby negative equal (forwards to TestNearlyEqual)"), KindaCloseNegative, Expected);
+	TestEqual(TEXT("Nearby negative equal low tolerance (forwards to TestNearlyEqual)"), KindaCloseNegative, Expected, LargeTolerance);
+	TestEqual(TEXT("Nearby negative equal high tolerance (forwards to TestNearlyEqual)"), ExtremelyCloseNegative, Expected, SmallTolerance);
+
+	TestNotEqual(TEXT("Zero unequal"), Zero, Expected);
+	TestNotEqual(TEXT("Positive unequal"), LargePositive, Expected);
+	TestNotEqual(TEXT("Nearby positive unequal due to high tolerance"), KindaClosePositive, Expected, SmallTolerance);
+	TestNotEqual(TEXT("Negative unequal"), LargeNegative, Expected);
+	TestNotEqual(TEXT("Nearby negative unequal due to high tolerance"), KindaClosePositive, Expected, SmallTolerance);
+	TestNotEqual(TEXT("Max unequal"), DBL_MAX, Expected);
+	TestNotEqual(TEXT("Min unequal"), DBL_MIN, Expected);
+
+	TestLessEqual(TEXT("Identity LE"), Identical, Expected);
+	TestLessEqual(TEXT("Identity LE low tolerance"), Identical, Expected, LargeTolerance);
+	TestLessEqual(TEXT("Identity LE high tolerance"), Identical, Expected, SmallTolerance);
+	TestLessEqual(TEXT("Less LE"), LargeNegative, Expected);
+	TestLessThan(TEXT("Min LE"), DBL_MIN, Expected);
+	TestLessThan(TEXT("Less than"), LargeNegative, Expected);
+	TestLessThan(TEXT("Min less than"), DBL_MIN, Expected);
+
+	TestGreaterEqual(TEXT("Identity GE"), Identical, Expected);
+	TestGreaterEqual(TEXT("Identity GE low tolerance"), Identical, Expected, LargeTolerance);
+	TestGreaterEqual(TEXT("Identity GE high tolerance"), Identical, Expected, SmallTolerance);
+	TestGreaterEqual(TEXT("Less GE"), LargePositive, Expected);
+	TestGreaterEqual(TEXT("Max GE"), DBL_MAX, Expected);
+	TestGreaterThan(TEXT("Greater than"), LargePositive, Expected);
+	TestGreaterThan(TEXT("Max greater than"), DBL_MAX, Expected);
+
+	return true;
+}
+
+IMPLEMENT_SIMPLE_AUTOMATION_TEST(FAutomationTestInequalityString, "TestFramework.Validation.TestInequalityString", EAutomationTestFlags::ApplicationContextMask | EAutomationTestFlags::EngineFilter);
+bool FAutomationTestInequalityString::RunTest(const FString& Parameters)
+{
+	FString ExpectedString("Forty-two");
+	const TCHAR* ExpectedCharPtr = TEXT("Forty-two");
+	FString IdenticalString("Forty-two");
+	const TCHAR* IdenticalCharPtr = TEXT("Forty-two");
+	FString LowercaseString("forty-two");
+	const TCHAR* LowercaseCharPtr = TEXT("forty-two");
+	FString UppercaseString("FORTY-TWO");
+	const TCHAR* UppercaseCharPtr = TEXT("FORTY-TWO");
+	FString EmptyString("");
+	const TCHAR* EmptyCharPtr = TEXT("");
+	FString DifferentString("42");
+	const TCHAR* DifferentCharPtr = TEXT("42");
+	const TCHAR* NullCharPtr = nullptr;
+
+	TestEqual(TEXT("String identity equal"), IdenticalString, ExpectedString);
+	TestEqual(TEXT("char* identity equal"), IdenticalCharPtr, ExpectedCharPtr);
+	TestEqual(TEXT("String equals char*"), ExpectedString, ExpectedCharPtr);
+	TestEqual(TEXT("char* equals string"), ExpectedCharPtr, ExpectedString);
+	TestEqual(TEXT("String equals char* empty"), EmptyString, EmptyCharPtr);
+	TestEqual(TEXT("char* equals string empty"), EmptyCharPtr, EmptyString);
+
+	TestNotEqual(TEXT("String unequal"), DifferentString, ExpectedString);
+	TestNotEqual(TEXT("char* unequal"), DifferentCharPtr, ExpectedCharPtr);	
+	TestNotEqual(TEXT("String unequal empty"), EmptyString, ExpectedString);
+	TestNotEqual(TEXT("char* unequal empty"), EmptyCharPtr, ExpectedCharPtr);
+	TestNotEqual(TEXT("char* unequal null"), NullCharPtr, ExpectedCharPtr);
+
+	TestEqualInsensitive(TEXT("String insensitive equal identity"), IdenticalString, ExpectedString);
+	TestEqualInsensitive(TEXT("char* insensitive equal identity"), IdenticalCharPtr, ExpectedCharPtr);
+	TestEqualInsensitive(TEXT("String insensitive equal lower"), LowercaseString, ExpectedString);
+	TestEqualInsensitive(TEXT("char* insensitive equal lower"), LowercaseCharPtr, ExpectedCharPtr);
+	TestEqualInsensitive(TEXT("String insensitive equal upper"), UppercaseString, ExpectedString);
+	TestEqualInsensitive(TEXT("char* insensitive equal upper"), UppercaseCharPtr, ExpectedCharPtr);
+	TestNotEqualInsensitive(TEXT("String insensitive unequal"), DifferentString, ExpectedString);
+	TestNotEqualInsensitive(TEXT("char* insensitive unequal"), DifferentCharPtr, ExpectedCharPtr);
+	TestNotEqualInsensitive(TEXT("char* insensitive unequal null"), NullCharPtr, ExpectedCharPtr);
+	
+	return true;
+}
+
 class FAutomationUTestMacrosExpr : public FAutomationTestBase
 {
 public:
@@ -849,9 +1150,13 @@ protected:
 	static const float ExpectedFloatValue;
 	static const float WrongFloatValue;
 	static const float ExpectedFloatValueOutOfTolerance;
+	static const float ExpectedFloatValueOutOfToleranceNegative;
+	static const float ExpectedFloatValueLess;
+	static const float ExpectedFloatValueGreater;
 	static const FString ActualFStringValue;
 	static const FString ExpectedFStringValueLowerCase;
 	static const FString UnexpectedFStringValueLowerCase;
+	static const FString CustomDescriptionString;
 };
 
 const float FAutomationUTestMacrosExpr::PositiveToleranceFloat(1.e-4f);
@@ -859,20 +1164,30 @@ const float FAutomationUTestMacrosExpr::ActualFloatValue(0.f);
 const float FAutomationUTestMacrosExpr::ExpectedFloatValue(ActualFloatValue);
 const float FAutomationUTestMacrosExpr::WrongFloatValue(ActualFloatValue + 1.f);
 const float FAutomationUTestMacrosExpr::ExpectedFloatValueOutOfTolerance(ActualFloatValue + PositiveToleranceFloat);
+const float FAutomationUTestMacrosExpr::ExpectedFloatValueOutOfToleranceNegative(ActualFloatValue - PositiveToleranceFloat);
+const float FAutomationUTestMacrosExpr::ExpectedFloatValueLess(ActualFloatValue + (PositiveToleranceFloat*2)); //actual < expected
+const float FAutomationUTestMacrosExpr::ExpectedFloatValueGreater(ActualFloatValue - (PositiveToleranceFloat*2)); //actual > expected
 const FString FAutomationUTestMacrosExpr::ActualFStringValue(TEXT("EQUALS"));
 const FString FAutomationUTestMacrosExpr::ExpectedFStringValueLowerCase(TEXT("equals"));
 const FString FAutomationUTestMacrosExpr::UnexpectedFStringValueLowerCase(TEXT("not-equals"));
+const FString FAutomationUTestMacrosExpr::CustomDescriptionString(TEXT("Error string appears when UTEST_ macro diverges from _EXPR variant"));
 
 IMPLEMENT_CUSTOM_SIMPLE_AUTOMATION_TEST(FAutomationEqualEXPR, FAutomationUTestMacrosExpr, "TestFramework.Validation.UTestEqual", EAutomationTestFlags::ApplicationContextMask | EAutomationTestFlags::EngineFilter);
 bool FAutomationEqualEXPR::RunTest(const FString& Parameters)
 {
 
 	UTEST_EQUAL_EXPR(ActualFloatValue, ExpectedFloatValue);
+	UTEST_EQUAL(CustomDescriptionString, ActualFloatValue, ExpectedFloatValue);
 	UTEST_NEARLY_EQUAL_EXPR(ActualFloatValue, ExpectedFloatValueOutOfTolerance, PositiveToleranceFloat);
+	UTEST_NEARLY_EQUAL(CustomDescriptionString, ActualFloatValue, ExpectedFloatValueOutOfTolerance, PositiveToleranceFloat);
 	UTEST_EQUAL_TOLERANCE_EXPR(ActualFloatValue, ExpectedFloatValueOutOfTolerance, PositiveToleranceFloat);
+	UTEST_EQUAL_TOLERANCE(CustomDescriptionString, ActualFloatValue, ExpectedFloatValueOutOfTolerance, PositiveToleranceFloat);
 	UTEST_NOT_EQUAL_EXPR(ActualFloatValue, WrongFloatValue);
+	UTEST_NOT_EQUAL(CustomDescriptionString, ActualFloatValue, WrongFloatValue);
 	UTEST_EQUAL_INSENSITIVE_EXPR(*ActualFStringValue, *ExpectedFStringValueLowerCase);
+	UTEST_EQUAL_INSENSITIVE(*CustomDescriptionString, *ActualFStringValue, *ExpectedFStringValueLowerCase);
 	UTEST_NOT_EQUAL_INSENSITIVE_EXPR(*ActualFStringValue, *UnexpectedFStringValueLowerCase);
+	UTEST_NOT_EQUAL_INSENSITIVE(*CustomDescriptionString, *ActualFStringValue, *UnexpectedFStringValueLowerCase);
 
 	return true;
 }
@@ -882,7 +1197,9 @@ bool FAutomationSameNotSameEXPR::RunTest(const FString& Parameters)
 {
 
 	UTEST_SAME_EXPR(ActualFStringValue, ActualFStringValue);
+	UTEST_SAME(CustomDescriptionString, ActualFStringValue, ActualFStringValue);
 	UTEST_NOT_SAME_EXPR(ActualFStringValue, ExpectedFStringValueLowerCase);
+	UTEST_NOT_SAME(CustomDescriptionString, ActualFStringValue, ExpectedFStringValueLowerCase);
 
 	return true;
 }
@@ -890,9 +1207,10 @@ bool FAutomationSameNotSameEXPR::RunTest(const FString& Parameters)
 IMPLEMENT_CUSTOM_SIMPLE_AUTOMATION_TEST(FAutomationTrueFalseEXPR, FAutomationUTestMacrosExpr, "TestFramework.Validation.UTestTrueFalse", EAutomationTestFlags::ApplicationContextMask | EAutomationTestFlags::EngineFilter);
 bool FAutomationTrueFalseEXPR::RunTest(const FString& Parameters)
 {
-
 	UTEST_TRUE_EXPR(ActualFloatValue == ExpectedFloatValue);
+	UTEST_TRUE(CustomDescriptionString, ActualFloatValue == ExpectedFloatValue);
 	UTEST_FALSE_EXPR(ActualFloatValue > ExpectedFloatValue);
+	UTEST_FALSE(CustomDescriptionString, ActualFloatValue > ExpectedFloatValue);
 
 	return true;
 }
@@ -921,10 +1239,14 @@ bool FAutomationValidInvalidEXPR::RunTest(const FString& Parameters)
 
 	//** VERIFY **//
 	UTEST_VALID_EXPR(ValidSharedPtr);
+	UTEST_VALID(CustomDescriptionString, ValidSharedPtr);
 	UTEST_INVALID_EXPR(InvalidSharedPtr);
+	UTEST_INVALID(CustomDescriptionString, InvalidSharedPtr);
 
 	UTEST_VALID_EXPR(ValidObject);
+	UTEST_VALID(CustomDescriptionString, ValidObject);
 	UTEST_INVALID_EXPR(InvalidObject);
+	UTEST_INVALID(CustomDescriptionString, InvalidObject);
 
 	return true;
 }
@@ -935,13 +1257,49 @@ bool FAutomationNullNotNullPtrEXPR::RunTest(const FString& Parameters)
 	UWorld* World = UWorld::CreateWorld(EWorldType::Game, false);
 
 	UTEST_NULL_EXPR(nullptr);
+	UTEST_NULL(CustomDescriptionString, nullptr);
 	UTEST_NOT_NULL_EXPR(World);
+	UTEST_NOT_NULL(CustomDescriptionString, World);
 
 	World->DestroyWorld(false);
 
 	return true;
 }
 
+IMPLEMENT_CUSTOM_SIMPLE_AUTOMATION_TEST(FAutomationLessGreaterEXPR, FAutomationUTestMacrosExpr, "TestFramework.Validation.UTestLessGreater", EAutomationTestFlags::ApplicationContextMask | EAutomationTestFlags::EngineFilter);
+bool FAutomationLessGreaterEXPR::RunTest(const FString& Parameters)
+{
+	// inequality
+	UTEST_LESS_EXPR(ActualFloatValue, ExpectedFloatValueLess);
+	UTEST_LESS(CustomDescriptionString, ActualFloatValue, ExpectedFloatValueLess);
+	UTEST_LESS_TOLERANCE_EXPR(ActualFloatValue, ExpectedFloatValueLess, PositiveToleranceFloat);
+	UTEST_LESS_TOLERANCE(CustomDescriptionString, ActualFloatValue, ExpectedFloatValueLess, PositiveToleranceFloat);
+	UTEST_LESS_EQUAL_EXPR(ActualFloatValue, ExpectedFloatValueLess);
+	UTEST_LESS_EQUAL(CustomDescriptionString, ActualFloatValue, ExpectedFloatValueLess);
+	UTEST_LESS_EQUAL_TOLERANCE_EXPR(ActualFloatValue, ExpectedFloatValueLess, PositiveToleranceFloat);
+	UTEST_LESS_EQUAL_TOLERANCE(CustomDescriptionString, ActualFloatValue, ExpectedFloatValueLess, PositiveToleranceFloat);
+	UTEST_GREATER_EXPR(ActualFloatValue, ExpectedFloatValueGreater);
+	UTEST_GREATER(CustomDescriptionString, ActualFloatValue, ExpectedFloatValueGreater);
+	UTEST_GREATER_TOLERANCE_EXPR(ActualFloatValue, ExpectedFloatValueGreater, PositiveToleranceFloat);
+	UTEST_GREATER_TOLERANCE(CustomDescriptionString, ActualFloatValue, ExpectedFloatValueGreater, PositiveToleranceFloat);
+	UTEST_GREATER_EQUAL_EXPR(ActualFloatValue, ExpectedFloatValueGreater);
+	UTEST_GREATER_EQUAL(CustomDescriptionString, ActualFloatValue, ExpectedFloatValueGreater);
+	UTEST_GREATER_EQUAL_TOLERANCE_EXPR(ActualFloatValue, ExpectedFloatValueGreater, PositiveToleranceFloat);
+	UTEST_GREATER_EQUAL_TOLERANCE(CustomDescriptionString, ActualFloatValue, ExpectedFloatValueGreater, PositiveToleranceFloat);
+
+	// equality
+	UTEST_LESS_EQUAL_EXPR(ActualFloatValue, ExpectedFloatValue);
+	UTEST_LESS_EQUAL(CustomDescriptionString, ActualFloatValue, ExpectedFloatValue);
+	UTEST_LESS_EQUAL_TOLERANCE_EXPR(ActualFloatValue, ExpectedFloatValue, PositiveToleranceFloat);
+	UTEST_LESS_EQUAL_TOLERANCE(CustomDescriptionString, ActualFloatValue, ExpectedFloatValue, PositiveToleranceFloat);
+	UTEST_GREATER_EQUAL_EXPR(ActualFloatValue, ExpectedFloatValue);
+	UTEST_GREATER_EQUAL(CustomDescriptionString, ActualFloatValue, ExpectedFloatValue);
+	UTEST_GREATER_EQUAL_TOLERANCE_EXPR(ActualFloatValue, ExpectedFloatValue, PositiveToleranceFloat);
+	UTEST_GREATER_EQUAL_TOLERANCE(CustomDescriptionString, ActualFloatValue, ExpectedFloatValue, PositiveToleranceFloat);
+
+	return true;
+}
+
 IMPLEMENT_SIMPLE_AUTOMATION_TEST(FAutomationAttachment, "System.Engine.Attachment", EAutomationTestFlags::EditorContext | EAutomationTestFlags::EngineFilter)
 
 #define DUMP_EXPECTED_TRANSFORMS 0
