World Partition: add function GetStreamingBounds to cells data that returns what bounds the underlying runtime hash uses for streaming.
This makes it possible to flip between different runtime hash implementations without breaking late HLODs metrics.
Fix FEventsTracker to use the streaming cell bounds.

diff --git a/Engine/Source/Runtime/Engine/Private/WorldPartition/WorldPartitionRuntimeCellData.cpp b/Engine/Source/Runtime/Engine/Private/WorldPartition/WorldPartitionRuntimeCellData.cpp
--- a/Engine/Source/Runtime/Engine/Private/WorldPartition/WorldPartitionRuntimeCellData.cpp
+++ b/Engine/Source/Runtime/Engine/Private/WorldPartition/WorldPartitionRuntimeCellData.cpp
@@ -115,6 +115,11 @@ FBox UWorldPartitionRuntimeCellData::GetCellBounds() const
 	return CellBounds.IsSet() ? *CellBounds : ContentBounds;
 }
 
+FBox UWorldPartitionRuntimeCellData::GetStreamingBounds() const
+{
+	return GetContentBounds();
+}
+
 FString UWorldPartitionRuntimeCellData::GetDebugName() const
 {
 	return DebugName.GetString();
diff --git a/Engine/Source/Runtime/Engine/Private/WorldPartition/WorldPartitionRuntimeCellDataSpatialHash.cpp b/Engine/Source/Runtime/Engine/Private/WorldPartition/WorldPartitionRuntimeCellDataSpatialHash.cpp
--- a/Engine/Source/Runtime/Engine/Private/WorldPartition/WorldPartitionRuntimeCellDataSpatialHash.cpp
+++ b/Engine/Source/Runtime/Engine/Private/WorldPartition/WorldPartitionRuntimeCellDataSpatialHash.cpp
@@ -198,6 +198,11 @@ FBox UWorldPartitionRuntimeCellDataSpatialHash::GetCellBounds() const
 	return Box;
 }
 
+FBox UWorldPartitionRuntimeCellDataSpatialHash::GetStreamingBounds() const
+{
+	return GetCellBounds();
+}
+
 bool UWorldPartitionRuntimeCellDataSpatialHash::IsDebugShown() const
 {
 	return Super::IsDebugShown() && 
diff --git a/Engine/Source/Runtime/Engine/Public/WorldPartition/WorldPartitionRuntimeCell.h b/Engine/Source/Runtime/Engine/Public/WorldPartition/WorldPartitionRuntimeCell.h
--- a/Engine/Source/Runtime/Engine/Public/WorldPartition/WorldPartitionRuntimeCell.h
+++ b/Engine/Source/Runtime/Engine/Public/WorldPartition/WorldPartitionRuntimeCell.h
@@ -412,8 +412,16 @@ public:
 	inline void AppendStreamingSourceInfo(const FWorldPartitionStreamingSource& Source, const FSphericalSector& SourceShape) const { RuntimeCellData->AppendStreamingSourceInfo(Source, SourceShape); }
 	inline void MergeStreamingSourceInfo() const { RuntimeCellData->MergeStreamingSourceInfo(); }
 	ENGINE_API int32 SortCompare(const UWorldPartitionRuntimeCell* Other) const;
+	
+	/** Returns the cell'content s bounds, which is the sum of all actor bounds inside the cell. */
 	inline const FBox& GetContentBounds() const { return RuntimeCellData->GetContentBounds(); }
+
+	/** Returns the cell's bounds, which is the uniform size of the cell. */
 	inline FBox GetCellBounds() const { return RuntimeCellData->GetCellBounds(); }
+
+	/** Returns the cell's streaming bounds, which is what the underlying runtime hash uses to intersect cells.  */
+	inline FBox GetstreamingBounds() const { return RuntimeCellData->GetStreamingBounds(); }
+
 	ENGINE_API virtual bool IsDebugShown() const;
 	//~End UWorldPartitionRuntimeCellData Proxy
 
diff --git a/Engine/Source/Runtime/Engine/Public/WorldPartition/WorldPartitionRuntimeCellData.h b/Engine/Source/Runtime/Engine/Public/WorldPartition/WorldPartitionRuntimeCellData.h
--- a/Engine/Source/Runtime/Engine/Public/WorldPartition/WorldPartitionRuntimeCellData.h
+++ b/Engine/Source/Runtime/Engine/Public/WorldPartition/WorldPartitionRuntimeCellData.h
@@ -25,9 +25,15 @@ class UWorldPartitionRuntimeCellData : public UObject
 	ENGINE_API virtual void MergeStreamingSourceInfo() const {}
 	ENGINE_API virtual int32 SortCompare(const UWorldPartitionRuntimeCellData* InOther) const;
 
+	/** Returns the cell's content bounds, which is the sum of all actor bounds inside the cell. */
 	ENGINE_API virtual const FBox& GetContentBounds() const;
+
+	/** Returns the cell's bounds, which is the uniform size of the cell. */
 	ENGINE_API virtual FBox GetCellBounds() const;
 
+	/** Returns the cell's streaming bounds, which is what the underlying runtime hash uses to intersect cells.  */
+	ENGINE_API virtual FBox GetStreamingBounds() const;	
+
 	virtual bool IsDebugShown() const { return true; }
 	ENGINE_API virtual FString GetDebugName() const;
 
diff --git a/Engine/Source/Runtime/Engine/Public/WorldPartition/WorldPartitionRuntimeCellDataSpatialHash.h b/Engine/Source/Runtime/Engine/Public/WorldPartition/WorldPartitionRuntimeCellDataSpatialHash.h
--- a/Engine/Source/Runtime/Engine/Public/WorldPartition/WorldPartitionRuntimeCellDataSpatialHash.h
+++ b/Engine/Source/Runtime/Engine/Public/WorldPartition/WorldPartitionRuntimeCellDataSpatialHash.h
@@ -19,6 +19,7 @@ class UWorldPartitionRuntimeCellDataSpatialHash : public UWorldPartitionRuntimeC
 	ENGINE_API virtual void MergeStreamingSourceInfo() const override;
 	ENGINE_API virtual int32 SortCompare(const UWorldPartitionRuntimeCellData* InOther) const override;
 	ENGINE_API virtual FBox GetCellBounds() const override;
+	ENGINE_API virtual FBox GetStreamingBounds() const override;
 	ENGINE_API virtual bool IsDebugShown() const override;
 	//~End UWorldPartitionRuntimeCellData
 
