Added ENGINE_API to ALight::PostLoad.

Follow-up fix to a linker error from CL 27595673 that was reported by a licensee.

If you created a custom class derived from ALight, there would be a linker error reporting a missing definition of ALight::PostLoad. The vtable of a derived class needs to call PostLoad, which requires that the symbol has been exported.

diff --git a/Engine/Source/Runtime/Engine/Classes/Engine/Light.h b/Engine/Source/Runtime/Engine/Classes/Engine/Light.h
--- a/Engine/Source/Runtime/Engine/Classes/Engine/Light.h
+++ b/Engine/Source/Runtime/Engine/Classes/Engine/Light.h
@@ -62,7 +62,7 @@ public:
 
 public:
 	//~ Begin UObject Interface.
-	virtual void PostLoad() override;
+	ENGINE_API virtual void PostLoad() override;
 	//~ End UObject Interface.
 
 #if WITH_EDITOR
