Skeletal mesh morph targets:  preserve motion vectors across render state recreate, to fix issues with motion blur and temporal effects that use motion vectors, in scenarios where the render state is modified between renders in a given frame (such as in a Movie Render Queue graph), or across frames.  Preserving motion vectors for morph targets involves copying previous frame vertex buffers from the old to new render state, and any logic that may invalidate the vertex buffers needs to be suppressed.  Further, there are two separate code paths -- one with the GPU skin cache enabled, and another when it's not.

Deletion of the previous render state (FSkeletalMeshObjectGPUSkin, etc) is deferred so any necessary data can be copied across, which should facilitate fixes in other systems in the future if necessary (for example, cloth and deformer meshes may run into similar issues).  The morph vertex buffers in the non-skin-cache code path are made reference counted.  Reference counting was necessary because new state is allocated in the main thread, and deletion is deferred, meaning up to 3 references can be in flight (deletion queue, current render thread copy, new main thread copy), so moving the reference doesn't work.  In the skin-cache code path, we can get by with moving the reference, because it's only interacted with in the render thread.  Two pieces of invalidation logic needed to be suppressed in the skin-cache code path as well.

#rnx

diff --git a/Engine/Source/Runtime/Engine/Classes/Components/ActorComponent.h b/Engine/Source/Runtime/Engine/Classes/Components/ActorComponent.h
--- a/Engine/Source/Runtime/Engine/Classes/Components/ActorComponent.h
+++ b/Engine/Source/Runtime/Engine/Classes/Components/ActorComponent.h
@@ -181,6 +181,14 @@ protected:
 	/** If the render state is currently created for this component */
 	uint8 bRenderStateCreated:1;
 
+	/**
+	 * Render state is being recreated for this component.  Useful if a component wants to preserve certain render state across recreate -- set before a
+	 * call to DestroyRenderState_Concurrent and cleared after the corresponding call to CreateRenderState_Concurrent.  By design, only set for re-creation
+	 * due to MarkRenderStateDirty, not RecreateRenderStateContext variations.  The latter are used for editor operations, file loads, or bulk setting
+	 * changes, which are assumed to potentially change render data in a way that could make preserved render state incompatible.
+	 */
+	uint8 bRenderStateRecreating:1;
+
 	/** If the physics state is currently created for this component */
 	uint8 bPhysicsStateCreated:1;
 
@@ -985,6 +993,12 @@ public:
 		return bRenderStateCreated;
 	}
 
+	/** Returns true if the render 'state' is being recreated for this component (see additional comments on variable above) */
+	bool IsRenderStateRecreating() const
+	{
+		return bRenderStateRecreating;
+	}
+
 	/** Returns true if the physics 'state' (e.g. physx bodies) are created for this component */
 	bool IsPhysicsStateCreated() const
 	{
diff --git a/Engine/Source/Runtime/Engine/Classes/Components/SkinnedMeshComponent.h b/Engine/Source/Runtime/Engine/Classes/Components/SkinnedMeshComponent.h
--- a/Engine/Source/Runtime/Engine/Classes/Components/SkinnedMeshComponent.h
+++ b/Engine/Source/Runtime/Engine/Classes/Components/SkinnedMeshComponent.h
@@ -843,6 +843,9 @@ public:
 	/** Passed into MeshObjectFactory */
 	void* MeshObjectFactoryUserData;
 
+	/** Previous copy of MeshObject set during recreate of the render state */
+	class FSkeletalMeshObject* PreviousMeshObject;
+
 	/** Gets the skeletal mesh resource used for rendering the component. */
 	ENGINE_API FSkeletalMeshRenderData* GetSkeletalMeshRenderData() const;
 
diff --git a/Engine/Source/Runtime/Engine/Private/Components/ActorComponent.cpp b/Engine/Source/Runtime/Engine/Private/Components/ActorComponent.cpp
--- a/Engine/Source/Runtime/Engine/Private/Components/ActorComponent.cpp
+++ b/Engine/Source/Runtime/Engine/Private/Components/ActorComponent.cpp
@@ -1807,16 +1807,26 @@ void UActorComponent::ReregisterComponent()
 
 void UActorComponent::RecreateRenderState_Concurrent()
 {
+	bool bCanRecreate = IsRegistered() && WorldPrivate->Scene;
+
 	if(bRenderStateCreated)
 	{
+		// Only set bRenderStateRecreating if we know for sure we are going to actually re-create it, so components can always count on the
+		// calls happening in sequence if bRenderStateRecreating is set, and don't need to handle edge cases where the latter isn't called.
+		if (bCanRecreate)
+		{
+			check(bRenderStateRecreating == false);
+			bRenderStateRecreating = true;
+		}
 		check(IsRegistered()); // Should never have render state unless registered
 		DestroyRenderState_Concurrent();
 		checkf(!bRenderStateCreated, TEXT("Failed to route DestroyRenderState_Concurrent (%s)"), *GetFullName());
 	}
 
-	if(IsRegistered() && WorldPrivate->Scene)
+	if (bCanRecreate)
 	{
 		CreateRenderState_Concurrent(nullptr);
+		bRenderStateRecreating = false;
 		checkf(bRenderStateCreated, TEXT("Failed to route CreateRenderState_Concurrent (%s)"), *GetFullName());
 	}
 }
diff --git a/Engine/Source/Runtime/Engine/Private/Components/SkinnedMeshComponent.cpp b/Engine/Source/Runtime/Engine/Private/Components/SkinnedMeshComponent.cpp
--- a/Engine/Source/Runtime/Engine/Private/Components/SkinnedMeshComponent.cpp
+++ b/Engine/Source/Runtime/Engine/Private/Components/SkinnedMeshComponent.cpp
@@ -396,6 +396,7 @@ USkinnedMeshComponent::USkinnedMeshComponent(const FObjectInitializer& ObjectIni
 	: Super(ObjectInitializer)
 	, MeshObjectFactory(nullptr)
 	, MeshObjectFactoryUserData(nullptr)
+	, PreviousMeshObject(nullptr)
 	, AnimUpdateRateParams(nullptr)
 {
 	bAutoActivate = true;
@@ -1006,6 +1007,16 @@ void USkinnedMeshComponent::CreateRenderState_Concurrent(FRegisterComponentConte
 		}
 	}
 
+	// Any data needed from PreviousMeshObject will have been propagated in the construction of the new MeshObject above,
+	// so we can delete it now.
+	if (PreviousMeshObject)
+	{
+		PreviousMeshObject->ReleaseResources();
+		BeginCleanup(PreviousMeshObject);
+
+		PreviousMeshObject = nullptr;
+	}
+
 	Super::CreateRenderState_Concurrent(Context);
 
 	if (GetSkinnedAsset())
@@ -1080,14 +1091,23 @@ void USkinnedMeshComponent::DestroyRenderState_Concurrent()
 
 	if(MeshObject)
 	{
-		// Begin releasing the RHI resources used by this skeletal mesh component.
-		// This doesn't immediately destroy anything, since the rendering thread may still be using the resources.
-		MeshObject->ReleaseResources();
+		if (bRenderStateRecreating)
+		{
+			// Preserve previous mesh object for recreate purposes
+			PreviousMeshObject = MeshObject;
+			MeshObject = nullptr;
+		}
+		else
+		{
+			// Begin releasing the RHI resources used by this skeletal mesh component.
+			// This doesn't immediately destroy anything, since the rendering thread may still be using the resources.
+			MeshObject->ReleaseResources();
 
-		// Begin a deferred delete of MeshObject.  BeginCleanup will call MeshObject->FinishDestroy after the above release resource
-		// commands execute in the rendering thread.
-		BeginCleanup(MeshObject);
-		MeshObject = nullptr;
+			// Begin a deferred delete of MeshObject.  BeginCleanup will call MeshObject->FinishDestroy after the above release resource
+			// commands execute in the rendering thread.
+			BeginCleanup(MeshObject);
+			MeshObject = nullptr;
+		}
 	}
 }
 
diff --git a/Engine/Source/Runtime/Engine/Private/GPUSkinCache.cpp b/Engine/Source/Runtime/Engine/Private/GPUSkinCache.cpp
--- a/Engine/Source/Runtime/Engine/Private/GPUSkinCache.cpp
+++ b/Engine/Source/Runtime/Engine/Private/GPUSkinCache.cpp
@@ -1682,6 +1682,7 @@ bool FGPUSkinCache::ProcessEntry(
 	uint32 RevisionNumber, 
 	int32 Section,
 	int32 LODIndex,
+	bool bRecreating,
 	FGPUSkinCacheEntry*& InOutEntry
 	)
 {
@@ -1725,6 +1726,12 @@ bool FGPUSkinCache::ProcessEntry(
 
 	if (InOutEntry)
 	{
+		// If re-creating, the FSkeletalMeshObjectGPUSkin will have changed, update it here so the "IsValid" call below succeeds
+		if (bRecreating)
+		{
+			InOutEntry->GPUSkin = Skin;
+		}
+
 		// If the LOD changed, the entry has to be invalidated
 		if (!InOutEntry->IsValid(Skin, LODIndex))
 		{
@@ -1745,6 +1752,9 @@ bool FGPUSkinCache::ProcessEntry(
 	// Try to allocate a new entry
 	if (!InOutEntry)
 	{
+		// If something caused the existing entry to be invalid, disable recreate logic for the rest of the function
+		bRecreating = false;
+
 		const bool WithTangents = true;
 		int32 TotalNumVertices = VertexFactory->GetNumVertices();
 		
@@ -1789,6 +1799,8 @@ bool FGPUSkinCache::ProcessEntry(
 		Entries.Add(InOutEntry);
 	}
 
+	FGPUSkinCacheEntry::FSectionDispatchData& SectionDispatchData = InOutEntry->DispatchData[Section];
+
 	const bool bMorph = MorphVertexBuffer && MorphVertexBuffer->SectionIds.Contains(Section);
 	if (bMorph)
 	{
@@ -1800,12 +1812,12 @@ bool FGPUSkinCache::ProcessEntry(
 		// see GPU code "check(MorphStride == sizeof(float) * 6);"
 		check(MorphStride == sizeof(float) * 6);
 
-		InOutEntry->DispatchData[Section].MorphBufferOffset = BatchElement.BaseVertexIndex;
+		SectionDispatchData.MorphBufferOffset = BatchElement.BaseVertexIndex;
 
 		// weight buffer
 		FSkinWeightVertexBuffer* WeightBuffer = Skin->GetSkinWeightVertexBuffer(LODIndex);
 		uint32 WeightStride = WeightBuffer->GetConstantInfluencesVertexStride();
-		InOutEntry->DispatchData[Section].InputWeightStart = (WeightStride * BatchElement.BaseVertexIndex) / sizeof(float);
+		SectionDispatchData.InputWeightStart = (WeightStride * BatchElement.BaseVertexIndex) / sizeof(float);
 		InOutEntry->InputWeightStride = WeightStride;
 		InOutEntry->InputWeightStreamSRV = WeightBuffer->GetDataVertexBuffer()->GetSRV();
 	}
@@ -1839,7 +1851,7 @@ bool FGPUSkinCache::ProcessEntry(
 
 				// Set the buffer offset depending on whether enough deformer mapping data exists (RaytracingMinLOD/RaytracingLODBias/ClothLODBiasMode settings)
 				const uint32 NumInfluences = NumVertices ? ClothBufferIndexMapping.LODBiasStride / NumVertices : 1;
-				InOutEntry->DispatchData[Section].ClothBufferOffset = (ClothBufferOffset + NumVertices * NumInfluences <= ClothVertexBuffer->GetNumVertices()) ?
+				SectionDispatchData.ClothBufferOffset = (ClothBufferOffset + NumVertices * NumInfluences <= ClothVertexBuffer->GetNumVertices()) ?
 					ClothBufferOffset :                     // If the offset is valid, set the calculated LODBias offset
 					ClothBufferIndexMapping.MappingOffset;  // Otherwise fallback to a 0 ClothLODBias to prevent from reading pass the buffer (but still raytrace broken shadows/reflections/etc.)
 			}
@@ -1856,18 +1868,29 @@ bool FGPUSkinCache::ProcessEntry(
 	        FRHIResourceCreateInfo CreateInfo(TEXT("ClothPositionAndNormalsBuffer"), ResourceArray);
 	        ClothPositionAndNormalsBuffer.VertexBufferRHI = RHICmdList.CreateVertexBuffer( ResourceArray->GetResourceDataSize(), BUF_Static | BUF_ShaderResource, CreateInfo);
 	        ClothPositionAndNormalsBuffer.VertexBufferSRV = RHICmdList.CreateShaderResourceView(ClothPositionAndNormalsBuffer.VertexBufferRHI, sizeof(FVector2f), PF_G32R32F);
-	        InOutEntry->DispatchData[Section].ClothPositionsAndNormalsBuffer = ClothPositionAndNormalsBuffer.VertexBufferSRV;
+			SectionDispatchData.ClothPositionsAndNormalsBuffer = ClothPositionAndNormalsBuffer.VertexBufferSRV;
 		}
 		else
 		{
 			UE_LOG(LogSkinCache, Error, TEXT("Cloth sim data is missing on mesh %s"), *GetSkeletalMeshObjectName(Skin));
 		}
 
-        InOutEntry->DispatchData[Section].ClothBlendWeight = ClothBlendWeight;
-        InOutEntry->DispatchData[Section].ClothToLocal = ClothToLocal;
-		InOutEntry->DispatchData[Section].WorldScale = WorldScale;
+		SectionDispatchData.ClothBlendWeight = ClothBlendWeight;
+		SectionDispatchData.ClothToLocal = ClothToLocal;
+		SectionDispatchData.WorldScale = WorldScale;
     }
-    InOutEntry->DispatchData[Section].SkinType = ClothVertexBuffer && InOutEntry->DispatchData[Section].ClothPositionsAndNormalsBuffer ? 2 : (bMorph ? 1 : 0);
+	SectionDispatchData.SkinType = ClothVertexBuffer && SectionDispatchData.ClothPositionsAndNormalsBuffer ? 2 : (bMorph ? 1 : 0);
+
+	// Need to update the previous bone buffer pointer, so logic that checks if the bone buffers changed (FGPUSkinCache::FRWBufferTracker::Find)
+	// doesn't invalidate the previous frame position data.  Recreating the render state will have generated new bone buffers.
+	if (bRecreating)
+	{
+		FGPUBaseSkinVertexFactory::FShaderDataType& ShaderData = VertexFactory->GetShaderData();
+		if (ShaderData.HasBoneBufferForReading(true))
+		{
+			SectionDispatchData.PositionTracker.UpdatePreviousBoneBuffer(ShaderData.GetBoneBufferForReading(true), VertexFactory->GetShaderData().GetRevisionNumber(true));
+		}
+	}
 
 	if (bShouldBatchDispatches)
 	{
@@ -1875,16 +1898,16 @@ bool FGPUSkinCache::ProcessEntry(
 
 		bool bFoundEntry = false;
 
-		if (InOutEntry->DispatchData[Section].RevisionNumber != 0)
+		if (SectionDispatchData.RevisionNumber != 0)
 		{
 			// Check if the combo of skin cache entry and section index already exists, if so use the entry and update to latest revision number.
-			InOutEntry->DispatchData[Section].RevisionNumber = FMath::Max(InOutEntry->DispatchData[Section].RevisionNumber, RevisionNumber);
+			SectionDispatchData.RevisionNumber = FMath::Max(InOutEntry->DispatchData[Section].RevisionNumber, RevisionNumber);
 			bFoundEntry = true;
 		}
 
 		if (!bFoundEntry)
 		{
-			InOutEntry->DispatchData[Section].RevisionNumber = RevisionNumber;
+			SectionDispatchData.RevisionNumber = RevisionNumber;
 			BatchDispatches.Add({ InOutEntry, uint32(Section) });
 		}
 	}
diff --git a/Engine/Source/Runtime/Engine/Private/SkeletalRenderGPUSkin.cpp b/Engine/Source/Runtime/Engine/Private/SkeletalRenderGPUSkin.cpp
--- a/Engine/Source/Runtime/Engine/Private/SkeletalRenderGPUSkin.cpp
+++ b/Engine/Source/Runtime/Engine/Private/SkeletalRenderGPUSkin.cpp
@@ -138,22 +138,30 @@ FMorphVertexBufferPool
 -----------------------------------------------------------------------------*/
 void FMorphVertexBufferPool::InitResources(const FName& OwnerName)
 {
-	MorphVertexBuffers[0].SetOwnerName(OwnerName);
-	MorphVertexBuffers[1].SetOwnerName(OwnerName);
-
-	check(!MorphVertexBuffers[0].VertexBufferRHI.IsValid());
-	check(!MorphVertexBuffers[1].VertexBufferRHI.IsValid());
-	BeginInitResource(&MorphVertexBuffers[0], &UE::RenderCommandPipe::SkeletalMesh);
-	if (bDoubleBuffer)
+	// InitResources may be called again when morph vertex data is persisted during render state re-creation.
+	if (!bInitializedResources)
 	{
-		BeginInitResource(&MorphVertexBuffers[1], &UE::RenderCommandPipe::SkeletalMesh);
+		MorphVertexBuffers[0].SetOwnerName(OwnerName);
+		MorphVertexBuffers[1].SetOwnerName(OwnerName);
+
+		check(!MorphVertexBuffers[0].VertexBufferRHI.IsValid());
+		check(!MorphVertexBuffers[1].VertexBufferRHI.IsValid());
+		BeginInitResource(&MorphVertexBuffers[0], &UE::RenderCommandPipe::SkeletalMesh);
+		if (bDoubleBuffer)
+		{
+			BeginInitResource(&MorphVertexBuffers[1], &UE::RenderCommandPipe::SkeletalMesh);
+		}
+
+		bInitializedResources = true;
 	}
 }
 
 void FMorphVertexBufferPool::ReleaseResources()
 {
-	BeginReleaseResource(&MorphVertexBuffers[0], &UE::RenderCommandPipe::SkeletalMesh);
-	BeginReleaseResource(&MorphVertexBuffers[1], &UE::RenderCommandPipe::SkeletalMesh);
+	check(bInitializedResources);
+	MorphVertexBuffers[0].ReleaseResource();
+	MorphVertexBuffers[1].ReleaseResource();
+	bInitializedResources = false;
 }
 
 SIZE_T FMorphVertexBufferPool::GetResourceSize() const
@@ -231,11 +239,44 @@ FSkeletalMeshObjectGPUSkin::FSkeletalMeshObjectGPUSkin(USkinnedMeshComponent* In
 	, 	LastBoneTransformRevisionNumber(0)
 	,	bAlwaysUpdateMorphVertexBuffer(false)
 {
+	FSkeletalMeshObjectGPUSkin* PreviousMeshObject = nullptr;
+	if (InMeshComponent->PreviousMeshObject && InMeshComponent->PreviousMeshObject->IsGPUSkinMesh())
+	{
+		PreviousMeshObject = (FSkeletalMeshObjectGPUSkin*)InMeshComponent->PreviousMeshObject;
+
+		// Don't use re-create data if the mesh or feature level changed
+		if (PreviousMeshObject->SkeletalMeshRenderData != InSkelMeshRenderData || PreviousMeshObject->FeatureLevel != InFeatureLevel)
+		{
+			PreviousMeshObject = nullptr;
+		}
+	}
+
+	if (PreviousMeshObject)
+	{
+		// Transfer GPU skin cache from PreviousMeshObject -- needs to happen on render thread.  PreviousMeshObject is defer deleted, so it's safe to access it there.
+		ENQUEUE_RENDER_COMMAND(ReleaseSkeletalMeshSkinCacheResources)(UE::RenderCommandPipe::SkeletalMesh,
+			[this, PreviousMeshObject](FRHICommandList& RHICmdList)
+			{
+				SkinCacheEntry = PreviousMeshObject->SkinCacheEntry;
+				SkinCacheEntryForRayTracing = PreviousMeshObject->SkinCacheEntryForRayTracing;
+
+				PreviousMeshObject->SkinCacheEntry = nullptr;
+				PreviousMeshObject->SkinCacheEntryForRayTracing = nullptr;
+			}
+		);
+	}
+
 	// create LODs to match the base mesh
 	LODs.Empty(SkeletalMeshRenderData->LODRenderData.Num());
 	for( int32 LODIndex=0;LODIndex < SkeletalMeshRenderData->LODRenderData.Num();LODIndex++ )
 	{
-		new(LODs) FSkeletalMeshObjectLOD(SkeletalMeshRenderData, LODIndex, InFeatureLevel);
+		FMorphVertexBufferPool* RecreateMorphVertexBuffer = nullptr;
+		if (PreviousMeshObject)
+		{
+			RecreateMorphVertexBuffer = PreviousMeshObject->LODs[LODIndex].MorphVertexBufferPool;
+		}
+
+		new(LODs) FSkeletalMeshObjectLOD(SkeletalMeshRenderData, LODIndex, InFeatureLevel, RecreateMorphVertexBuffer);
 	}
 
 	InitResources(InMeshComponent);
@@ -410,10 +451,10 @@ void FSkeletalMeshObjectGPUSkin::Update(
 
 	FGPUSkinCache* GPUSkinCache = nullptr;
 	FSceneInterface* Scene = nullptr;
-	if (InMeshComponent && InMeshComponent->SceneProxy)
+	if (InMeshComponent && InMeshComponent->GetScene())
 	{
 		// We allow caching of per-frame, per-scene data
-		Scene = &(InMeshComponent->SceneProxy->GetScene());
+		Scene = InMeshComponent->GetScene();
 		GPUSkinCache = Scene->GetGPUSkinCache();
 		RevisionNumber = InMeshComponent->GetBoneTransformRevisionNumber();
 		PreviousRevisionNumber = InMeshComponent->GetPreviousBoneTransformRevisionNumber();
@@ -421,11 +462,12 @@ void FSkeletalMeshObjectGPUSkin::Update(
 
 	// queue a call to update this data
 	FSkeletalMeshObjectGPUSkin* MeshObject = this;
+	bool bRecreating = InMeshComponent ? InMeshComponent->IsRenderStateRecreating() : false;
 	ENQUEUE_RENDER_COMMAND(SkelMeshObjectUpdateDataCommand)(UE::RenderCommandPipe::SkeletalMesh,
-		[MeshObject, FrameNumberToPrepare, RevisionNumber, PreviousRevisionNumber, NewDynamicData, GPUSkinCache, Scene](FRHICommandList& RHICmdList)
+		[MeshObject, FrameNumberToPrepare, RevisionNumber, PreviousRevisionNumber, NewDynamicData, GPUSkinCache, Scene, bRecreating](FRHICommandList& RHICmdList)
 		{
 			FScopeCycleCounter Context(MeshObject->GetStatId());
-			MeshObject->UpdateDynamicData_RenderThread(GPUSkinCache, RHICmdList, NewDynamicData, Scene, FrameNumberToPrepare, RevisionNumber, PreviousRevisionNumber);
+			MeshObject->UpdateDynamicData_RenderThread(GPUSkinCache, RHICmdList, NewDynamicData, Scene, FrameNumberToPrepare, RevisionNumber, PreviousRevisionNumber, bRecreating);
 		}
 	);
 }
@@ -483,7 +525,7 @@ FORCEINLINE bool IsDeferredSkeletalDynamicDataUpdateEnabled()
 	return CVarDeferSkeletalDynamicDataUpdateUntilGDME.GetValueOnRenderThread() > 0 && !IsParallelGatherDynamicMeshElementsEnabled();
 }
 
-void FSkeletalMeshObjectGPUSkin::UpdateDynamicData_RenderThread(FGPUSkinCache* GPUSkinCache, FRHICommandList& RHICmdList, FDynamicSkelMeshObjectDataGPUSkin* InDynamicData, FSceneInterface* Scene, uint64 FrameNumberToPrepare, uint32 RevisionNumber, uint32 PreviousRevisionNumber)
+void FSkeletalMeshObjectGPUSkin::UpdateDynamicData_RenderThread(FGPUSkinCache* GPUSkinCache, FRHICommandList& RHICmdList, FDynamicSkelMeshObjectDataGPUSkin* InDynamicData, FSceneInterface* Scene, uint64 FrameNumberToPrepare, uint32 RevisionNumber, uint32 PreviousRevisionNumber, bool bRecreating)
 {
 	TRACE_CPUPROFILER_EVENT_SCOPE(GPUSkin::UpdateDynamicData_RT);
 
@@ -522,13 +564,13 @@ void FSkeletalMeshObjectGPUSkin::UpdateDynamicData_RenderThread(FGPUSkinCache* G
 	}
 	else
 	{
-		ProcessUpdatedDynamicData(EGPUSkinCacheEntryMode::Raster, GPUSkinCache, RHICmdList, FrameNumberToPrepare, RevisionNumber, PreviousRevisionNumber, bMorphNeedsUpdate, DynamicData->LODIndex);
+		ProcessUpdatedDynamicData(EGPUSkinCacheEntryMode::Raster, GPUSkinCache, RHICmdList, FrameNumberToPrepare, RevisionNumber, PreviousRevisionNumber, bMorphNeedsUpdate, DynamicData->LODIndex, bRecreating);
 
 	#if RHI_RAYTRACING
 		if (ShouldUseSeparateSkinCacheEntryForRayTracing() && FGPUSkinCache::IsGPUSkinCacheRayTracingSupported() && GPUSkinCache && SkeletalMeshRenderData->bSupportRayTracing)
 		{
 			// Morph delta is updated in raster pass above, no need to update again for ray tracing
-			ProcessUpdatedDynamicData(EGPUSkinCacheEntryMode::RayTracing, GPUSkinCache, RHICmdList, FrameNumberToPrepare, RevisionNumber, PreviousRevisionNumber, /*bMorphNeedsUpdate=*/false, DynamicData->RayTracingLODIndex);
+			ProcessUpdatedDynamicData(EGPUSkinCacheEntryMode::RayTracing, GPUSkinCache, RHICmdList, FrameNumberToPrepare, RevisionNumber, PreviousRevisionNumber, /*bMorphNeedsUpdate=*/false, DynamicData->RayTracingLODIndex, bRecreating);
 		}
 		else
 		{
@@ -552,11 +594,11 @@ void FSkeletalMeshObjectGPUSkin::PreGDMECallback(FRHICommandList& RHICmdList, FG
 {
 	if (bNeedsUpdateDeferred)
 	{
-		ProcessUpdatedDynamicData(EGPUSkinCacheEntryMode::Raster, GPUSkinCache, RHICmdList, FrameNumber, LastBoneTransformRevisionNumber, 0, bMorphNeedsUpdateDeferred, DynamicData->LODIndex);
+		ProcessUpdatedDynamicData(EGPUSkinCacheEntryMode::Raster, GPUSkinCache, RHICmdList, FrameNumber, LastBoneTransformRevisionNumber, 0, bMorphNeedsUpdateDeferred, DynamicData->LODIndex, false);
 	}
 }
 
-void FSkeletalMeshObjectGPUSkin::ProcessUpdatedDynamicData(EGPUSkinCacheEntryMode Mode, FGPUSkinCache* GPUSkinCache, FRHICommandList& RHICmdList, uint32 FrameNumberToPrepare, uint32 RevisionNumber, uint32 PreviousRevisionNumber, bool bMorphNeedsUpdate, int32 LODIndex)
+void FSkeletalMeshObjectGPUSkin::ProcessUpdatedDynamicData(EGPUSkinCacheEntryMode Mode, FGPUSkinCache* GPUSkinCache, FRHICommandList& RHICmdList, uint32 FrameNumberToPrepare, uint32 RevisionNumber, uint32 PreviousRevisionNumber, bool bMorphNeedsUpdate, int32 LODIndex, bool bRecreating)
 {
 	QUICK_SCOPE_CYCLE_COUNTER(STAT_FSkeletalMeshObjectGPUSkin_ProcessUpdatedDynamicData);
 	bNeedsUpdateDeferred = false;
@@ -577,10 +619,10 @@ void FSkeletalMeshObjectGPUSkin::ProcessUpdatedDynamicData(EGPUSkinCacheEntryMod
 	if (bMorph && bMorphNeedsUpdate)
 	{
 		// Morph vertex buffer needs updating so advance the revision number
-		LOD.MorphVertexBufferPool.SetCurrentRevisionNumber(RevisionNumber);
+		LOD.MorphVertexBufferPool->SetCurrentRevisionNumber(RevisionNumber);
 	}
 
-	FMorphVertexBuffer& MorphVertexBuffer = LOD.MorphVertexBufferPool.GetMorphVertexBufferForWriting();
+	FMorphVertexBuffer& MorphVertexBuffer = LOD.MorphVertexBufferPool->GetMorphVertexBufferForWriting();
 	// if morph buffer hasn't been updated, force update it
 	bMorphNeedsUpdate = MorphVertexBuffer.bHasBeenUpdated ? bMorphNeedsUpdate : true;
 
@@ -622,7 +664,7 @@ void FSkeletalMeshObjectGPUSkin::ProcessUpdatedDynamicData(EGPUSkinCacheEntryMod
 		if(bMorphNeedsUpdate)
 		{
 			UpdateMorphVertexBuffer(RHICmdList, Mode, LOD, LODData, bGPUSkinCacheEnabled, MorphVertexBuffer);
-			LOD.MorphVertexBufferPool.SetUpdatedFrameNumber(FrameNumberToPrepare);
+			LOD.MorphVertexBufferPool->SetUpdatedFrameNumber(FrameNumberToPrepare);
 		}
 	}
 	else
@@ -733,6 +775,7 @@ void FSkeletalMeshObjectGPUSkin::ProcessUpdatedDynamicData(EGPUSkinCacheEntryMod
 						RevisionNumber,
 						SectionIdx,
 						LODIndex,
+						bRecreating,
 						Mode == EGPUSkinCacheEntryMode::RayTracing ? SkinCacheEntryForRayTracing : SkinCacheEntry
 						);
 				}
@@ -749,10 +792,10 @@ void FSkeletalMeshObjectGPUSkin::ProcessUpdatedDynamicData(EGPUSkinCacheEntryMod
 	}
 
 	// Mobile doesn't support motion blur so no need to double buffer morph deltas
-	if (bMorph && !bIsMobile && !SkinCacheEntry && Mode == EGPUSkinCacheEntryMode::Raster && !LOD.MorphVertexBufferPool.IsDoubleBuffered())
+	if (bMorph && !bIsMobile && !SkinCacheEntry && Mode == EGPUSkinCacheEntryMode::Raster && !LOD.MorphVertexBufferPool->IsDoubleBuffered())
 	{
 		// Going through GPU skinned vertex factory, turn on double buffering for motion blur
-		LOD.MorphVertexBufferPool.EnableDoubleBuffer(RHICmdList);
+		LOD.MorphVertexBufferPool->EnableDoubleBuffer(RHICmdList);
 	}
 }
 
@@ -1076,7 +1119,7 @@ void FSkeletalMeshObjectGPUSkin::UpdateMorphVertexBuffer(FRHICommandList& RHICmd
 		LOD.UpdateMorphVertexBufferCPU(RHICmdList, DynamicData->ActiveMorphTargets, DynamicData->MorphTargetWeights, DynamicData->SectionIdsUseByActiveMorphTargets, bGPUSkinCacheEnabled, MorphVertexBuffer);
 	}
 
-	if (LOD.MorphVertexBufferPool.IsDoubleBuffered())
+	if (LOD.MorphVertexBufferPool->IsDoubleBuffered())
 	{
 		// Update vertex stream binding when morphed delta data is double buffered
 		for (int32 SectionIdx = 0; SectionIdx < LODData.RenderSections.Num(); ++SectionIdx)
@@ -2020,7 +2063,7 @@ void FSkeletalMeshObjectGPUSkin::GetUsedVertexFactoryData(
 	}
 
 	// Setup tmp MeshObjectLOD object to extract the vertex factory buffers
-	FSkeletalMeshObjectGPUSkin::FSkeletalMeshObjectLOD SkeletalMeshObjectLOD(SkelMeshRenderData, LODIndex, InFeatureLevel);
+	FSkeletalMeshObjectGPUSkin::FSkeletalMeshObjectLOD SkeletalMeshObjectLOD(SkelMeshRenderData, LODIndex, InFeatureLevel, nullptr);
 	SkeletalMeshObjectLOD.MeshObjectWeightBuffer = FSkeletalMeshObject::GetSkinWeightVertexBuffer(LODRenderData, CompLODInfo);
 	SkeletalMeshObjectLOD.MeshObjectColorBuffer = FSkeletalMeshObject::GetColorVertexBuffer(LODRenderData, CompLODInfo);
 
@@ -2072,7 +2115,7 @@ void FSkeletalMeshObjectGPUSkin::FSkeletalMeshObjectLOD::GetVertexBuffers(FVerte
 	OutVertexBuffers.StaticVertexBuffers = &LODData.StaticVertexBuffers;
 	OutVertexBuffers.ColorVertexBuffer = MeshObjectColorBuffer;
 	OutVertexBuffers.SkinWeightVertexBuffer = MeshObjectWeightBuffer;
-	OutVertexBuffers.MorphVertexBufferPool = &MorphVertexBufferPool;
+	OutVertexBuffers.MorphVertexBufferPool = MorphVertexBufferPool;
 	OutVertexBuffers.APEXClothVertexBuffer = &LODData.ClothVertexBuffer;
 	OutVertexBuffers.NumVertices = LODData.GetNumVertices();
 }
@@ -2241,7 +2284,7 @@ void FSkeletalMeshObjectGPUSkin::FSkeletalMeshObjectLOD::InitMorphResources(cons
 
 	// init the delta vertex buffer for this LOD
 	const FName OwnerName = LODData.MorphTargetVertexInfoBuffers.GetOwnerName();
-	MorphVertexBufferPool.InitResources(OwnerName);
+	MorphVertexBufferPool->InitResources(OwnerName);
 
 	// Vertex buffers available for the LOD
 	FVertexFactoryBuffers VertexBuffers;
@@ -2257,8 +2300,9 @@ void FSkeletalMeshObjectGPUSkin::FSkeletalMeshObjectLOD::ReleaseMorphResources()
 {
 	// Release morph vertex factories
 	GPUSkinVertexFactories.ReleaseMorphVertexFactories();
-	// release the delta vertex buffer
-	MorphVertexBufferPool.ReleaseResources();
+
+	// By design, we do not release MorphVertexBufferPool, as it may persist when render state gets re-created.  Instead, it gets released
+	// when its ref count goes to zero in the FSkeletalMeshObjectLOD destructor.
 }
 
 
@@ -2683,4 +2727,4 @@ bool FDynamicSkelMeshObjectDataGPUSkin::UpdateClothSimulationData(USkinnedMeshCo
 		return true;
 	}
 	return false;
-}
\ No newline at end of file
+}
diff --git a/Engine/Source/Runtime/Engine/Private/SkeletalRenderGPUSkin.h b/Engine/Source/Runtime/Engine/Private/SkeletalRenderGPUSkin.h
--- a/Engine/Source/Runtime/Engine/Private/SkeletalRenderGPUSkin.h
+++ b/Engine/Source/Runtime/Engine/Private/SkeletalRenderGPUSkin.h
@@ -303,7 +303,7 @@ private:
 /**
 * Pooled morph vertex buffers that store the vertex deltas.
 */
-class FMorphVertexBufferPool
+class FMorphVertexBufferPool : public FThreadSafeRefCountedObject
 {
 public:
 	FMorphVertexBufferPool(FSkeletalMeshRenderData* InSkelMeshRenderData, int32 InLOD, ERHIFeatureLevel::Type InFeatureLevel)
@@ -312,10 +312,21 @@ public:
 		MorphVertexBuffers[1] = FMorphVertexBuffer(InSkelMeshRenderData, InLOD, InFeatureLevel);
 	}
 
+	~FMorphVertexBufferPool()
+	{
+		// Note that destruction of this class must occur on the render thread if InitResources has been called!
+		// This is normally pointed to by FSkeletalMeshObjectGPUSkin, which is defer deleted on the render thread.
+		if (bInitializedResources)
+		{
+			ReleaseResources();
+		}
+	}
+
 	void InitResources(const FName& OwnerName);
 	void ReleaseResources();
 	SIZE_T GetResourceSize() const;
 	void EnableDoubleBuffer(FRHICommandListBase& RHICmdList);
+	bool IsInitialized() const						{ return bInitializedResources; }
 	bool IsDoubleBuffered() const					{ return bDoubleBuffer; }
 	void SetUpdatedFrameNumber(uint32 FrameNumber)	{ UpdatedFrameNumber = FrameNumber; }
 	uint32 GetUpdatedFrameNumber() const			{ return UpdatedFrameNumber; }
@@ -326,6 +337,8 @@ public:
 private:
 	/** Vertex buffer that stores the morph target vertex deltas. */
 	FMorphVertexBuffer MorphVertexBuffers[2];
+	/** If data is preserved when recreating render state, resources will already be initialized, so we need a flag to track that. */
+	bool bInitializedResources = false;
 	/** whether to double buffer. If going through skin cache, then use single buffer; otherwise double buffer. */
 	bool bDoubleBuffer = false;
 
@@ -352,11 +365,12 @@ public:
 	ENGINE_API virtual void InitResources(USkinnedMeshComponent* InMeshComponent) override;
 	ENGINE_API virtual void ReleaseResources() override;
 	ENGINE_API virtual void Update(int32 LODIndex,USkinnedMeshComponent* InMeshComponent,const FMorphTargetWeightMap& InActiveMorphTargets, const TArray<float>& InMorphTargetWeights, EPreviousBoneTransformUpdateMode PreviousBoneTransformUpdateMode, const FExternalMorphWeightData& InExternalMorphWeightData) override;
-	ENGINE_API void UpdateDynamicData_RenderThread(FGPUSkinCache* GPUSkinCache, FRHICommandList& RHICmdList, FDynamicSkelMeshObjectDataGPUSkin* InDynamicData, FSceneInterface* Scene, uint64 FrameNumberToPrepare, uint32 RevisionNumber, uint32 PreviousRevisionNumber);
+	ENGINE_API void UpdateDynamicData_RenderThread(FGPUSkinCache* GPUSkinCache, FRHICommandList& RHICmdList, FDynamicSkelMeshObjectDataGPUSkin* InDynamicData, FSceneInterface* Scene, uint64 FrameNumberToPrepare, uint32 RevisionNumber, uint32 PreviousRevisionNumber, bool bRecreating);
 	ENGINE_API virtual void PreGDMECallback(FRHICommandList& RHICmdList, FGPUSkinCache* GPUSkinCache, uint32 FrameNumber) override;
 	ENGINE_API virtual const FVertexFactory* GetSkinVertexFactory(const FSceneView* View, int32 LODIndex,int32 ChunkIdx, ESkinVertexFactoryMode VFMode = ESkinVertexFactoryMode::Default) const override;
 	ENGINE_API virtual const FSkinBatchVertexFactoryUserData* GetVertexFactoryUserData(const int32 LODIndex, int32 ChunkIdx, ESkinVertexFactoryMode VFMode) const override;
 	virtual bool IsCPUSkinned() const override { return false; }
+	virtual bool IsGPUSkinMesh() const override { return true; }
 	ENGINE_API virtual TArray<FTransform>* GetComponentSpaceTransforms() const override;
 	ENGINE_API virtual const TArray<FMatrix44f>& GetReferenceToLocalMatrices() const override;
 	ENGINE_API virtual bool GetCachedGeometry(FCachedGeometry& OutCachedGeometry) const override;
@@ -565,14 +579,21 @@ protected:
 	/** vertex data for rendering a single LOD */
 	struct FSkeletalMeshObjectLOD
 	{
-		FSkeletalMeshObjectLOD(FSkeletalMeshRenderData* InSkelMeshRenderData,int32 InLOD, ERHIFeatureLevel::Type InFeatureLevel)
+		FSkeletalMeshObjectLOD(FSkeletalMeshRenderData* InSkelMeshRenderData,int32 InLOD, ERHIFeatureLevel::Type InFeatureLevel, FMorphVertexBufferPool* InRecreateBufferPool)
 			: SkelMeshRenderData(InSkelMeshRenderData)
 			, LODIndex(InLOD)
 			, FeatureLevel(InFeatureLevel)
-			, MorphVertexBufferPool(InSkelMeshRenderData, LODIndex, FeatureLevel)
 			, MeshObjectWeightBuffer(nullptr)
 			, MeshObjectColorBuffer(nullptr)
 		{
+			if (InRecreateBufferPool)
+			{
+				MorphVertexBufferPool = InRecreateBufferPool;
+			}
+			else
+			{
+				MorphVertexBufferPool = new FMorphVertexBufferPool(InSkelMeshRenderData, LODIndex, FeatureLevel);
+			}
 		}
 
 		/** 
@@ -604,7 +625,7 @@ protected:
 		 */
 		void GetResourceSizeEx(FResourceSizeEx& CumulativeResourceSize)
 		{
-			CumulativeResourceSize.AddUnknownMemoryBytes(MorphVertexBufferPool.GetResourceSize());
+			CumulativeResourceSize.AddUnknownMemoryBytes(MorphVertexBufferPool->GetResourceSize());
 			CumulativeResourceSize.AddUnknownMemoryBytes(GPUSkinVertexFactories.GetResourceSize());
 		}
 
@@ -615,7 +636,7 @@ protected:
 		ERHIFeatureLevel::Type FeatureLevel;
 
 		/** Pooled vertex buffers that store the morph target vertex deltas. */
-		FMorphVertexBufferPool	MorphVertexBufferPool;
+		TRefCountPtr<FMorphVertexBufferPool> MorphVertexBufferPool;
 
 		/** Default GPU skinning vertex factories and matrices */
 		FVertexFactoryData GPUSkinVertexFactories;
@@ -665,7 +686,7 @@ protected:
 	*/
 	ENGINE_API void ReleaseMorphResources();
 
-	ENGINE_API void ProcessUpdatedDynamicData(EGPUSkinCacheEntryMode Mode, FGPUSkinCache* GPUSkinCache, FRHICommandList& RHICmdList, uint32 FrameNumberToPrepare, uint32 RevisionNumber, uint32 PreviousRevisionNumber, bool bMorphNeedsUpdate, int32 LODIndex);
+	ENGINE_API void ProcessUpdatedDynamicData(EGPUSkinCacheEntryMode Mode, FGPUSkinCache* GPUSkinCache, FRHICommandList& RHICmdList, uint32 FrameNumberToPrepare, uint32 RevisionNumber, uint32 PreviousRevisionNumber, bool bMorphNeedsUpdate, int32 LODIndex, bool bRecreating);
 
 	ENGINE_API virtual void UpdateMorphVertexBuffer(FRHICommandList& RHICmdList, EGPUSkinCacheEntryMode Mode, FSkeletalMeshObjectLOD& LOD, const FSkeletalMeshLODRenderData& LODData, bool bGPUSkinCacheEnabled, FMorphVertexBuffer& MorphVertexBuffer);
 
diff --git a/Engine/Source/Runtime/Engine/Public/GPUSkinCache.h b/Engine/Source/Runtime/Engine/Public/GPUSkinCache.h
--- a/Engine/Source/Runtime/Engine/Public/GPUSkinCache.h
+++ b/Engine/Source/Runtime/Engine/Public/GPUSkinCache.h
@@ -146,6 +146,7 @@ public:
 		uint32 RevisionNumber, 
 		int32 Section,
 		int32 LOD,
+		bool bRecreating,
 		FGPUSkinCacheEntry*& InOutEntry
 		);
 
@@ -398,6 +399,21 @@ public:
 			return Result;
 		}
 
+		// On recreate of the render state where the GPU skin cache entry is preserved, the bone buffer will have been reallocated,
+		// even though the transforms didn't change.  We need to force the Find() call above to treat the data as up-to-date, which
+		// can be accomplished by updating the BoneBuffer pointer for the previous Revision, so it matches again.
+		void UpdatePreviousBoneBuffer(const FVertexBufferAndSRV& PreviousBoneBuffer, uint32 PreviousRevision)
+		{
+			for (int32 Index = 0; Index < NUM_BUFFERS; ++Index)
+			{
+				if (Revisions[Index] == PreviousRevision)
+				{
+					BoneBuffers[Index] = &PreviousBoneBuffer;
+					break;
+				}
+			}
+		}
+
 	private:
 		uint32 Revisions[NUM_BUFFERS];
 		const FVertexBufferAndSRV* BoneBuffers[NUM_BUFFERS];
diff --git a/Engine/Source/Runtime/Engine/Public/SkeletalRenderPublic.h b/Engine/Source/Runtime/Engine/Public/SkeletalRenderPublic.h
--- a/Engine/Source/Runtime/Engine/Public/SkeletalRenderPublic.h
+++ b/Engine/Source/Runtime/Engine/Public/SkeletalRenderPublic.h
@@ -124,6 +124,11 @@ public:
 	 */
 	virtual bool IsCPUSkinned() const = 0;
 
+	/**
+	 * Returns true if this mesh is an FSkeletalMeshObjectGPUSkin
+	 */
+	virtual bool IsGPUSkinMesh() const { return false; }
+
 	/** 
 	 *	Get the array of component-space bone transforms. 
 	 *	Not safe to hold this point between frames, because it exists in dynamic data passed from main thread.
