SkyAtmosphere - added node to evaluate outspace dir light illuminance for a desired illuminance on ground according to a light direction.

diff --git a/Engine/Source/Runtime/Engine/Classes/Components/SkyAtmosphereComponent.h b/Engine/Source/Runtime/Engine/Classes/Components/SkyAtmosphereComponent.h
--- a/Engine/Source/Runtime/Engine/Classes/Components/SkyAtmosphereComponent.h
+++ b/Engine/Source/Runtime/Engine/Classes/Components/SkyAtmosphereComponent.h
@@ -229,6 +229,11 @@ class USkyAtmosphereComponent : public USceneComponent
 	UFUNCTION(BlueprintCallable, Category = "Utilities", meta = (DisplayName = "Get Atmosphere Transmitance On Ground At Planet Top"))
 	ENGINE_API FLinearColor GetAtmosphereTransmitanceOnGroundAtPlanetTop(UDirectionalLightComponent* DirectionalLight);
 
+	// This function can be used for instance in order to evaluate a directional atmospheric light outer space illuminance for a desired illuminance on ground given a direction. 
+	// This is given for the position at the top of the virtual planet. Plus the output outer space illuminance into the light intensity.
+	UFUNCTION(BlueprintCallable, Category = "Utilities", meta = (DisplayName = "Get Atmospheric Light To Match Illuminance On Ground"))
+	ENGINE_API float GetAtmosphericLightToMatchIlluminanceOnGround(FVector LightDirection = FVector(0.0f, 0.0f, 1.0f), float IlluminanceOnGround = 1.0);
+
 	// This is used to position the SkyAtmosphere similarly to the deprecated AtmosphericFog component
 	void SetPositionToMatchDeprecatedAtmosphericFog();
 
diff --git a/Engine/Source/Runtime/Engine/Private/Components/SkyAtmosphereComponent.cpp b/Engine/Source/Runtime/Engine/Private/Components/SkyAtmosphereComponent.cpp
--- a/Engine/Source/Runtime/Engine/Private/Components/SkyAtmosphereComponent.cpp
+++ b/Engine/Source/Runtime/Engine/Private/Components/SkyAtmosphereComponent.cpp
@@ -17,6 +17,7 @@
 #include "UObject/ConstructorHelpers.h"
 #include "Components/DirectionalLightComponent.h"
 #include "UObject/UE5MainStreamObjectVersion.h"
+#include "ColorSpace.h"
 
 #include UE_INLINE_GENERATED_CPP_BY_NAME(SkyAtmosphereComponent)
 
@@ -419,6 +420,14 @@ FLinearColor USkyAtmosphereComponent::GetAtmosphereTransmitanceOnGroundAtPlanetT
 	return FLinearColor::White;
 }
 
+float USkyAtmosphereComponent::GetAtmosphericLightToMatchIlluminanceOnGround(FVector LightDirection, float IlluminanceOnGround)
+{
+	FAtmosphereSetup AtmosphereSetup(*this);
+	const FLinearColor TransmittanceAtDirLight = AtmosphereSetup.GetTransmittanceAtGroundLevel(LightDirection);
+	const float OuterSpaceIlluminance = IlluminanceOnGround / FMath::Max(UE_SMALL_NUMBER, UE::Color::FColorSpace::GetWorking().GetLuminance(TransmittanceAtDirLight));
+	return OuterSpaceIlluminance;
+}
+
 /*=============================================================================
 	ASkyAtmosphere implementation.
 =============================================================================*/
