World Partition: fix UActorDescContainerInstance delegates leak.

diff --git a/Engine/Source/Runtime/Engine/Private/WorldPartition/ActorDescContainerInstance.cpp b/Engine/Source/Runtime/Engine/Private/WorldPartition/ActorDescContainerInstance.cpp
--- a/Engine/Source/Runtime/Engine/Private/WorldPartition/ActorDescContainerInstance.cpp
+++ b/Engine/Source/Runtime/Engine/Private/WorldPartition/ActorDescContainerInstance.cpp
@@ -16,6 +16,16 @@
 #if WITH_EDITOR
 UActorDescContainerInstance::FActorDescContainerInstanceInitializeDelegate UActorDescContainerInstance::OnActorDescContainerInstanceInitialized;
 
+void UActorDescContainerInstance::BeginDestroy()
+{
+	Super::BeginDestroy();
+
+	if (bIsInitialized)
+	{
+		Uninitialize();
+	}
+}
+
 FName UActorDescContainerInstance::GetContainerPackageNameFromWorld(UWorld* InWorld)
 {
 	check(InWorld);
diff --git a/Engine/Source/Runtime/Engine/Public/WorldPartition/ActorDescContainerInstance.h b/Engine/Source/Runtime/Engine/Public/WorldPartition/ActorDescContainerInstance.h
--- a/Engine/Source/Runtime/Engine/Public/WorldPartition/ActorDescContainerInstance.h
+++ b/Engine/Source/Runtime/Engine/Public/WorldPartition/ActorDescContainerInstance.h
@@ -31,6 +31,10 @@ protected:
 #endif
 	{}
 
+	//~ Begin UObject Interface
+	virtual void BeginDestroy() override;
+	//~ End UObject Interface
+
 #if WITH_EDITOR
 	friend FWorldPartitionActorDescInstance;
 	friend UWorldPartition;
