World Partition HLODs: fix runtime spatial hash to include SpatialHashRuntimeGridInfo actors when validating grids so we get a check map error instead of a cook ensure.

diff --git a/Engine/Source/Runtime/Engine/Private/WorldPartition/WorldPartitionActorDescView.cpp b/Engine/Source/Runtime/Engine/Private/WorldPartition/WorldPartitionActorDescView.cpp
--- a/Engine/Source/Runtime/Engine/Private/WorldPartition/WorldPartitionActorDescView.cpp
+++ b/Engine/Source/Runtime/Engine/Private/WorldPartition/WorldPartitionActorDescView.cpp
@@ -256,11 +256,6 @@ FName FWorldPartitionActorDescView::GetActorLabelOrName() const
 	return ActorDesc->GetActorLabelOrName();
 }
 
-bool FWorldPartitionActorDescView::ShouldValidateRuntimeGrid() const
-{
-	return ActorDesc->ShouldValidateRuntimeGrid();
-}
-
 void FWorldPartitionActorDescView::SetForcedNonSpatiallyLoaded()
 {
 	if (!bIsForcedNonSpatiallyLoaded)
diff --git a/Engine/Source/Runtime/Engine/Private/WorldPartition/WorldPartitionRuntimeSpatialHash.cpp b/Engine/Source/Runtime/Engine/Private/WorldPartition/WorldPartitionRuntimeSpatialHash.cpp
--- a/Engine/Source/Runtime/Engine/Private/WorldPartition/WorldPartitionRuntimeSpatialHash.cpp
+++ b/Engine/Source/Runtime/Engine/Private/WorldPartition/WorldPartitionRuntimeSpatialHash.cpp
@@ -1630,6 +1630,15 @@ bool UWorldPartitionRuntimeSpatialHash::IsValidGrid(FName GridName) const
 		}
 	}
 
+	/*const UWorldPartition* WorldPartition = GetOuterUWorldPartition();
+	for (FActorDescContainerInstanceCollection::TConstIterator<ASpatialHashRuntimeGridInfo> Iterator(WorldPartition); Iterator; ++Iterator)
+	{
+		if (Iterator->GetRuntimeGrid() == GridName)
+		{
+			return true;
+		}
+	}*/
+
 	return false;
 }
 
diff --git a/Engine/Source/Runtime/Engine/Private/WorldPartition/WorldPartitionStreamingGeneration.cpp b/Engine/Source/Runtime/Engine/Private/WorldPartition/WorldPartitionStreamingGeneration.cpp
--- a/Engine/Source/Runtime/Engine/Private/WorldPartition/WorldPartitionStreamingGeneration.cpp
+++ b/Engine/Source/Runtime/Engine/Private/WorldPartition/WorldPartitionStreamingGeneration.cpp
@@ -1512,37 +1512,32 @@ class FWorldPartitionStreamingGenerator
 			{
 				FContainerCollectionInstanceDescriptor::FPerInstanceData& PerInstanceData = ContainerCollectionInstanceDescriptor.GetPerInstanceData(ActorDescView.GetGuid());
 
-PRAGMA_DISABLE_DEPRECATION_WARNINGS
-				if (ActorDescView.GetActorDesc()->ShouldValidateRuntimeGrid())
-PRAGMA_ENABLE_DEPRECATION_WARNINGS
+				if (!IsValidGrid(PerInstanceData.RuntimeGrid))
 				{
-					if (!IsValidGrid(PerInstanceData.RuntimeGrid))
+					if (PassType == EPassType::ErrorReporting)
 					{
-						if (PassType == EPassType::ErrorReporting)
-						{
-							ErrorHandler->OnInvalidRuntimeGrid(ActorDescView, PerInstanceData.RuntimeGrid);
-						}
-						else
-						{
-							PerInstanceData.RuntimeGrid = NAME_None;
-						}
-
-						NbErrorsDetected++;
+						ErrorHandler->OnInvalidRuntimeGrid(ActorDescView, PerInstanceData.RuntimeGrid);
 					}
-
-					if (ActorDescView.GetHLODLayer().IsValid() && !IsValidHLODLayer(PerInstanceData.RuntimeGrid, ActorDescView.GetHLODLayer()))
+					else
 					{
-						if (PassType == EPassType::ErrorReporting)
-						{
-							ErrorHandler->OnInvalidHLODLayer(ActorDescView);
-						}
-						else
-						{
-							ActorDescView.SetForcedNoHLODLayer();
-						}
+						PerInstanceData.RuntimeGrid = NAME_None;
+					}
 
-						NbErrorsDetected++;
+					NbErrorsDetected++;
+				}
+
+				if (ActorDescView.GetHLODLayer().IsValid() && !IsValidHLODLayer(PerInstanceData.RuntimeGrid, ActorDescView.GetHLODLayer()))
+				{
+					if (PassType == EPassType::ErrorReporting)
+					{
+						ErrorHandler->OnInvalidHLODLayer(ActorDescView);
 					}
+					else
+					{
+						ActorDescView.SetForcedNoHLODLayer();
+					}
+
+					NbErrorsDetected++;
 				}
 
 				if (TSet<FGuid>* FilteredActors = ContainerFilteredActors.Find(ContainerCollectionInstanceDescriptor.ID))
diff --git a/Engine/Source/Runtime/Engine/Public/WorldPartition/HLOD/HLODActorDesc.h b/Engine/Source/Runtime/Engine/Public/WorldPartition/HLOD/HLODActorDesc.h
--- a/Engine/Source/Runtime/Engine/Public/WorldPartition/HLOD/HLODActorDesc.h
+++ b/Engine/Source/Runtime/Engine/Public/WorldPartition/HLOD/HLODActorDesc.h
@@ -31,7 +31,6 @@ public:
 	int64 GetStat(FName InStatName) const;
 
 	//~ Begin FWorldPartitionActorDesc Interface.
-	virtual bool ShouldValidateRuntimeGrid() const override { return false; }
 	virtual FBox GetEditorBounds() const override { return EditorBounds; }
 	//~ End FWorldPartitionActorDesc Interface.
 
diff --git a/Engine/Source/Runtime/Engine/Public/WorldPartition/WorldPartitionActorDescView.h b/Engine/Source/Runtime/Engine/Public/WorldPartition/WorldPartitionActorDescView.h
--- a/Engine/Source/Runtime/Engine/Public/WorldPartition/WorldPartitionActorDescView.h
+++ b/Engine/Source/Runtime/Engine/Public/WorldPartition/WorldPartitionActorDescView.h
@@ -74,9 +74,6 @@ public:
 
 	ENGINE_API FName GetActorLabelOrName() const;
 
-	UE_DEPRECATED(5.2, "ShouldValidateRuntimeGrid is deprecated and should not be used.")
-		ENGINE_API bool ShouldValidateRuntimeGrid() const;
-
 	ENGINE_API void SetForcedNonSpatiallyLoaded();
 	ENGINE_API void SetForcedNoRuntimeGrid();
 	ENGINE_API void SetForcedNoDataLayers();
