WorldPartition: Fix delegate leak
- FStreamingGenerationContainerInstanceCollection::InitializeCollection could remove elements from ActorDescContainerInstanceCollection when CollectionType was ECollectionType::BaseAndEDLs without calling UnregisterDelegates.
- Fix is to avoid registration of container delegates for FStreamingGenerationContainerInstanceCollection class as it is used only for streaming generation.

#rnx

diff --git a/Engine/Source/Runtime/Engine/Private/WorldPartition/WorldPartitionStreamingGeneration.cpp b/Engine/Source/Runtime/Engine/Private/WorldPartition/WorldPartitionStreamingGeneration.cpp
--- a/Engine/Source/Runtime/Engine/Private/WorldPartition/WorldPartitionStreamingGeneration.cpp
+++ b/Engine/Source/Runtime/Engine/Private/WorldPartition/WorldPartitionStreamingGeneration.cpp
@@ -2204,6 +2204,7 @@ void FStreamingGenerationContainerInstanceCollection::InitializeCollection()
 	{
 		// When type is set to BaseAndEDL, we remove ContentBundle containers from the collection.
 		// BaseAndEDL type assumes ContentBundle containers are generated separately one at a time.
+		check(!ShouldRegisterDelegates());
 		ActorDescContainerInstanceCollection.SetNum(Algo::RemoveIf(ActorDescContainerInstanceCollection, [](const UActorDescContainerInstance* ActorDescContainerInstance) { return ActorDescContainerInstance->GetContentBundleGuid().IsValid(); }));
 	}
 
diff --git a/Engine/Source/Runtime/Engine/Public/WorldPartition/ActorDescContainerInstanceCollection.h b/Engine/Source/Runtime/Engine/Public/WorldPartition/ActorDescContainerInstanceCollection.h
--- a/Engine/Source/Runtime/Engine/Public/WorldPartition/ActorDescContainerInstanceCollection.h
+++ b/Engine/Source/Runtime/Engine/Public/WorldPartition/ActorDescContainerInstanceCollection.h
@@ -74,6 +74,7 @@ public:
 
 protected:
 	virtual void OnCollectionChanged() {};
+	virtual bool ShouldRegisterDelegates() const { return true; }
 
 	TArray<ActorDescContPtrType> ActorDescContainerInstanceCollection;
 
@@ -299,17 +300,23 @@ void TActorDescContainerInstanceCollection<ActorDescContPtrType>::Empty()
 template<class ActorDescContPtrType>
 void TActorDescContainerInstanceCollection<ActorDescContPtrType>::RegisterDelegates(ActorDescContPtrType Container)
 {
-	ConstCast(Container)->OnActorDescInstanceAddedEvent.AddRaw(this, &TActorDescContainerInstanceCollection<ActorDescContPtrType>::OnActorDescInstanceAdded);
-	ConstCast(Container)->OnActorDescInstanceRemovedEvent.AddRaw(this, &TActorDescContainerInstanceCollection<ActorDescContPtrType>::OnActorDescInstanceRemoved);
-	ConstCast(Container)->OnActorReplacedEvent.AddRaw(this, &TActorDescContainerInstanceCollection<ActorDescContPtrType>::OnActorReplaced);
+	if (ShouldRegisterDelegates())
+	{
+		ConstCast(Container)->OnActorDescInstanceAddedEvent.AddRaw(this, &TActorDescContainerInstanceCollection<ActorDescContPtrType>::OnActorDescInstanceAdded);
+		ConstCast(Container)->OnActorDescInstanceRemovedEvent.AddRaw(this, &TActorDescContainerInstanceCollection<ActorDescContPtrType>::OnActorDescInstanceRemoved);
+		ConstCast(Container)->OnActorReplacedEvent.AddRaw(this, &TActorDescContainerInstanceCollection<ActorDescContPtrType>::OnActorReplaced);
+	}
 }
 
 template<class ActorDescContPtrType>
 void TActorDescContainerInstanceCollection<ActorDescContPtrType>::UnregisterDelegates(ActorDescContPtrType Container)
 {
-	ConstCast(Container)->OnActorDescInstanceAddedEvent.RemoveAll(this);
-	ConstCast(Container)->OnActorDescInstanceRemovedEvent.RemoveAll(this);
-	ConstCast(Container)->OnActorReplacedEvent.RemoveAll(this);
+	if (ShouldRegisterDelegates())
+	{
+		ConstCast(Container)->OnActorDescInstanceAddedEvent.RemoveAll(this);
+		ConstCast(Container)->OnActorDescInstanceRemovedEvent.RemoveAll(this);
+		ConstCast(Container)->OnActorReplacedEvent.RemoveAll(this);
+	}
 }
 
 template<class ActorDescContPtrType>
diff --git a/Engine/Source/Runtime/Engine/Public/WorldPartition/WorldPartitionStreamingGeneration.h b/Engine/Source/Runtime/Engine/Public/WorldPartition/WorldPartitionStreamingGeneration.h
--- a/Engine/Source/Runtime/Engine/Public/WorldPartition/WorldPartitionStreamingGeneration.h
+++ b/Engine/Source/Runtime/Engine/Public/WorldPartition/WorldPartitionStreamingGeneration.h
@@ -251,6 +251,7 @@ private:
 	TArrayView<const UActorDescContainerInstance* const> GetExternalDataLayerContainerInstances() const;
 	TArrayView<const UActorDescContainerInstance* const> GetContentBundleContainerInstances() const;
 	ENGINE_API virtual void OnCollectionChanged() override;
+	virtual bool ShouldRegisterDelegates() const override { return false; }
 
 	ECollectionType CollectionType;
 	int32 ExternalDataLayerStartIdx = INDEX_NONE;
