Added an additional option to SuggestProjectileVelocity, allowing a test with no results to return a velocity which will get the projectile as close to the target as possible.

diff --git a/Engine/Source/Runtime/AIModule/Private/AIController.cpp b/Engine/Source/Runtime/AIModule/Private/AIController.cpp
--- a/Engine/Source/Runtime/AIModule/Private/AIController.cpp
+++ b/Engine/Source/Runtime/AIModule/Private/AIController.cpp
@@ -1086,7 +1086,13 @@ bool AAIController::SuggestTossVelocity(FVector& OutTossVelocity, FVector Start,
 	float const GravityOverride = PhysicsVolume ? PhysicsVolume->GetGravityZ() : 0.f;
 	ESuggestProjVelocityTraceOption::Type const TraceOption = bOnlyTraceUp ? ESuggestProjVelocityTraceOption::OnlyTraceWhileAscending : ESuggestProjVelocityTraceOption::TraceFullPath;
 
-	return UGameplayStatics::SuggestProjectileVelocity(this, OutTossVelocity, Start, End, TossSpeed, bPreferHighArc, CollisionRadius, GravityOverride, TraceOption);
+	UGameplayStatics::FSuggestProjectileVelocityParameters VelocityParams = UGameplayStatics::FSuggestProjectileVelocityParameters(this, Start, End, TossSpeed);
+	VelocityParams.bFavorHighArc = bPreferHighArc;
+	VelocityParams.CollisionRadius = CollisionRadius;
+	VelocityParams.OverrideGravityZ = GravityOverride;
+	VelocityParams.TraceOption = TraceOption;
+
+	return UGameplayStatics::SuggestProjectileVelocity(VelocityParams, OutTossVelocity);
 }
 bool AAIController::PerformAction(UDEPRECATED_PawnAction& Action, EAIRequestPriority::Type Priority, UObject* const InInstigator /*= NULL*/)
 {
diff --git a/Engine/Source/Runtime/Engine/Classes/Kismet/GameplayStatics.h b/Engine/Source/Runtime/Engine/Classes/Kismet/GameplayStatics.h
--- a/Engine/Source/Runtime/Engine/Classes/Kismet/GameplayStatics.h
+++ b/Engine/Source/Runtime/Engine/Classes/Kismet/GameplayStatics.h
@@ -38,6 +38,7 @@ class APlayerController;
 struct FDialogueContext;
 
 
+
 /** Delegate called from AsyncLoadGameFromSlot. First two parameters are passed in SlotName and UserIndex, third parameter is a bool indicating success (true) or failure (false). */
 DECLARE_DELEGATE_ThreeParams(FAsyncSaveGameToSlotDelegate, const FString&, const int32, bool);
 
@@ -551,6 +552,27 @@ private:
 
 public:
 
+	/* Struct of parameters passed to SuggestProjectileVelocity function. */
+	struct FSuggestProjectileVelocityParameters
+	{
+	public:
+		const UObject* WorldContextObject = nullptr;
+		FVector Start = FVector::ZeroVector;
+		FVector End = FVector::ZeroVector;
+		float TossSpeed = 0.f;
+		bool bFavorHighArc = false;
+		float CollisionRadius = 0.f;
+		float OverrideGravityZ = 0;
+		ESuggestProjVelocityTraceOption::Type TraceOption = ESuggestProjVelocityTraceOption::TraceFullPath;
+		FCollisionResponseParams& ResponseParam = FCollisionResponseParams::DefaultResponseParam;
+		TArray<AActor*> ActorsToIgnore = TArray<AActor*>();
+		bool bDrawDebug = false;
+		bool bAcceptClosestOnNoSolutions = false;
+
+		FSuggestProjectileVelocityParameters(const UObject* World, FVector StartLocation, FVector EndLocation, float Speed)
+			: WorldContextObject{ World }, Start{ StartLocation }, End{ EndLocation }, TossSpeed{ Speed } {}
+	};
+
 	/** Plays the specified effect attached to and following the specified component. The system will go away when the effect is complete. Does not replicate.
 	* @param EmitterTemplate - particle system to create
 	 * @param AttachComponent - Component to attach to.
@@ -1253,22 +1275,30 @@ public:
 
 	/**
 	 * Calculates an launch velocity for a projectile to hit a specified point.
-	 * @param TossVelocity		(output) Result launch velocity.
-	 * @param StartLocation		Intended launch location
-	 * @param EndLocation		Desired landing location
-	 * @param LaunchSpeed		Desired launch speed
-	 * @param OverrideGravityZ	Optional gravity override.  0 means "do not override".
-	 * @param TraceOption		Controls whether or not to validate a clear path by tracing along the calculated arc
-	 * @param CollisionRadius	Radius of the projectile (assumed spherical), used when tracing
-	 * @param bFavorHighArc		If true and there are 2 valid solutions, will return the higher arc.  If false, will favor the lower arc.
-	 * @param bDrawDebug		When true, a debug arc is drawn (red for an invalid arc, green for a valid arc)
-	 * @return					Returns false if there is no valid solution or the valid solutions are blocked.  Returns true otherwise.
+	 * @param TossVelocity					(output) Result launch velocity.
+	 * @param StartLocation					Intended launch location
+	 * @param EndLocation					Desired landing location
+	 * @param LaunchSpeed					Desired launch speed
+	 * @param OverrideGravityZ				Optional gravity override.  0 means "do not override".
+	 * @param TraceOption					Controls whether or not to validate a clear path by tracing along the calculated arc
+	 * @param CollisionRadius				Radius of the projectile (assumed spherical), used when tracing
+	 * @param bFavorHighArc					If true and there are 2 valid solutions, will return the higher arc.  If false, will favor the lower arc.
+	 * @param bDrawDebug					When true, a debug arc is drawn (red for an invalid arc, green for a valid arc)
+	 * @param bAcceptClosestOnNoSolutions	If target is unreachable and no solutions are possible, provide a velocity that gets as close to the target as possible given this launch speed
+	 * @return								Returns false if there is no valid solution or the valid solutions are blocked.  Returns true otherwise.
 	 */
 	UFUNCTION(BlueprintCallable, Category="Game", DisplayName="Suggest Projectile Velocity", meta=(WorldContext="WorldContextObject"))
-	static ENGINE_API bool BlueprintSuggestProjectileVelocity(const UObject* WorldContextObject, FVector& TossVelocity, FVector StartLocation, FVector EndLocation, float LaunchSpeed, float OverrideGravityZ, ESuggestProjVelocityTraceOption::Type TraceOption, float CollisionRadius, bool bFavorHighArc, bool bDrawDebug);
+	static bool BlueprintSuggestProjectileVelocity(const UObject* WorldContextObject, FVector& TossVelocity, FVector StartLocation, FVector EndLocation, float LaunchSpeed, float OverrideGravityZ, ESuggestProjVelocityTraceOption::Type TraceOption, float CollisionRadius, bool bFavorHighArc, bool bDrawDebug, bool bAcceptClosestOnNoSolutions = false);
+
+	UE_DEPRECATED(5.4, "Deprecated parameter list. Please use FSuggestProjectileVelocityParameters instead.")
+	static ENGINE_API bool SuggestProjectileVelocity(const UObject* WorldContextObject, FVector& TossVelocity, FVector StartLocation, FVector EndLocation, float TossSpeed, bool bHighArc = false, float CollisionRadius = 0.f, float OverrideGravityZ = 0, ESuggestProjVelocityTraceOption::Type TraceOption = ESuggestProjVelocityTraceOption::TraceFullPath, FCollisionResponseParams& ResponseParam = FCollisionResponseParams::DefaultResponseParam, TArray<AActor*> ActorsToIgnore = TArray<AActor*>(), bool bDrawDebug = false, bool bAcceptClosestOnNoSolutions = false);
 
 	/** Native version, has more options than the Blueprint version. */
-	static ENGINE_API bool SuggestProjectileVelocity(const UObject* WorldContextObject, FVector& TossVelocity, FVector StartLocation, FVector EndLocation, float TossSpeed, bool bHighArc = false, float CollisionRadius = 0.f, float OverrideGravityZ = 0, ESuggestProjVelocityTraceOption::Type TraceOption = ESuggestProjVelocityTraceOption::TraceFullPath, const FCollisionResponseParams& ResponseParam = FCollisionResponseParams::DefaultResponseParam, const TArray<AActor*>& ActorsToIgnore = TArray<AActor*>(), bool bDrawDebug = false);
+	static ENGINE_API bool SuggestProjectileVelocity(const FSuggestProjectileVelocityParameters& ProjectileParams, FVector& OutTossVelocity);
+
+
+	/** Stepwise trace along the path of a given velocity*/
+	static ENGINE_API bool IsProjectileTrajectoryBlocked(const UWorld* World, FVector StartLocation, FVector& ProjectileVelocity, float TargetDeltaXY, float GravityZ, float CollisionRadius = 0.f, ESuggestProjVelocityTraceOption::Type TraceOption = ESuggestProjVelocityTraceOption::TraceFullPath, const FCollisionResponseParams& ResponseParam = FCollisionResponseParams::DefaultResponseParam, const TArray<AActor*>& ActorsToIgnore = TArray<AActor*>(), bool bDrawDebug = false);
 
 	/**
 	* Predict the arc of a virtual projectile affected by gravity with collision checks along the arc. Returns a list of positions of the simulated arc and the destination reached by the simulation.
diff --git a/Engine/Source/Runtime/Engine/Private/GameplayStatics.cpp b/Engine/Source/Runtime/Engine/Private/GameplayStatics.cpp
--- a/Engine/Source/Runtime/Engine/Private/GameplayStatics.cpp
+++ b/Engine/Source/Runtime/Engine/Private/GameplayStatics.cpp
@@ -2594,188 +2594,223 @@ FString UGameplayStatics::GetPlatformName()
 	return FPlatformProperties::IniPlatformName();
 }
 
-bool UGameplayStatics::BlueprintSuggestProjectileVelocity(const UObject* WorldContextObject, FVector& OutTossVelocity, FVector StartLocation, FVector EndLocation, float LaunchSpeed, float OverrideGravityZ, ESuggestProjVelocityTraceOption::Type TraceOption, float CollisionRadius, bool bFavorHighArc, bool bDrawDebug)
+bool UGameplayStatics::BlueprintSuggestProjectileVelocity(const UObject* WorldContextObject, FVector& OutTossVelocity, FVector StartLocation, FVector EndLocation, float LaunchSpeed, float OverrideGravityZ, ESuggestProjVelocityTraceOption::Type TraceOption, float CollisionRadius, bool bFavorHighArc, bool bDrawDebug, bool bAcceptClosestOnNoSolutions)
 {
 	// simple pass-through to the C++ interface
-	return UGameplayStatics::SuggestProjectileVelocity(WorldContextObject, OutTossVelocity, StartLocation, EndLocation, LaunchSpeed, bFavorHighArc, CollisionRadius, OverrideGravityZ, TraceOption, FCollisionResponseParams::DefaultResponseParam, TArray<AActor*>(), bDrawDebug);
+	FSuggestProjectileVelocityParameters ProjectileParams = FSuggestProjectileVelocityParameters(WorldContextObject, StartLocation, EndLocation, LaunchSpeed);
+	ProjectileParams.bFavorHighArc = bFavorHighArc;
+	ProjectileParams.CollisionRadius = CollisionRadius;
+	ProjectileParams.OverrideGravityZ = OverrideGravityZ;
+	ProjectileParams.TraceOption = TraceOption;
+	ProjectileParams.bDrawDebug = bDrawDebug;
+	ProjectileParams.bAcceptClosestOnNoSolutions = bAcceptClosestOnNoSolutions;
+
+	return SuggestProjectileVelocity(ProjectileParams, OutTossVelocity);
+}
+
+bool UGameplayStatics::SuggestProjectileVelocity(const UObject* WorldContextObject, FVector& TossVelocity, FVector StartLocation, FVector EndLocation, float TossSpeed, bool bHighArc , float CollisionRadius, float OverrideGravityZ, ESuggestProjVelocityTraceOption::Type TraceOption, FCollisionResponseParams& ResponseParam, TArray<AActor*> ActorsToIgnore, bool bDrawDebug, bool bAcceptClosestOnNoSolutions)
+{
+	FSuggestProjectileVelocityParameters ProjectileParams = FSuggestProjectileVelocityParameters(WorldContextObject,StartLocation, EndLocation, TossSpeed);
+	ProjectileParams.bFavorHighArc = bHighArc;
+	ProjectileParams.CollisionRadius = CollisionRadius;
+	ProjectileParams.OverrideGravityZ = OverrideGravityZ;
+	ProjectileParams.TraceOption = TraceOption;
+	ProjectileParams.ResponseParam = ResponseParam;
+	ProjectileParams.ActorsToIgnore = ActorsToIgnore;
+	ProjectileParams.bDrawDebug = bDrawDebug;
+	ProjectileParams.bAcceptClosestOnNoSolutions = bAcceptClosestOnNoSolutions;
+
+	return SuggestProjectileVelocity(ProjectileParams, TossVelocity);
 }
 
 // note: this will automatically fall back to line test if radius is small enough
 // Based on analytic solution to ballistic angle of launch http://en.wikipedia.org/wiki/Trajectory_of_a_projectile#Angle_required_to_hit_coordinate_.28x.2Cy.29
-bool UGameplayStatics::SuggestProjectileVelocity(const UObject* WorldContextObject, FVector& OutTossVelocity, FVector Start, FVector End, float TossSpeed, bool bFavorHighArc, float CollisionRadius, float OverrideGravityZ, ESuggestProjVelocityTraceOption::Type TraceOption, const FCollisionResponseParams& ResponseParam, const TArray<AActor*>& ActorsToIgnore, bool bDrawDebug)
+bool UGameplayStatics::SuggestProjectileVelocity(const FSuggestProjectileVelocityParameters& ProjectileParams, FVector& OutTossVelocity)
 {
-	const FVector FlightDelta = End - Start;
+	const FVector FlightDelta = ProjectileParams.End - ProjectileParams.Start;
 	const FVector DirXY = FlightDelta.GetSafeNormal2D();
 	const float DeltaXY = FlightDelta.Size2D();
 
 	const float DeltaZ = FlightDelta.Z;
 
-	const float TossSpeedSq = FMath::Square(TossSpeed);
+	const float TossSpeedSq = FMath::Square(ProjectileParams.TossSpeed);
 
-	const UWorld* const World = GEngine->GetWorldFromContextObject(WorldContextObject, EGetWorldErrorMode::LogAndReturnNull);
+	const UWorld* const World = GEngine->GetWorldFromContextObject(ProjectileParams.WorldContextObject, EGetWorldErrorMode::LogAndReturnNull);
 	if (World == nullptr)
 	{
 		return false;
 	}
-	const float GravityZ = FMath::IsNearlyEqual(OverrideGravityZ, 0.0f) ? -World->GetGravityZ() : -OverrideGravityZ;
+	const float GravityZ = FMath::IsNearlyEqual(ProjectileParams.OverrideGravityZ, 0.0f) ? -World->GetGravityZ() : -ProjectileParams.OverrideGravityZ;
+
+	FVector PrioritizedProjVelocities[2];
+	int32 NumSolutions;
 
 	// v^4 - g*(g*x^2 + 2*y*v^2)
-	const float InsideTheSqrt = FMath::Square(TossSpeedSq) - GravityZ * ( (GravityZ * FMath::Square(DeltaXY)) + (2.f * DeltaZ * TossSpeedSq) );
+	const float InsideTheSqrt = FMath::Square(TossSpeedSq) - GravityZ * ((GravityZ * FMath::Square(DeltaXY)) + (2.f * DeltaZ * TossSpeedSq));
+
 	if (InsideTheSqrt < 0.f)
 	{
-		// sqrt will be imaginary, therefore no solutions
-		return false;
+		if (ProjectileParams.bAcceptClosestOnNoSolutions)
+		{
+			NumSolutions = 1;
+			// To get closest, we want to maximise the displacement in the direction of the target
+			// R = v^2/g * 1/cos^2(theta) * (sin(2phi + theta) - sin(theta))
+			// R = range, theta= angle to target (incline), phi = launch angle
+			// 2phi + theta maximises at PI/2 rads
+			// phi = PI/4 - theta/2    for max R
+			// alpha = phi + theta
+			// Create a vector at angle alpha up from DirXY and * by TossSpeed
+			const FVector DirFlightDelta = FlightDelta.GetSafeNormal();
+			const float AngleToTarget = FMath::Acos(FVector::DotProduct(DirFlightDelta, DirXY));
+			const float AngleAlpha = UE_PI / 4 - (AngleToTarget / 2) + AngleToTarget;
+			const FVector Cross = DirXY.Cross(FVector(0, 0, 1));
+			const FVector TrajectoryHeading = DirXY.RotateAngleAxisRad(AngleAlpha, Cross);
+			PrioritizedProjVelocities[0] = TrajectoryHeading * ProjectileParams.TossSpeed;
+		}
+		else
+		{
+			return false;
+		}
 	}
+	else
+	{
+		// if we got here, there are 2 solutions: one high-angle and one low-angle.
+		NumSolutions = 2;
 
-	// if we got here, there are 2 solutions: one high-angle and one low-angle.
+		const float SqrtPart = FMath::Sqrt(InsideTheSqrt);
 
-	const float SqrtPart = FMath::Sqrt(InsideTheSqrt);
+		// this is the tangent of the firing angle for the first (+) solution
+		const float TanSolutionAngleA = (TossSpeedSq + SqrtPart) / (GravityZ * DeltaXY);
+		// this is the tangent of the firing angle for the second (-) solution
+		const float TanSolutionAngleB = (TossSpeedSq - SqrtPart) / (GravityZ * DeltaXY);
 
-	// this is the tangent of the firing angle for the first (+) solution
-	const float TanSolutionAngleA = (TossSpeedSq + SqrtPart) / (GravityZ * DeltaXY);
-	// this is the tangent of the firing angle for the second (-) solution
-	const float TanSolutionAngleB = (TossSpeedSq - SqrtPart) / (GravityZ * DeltaXY);
-	
-	// mag in the XY dir = sqrt( TossSpeedSq / (TanSolutionAngle^2 + 1) );
-	const float MagXYSq_A = TossSpeedSq / (FMath::Square(TanSolutionAngleA) + 1.f);
-	const float MagXYSq_B = TossSpeedSq / (FMath::Square(TanSolutionAngleB) + 1.f);
+		// mag in the XY dir = sqrt( TossSpeedSq / (TanSolutionAngle^2 + 1) );
+		const float MagXYSq_A = TossSpeedSq / (FMath::Square(TanSolutionAngleA) + 1.f);
+		const float MagXYSq_B = TossSpeedSq / (FMath::Square(TanSolutionAngleB) + 1.f);
 
-	bool bFoundAValidSolution = false;
+		float PrioritizedSolutionsMagXYSq[2];
+		PrioritizedSolutionsMagXYSq[0] = ProjectileParams.bFavorHighArc ? FMath::Min(MagXYSq_A, MagXYSq_B) : FMath::Max(MagXYSq_A, MagXYSq_B);
+		PrioritizedSolutionsMagXYSq[1] = ProjectileParams.bFavorHighArc ? FMath::Max(MagXYSq_A, MagXYSq_B) : FMath::Min(MagXYSq_A, MagXYSq_B);
+
+		float PrioritizedSolutionZSign[2];
+		PrioritizedSolutionZSign[0] = ProjectileParams.bFavorHighArc ?
+			(MagXYSq_A < MagXYSq_B) ? FMath::Sign(TanSolutionAngleA) : FMath::Sign(TanSolutionAngleB) :
+			(MagXYSq_A > MagXYSq_B) ? FMath::Sign(TanSolutionAngleA) : FMath::Sign(TanSolutionAngleB);
+		PrioritizedSolutionZSign[1] = ProjectileParams.bFavorHighArc ?
+			(MagXYSq_A > MagXYSq_B) ? FMath::Sign(TanSolutionAngleA) : FMath::Sign(TanSolutionAngleB) :
+			(MagXYSq_A < MagXYSq_B) ? FMath::Sign(TanSolutionAngleA) : FMath::Sign(TanSolutionAngleB);
+
+		// solutions in priority order
+		for (int32 CurrentSolutionIdx = 0; (CurrentSolutionIdx < 2); ++CurrentSolutionIdx)
+		{
+			const float MagXY = FMath::Sqrt(PrioritizedSolutionsMagXYSq[CurrentSolutionIdx]);
+			const float MagZ = FMath::Sqrt(TossSpeedSq - PrioritizedSolutionsMagXYSq[CurrentSolutionIdx]);		// pythagorean
+			const float ZSign = PrioritizedSolutionZSign[CurrentSolutionIdx];
+
+			PrioritizedProjVelocities[CurrentSolutionIdx] = (DirXY * MagXY) + (FVector::UpVector * MagZ * ZSign);
+		}
+	}
 
 	// trace if desired
-	if (TraceOption == ESuggestProjVelocityTraceOption::DoNotTrace)
+	if (ProjectileParams.TraceOption == ESuggestProjVelocityTraceOption::DoNotTrace)
 	{
-		// choose which arc
-		const float FavoredMagXYSq = bFavorHighArc ? FMath::Min(MagXYSq_A, MagXYSq_B) : FMath::Max(MagXYSq_A, MagXYSq_B);
-		const float ZSign = bFavorHighArc ?
-							(MagXYSq_A < MagXYSq_B) ? FMath::Sign(TanSolutionAngleA) : FMath::Sign(TanSolutionAngleB) :
-							(MagXYSq_A > MagXYSq_B) ? FMath::Sign(TanSolutionAngleA) : FMath::Sign(TanSolutionAngleB);
-
-		// finish calculations
-		const float MagXY = FMath::Sqrt(FavoredMagXYSq);
-		const float MagZ = FMath::Sqrt(TossSpeedSq - FavoredMagXYSq);		// pythagorean
+		OutTossVelocity = PrioritizedProjVelocities[0];
+		const float MagXY = OutTossVelocity.Size2D();
+#if ENABLE_DRAW_DEBUG
+		if (ProjectileParams.bDrawDebug)
+		{
+			static const float StepSize = 0.125f;
+			FVector TraceStart = ProjectileParams.Start;
+			for (float Step = 0.f; Step < 1.f; Step += StepSize)
+			{
+				const float TimeInFlight = (Step + StepSize) * DeltaXY / MagXY;
 
-		// final answer!
-		OutTossVelocity = (DirXY * MagXY) + (FVector::UpVector * MagZ * ZSign);
-		bFoundAValidSolution = true;
+				// d = vt + .5 a t^2
+				const FVector TraceEnd = ProjectileParams.Start + OutTossVelocity * TimeInFlight + FVector(0.f, 0.f, 0.5f * -GravityZ * FMath::Square(TimeInFlight) - ProjectileParams.CollisionRadius);
 
-#if ENABLE_DRAW_DEBUG
-	 	if (bDrawDebug)
-	 	{
-	 		static const float StepSize = 0.125f;
-	 		FVector TraceStart = Start;
-	 		for ( float Step=0.f; Step<1.f; Step+=StepSize )
-	 		{
-	 			const float TimeInFlight = (Step+StepSize) * DeltaXY/MagXY;
-	 
-	 			// d = vt + .5 a t^2
-				const FVector TraceEnd = Start + OutTossVelocity*TimeInFlight + FVector(0.f, 0.f, 0.5f * -GravityZ * FMath::Square(TimeInFlight) - CollisionRadius);
-	 
-	 			DrawDebugLine( World, TraceStart, TraceEnd, (bFoundAValidSolution ? FColor::Yellow : FColor::Red), true );
-	 			TraceStart = TraceEnd;
-	 		}
-	 	}
+				DrawDebugLine(World, TraceStart, TraceEnd, FColor::Yellow, true);
+				TraceStart = TraceEnd;
+			}
+		}
 #endif // ENABLE_DRAW_DEBUG
+		return true;
 	}
 	else
 	{
 		// need to trace to validate
-
-		// sort potential solutions by priority
-		float PrioritizedSolutionsMagXYSq[2];
-		PrioritizedSolutionsMagXYSq[0] = bFavorHighArc ? FMath::Min(MagXYSq_A, MagXYSq_B) : FMath::Max(MagXYSq_A, MagXYSq_B);
-		PrioritizedSolutionsMagXYSq[1] = bFavorHighArc ? FMath::Max(MagXYSq_A, MagXYSq_B) : FMath::Min(MagXYSq_A, MagXYSq_B);
-
-		float PrioritizedSolutionZSign[2];
-		PrioritizedSolutionZSign[0] = bFavorHighArc ?
-										(MagXYSq_A < MagXYSq_B) ? FMath::Sign(TanSolutionAngleA) : FMath::Sign(TanSolutionAngleB) :
-										(MagXYSq_A > MagXYSq_B) ? FMath::Sign(TanSolutionAngleA) : FMath::Sign(TanSolutionAngleB);
-		PrioritizedSolutionZSign[1] = bFavorHighArc ?
-										(MagXYSq_A > MagXYSq_B) ? FMath::Sign(TanSolutionAngleA) : FMath::Sign(TanSolutionAngleB) :
-										(MagXYSq_A < MagXYSq_B) ? FMath::Sign(TanSolutionAngleA) : FMath::Sign(TanSolutionAngleB);
-
-		FVector PrioritizedProjVelocities[2];
-
-		// try solutions in priority order
 		int32 ValidSolutionIdx = INDEX_NONE;
-		for (int32 CurrentSolutionIdx=0; (CurrentSolutionIdx<2); ++CurrentSolutionIdx)
-		{
-			const float MagXY = FMath::Sqrt( PrioritizedSolutionsMagXYSq[CurrentSolutionIdx] );
-			const float MagZ = FMath::Sqrt( TossSpeedSq - PrioritizedSolutionsMagXYSq[CurrentSolutionIdx] );		// pythagorean
-			const float ZSign = PrioritizedSolutionZSign[CurrentSolutionIdx];
 
-			PrioritizedProjVelocities[CurrentSolutionIdx] = (DirXY * MagXY) + (FVector::UpVector * MagZ * ZSign);
-
-			// iterate along the arc, doing stepwise traces
-			bool bFailedTrace = false;
-			static const float StepSize = 0.125f;
-			FVector TraceStart = Start;
-			for ( float Step=0.f; Step<1.f; Step+=StepSize )
+		for (int32 SolutionIdx = 0; SolutionIdx < NumSolutions; ++SolutionIdx)
+		{
+			// trace each solution
+			if (!IsProjectileTrajectoryBlocked(World, ProjectileParams.Start, PrioritizedProjVelocities[SolutionIdx], DeltaXY, GravityZ, ProjectileParams.CollisionRadius, ProjectileParams.TraceOption, ProjectileParams.ResponseParam, ProjectileParams.ActorsToIgnore, ProjectileParams.bDrawDebug))
 			{
-				const float TimeInFlight = (Step+StepSize) * DeltaXY/MagXY;
+				OutTossVelocity = PrioritizedProjVelocities[SolutionIdx];
+				return true;
+			}
+		}
 
-				// d = vt + .5 a t^2
-				const FVector TraceEnd = Start + PrioritizedProjVelocities[CurrentSolutionIdx]*TimeInFlight + FVector(0.f, 0.f, 0.5f * -GravityZ * FMath::Square(TimeInFlight) - CollisionRadius);
+		return false;
+	}
+}
 
-				if ( (TraceOption == ESuggestProjVelocityTraceOption::OnlyTraceWhileAscending) && (TraceEnd.Z < TraceStart.Z) )
-				{
-					// falling, we are done tracing
-					if (!bDrawDebug)
-					{
-						// if we're drawing, we continue stepping without the traces
-						// else we can just trivially end the iteration loop
-						break;
-					}
-				}
-				else
-				{
-					FCollisionQueryParams QueryParams(SCENE_QUERY_STAT(SuggestProjVelTrace), true);
-					QueryParams.AddIgnoredActors(ActorsToIgnore);
-					if (World->SweepTestByChannel(TraceStart, TraceEnd, FQuat::Identity, ECC_WorldDynamic, FCollisionShape::MakeSphere(CollisionRadius), QueryParams, ResponseParam))
-					{
-						// hit something, failed
-						bFailedTrace = true;
+bool UGameplayStatics::IsProjectileTrajectoryBlocked(const UWorld* World, FVector StartLocation, FVector& ProjectileVelocity, float TargetDeltaXY, float GravityZ, float CollisionRadius, ESuggestProjVelocityTraceOption::Type TraceOption, const FCollisionResponseParams& ResponseParam, const TArray<AActor*>& ActorsToIgnore, bool bDrawDebug)
+{
+	
+	// iterate along the arc, doing stepwise traces
+	static const float StepSize = 0.125f;
+	const float MagXY = ProjectileVelocity.Size2D();
+	FVector TraceStart = StartLocation;
 
-#if ENABLE_DRAW_DEBUG
-						if (bDrawDebug)
-						{
-							// draw failed segment in red
-							DrawDebugLine( World, TraceStart, TraceEnd, FColor::Red, true );
-						}
-#endif // ENABLE_DRAW_DEBUG
+	for (float Step = 0.f; Step < 1.f; Step += StepSize)
+	{
+		const float TimeInFlight = (Step + StepSize) * TargetDeltaXY / MagXY;
 
-						break;
-					}
+		// d = vt + .5 a t^2
+		const FVector TraceEnd = StartLocation + ProjectileVelocity * TimeInFlight + FVector(0.f, 0.f, 0.5f * -GravityZ * FMath::Square(TimeInFlight) - CollisionRadius);
+
+		if ((TraceOption == ESuggestProjVelocityTraceOption::OnlyTraceWhileAscending) && (TraceEnd.Z < TraceStart.Z))
+		{
+			// falling, we are done tracing
+			if (!bDrawDebug)
+			{
+				// if we're drawing, we continue stepping without the traces
+				// else we can just trivially end the iteration loop
+				break;
+			}
+		}
+		else
+		{
+			FCollisionQueryParams QueryParams(SCENE_QUERY_STAT(SuggestProjVelTrace), true);
+			QueryParams.AddIgnoredActors(ActorsToIgnore);
+			if (World->SweepTestByChannel(TraceStart, TraceEnd, FQuat::Identity, ECC_WorldDynamic, FCollisionShape::MakeSphere(CollisionRadius), QueryParams, ResponseParam))
+			{
 
-				}
 
 #if ENABLE_DRAW_DEBUG
 				if (bDrawDebug)
 				{
-					DrawDebugLine( World, TraceStart, TraceEnd, FColor::Yellow, true );
+					// draw failed segment in red
+					DrawDebugLine(World, TraceStart, TraceEnd, FColor::Red, true);
 				}
 #endif // ENABLE_DRAW_DEBUG
-
-				// advance
-				TraceStart = TraceEnd;
-			}
-
-			if (bFailedTrace == false)
-			{
-				// passes all traces along the arc, we have a valid solution and can be done
-				ValidSolutionIdx = CurrentSolutionIdx;
-				break;
+				// hit something, failed
+				return true;
 			}
 		}
 
-		if (ValidSolutionIdx != INDEX_NONE)
+#if ENABLE_DRAW_DEBUG
+		if (bDrawDebug)
 		{
-			OutTossVelocity = PrioritizedProjVelocities[ValidSolutionIdx];
-			bFoundAValidSolution = true;
+			DrawDebugLine(World, TraceStart, TraceEnd, FColor::Yellow, true);
 		}
+#endif // ENABLE_DRAW_DEBUG
+
+		// advance
+		TraceStart = TraceEnd;
 	}
 
-	return bFoundAValidSolution;
+	return false;
 }
 
 // note: this will automatically fall back to line test if radius is small enough
