Fix conflicting localization

#rnx

diff --git a/Engine/Source/Runtime/Engine/Private/LevelInstance/LevelInstanceSubsystem.cpp b/Engine/Source/Runtime/Engine/Private/LevelInstance/LevelInstanceSubsystem.cpp
--- a/Engine/Source/Runtime/Engine/Private/LevelInstance/LevelInstanceSubsystem.cpp
+++ b/Engine/Source/Runtime/Engine/Private/LevelInstance/LevelInstanceSubsystem.cpp
@@ -917,11 +917,11 @@ FText ULevelInstanceSubsystem::GetToolKitDisplayText() const
 {
 	if (ILevelInstanceInterface* PropertyOverrideInstance = GetEditingPropertyOverridesLevelInstance())
 	{
-		return LOCTEXT("ExitTooltip", "Exit Property Override Edit");
+		return LOCTEXT("ExitEditTooltip", "Exit Property Override Edit");
 	}
 	else if (ILevelInstanceInterface* EditingLevelInstance = GetEditingLevelInstance())
 	{
-		return LOCTEXT("ExitTooltip", "Exit Level Instance Edit");
+		return LOCTEXT("ExitOverrideTooltip", "Exit Level Instance Edit");
 	}
 
 	return FText::GetEmpty();
@@ -2219,7 +2219,7 @@ bool ULevelInstanceSubsystem::CanEditLevelInstanceCommon(const ILevelInstanceInt
 	{
 		if (OutReason)
 		{
-			*OutReason = FText::Format(LOCTEXT("CanEditLevelInstanceAlreadyBeingEdited", "Level Instance already in property override edit ({0})."), FText::FromString(LevelInstance->GetWorldAssetPackage()));
+			*OutReason = FText::Format(LOCTEXT("CanEditLevelInstanceAlreadyBeingOverriden", "Level Instance already in property override edit ({0})."), FText::FromString(LevelInstance->GetWorldAssetPackage()));
 		}
 		return false;
 	}
@@ -2292,7 +2292,7 @@ bool ULevelInstanceSubsystem::CanEditLevelInstancePropertyOverrides(const ILevel
 	{
 		if (OutReason)
 		{
-			*OutReason = LOCTEXT("LevelInstanceNoSupportPropertyOverrides", "Level Instance property override feature is not enabled");
+			*OutReason = LOCTEXT("LevelInstanceNotEnabledPropertyOverrides", "Level Instance property override feature is not enabled");
 		}
 		return false;
 	}
@@ -2364,7 +2364,7 @@ bool ULevelInstanceSubsystem::CanCommitLevelInstancePropertyOverrides(const ILev
 
 	if (OutReason)
 	{
-		*OutReason = LOCTEXT("CanCommitLevelInstanceNotEditing", "Level Instance is not currently in property override edit");
+		*OutReason = LOCTEXT("CanCommitLevelInstanceNotOverriding", "Level Instance is not currently in property override edit");
 	}
 	return false;
 }
