Adding "Bindless Switch" material expression that chooses an input based on whether bindless is being used or not.

diff --git a/Engine/Source/Runtime/Engine/Classes/Materials/MaterialExpressionBindlessSwitch.h b/Engine/Source/Runtime/Engine/Classes/Materials/MaterialExpressionBindlessSwitch.h
new file mode 100644
--- /dev/null
+++ b/Engine/Source/Runtime/Engine/Classes/Materials/MaterialExpressionBindlessSwitch.h
@@ -0,0 +1,28 @@
+// Copyright Epic Games, Inc. All Rights Reserved.
+
+#pragma once
+
+#include "Materials/MaterialExpression.h"
+#include "MaterialExpressionBindlessSwitch.generated.h"
+
+UCLASS(MinimalAPI)
+class UMaterialExpressionBindlessSwitch : public UMaterialExpression
+{
+	GENERATED_UCLASS_BODY()
+
+	UPROPERTY(meta = (RequiredInput = "true", ToolTip = "Input will be used when rendering without bindless support"))
+	FExpressionInput Default;
+
+	UPROPERTY(meta = (RequiredInput = "true", ToolTip = "Input will be used when rendering with bindless support"))
+	FExpressionInput Bindless;
+
+	//~ Begin UMaterialExpression Interface
+#if WITH_EDITOR
+	virtual int32 Compile(FMaterialCompiler* Compiler, int32 OutputIndex) override;
+	virtual void GetCaption(TArray<FString>& OutCaptions) const override;
+	virtual bool IsInputConnectionRequired(int32 InputIndex) const override;
+	virtual bool IsResultMaterialAttributes(int32 OutputIndex) override;
+	virtual void GetExpressionToolTip(TArray<FString>& OutToolTip) override;
+#endif
+	//~ End UMaterialExpression Interface
+};
diff --git a/Engine/Source/Runtime/Engine/Private/Materials/MaterialExpressions.cpp b/Engine/Source/Runtime/Engine/Private/Materials/MaterialExpressions.cpp
--- a/Engine/Source/Runtime/Engine/Private/Materials/MaterialExpressions.cpp
+++ b/Engine/Source/Runtime/Engine/Private/Materials/MaterialExpressions.cpp
@@ -59,6 +59,7 @@
 #include "Materials/MaterialExpressionArctangent2Fast.h"
 #include "Materials/MaterialExpressionAtmosphericFogColor.h"
 #include "Materials/MaterialExpressionBentNormalCustomOutput.h"
+#include "Materials/MaterialExpressionBindlessSwitch.h"
 #include "Materials/MaterialExpressionBlackBody.h"
 #include "Materials/MaterialExpressionBlendMaterialAttributes.h"
 #include "Materials/MaterialExpressionBreakMaterialAttributes.h"
@@ -10027,6 +10028,75 @@ void UMaterialExpressionDataDrivenShaderPlatformInfoSwitch::PostEditChangeProper
 }
 #endif // WITH_EDITOR
 
+//
+//	UMaterialExpressionBindlessSwitch
+//
+
+UMaterialExpressionBindlessSwitch::UMaterialExpressionBindlessSwitch(const FObjectInitializer& ObjectInitializer)
+	: Super(ObjectInitializer)
+{
+}
+
+#if WITH_EDITOR
+static bool IsBindlessEnabledForCompiler(FMaterialCompiler* Compiler)
+{
+	const EShaderPlatform ShaderPlatform = Compiler->GetShaderPlatform();
+	const ERHIBindlessConfiguration BindlessConfiguration = UE::ShaderCompiler::GetBindlessResourcesConfiguration(LegacyShaderPlatformToShaderFormat(ShaderPlatform));
+
+	if (BindlessConfiguration == ERHIBindlessConfiguration::AllShaders)
+	{
+		return true;
+	}
+
+	if (BindlessConfiguration == ERHIBindlessConfiguration::RayTracingShaders && IsRayTracingShaderFrequency(Compiler->GetCurrentShaderFrequency()))
+	{
+		return true;
+	}
+
+	return false;
+}
+
+int32 UMaterialExpressionBindlessSwitch::Compile(FMaterialCompiler* Compiler, int32 OutputIndex)
+{
+	if (IsBindlessEnabledForCompiler(Compiler))
+	{
+		return Bindless.Compile(Compiler);
+	}
+
+	return Default.Compile(Compiler);
+}
+
+void UMaterialExpressionBindlessSwitch::GetCaption(TArray<FString>& OutCaptions) const
+{
+	OutCaptions.Add(TEXT("Bindless Switch"));
+}
+
+bool UMaterialExpressionBindlessSwitch::IsInputConnectionRequired(int32 InputIndex) const
+{
+	return true;
+}
+
+bool UMaterialExpressionBindlessSwitch::IsResultMaterialAttributes(int32 OutputIndex)
+{
+	check(OutputIndex == 0);
+	for (FExpressionInput* ExpressionInput : GetInputsView())
+	{
+		// If there is a loop anywhere in this expression's inputs then we can't risk checking them
+		if (ExpressionInput->Expression && ExpressionInput->Expression->IsResultMaterialAttributes(ExpressionInput->OutputIndex))
+		{
+			return true;
+		}
+	}
+
+	return false;
+}
+
+void UMaterialExpressionBindlessSwitch::GetExpressionToolTip(TArray<FString>& OutToolTip)
+{
+	ConvertToMultilineToolTip(TEXT("Allows material to define specialized behavior when being rendered with bindless enabled."), 40, OutToolTip);
+}
+#endif
+
 //
 //	UMaterialExpressionRequiredSamplersSwitch
 //
