[EntityFramework] added ai_navigation_relevant_component that handles static_mesh_component affecting navigatio data
#okforversepublic

diff --git a/Engine/Source/Runtime/Engine/Classes/Engine/StaticMesh.h b/Engine/Source/Runtime/Engine/Classes/Engine/StaticMesh.h
--- a/Engine/Source/Runtime/Engine/Classes/Engine/StaticMesh.h
+++ b/Engine/Source/Runtime/Engine/Classes/Engine/StaticMesh.h
@@ -1561,6 +1561,7 @@ private:
 public:
 	ENGINE_API void SetNavCollision(UNavCollisionBase*);
 	ENGINE_API UNavCollisionBase* GetNavCollision() const;
+	ENGINE_API FBox GetNavigationBounds(const FTransform& LocalToWorld) const;
 	ENGINE_API bool IsNavigationRelevant() const;
 
 	/**
diff --git a/Engine/Source/Runtime/Engine/Private/Components/StaticMeshComponent.cpp b/Engine/Source/Runtime/Engine/Private/Components/StaticMeshComponent.cpp
--- a/Engine/Source/Runtime/Engine/Private/Components/StaticMeshComponent.cpp
+++ b/Engine/Source/Runtime/Engine/Private/Components/StaticMeshComponent.cpp
@@ -3166,15 +3166,7 @@ FBox UStaticMeshComponent::GetNavigationBounds() const
 	{
 		if (ensureMsgf(!Mesh->IsCompiling(), TEXT("%s is not considered relevant to navigation until associated mesh is compiled."), *GetFullName()))
 		{
-			if (const UNavCollisionBase* NavCollision = Mesh->GetNavCollision())
-			{
-				FBox NavBounds = NavCollision->GetBounds();
-				if (!NavBounds.IsValid)
-				{
-					NavBounds = Mesh->GetBounds().GetBox();
-				}
-				return NavBounds.TransformBy(GetComponentTransform());
-			}
+			return Mesh->GetNavigationBounds(GetComponentTransform());
 		}
 	}
 
diff --git a/Engine/Source/Runtime/Engine/Private/StaticMesh.cpp b/Engine/Source/Runtime/Engine/Private/StaticMesh.cpp
--- a/Engine/Source/Runtime/Engine/Private/StaticMesh.cpp
+++ b/Engine/Source/Runtime/Engine/Private/StaticMesh.cpp
@@ -7756,6 +7756,20 @@ UNavCollisionBase* UStaticMesh::GetNavCollision() const
 	PRAGMA_ENABLE_DEPRECATION_WARNINGS
 }
 
+FBox UStaticMesh::GetNavigationBounds(const FTransform& LocalToWorld) const
+{
+	FBox NavBounds = GetBounds().GetBox();
+	if (const UNavCollisionBase* NavCol = GetNavCollision())
+	{
+		const FBox NavCollisionBounds = NavCol->GetBounds();
+		if (NavCollisionBounds.IsValid)
+		{
+			NavBounds = NavCollisionBounds;
+		}
+	}
+	return NavBounds.TransformBy(LocalToWorld);
+}
+
 void UStaticMesh::MarkAsNotHavingNavigationData()
 {
 	bHasNavigationData = false;
