Fix function name typo to avoid deprecation in future releases.

diff --git a/Engine/Source/Runtime/Engine/Public/WorldPartition/WorldPartitionRuntimeCell.h b/Engine/Source/Runtime/Engine/Public/WorldPartition/WorldPartitionRuntimeCell.h
--- a/Engine/Source/Runtime/Engine/Public/WorldPartition/WorldPartitionRuntimeCell.h
+++ b/Engine/Source/Runtime/Engine/Public/WorldPartition/WorldPartitionRuntimeCell.h
@@ -420,7 +420,7 @@ public:
 	inline FBox GetCellBounds() const { return RuntimeCellData->GetCellBounds(); }
 
 	/** Returns the cell's streaming bounds, which is what the underlying runtime hash uses to intersect cells.  */
-	inline FBox GetstreamingBounds() const { return RuntimeCellData->GetStreamingBounds(); }
+	inline FBox GetStreamingBounds() const { return RuntimeCellData->GetStreamingBounds(); }
 
 	ENGINE_API virtual bool IsDebugShown() const;
 	//~End UWorldPartitionRuntimeCellData Proxy
