Improve help text for functions moving actors etc, giving more info about the behaviour when teleport is/isn't used.

diff --git a/Engine/Source/Runtime/Engine/Classes/Components/SkeletalMeshComponent.h b/Engine/Source/Runtime/Engine/Classes/Components/SkeletalMeshComponent.h
--- a/Engine/Source/Runtime/Engine/Classes/Components/SkeletalMeshComponent.h
+++ b/Engine/Source/Runtime/Engine/Classes/Components/SkeletalMeshComponent.h
@@ -2193,7 +2193,7 @@ public:
 
 	/** 
 	 *	Iterate over each physics body in the physics for this mesh, and for each 'kinematic' (ie fixed or default if owner isn't simulating) one, update its
-	 *	transform based on the animated transform.
+	 *	transform based on the animated transform. This update is also done for simulating bodies if a teleport is being requested. 
 	 *	@param	InComponentSpaceTransforms	Array of bone transforms in component space
 	 *	@param	Teleport					Whether movement is a 'teleport' (ie infers no physics velocity, but moves simulating bodies) or not
 	 *	@param	bNeedsSkinning				Whether we may need  to send new triangle data for per-poly skeletal mesh collision
diff --git a/Engine/Source/Runtime/Engine/Classes/GameFramework/Actor.h b/Engine/Source/Runtime/Engine/Classes/GameFramework/Actor.h
--- a/Engine/Source/Runtime/Engine/Classes/GameFramework/Actor.h
+++ b/Engine/Source/Runtime/Engine/Classes/GameFramework/Actor.h
@@ -1425,6 +1425,8 @@ public:
 	 *						If true, physics velocity for this object is unchanged (so ragdoll parts are not affected by change in location).
 	 *						If false, physics velocity is updated based on the change in position (affecting ragdoll parts).
 	 *						If CCD is on and not teleporting, this will affect objects along the entire swept volume.
+	 *                      Note that when teleporting, any child/attached components will be teleported too, maintaining their current offset even if they are being simulated. 
+	 *                      Setting the location without teleporting will not update the location of simulated child/attached components.
 	 * @param SweepHitResult	The hit result from the move if swept.
 	 * @return	Whether the location was successfully set (if not swept), or whether movement occurred at all (if swept).
 	 */
@@ -1475,6 +1477,8 @@ public:
 	 *						If equal to ETeleportType::TeleportPhysics, physics velocity for this object is unchanged (so ragdoll parts are not affected by change in location).
 	 *						If equal to ETeleportType::None, physics velocity is updated based on the change in position (affecting ragdoll parts).
 	 *						If CCD is on and not teleporting, this will affect objects along the entire swept volume.
+	 *                      Note that when teleporting, any child/attached components will be teleported too, maintaining their current offset even if they are being simulated. 
+	 *                      Setting the location without teleporting will not update the location of simulated child/attached components.
 	 * @param OutSweepHitResult The hit result from the move if swept.
 	 * @return	Whether the location was successfully set if not swept, or whether movement occurred if swept.
 	 */
@@ -1487,6 +1491,8 @@ public:
 	 * @param	bTeleportPhysics Whether we teleport the physics state (if physics collision is enabled for this object).
 	 *			If true, physics velocity for this object is unchanged (so ragdoll parts are not affected by change in location).
 	 *			If false, physics velocity is updated based on the change in position (affecting ragdoll parts).
+	 *          Note that when teleporting, any child/attached components will be teleported too, maintaining their current offset even if they are being simulated. 
+	 *          Setting the rotation without teleporting will not update the rotation of simulated child/attached components.
 	 * @return	Whether the rotation was successfully set.
 	 */
 	UFUNCTION(BlueprintCallable, meta=(DisplayName = "Set Actor Rotation", ScriptName = "SetActorRotation"), Category="Transformation")
@@ -1499,6 +1505,8 @@ public:
 	 * @param	Teleport	How we teleport the physics state (if physics collision is enabled for this object).
 	 *						If equal to ETeleportType::TeleportPhysics, physics velocity for this object is unchanged (so ragdoll parts are not affected by change in location).
 	 *						If equal to ETeleportType::None, physics velocity is updated based on the change in position (affecting ragdoll parts).
+	 *                      Note that when teleporting, any child/attached components will be teleported too, maintaining their current offset even if they are being simulated.
+	 *                      Setting the rotation without teleporting will not update the rotation of simulated child/attached components.
 	 * @return	Whether the rotation was successfully set.
 	 */
 	ENGINE_API bool SetActorRotation(FRotator NewRotation, ETeleportType Teleport = ETeleportType::None);
@@ -1515,6 +1523,8 @@ public:
 	 *							If true, physics velocity for this object is unchanged (so ragdoll parts are not affected by change in location).
 	 *							If false, physics velocity is updated based on the change in position (affecting ragdoll parts).
 	 *							If CCD is on and not teleporting, this will affect objects along the entire swept volume.
+	 *                          Note that when teleporting, any child/attached components will be teleported too, maintaining their current offset even if they are being simulated. 
+	 *                          Setting the location without teleporting will not update the location of simulated child/attached components.
 	 * @param SweepHitResult	The hit result from the move if swept.
 	 * @return	Whether the rotation was successfully set.
 	 */
@@ -1532,6 +1542,8 @@ public:
 	 *							If equal to ETeleportType::TeleportPhysics, physics velocity for this object is unchanged (so ragdoll parts are not affected by change in location).
 	 *							If equal to ETeleportType::None, physics velocity is updated based on the change in position (affecting ragdoll parts).
 	 *							If CCD is on and not teleporting, this will affect objects along the entire swept volume.
+	 *                          Note that when teleporting, any child/attached components will be teleported too, maintaining their current offset even if they are being simulated.
+	 *                          Setting the location without teleporting will not update the location of simulated child/attached components.
 	 * @param OutSweepHitResult	The hit result from the move if swept.
 	 * @return	Whether the rotation was successfully set.
 	 */
@@ -1584,6 +1596,8 @@ public:
 	 *							If true, physics velocity for this object is unchanged (so ragdoll parts are not affected by change in location).
 	 *							If false, physics velocity is updated based on the change in position (affecting ragdoll parts).
 	 *							If CCD is on and not teleporting, this will affect objects along the entire swept volume.
+	 *                          Note that when teleporting, any child/attached components will be teleported too, maintaining their current offset even if they are being simulated.
+	 *                          Setting the location without teleporting will not update the location of simulated child/attached components.
 	 * @param SweepHitResult	The hit result from the move if swept.
 	 */
 	UFUNCTION(BlueprintCallable, Category="Transformation", meta=(DisplayName="Add Actor World Offset", ScriptName="AddActorWorldOffset", Keywords="location position"))
@@ -1599,6 +1613,8 @@ public:
 	 *							If true, physics velocity for this object is unchanged (so ragdoll parts are not affected by change in location).
 	 *							If false, physics velocity is updated based on the change in position (affecting ragdoll parts).
 	 *							If CCD is on and not teleporting, this will affect objects along the entire swept volume.
+	 *                          Note that when teleporting, any child/attached components will be teleported too, maintaining their current offset even if they are being simulated.
+	 *                          Setting the rotation without teleporting will not update the rotation of simulated child/attached components.
 	 * @param SweepHitResult	The hit result from the move if swept.
 	 */
 	UFUNCTION(BlueprintCallable, Category="Transformation", meta=(DisplayName="Add Actor World Rotation", ScriptName="AddActorWorldRotation", AdvancedDisplay="bSweep,SweepHitResult,bTeleport"))
@@ -1625,6 +1641,8 @@ public:
 	 *							If true, physics velocity for this object is unchanged (so ragdoll parts are not affected by change in location).
 	 *							If false, physics velocity is updated based on the change in position (affecting ragdoll parts).
 	 *							If CCD is on and not teleporting, this will affect objects along the entire swept volume.
+	 *                          Note that when teleporting, any child/attached components will be teleported too, maintaining their current offset even if they are being simulated.
+	 *                          Setting the transform without teleporting will not update the transform of simulated child/attached components.
 	 */
 	UFUNCTION(BlueprintCallable, Category="Transformation", meta=(DisplayName="Set Actor Transform", ScriptName="SetActorTransform"))
 	ENGINE_API bool K2_SetActorTransform(const FTransform& NewTransform, bool bSweep, FHitResult& SweepHitResult, bool bTeleport);
@@ -1639,6 +1657,8 @@ public:
 	 *							If true, physics velocity for this object is unchanged (so ragdoll parts are not affected by change in location).
 	 *							If false, physics velocity is updated based on the change in position (affecting ragdoll parts).
 	 *							If CCD is on and not teleporting, this will affect objects along the entire swept volume.
+	 *                          Note that when teleporting, any child/attached components will be teleported too, maintaining their current offset even if they are being simulated.
+	 *                          Setting the location without teleporting will not update the location of simulated child/attached components.
 	 */
 	UFUNCTION(BlueprintCallable, Category="Transformation", meta=(DisplayName="Add Actor Local Offset", ScriptName="AddActorLocalOffset", Keywords="location position"))
 	ENGINE_API void K2_AddActorLocalOffset(FVector DeltaLocation, bool bSweep, FHitResult& SweepHitResult, bool bTeleport);
@@ -1653,6 +1673,8 @@ public:
 	 *							If true, physics velocity for this object is unchanged (so ragdoll parts are not affected by change in location).
 	 *							If false, physics velocity is updated based on the change in position (affecting ragdoll parts).
 	 *							If CCD is on and not teleporting, this will affect objects along the entire swept volume.
+	 *                          Note that when teleporting, any child/attached components will be teleported too, maintaining their current offset even if they are being simulated.
+	 *                          Setting the rotation without teleporting will not update the rotation of simulated child/attached components.
 	 */
 	UFUNCTION(BlueprintCallable, Category="Transformation", meta=(DisplayName="Add Actor Local Rotation", ScriptName="AddActorLocalRotation", AdvancedDisplay="bSweep,SweepHitResult,bTeleport"))
 	ENGINE_API void K2_AddActorLocalRotation(FRotator DeltaRotation, bool bSweep, FHitResult& SweepHitResult, bool bTeleport);
@@ -1668,6 +1690,8 @@ public:
 	 *							If true, physics velocity for this object is unchanged (so ragdoll parts are not affected by change in location).
 	 *							If false, physics velocity is updated based on the change in position (affecting ragdoll parts).
 	 *							If CCD is on and not teleporting, this will affect objects along the entire swept volume.
+	 *                          Note that when teleporting, any child/attached components will be teleported too, maintaining their current offset even if they are being simulated.
+	 *                          Setting the transform without teleporting will not update the transform of simulated child/attached components.
 	 */
 	UFUNCTION(BlueprintCallable, Category="Transformation", meta=(DisplayName="Add Actor Local Transform", ScriptName="AddActorLocalTransform"))
 	ENGINE_API void K2_AddActorLocalTransform(const FTransform& NewTransform, bool bSweep, FHitResult& SweepHitResult, bool bTeleport);
@@ -1682,6 +1706,8 @@ public:
 	 *								If true, physics velocity for this object is unchanged (so ragdoll parts are not affected by change in location).
 	 *								If false, physics velocity is updated based on the change in position (affecting ragdoll parts).
 	 *								If CCD is on and not teleporting, this will affect objects along the entire swept volume.
+	 *                              Note that when teleporting, any child/attached components will be teleported too, maintaining their current offset even if they are being simulated. 
+	 *                              Setting the location without teleporting will not update the location of simulated child/attached components.
 	 */
 	UFUNCTION(BlueprintCallable, Category="Transformation", meta=(DisplayName="Set Actor Relative Location", ScriptName="SetActorRelativeLocation"))
 	ENGINE_API void K2_SetActorRelativeLocation(FVector NewRelativeLocation, bool bSweep, FHitResult& SweepHitResult, bool bTeleport);
@@ -1696,6 +1722,8 @@ public:
 	 *								If true, physics velocity for this object is unchanged (so ragdoll parts are not affected by change in location).
 	 *								If false, physics velocity is updated based on the change in position (affecting ragdoll parts).
 	 *								If CCD is on and not teleporting, this will affect objects along the entire swept volume.
+	 *                              Note that when teleporting, any child/attached components will be teleported too, maintaining their current offset even if they are being simulated.
+	 *                              Setting the rotation without teleporting will not update the rotation of simulated child/attached components.
 	 */
 	UFUNCTION(BlueprintCallable, Category="Transformation", meta=(DisplayName="Set Actor Relative Rotation", ScriptName="SetActorRelativeRotation", AdvancedDisplay="bSweep,SweepHitResult,bTeleport"))
 	ENGINE_API void K2_SetActorRelativeRotation(FRotator NewRelativeRotation, bool bSweep, FHitResult& SweepHitResult, bool bTeleport);
@@ -1711,6 +1739,8 @@ public:
 	 *							If true, physics velocity for this object is unchanged (so ragdoll parts are not affected by change in location).
 	 *							If false, physics velocity is updated based on the change in position (affecting ragdoll parts).
 	 *							If CCD is on and not teleporting, this will affect objects along the entire swept volume.
+	 *                          Note that when teleporting, any child/attached components will be teleported too, maintaining their current offset even if they are being simulated. 
+	 *                          Setting the transform without teleporting will not update the transform of simulated child/attached components.
 	 */
 	UFUNCTION(BlueprintCallable, Category="Transformation", meta=(DisplayName="Set Actor Relative Transform", ScriptName="SetActorRelativeTransform"))
 	ENGINE_API void K2_SetActorRelativeTransform(const FTransform& NewRelativeTransform, bool bSweep, FHitResult& SweepHitResult, bool bTeleport);
