Changed streaming asset task back to foreground priority to fix hitching on certain platforms.

diff --git a/Engine/Source/Runtime/Engine/Private/Streaming/StreamingManagerTexture.cpp b/Engine/Source/Runtime/Engine/Private/Streaming/StreamingManagerTexture.cpp
--- a/Engine/Source/Runtime/Engine/Private/Streaming/StreamingManagerTexture.cpp
+++ b/Engine/Source/Runtime/Engine/Private/Streaming/StreamingManagerTexture.cpp
@@ -1731,7 +1731,7 @@ public:
 	}
 	static FORCEINLINE ENamedThreads::Type GetDesiredThread()
 	{
-		return ENamedThreads::AnyBackgroundThreadNormalTask;
+		return ENamedThreads::AnyNormalThreadNormalTask;
 	}
 	static FORCEINLINE ESubsequentsMode::Type GetSubsequentsMode()
 	{
