Initial functionality of GameFeatureDataTableSubsystem

diff --git a/Engine/Source/Runtime/Engine/Classes/Engine/CompositeDataTable.h b/Engine/Source/Runtime/Engine/Classes/Engine/CompositeDataTable.h
--- a/Engine/Source/Runtime/Engine/Classes/Engine/CompositeDataTable.h
+++ b/Engine/Source/Runtime/Engine/Classes/Engine/CompositeDataTable.h
@@ -56,6 +56,9 @@ class UCompositeDataTable
 	// Support for runtime modification of parent tables
 	// Be aware this can be slow and can cause hitches during gameplay
 	ENGINE_API void AppendParentTables(const TArray<UDataTable*>& NewTables);
+	ENGINE_API void RemoveParentTables(const TArray<UDataTable*>& TablesToRemove);
+	ENGINE_API void AddParentTable(const TObjectPtr<UDataTable>& TableToAdd);
+	ENGINE_API void RemoveParentTable(const TObjectPtr<UDataTable>& TableToRemove);
 
 protected:
 
diff --git a/Engine/Source/Runtime/Engine/Private/CompositeDataTable.cpp b/Engine/Source/Runtime/Engine/Private/CompositeDataTable.cpp
--- a/Engine/Source/Runtime/Engine/Private/CompositeDataTable.cpp
+++ b/Engine/Source/Runtime/Engine/Private/CompositeDataTable.cpp
@@ -279,6 +279,39 @@ void UCompositeDataTable::AppendParentTables(const TArray<UDataTable*>& NewTable
 	OnParentTablesUpdated(EPropertyChangeType::ValueSet);
 }
 
+void UCompositeDataTable::RemoveParentTables(const TArray<UDataTable*>& TablesToRemove)
+{
+	for(UDataTable* DataTable : TablesToRemove)
+	{
+		if(DataTable)
+		{
+			ParentTables.Remove(DataTable);
+		}
+	}
+	OnParentTablesUpdated(EPropertyChangeType::ValueSet);
+}
+
+void UCompositeDataTable::AddParentTable(const TObjectPtr<UDataTable>& TableToAdd)
+{
+	if(TableToAdd == nullptr)
+	{
+		return;
+	}
+	ParentTables.Add(TableToAdd);
+	OnParentTablesUpdated(EPropertyChangeType::ValueSet);
+}
+
+void UCompositeDataTable::RemoveParentTable(const TObjectPtr<UDataTable>& TableToRemove)
+{
+	if(TableToRemove == nullptr)
+	{
+		return;
+	}
+	ParentTables.Remove(TableToRemove);
+	OnParentTablesUpdated(EPropertyChangeType::ValueSet);
+}
+
+
 void UCompositeDataTable::OnParentTablesUpdated(EPropertyChangeType::Type ChangeType)
 {
 	// Prevent recursion when there was a cycle in the parent hierarchy (or during the undo of the action that created the cycle; in that case PostEditUndo will recall OnParentTablesUpdated when the dust has settled)
