Fix RHICmdList not being submitted when skeletal mesh stream in request is cancelled.

diff --git a/Engine/Source/Runtime/Engine/Private/Streaming/SkeletalMeshUpdate.cpp b/Engine/Source/Runtime/Engine/Private/Streaming/SkeletalMeshUpdate.cpp
--- a/Engine/Source/Runtime/Engine/Private/Streaming/SkeletalMeshUpdate.cpp
+++ b/Engine/Source/Runtime/Engine/Private/Streaming/SkeletalMeshUpdate.cpp
@@ -554,7 +554,13 @@ void FSkeletalMeshStreamIn_IO::DoCreateBuffers(const FContext& Context)
 	CreateBuffers(Context);
 
 	check(!TaskSynchronization.GetValue());
-	PushTask(Context, TT_Render, SRA_UPDATE_CALLBACK(DoFinishUpdate), TT_None, nullptr);
+
+	// We cannot cancel once DoCreateBuffers has started executing, as there's an RHICmdList that must be submitted.
+	// Pass the same callback for both task and cancel.
+	PushTask(Context
+		, TT_Render, SRA_UPDATE_CALLBACK(DoFinishUpdate)
+		, TT_Render, SRA_UPDATE_CALLBACK(DoFinishUpdate)
+	);
 }
 
 void FSkeletalMeshStreamIn_IO::DoCancelIO(const FContext& Context)
