Allow Editor Data Layers to be created under External Data Layers.

diff --git a/Engine/Source/Editor/DataLayerEditor/Private/DataLayer/DataLayerMode.cpp b/Engine/Source/Editor/DataLayerEditor/Private/DataLayer/DataLayerMode.cpp
--- a/Engine/Source/Editor/DataLayerEditor/Private/DataLayer/DataLayerMode.cpp
+++ b/Engine/Source/Editor/DataLayerEditor/Private/DataLayer/DataLayerMode.cpp
@@ -871,7 +871,7 @@ FSceneOutlinerDragValidationInfo FDataLayerMode::ValidateDataLayerAssetDrop(cons
 		for (const UDataLayerAsset* DataLayerAssetToDrop : DataLayerAssetsToDrop)
 		{
 			FText Reason;
-			if (!FDataLayerUtils::AreDataLayerTypesCompatible(ParentType, DataLayerAssetToDrop->GetType(), &Reason))
+			if (!FDataLayerUtils::AreDataLayerTypesCompatible(ParentType, DataLayerAssetToDrop->GetType(), DropTargetDataLayerWithAsset->IsA<UExternalDataLayerInstance>(), &Reason))
 			{
 				return FSceneOutlinerDragValidationInfo(ESceneOutlinerDropCompatibility::Incompatible, FText::Format(LOCTEXT("CantCreateDataLayerInstanceIncompatibleChildType", "Cannot create Data Layer Instance : {0}"), Reason));
 			}
diff --git a/Engine/Source/Runtime/Engine/Private/WorldPartition/DataLayer/DataLayerInstance.cpp b/Engine/Source/Runtime/Engine/Private/WorldPartition/DataLayer/DataLayerInstance.cpp
--- a/Engine/Source/Runtime/Engine/Private/WorldPartition/DataLayer/DataLayerInstance.cpp
+++ b/Engine/Source/Runtime/Engine/Private/WorldPartition/DataLayer/DataLayerInstance.cpp
@@ -457,7 +457,7 @@ bool UDataLayerInstance::IsParentDataLayerTypeCompatible(const UDataLayerInstanc
 		return false;
 	}
 
-	if (!FDataLayerUtils::AreDataLayerTypesCompatible(InParent->GetType(), GetType(), OutReasonText))
+	if (!FDataLayerUtils::AreDataLayerTypesCompatible(InParent->GetType(), GetType(), InParent->IsA<UExternalDataLayerInstance>(), OutReasonText))
 	{
 		if (OutReason)
 		{
diff --git a/Engine/Source/Runtime/Engine/Private/WorldPartition/DataLayer/DataLayerUtils.cpp b/Engine/Source/Runtime/Engine/Private/WorldPartition/DataLayer/DataLayerUtils.cpp
--- a/Engine/Source/Runtime/Engine/Private/WorldPartition/DataLayer/DataLayerUtils.cpp
+++ b/Engine/Source/Runtime/Engine/Private/WorldPartition/DataLayer/DataLayerUtils.cpp
@@ -244,11 +244,11 @@ FString FDataLayerUtils::GenerateUniqueDataLayerShortName(const UDataLayerManage
 	return UniqueNewDataLayerShortName;
 }
 
-bool FDataLayerUtils::AreDataLayerTypesCompatible(EDataLayerType ParentDataLayerType, EDataLayerType ChildDataLayerType, FText* OutReason)
+bool FDataLayerUtils::AreDataLayerTypesCompatible(EDataLayerType ParentDataLayerType, EDataLayerType ChildDataLayerType, bool bIsParentExternalDataLayer, FText* OutReason)
 {
 	if ((ChildDataLayerType == EDataLayerType::Unknown) ||
 		(ParentDataLayerType == EDataLayerType::Unknown) ||
-		(ParentDataLayerType != EDataLayerType::Editor && ChildDataLayerType != EDataLayerType::Runtime))
+		(!bIsParentExternalDataLayer && ParentDataLayerType != EDataLayerType::Editor && ChildDataLayerType != EDataLayerType::Runtime))
 	{
 		if (OutReason)
 		{
diff --git a/Engine/Source/Runtime/Engine/Public/WorldPartition/DataLayer/DataLayerUtils.h b/Engine/Source/Runtime/Engine/Public/WorldPartition/DataLayer/DataLayerUtils.h
--- a/Engine/Source/Runtime/Engine/Public/WorldPartition/DataLayer/DataLayerUtils.h
+++ b/Engine/Source/Runtime/Engine/Public/WorldPartition/DataLayer/DataLayerUtils.h
@@ -58,7 +58,7 @@ public:
 
 	static ENGINE_API bool FindDataLayerByShortName(const UDataLayerManager* InDataLayerManager, const FString& InShortName, TSet<UDataLayerInstance*>& OutDataLayerInstances);
 
-	static ENGINE_API bool AreDataLayerTypesCompatible(EDataLayerType ParentDataLayerType, EDataLayerType ChildDataLayerType, FText* OutReason = nullptr);
+	static ENGINE_API bool AreDataLayerTypesCompatible(EDataLayerType ParentDataLayerType, EDataLayerType ChildDataLayerType, bool bIsParentExternalDataLayer, FText* OutReason = nullptr);
 #endif
 
 	static FString GetSanitizedDataLayerShortName(FString InShortName)
