Add additional validation to prevent incompatible streaming mips with runtime virtual texture. This is extra defence against any code logic issues that could generate an incompatible SVT format. If that does happen it will generate D3D device errors when trying to copy the incompatible SVT format to the physical texture.

diff --git a/Engine/Source/Runtime/Engine/Private/Components/RuntimeVirtualTextureComponent.cpp b/Engine/Source/Runtime/Engine/Private/Components/RuntimeVirtualTextureComponent.cpp
--- a/Engine/Source/Runtime/Engine/Private/Components/RuntimeVirtualTextureComponent.cpp
+++ b/Engine/Source/Runtime/Engine/Private/Components/RuntimeVirtualTextureComponent.cpp
@@ -246,7 +246,7 @@ bool IsCompatibleFormat(URuntimeVirtualTexture const& RuntimeVirtualTexture, UVi
 {
 	// During texture compilation we can't validate anything other than first layer, so restrict validation to that.
 	// This should catch any 99% of issues anyway. 
-	return (RuntimeVirtualTexture.GetLayerFormat(0) != StreamingVirtualTexture.GetPixelFormat(0));
+	return (RuntimeVirtualTexture.GetLayerFormat(0) == StreamingVirtualTexture.GetPixelFormat(0));
 }
 
 bool URuntimeVirtualTextureComponent::IsStreamingTextureInvalid(EShadingPath ShadingPath) const
@@ -257,22 +257,14 @@ bool URuntimeVirtualTextureComponent::IsStreamingTextureInvalid(EShadingPath Sha
 		VirtualTexture != nullptr && 
 		StreamingTexture != nullptr && 
 		StreamingTexture->GetVirtualTexture(ShadingPath) != nullptr && 
-		IsCompatibleFormat(*VirtualTexture, *StreamingTexture->GetVirtualTexture(ShadingPath)) &&
-		StreamingTexture->BuildHash != CalculateStreamingTextureSettingsHash();
+		(StreamingTexture->BuildHash != CalculateStreamingTextureSettingsHash() || !IsCompatibleFormat(*VirtualTexture, *StreamingTexture->GetVirtualTexture(ShadingPath)));
 }
 
 #if WITH_EDITOR
 
 bool URuntimeVirtualTextureComponent::IsStreamingTextureInvalid() const
 {
-	return 
-		VirtualTexture != nullptr && 
-		StreamingTexture != nullptr && 
-		StreamingTexture->GetVirtualTexture(EShadingPath::Deferred) != nullptr && 
-		StreamingTexture->GetVirtualTexture(EShadingPath::Mobile) != nullptr && 
-		IsCompatibleFormat(*VirtualTexture, *StreamingTexture->GetVirtualTexture(EShadingPath::Deferred)) &&
-		IsCompatibleFormat(*VirtualTexture, *StreamingTexture->GetVirtualTexture(EShadingPath::Mobile)) &&
-		StreamingTexture->BuildHash != CalculateStreamingTextureSettingsHash();
+	return IsStreamingTextureInvalid(EShadingPath::Mobile) || IsStreamingTextureInvalid(EShadingPath::Deferred);
 }
 
 FLinearColor URuntimeVirtualTextureComponent::GetStreamingMipsFixedColor() const 
