Move fetching of PlatformRenderData for the StaticMesh in an FAsyncDistanceFieldTask to be done as part of the async task rather than block on the main thread if the mesh needs to be built

Opening CitySample Editor to Big_City_LVL map with -ddc=cold, time to a stable editor (no longer building and streaming content) 15m30s -> 6m20s

diff --git a/Engine/Source/Runtime/Engine/Private/DistanceFieldAtlas.cpp b/Engine/Source/Runtime/Engine/Private/DistanceFieldAtlas.cpp
--- a/Engine/Source/Runtime/Engine/Private/DistanceFieldAtlas.cpp
+++ b/Engine/Source/Runtime/Engine/Private/DistanceFieldAtlas.cpp
@@ -733,6 +733,10 @@ void FDistanceFieldAsyncQueue::Build(FAsyncDistanceFieldTask* Task, FQueuedThrea
 			Task->DistanceFieldResolutionScale,
 			Task->bGenerateDistanceFieldAsIfTwoSided,
 			*Task->GeneratedVolumeData);
+
+		// Fetch the platform data as part of the build step since we will lazily create the data if it doesn't exist.
+		// Building the platform data is expensive so doing it asyncronously here is ideal instead of on the game thread.
+		Task->StaticMeshPlatformRenderData = &Task->StaticMesh->GetPlatformStaticMeshRenderData(Task->StaticMesh, Task->TargetPlatform);
 	}
 
 	{
@@ -847,7 +851,7 @@ void FDistanceFieldAsyncQueue::ProcessAsyncTasks(bool bLimitExecutionTime)
 			BeginCacheMeshCardRepresentation(
 				Task->TargetPlatform,
 				Task->StaticMesh,
-				Task->StaticMesh->GetPlatformStaticMeshRenderData(Task->StaticMesh, Task->TargetPlatform),
+				*Task->StaticMeshPlatformRenderData,
 				Task->DDCKey,
 				&Task->SourceMeshData);
 		}
diff --git a/Engine/Source/Runtime/Engine/Public/DistanceFieldAtlas.h b/Engine/Source/Runtime/Engine/Public/DistanceFieldAtlas.h
--- a/Engine/Source/Runtime/Engine/Public/DistanceFieldAtlas.h
+++ b/Engine/Source/Runtime/Engine/Public/DistanceFieldAtlas.h
@@ -328,7 +328,8 @@ public:
 	TArray<FSignedDistanceFieldBuildSectionData> SectionData;
 #endif
 	FSourceMeshDataForDerivedDataTask SourceMeshData;
-	UStaticMesh* StaticMesh;
+	UStaticMesh* StaticMesh; 
+	FStaticMeshRenderData* StaticMeshPlatformRenderData;
 	UStaticMesh* GenerateSource;
 	float DistanceFieldResolutionScale;
 	bool bGenerateDistanceFieldAsIfTwoSided;
