World Partition Runtime Spatial Hash: don't validate target runtime grid for auto-generated HLODs actors.

diff --git a/Engine/Source/Editor/UnrealEd/Public/WorldPartition/WorldPartitionBuilder.h b/Engine/Source/Editor/UnrealEd/Public/WorldPartition/WorldPartitionBuilder.h
--- a/Engine/Source/Editor/UnrealEd/Public/WorldPartition/WorldPartitionBuilder.h
+++ b/Engine/Source/Editor/UnrealEd/Public/WorldPartition/WorldPartitionBuilder.h
@@ -6,7 +6,7 @@
 #include "Engine/World.h"
 #include "WorldPartitionBuilder.generated.h"
 
-typedef UE::Math::TIntVector3<int64> FWorldBuilderCellCoord;
+typedef FInt64Vector3 FWorldBuilderCellCoord;
 
 /**
  * Structure containing information about a World Partition Builder cell
diff --git a/Engine/Source/Runtime/Engine/Private/WorldPartition/RuntimeHashSet/WorldPartitionRuntimeHashSet.cpp b/Engine/Source/Runtime/Engine/Private/WorldPartition/RuntimeHashSet/WorldPartitionRuntimeHashSet.cpp
--- a/Engine/Source/Runtime/Engine/Private/WorldPartition/RuntimeHashSet/WorldPartitionRuntimeHashSet.cpp
+++ b/Engine/Source/Runtime/Engine/Private/WorldPartition/RuntimeHashSet/WorldPartitionRuntimeHashSet.cpp
@@ -151,7 +151,7 @@ void UWorldPartitionRuntimeHashSet::FlushStreamingContent()
 	RuntimeStreamingData.Empty();
 }
 
-bool UWorldPartitionRuntimeHashSet::IsValidGrid(FName GridName) const
+bool UWorldPartitionRuntimeHashSet::IsValidGrid(FName GridName, const UClass* ActorClass) const
 {
 	TArray<FName> MainPartitionTokens;
 	TArray<FName> HLODPartitionTokens;
diff --git a/Engine/Source/Runtime/Engine/Private/WorldPartition/WorldPartitionRuntimeSpatialHash.cpp b/Engine/Source/Runtime/Engine/Private/WorldPartition/WorldPartitionRuntimeSpatialHash.cpp
--- a/Engine/Source/Runtime/Engine/Private/WorldPartition/WorldPartitionRuntimeSpatialHash.cpp
+++ b/Engine/Source/Runtime/Engine/Private/WorldPartition/WorldPartitionRuntimeSpatialHash.cpp
@@ -1615,8 +1615,16 @@ void UWorldPartitionRuntimeSpatialHash::FlushStreamingContent()
 	StreamingGrids.Empty();
 }
 
-bool UWorldPartitionRuntimeSpatialHash::IsValidGrid(FName GridName) const
+bool UWorldPartitionRuntimeSpatialHash::IsValidGrid(FName GridName, const UClass* ActorClass) const
 {
+	// We always consider grids from auto-generated HLODs to be valid
+	static const UClass* WorldPartitionHLODClass = FindObjectChecked<UClass>(nullptr, TEXT("/Script/Engine.WorldPartitionHLOD"));
+	if (ActorClass->IsChildOf(WorldPartitionHLODClass))
+	{
+		return true;
+	}
+
+	// None always maps to the default grid
 	if (GridName.IsNone())
 	{
 		return true;
@@ -1630,15 +1638,6 @@ bool UWorldPartitionRuntimeSpatialHash::IsValidGrid(FName GridName) const
 		}
 	}
 
-	/*const UWorldPartition* WorldPartition = GetOuterUWorldPartition();
-	for (FActorDescContainerInstanceCollection::TConstIterator<ASpatialHashRuntimeGridInfo> Iterator(WorldPartition); Iterator; ++Iterator)
-	{
-		if (Iterator->GetRuntimeGrid() == GridName)
-		{
-			return true;
-		}
-	}*/
-
 	return false;
 }
 
diff --git a/Engine/Source/Runtime/Engine/Private/WorldPartition/WorldPartitionStreamingGeneration.cpp b/Engine/Source/Runtime/Engine/Private/WorldPartition/WorldPartitionStreamingGeneration.cpp
--- a/Engine/Source/Runtime/Engine/Private/WorldPartition/WorldPartitionStreamingGeneration.cpp
+++ b/Engine/Source/Runtime/Engine/Private/WorldPartition/WorldPartitionStreamingGeneration.cpp
@@ -1512,7 +1512,7 @@ class FWorldPartitionStreamingGenerator
 			{
 				FContainerCollectionInstanceDescriptor::FPerInstanceData& PerInstanceData = ContainerCollectionInstanceDescriptor.GetPerInstanceData(ActorDescView.GetGuid());
 
-				if (!IsValidGrid(PerInstanceData.RuntimeGrid))
+				if (!IsValidGrid(PerInstanceData.RuntimeGrid, ActorDescView.GetActorNativeClass()))
 				{
 					if (PassType == EPassType::ErrorReporting)
 					{
@@ -1621,7 +1621,7 @@ public:
 		bool bHandleUnsavedActors;
 		TMap<FGuid, const UActorDescContainerInstance*> ActorGuidsToContainerInstanceMap;
 		TArray<TSubclassOf<AActor>> FilteredClasses;
-		TFunction<bool(FName)> IsValidGrid;
+		TFunction<bool(FName, const UClass*)> IsValidGrid;
 		TFunction<bool(FName GridName, const FSoftObjectPath&)> IsValidHLODLayer;
 
 		FWorldPartitionStreamingGeneratorParams& SetWorldPartitionContext(const UWorldPartition* InWorldPartitionContext) { WorldPartitionContext = InWorldPartitionContext; return *this; }
@@ -1631,7 +1631,7 @@ public:
 		FWorldPartitionStreamingGeneratorParams& SetHandleUnsavedActors(bool bInHandleUnsavedActors) { bHandleUnsavedActors = bInHandleUnsavedActors; return *this; }
 		FWorldPartitionStreamingGeneratorParams& SetActorGuidsToContainerInstanceMap(const TMap<FGuid, const UActorDescContainerInstance*>& InActorGuidsToContainerInstanceMap) { ActorGuidsToContainerInstanceMap = InActorGuidsToContainerInstanceMap; return *this; }
 		FWorldPartitionStreamingGeneratorParams& SetFilteredClasses(const TArray<TSubclassOf<AActor>>& InFilteredClasses) { FilteredClasses = InFilteredClasses; return *this; }
-		FWorldPartitionStreamingGeneratorParams& SetIsValidGrid(TFunction<bool(FName)> InIsValidGrid) { IsValidGrid = InIsValidGrid; return *this; }
+		FWorldPartitionStreamingGeneratorParams& SetIsValidGrid(TFunction<bool(FName, const UClass*)> InIsValidGrid) { IsValidGrid = InIsValidGrid; return *this; }
 		FWorldPartitionStreamingGeneratorParams& SetIsValidHLODLayer(TFunction<bool(FName GridName, const FSoftObjectPath&)> InIsValidHLODLayer) { IsValidHLODLayer = InIsValidHLODLayer; return *this; }
 				
 		inline static FStreamingGenerationNullErrorHandler NullErrorHandler;
@@ -1817,7 +1817,7 @@ private:
 	bool bCreateContainerResolver;
 	bool bHandleUnsavedActors;
 	TArray<TSubclassOf<AActor>> FilteredClasses;
-	TFunction<bool(FName)> IsValidGrid;
+	TFunction<bool(FName, const UClass*)> IsValidGrid;
 	TFunction<bool(FName, const FSoftObjectPath&)> IsValidHLODLayer;
 	IStreamingGenerationErrorHandler* ErrorHandler;
 	FWorldPartitionRuntimeContainerResolver ContainerResolver;
@@ -1892,7 +1892,7 @@ bool UWorldPartition::GenerateContainerStreaming(const FGenerateStreamingParams&
 	FWorldPartitionStreamingGenerator::FWorldPartitionStreamingGeneratorParams StreamingGeneratorParams = FWorldPartitionStreamingGenerator::FWorldPartitionStreamingGeneratorParams()
 		.SetWorldPartitionContext(this)
 		.SetHandleUnsavedActors(bIsPIE)
-		.SetIsValidGrid([this](FName GridName) { return RuntimeHash->IsValidGrid(GridName); })
+		.SetIsValidGrid([this](FName GridName, const UClass* ActorClass) { return RuntimeHash->IsValidGrid(GridName, ActorClass); })
 		.SetIsValidHLODLayer([this](FName GridName, const FSoftObjectPath& HLODLayerPath) { return RuntimeHash->IsValidHLODLayer(GridName, HLODLayerPath); })
 		.SetErrorHandler(ErrorHandlerSelector.Get())
 		.SetEnableStreaming(IsStreamingEnabled())
@@ -1991,7 +1991,7 @@ TUniquePtr<IStreamingGenerationContext> UWorldPartition::GenerateStreamingGenera
 
 	FWorldPartitionStreamingGenerator::FWorldPartitionStreamingGeneratorParams StreamingGeneratorParams = FWorldPartitionStreamingGenerator::FWorldPartitionStreamingGeneratorParams()
 		.SetWorldPartitionContext(this)
-		.SetIsValidGrid([this](FName GridName) { return RuntimeHash->IsValidGrid(GridName); })
+		.SetIsValidGrid([this](FName GridName, const UClass* ActorClass) { return RuntimeHash->IsValidGrid(GridName, ActorClass); })
 		.SetIsValidHLODLayer([this](FName GridName, const FSoftObjectPath& HLODLayerPath) { return RuntimeHash->IsValidHLODLayer(GridName, HLODLayerPath); })
 		.SetErrorHandler(ErrorHandlerSelector.Get())
 		.SetEnableStreaming(IsStreamingEnabled())
@@ -2061,7 +2061,7 @@ void UWorldPartition::SetupHLODActors(const FSetupHLODActorsParams& Params)
 			.SetErrorHandler(ErrorHandlerSelector.Get())
 			.SetEnableStreaming(IsStreamingEnabled())
 			.SetFilteredClasses({ AWorldPartitionHLOD::StaticClass() })
-			.SetIsValidGrid([this](FName GridName) { return RuntimeHash->IsValidGrid(GridName); })
+			.SetIsValidGrid([this](FName GridName, const UClass* ActorClass) { return RuntimeHash->IsValidGrid(GridName, ActorClass); })
 			.SetIsValidHLODLayer([this](FName GridName, const FSoftObjectPath& HLODLayerPath) { return RuntimeHash->IsValidHLODLayer(GridName, HLODLayerPath); });
 
 		FWorldPartitionStreamingGenerator StreamingGenerator(StreamingGeneratorParams);
@@ -2290,7 +2290,7 @@ void UWorldPartition::CheckForErrors(const FCheckForErrorsParams& InParams)
 			.SetWorldPartitionContext(WorldPartition)
 			.SetHandleUnsavedActors(!!WorldPartition)
 			.SetErrorHandler(ErrorHandlerSelector.Get())
-			.SetIsValidGrid([WorldPartitionRuntimeHash](FName GridName) { return WorldPartitionRuntimeHash ? WorldPartitionRuntimeHash->IsValidGrid(GridName) : true; })
+			.SetIsValidGrid([WorldPartitionRuntimeHash](FName GridName, const UClass* ActorClass) { return WorldPartitionRuntimeHash ? WorldPartitionRuntimeHash->IsValidGrid(GridName, ActorClass) : true; })
 			.SetIsValidHLODLayer([WorldPartitionRuntimeHash](FName GridName, const FSoftObjectPath& HLODLayerPath) { return WorldPartitionRuntimeHash ? WorldPartitionRuntimeHash->IsValidHLODLayer(GridName, HLODLayerPath) : true; })
 			.SetEnableStreaming(InParams.bEnableStreaming)
 			.SetActorGuidsToContainerInstanceMap(ActorGuidsToContainerInstanceMap);
diff --git a/Engine/Source/Runtime/Engine/Public/WorldPartition/RuntimeHashSet/WorldPartitionRuntimeHashSet.h b/Engine/Source/Runtime/Engine/Public/WorldPartition/RuntimeHashSet/WorldPartitionRuntimeHashSet.h
--- a/Engine/Source/Runtime/Engine/Public/WorldPartition/RuntimeHashSet/WorldPartitionRuntimeHashSet.h
+++ b/Engine/Source/Runtime/Engine/Public/WorldPartition/RuntimeHashSet/WorldPartitionRuntimeHashSet.h
@@ -137,7 +137,7 @@ public:
 	ENGINE_API virtual bool SupportsHLODs() const override;
 	ENGINE_API virtual bool SetupHLODActors(const IStreamingGenerationContext* StreamingGenerationContext, const UWorldPartition::FSetupHLODActorsParams& Params) const override;
 	ENGINE_API virtual bool GenerateStreaming(class UWorldPartitionStreamingPolicy* StreamingPolicy, const IStreamingGenerationContext* StreamingGenerationContext, TArray<FString>* OutPackagesToGenerate) override;
-	ENGINE_API virtual bool IsValidGrid(FName GridName) const;
+	ENGINE_API virtual bool IsValidGrid(FName GridName, const UClass* ActorClass) const;
 	ENGINE_API virtual bool IsValidHLODLayer(FName GridName, const FSoftObjectPath& HLODLayerPath) const;
 	ENGINE_API virtual void DumpStateLog(FHierarchicalLogArchive& Ar) const override;
 
diff --git a/Engine/Source/Runtime/Engine/Public/WorldPartition/WorldPartitionRuntimeHash.h b/Engine/Source/Runtime/Engine/Public/WorldPartition/WorldPartitionRuntimeHash.h
--- a/Engine/Source/Runtime/Engine/Public/WorldPartition/WorldPartitionRuntimeHash.h
+++ b/Engine/Source/Runtime/Engine/Public/WorldPartition/WorldPartitionRuntimeHash.h
@@ -158,7 +158,7 @@ class UWorldPartitionRuntimeHash : public UObject
 	ENGINE_API virtual TArray<UWorldPartitionRuntimeCell*> GetAlwaysLoadedCells() const;
 	ENGINE_API virtual bool GenerateStreaming(class UWorldPartitionStreamingPolicy* StreamingPolicy, const IStreamingGenerationContext* StreamingGenerationContext, TArray<FString>* OutPackagesToGenerate);
 	virtual bool SetupHLODActors(const IStreamingGenerationContext* StreamingGenerationContext, const UWorldPartition::FSetupHLODActorsParams& Params) const { return false; }
-	virtual bool IsValidGrid(FName GridName) const { return false; }
+	virtual bool IsValidGrid(FName GridName, const UClass* ActorClass) const { return false; }
 	virtual bool IsValidHLODLayer(FName GridName, const FSoftObjectPath& HLODLayerPath) const { return false; }
 	virtual void DrawPreview() const {}
 
diff --git a/Engine/Source/Runtime/Engine/Public/WorldPartition/WorldPartitionRuntimeSpatialHash.h b/Engine/Source/Runtime/Engine/Public/WorldPartition/WorldPartitionRuntimeSpatialHash.h
--- a/Engine/Source/Runtime/Engine/Public/WorldPartition/WorldPartitionRuntimeSpatialHash.h
+++ b/Engine/Source/Runtime/Engine/Public/WorldPartition/WorldPartitionRuntimeSpatialHash.h
@@ -15,8 +15,8 @@
 #include "WorldPartitionRuntimeSpatialHashGridPreviewer.h"
 #include "WorldPartitionRuntimeSpatialHash.generated.h"
 
-typedef UE::Math::TIntVector3<int64> FGridCellCoord;
-typedef UE::Math::TIntVector2<int64> FGridCellCoord2;
+typedef FInt64Vector3 FGridCellCoord;
+typedef FInt64Vector2 FGridCellCoord2;
 class FWorldPartitionDraw2DContext;
 
 UENUM()
@@ -322,7 +322,7 @@ public:
 	ENGINE_API virtual void SetDefaultValues() override;
 	virtual bool SupportsHLODs() const override { return true; }
 	ENGINE_API virtual bool SetupHLODActors(const IStreamingGenerationContext* StreamingGenerationContext, const UWorldPartition::FSetupHLODActorsParams& Params) const override;
-	ENGINE_API virtual bool IsValidGrid(FName GridName) const override;
+	ENGINE_API virtual bool IsValidGrid(FName GridName, const UClass* ActorClass) const override;
 	ENGINE_API virtual bool IsValidHLODLayer(FName GridName, const FSoftObjectPath& HLODLayerPath) const override { return true; }
 	ENGINE_API virtual void DrawPreview() const override;
 
