Remove redundant code

diff --git a/Engine/Source/Runtime/Engine/Private/Animation/AnimCompressionTypes.cpp b/Engine/Source/Runtime/Engine/Private/Animation/AnimCompressionTypes.cpp
--- a/Engine/Source/Runtime/Engine/Private/Animation/AnimCompressionTypes.cpp
+++ b/Engine/Source/Runtime/Engine/Private/Animation/AnimCompressionTypes.cpp
@@ -444,10 +444,6 @@ void FCompressibleAnimData::ResampleAnimationTrackData(const FFrameRate& SampleR
             	return;
             }
 
-			// Make a copy, deals with bone name and index
-			TArray<FName> TrackNames;
-			DataModelInterface->GetBoneTrackNames(TrackNames);
-			
 			FMemMark Mark(FMemStack::Get());
 			FByFramePoseEvalContext EvalContext(AnimSequence);
 			EvalContext.RequiredBones.SetDisableRetargeting(true);
