Add a Kismet Math library function to access the value of a FRuntimeFloatCurve.

diff --git a/Engine/Source/Runtime/Engine/Classes/Kismet/KismetMathLibrary.h b/Engine/Source/Runtime/Engine/Classes/Kismet/KismetMathLibrary.h
--- a/Engine/Source/Runtime/Engine/Classes/Kismet/KismetMathLibrary.h
+++ b/Engine/Source/Runtime/Engine/Classes/Kismet/KismetMathLibrary.h
@@ -177,6 +177,8 @@ struct FQuaternionSpringState
 	}
 };
 
+struct FRuntimeFloatCurve;
+
 UCLASS(meta=(BlueprintThreadSafe, ScriptName = "MathLibrary"), MinimalAPI)
 class UKismetMathLibrary : public UBlueprintFunctionLibrary
 {
@@ -873,6 +875,20 @@ class UKismetMathLibrary : public UBlueprintFunctionLibrary
 	UFUNCTION(BlueprintPure, Category = "Math|Float")
 	static ENGINE_API double FInterpEaseInOut(double A, double B, double Alpha, double Exponent);
 
+	/**
+	* Evaluate this runtime float curve at the specified time 
+	* 
+	* @param Curve				The runtime float curve to evaluate
+	* @param InTime				The time at which to evaluate the curve
+	* @param InDefaultValue		The default value which should be used if the curve cannot be evaluated at the given time.
+	* 
+	* @return					The curve's value at the given time.
+	* 
+	* @see FRichCurve::Eval
+	*/
+	UFUNCTION(BlueprintPure, Category = "Math|Curves")
+	static ENGINE_API float GetRuntimeFloatCurveValue(const FRuntimeFloatCurve& Curve, const float InTime, const float InDefaultValue = 0.0f);
+
 	/**
 	* Simple function to create a pulsating scalar value
 	*
diff --git a/Engine/Source/Runtime/Engine/Private/KismetMathLibrary.cpp b/Engine/Source/Runtime/Engine/Private/KismetMathLibrary.cpp
--- a/Engine/Source/Runtime/Engine/Private/KismetMathLibrary.cpp
+++ b/Engine/Source/Runtime/Engine/Private/KismetMathLibrary.cpp
@@ -2,6 +2,7 @@
 
 #include "Kismet/KismetMathLibrary.h"
 #include "Engine/Engine.h"
+#include "Curves/CurveFloat.h"
 #include "DrawDebugHelpers.h"
 
 #include "Blueprint/BlueprintSupport.h"
@@ -352,6 +353,16 @@ double UKismetMathLibrary::FInterpEaseInOut(double A, double B, double Alpha, do
 	return FMath::InterpEaseInOut<double>(A, B, Alpha, Exponent);
 }
 
+float UKismetMathLibrary::GetRuntimeFloatCurveValue(const FRuntimeFloatCurve& Curve, const float InTime, const float InDefaultValue /*= 0.0f*/)
+{
+	if (const FRichCurve* RichCurve = Curve.GetRichCurveConst())
+	{
+		return RichCurve->Eval(InTime, InDefaultValue);
+	}
+
+	return InDefaultValue;
+}
+
 float UKismetMathLibrary::MakePulsatingValue(float InCurrentTime, float InPulsesPerSecond, float InPhase)
 {
 	return FMath::MakePulsatingValue((double)InCurrentTime, InPulsesPerSecond, InPhase);
