WorldPartition: Add missing UActorDescContainerInstance::Uninitialize calls

#rnx

diff --git a/Engine/Source/Runtime/Engine/Private/Tests/WorldPartition/WorldPartitionHandleTests.cpp b/Engine/Source/Runtime/Engine/Private/Tests/WorldPartition/WorldPartitionHandleTests.cpp
--- a/Engine/Source/Runtime/Engine/Private/Tests/WorldPartition/WorldPartitionHandleTests.cpp
+++ b/Engine/Source/Runtime/Engine/Private/Tests/WorldPartition/WorldPartitionHandleTests.cpp
@@ -294,6 +294,8 @@ namespace WorldPartitionTests
 			TestTrue(TEXT("Invalid container test"), Handle.IsValid());
 			TestTrue(TEXT("Invalid container test"), Reference.IsValid());
 
+			ActorDescContainerInstance->Uninitialize();
+
 			// Make sure to cleanup world before collecting garbage so it gets uninitialized
 			ScopedEditorWorld.Reset();
 			CollectGarbage(GARBAGE_COLLECTION_KEEPFLAGS);
@@ -335,6 +337,8 @@ namespace WorldPartitionTests
 				TestTrue(TEXT("Actor Descriptor Serialization"), NewActorDesc->Equals(Iterator->GetActorDesc()));
 			}
 		}
+
+		ActorDescContainerInstance->Uninitialize();
 #endif
 		return true;
 	}
