Populate FEngineFontServices::SlateFontServices if the slate application is initialised and it has a renderer, rather than on some other conditions. Allows for use of FEngineFontServices in standalone programs where a commandlet is being executed, for example.

diff --git a/Engine/Source/Runtime/Engine/Private/UserInterface/EngineFontServices.cpp b/Engine/Source/Runtime/Engine/Private/UserInterface/EngineFontServices.cpp
--- a/Engine/Source/Runtime/Engine/Private/UserInterface/EngineFontServices.cpp
+++ b/Engine/Source/Runtime/Engine/Private/UserInterface/EngineFontServices.cpp
@@ -10,12 +10,13 @@ FEngineFontServices::FEngineFontServices()
 {
 	check(IsInGameThread());
 
-#if !UE_SERVER
-	if (!IsRunningDedicatedServer() && !IsRunningCommandlet())
+	if (FSlateApplication::IsInitialized())
 	{
-		SlateFontServices = FSlateApplication::Get().GetRenderer()->GetFontServices();
+		if (const FSlateRenderer* Renderer = FSlateApplication::Get().GetRenderer())
+		{
+			SlateFontServices = Renderer->GetFontServices();
+		}
 	}
-#endif
 }
 
 FEngineFontServices::~FEngineFontServices()
