[Material Editor] Added safety code that prevents Material Functions from saving expression collections with nulls in them, causing continuous material invalidations on PostLoad().

diff --git a/Engine/Source/Runtime/Engine/Private/Materials/MaterialExpressions.cpp b/Engine/Source/Runtime/Engine/Private/Materials/MaterialExpressions.cpp
--- a/Engine/Source/Runtime/Engine/Private/Materials/MaterialExpressions.cpp
+++ b/Engine/Source/Runtime/Engine/Private/Materials/MaterialExpressions.cpp
@@ -15004,13 +15004,22 @@ void UMaterialFunction::ForceRecompileForRendering(FMaterialUpdateContext& Updat
 
 void UMaterialFunction::Serialize(FArchive& Ar)
 {
-#if WITH_EDITOR
-	// Temporary debugging code. This will populate the DebugExpressionInfos with information about each expression
-	// in ExpressionCollection.Expressions, to gather more information when some expression is null upon function
-	// PostLoad() to help solve UE-198712.
+#if WITH_EDITORONLY_DATA
 	UMaterialFunctionEditorOnlyData* EditorOnly = GetEditorOnlyData();
-	if (EditorOnly  && Ar.IsSaving() && !Ar.IsCooking() && !Ar.IsObjectReferenceCollector())
+	if (EditorOnly && Ar.IsSaving() && !Ar.IsCooking() && !Ar.IsObjectReferenceCollector())
 	{
+		// If the collection of expressions got some null expressions remove them now, but warn the user about it.
+		if (EditorOnly->ExpressionCollection.Expressions.Remove(nullptr))
+		{
+			FText Message = FText::Format(NSLOCTEXT("MaterialExpressions", "Error_NullExpressionsInMaterialFunction",
+				"Material Function {0} editor only data contained null expression and some expressions may be missing."
+				"\n\nPlease close and repoen this Material Function and verify it is still valid."), FText::FromString(GetFullName()));
+			FMessageDialog::Open(EAppMsgType::Ok, Message);
+		}
+
+		// Temporary debugging code. This will populate the DebugExpressionInfos with information about each expression
+		// in ExpressionCollection.Expressions, to gather more information when some expression is null upon function
+		// PostLoad() to help solve UE-198712.
 		EditorOnly->ExpressionCollection.DebugExpressionInfos.Empty();
 		EditorOnly->ExpressionCollection.DebugExpressionInfos.Reserve(EditorOnly->ExpressionCollection.Expressions.Num());
 		for (UMaterialExpression* Expression : EditorOnly->ExpressionCollection.Expressions)
