Fixes for not setting all the constraint parameters when switching profiles. This also exposes the function that does this, so it can be used by RigidBodyWithControl

diff --git a/Engine/Plugins/Experimental/PhysicsControl/Source/PhysicsControl/Private/RigidBodyWithControl.cpp b/Engine/Plugins/Experimental/PhysicsControl/Source/PhysicsControl/Private/RigidBodyWithControl.cpp
--- a/Engine/Plugins/Experimental/PhysicsControl/Source/PhysicsControl/Private/RigidBodyWithControl.cpp
+++ b/Engine/Plugins/Experimental/PhysicsControl/Source/PhysicsControl/Private/RigidBodyWithControl.cpp
@@ -890,42 +890,6 @@ void FAnimNode_RigidBodyWithControl::ApplyKinematicTargets()
 	}
 }
 
-//======================================================================================================================
-static Chaos::EJointMotionType ConvertMotionType(ELinearConstraintMotion InEngineType)
-{
-	switch (InEngineType)
-	{
-	case ELinearConstraintMotion::LCM_Free: return Chaos::EJointMotionType::Free;
-	case ELinearConstraintMotion::LCM_Limited: return Chaos::EJointMotionType::Limited;
-	case ELinearConstraintMotion::LCM_Locked : return Chaos::EJointMotionType::Locked;
-	default: ensure(false); return Chaos::EJointMotionType::Locked;
-	}
-};
-
-//======================================================================================================================
-static Chaos::EJointMotionType ConvertMotionType(EAngularConstraintMotion InEngineType)
-{
-	switch(InEngineType)
-	{
-	case EAngularConstraintMotion::ACM_Free: return Chaos::EJointMotionType::Free;
-	case EAngularConstraintMotion::ACM_Limited: return Chaos::EJointMotionType::Limited;
-	case EAngularConstraintMotion::ACM_Locked: return Chaos::EJointMotionType::Locked;
-	default: ensure(false); return Chaos::EJointMotionType::Locked;
-	}
-};
-
-//======================================================================================================================
-static Chaos::EPlasticityType ConvertPlasticityType(EConstraintPlasticityType InType)
-{
-	switch (InType)
-	{
-	case EConstraintPlasticityType::CCPT_Free: return Chaos::EPlasticityType::Free;
-	case EConstraintPlasticityType::CCPT_Shrink: return Chaos::EPlasticityType::Shrink;
-	case EConstraintPlasticityType::CCPT_Grow: return Chaos::EPlasticityType::Grow;
-	default: ensure(false); return Chaos::EPlasticityType::Free;
-	}
-}
-
 //======================================================================================================================
 void FAnimNode_RigidBodyWithControl::ApplyCurrentConstraintProfile()
 {
@@ -960,143 +924,7 @@ void FAnimNode_RigidBodyWithControl::ApplyCurrentConstraintProfile()
 
 					FPBDJointSettings JointSettings = ConstraintHandle->GetSettings();
 					
-					// See ImmediatePhysics_Chaos::TransferJointSettings TODO avoid code duplication
-					// with that function. However, it uses ConstraintInstance, but we don't have
-					// one of those.
-
-					JointSettings.Stiffness = ConstraintSettings::JointStiffness();
-					JointSettings.LinearProjection = Profile.bEnableProjection ? Profile.ProjectionLinearAlpha : 0.0f;
-					JointSettings.AngularProjection = Profile.bEnableProjection ? Profile.ProjectionAngularAlpha : 0.0f;
-					JointSettings.ShockPropagation = Profile.bEnableShockPropagation ? Profile.ShockPropagationAlpha : 0.0f;
-					JointSettings.TeleportDistance = Profile.bEnableProjection ? Profile.ProjectionLinearTolerance : -1.0f;
-					JointSettings.TeleportAngle = Profile.bEnableProjection ? 
-						FMath::DegreesToRadians(Profile.ProjectionAngularTolerance) : -1.0f;
-					JointSettings.ParentInvMassScale = Profile.bParentDominates ? (FReal)0 : (FReal)1;
-
-					JointSettings.bCollisionEnabled = !Profile.bDisableCollision;
-					JointSettings.bProjectionEnabled = Profile.bEnableProjection;
-					JointSettings.bShockPropagationEnabled = Profile.bEnableShockPropagation;
-					JointSettings.bMassConditioningEnabled = Profile.bEnableMassConditioning;
-
-					JointSettings.LinearMotionTypes[0] = ConvertMotionType(Profile.LinearLimit.XMotion);
-					JointSettings.LinearMotionTypes[1] = ConvertMotionType(Profile.LinearLimit.YMotion);
-					JointSettings.LinearMotionTypes[2] = ConvertMotionType(Profile.LinearLimit.ZMotion);
-
-					JointSettings.LinearLimit = Profile.LinearLimit.Limit; // Is float to vector the best way?!
-
-					// Order is twist, swing1, swing2 and in degrees
-					JointSettings.AngularMotionTypes[(int32)Chaos::EJointAngularConstraintIndex::Twist] = ConvertMotionType(Profile.TwistLimit.TwistMotion);
-					JointSettings.AngularMotionTypes[(int32)Chaos::EJointAngularConstraintIndex::Swing1] = ConvertMotionType(Profile.ConeLimit.Swing1Motion);
-					JointSettings.AngularMotionTypes[(int32)Chaos::EJointAngularConstraintIndex::Swing2] = ConvertMotionType(Profile.ConeLimit.Swing2Motion);
-
-					JointSettings.AngularLimits[(int32)Chaos::EJointAngularConstraintIndex::Twist] = FMath::DegreesToRadians(Profile.TwistLimit.TwistLimitDegrees);
-					JointSettings.AngularLimits[(int32)Chaos::EJointAngularConstraintIndex::Swing1] = FMath::DegreesToRadians(Profile.ConeLimit.Swing1LimitDegrees);
-					JointSettings.AngularLimits[(int32)Chaos::EJointAngularConstraintIndex::Swing2] = FMath::DegreesToRadians(Profile.ConeLimit.Swing2LimitDegrees);
-
-					JointSettings.bSoftLinearLimitsEnabled = Profile.LinearLimit.bSoftConstraint;
-					JointSettings.bSoftTwistLimitsEnabled = Profile.TwistLimit.bSoftConstraint;
-					JointSettings.bSoftSwingLimitsEnabled = Profile.ConeLimit.bSoftConstraint;
-
-					JointSettings.LinearSoftForceMode = (ConstraintSettings::SoftLinearForceMode() == 0) ? 
-						EJointForceMode::Acceleration : EJointForceMode::Force;
-					JointSettings.AngularSoftForceMode = (ConstraintSettings::SoftAngularForceMode() == 0) ? 
-						EJointForceMode::Acceleration : EJointForceMode::Force;
-
-					JointSettings.SoftLinearStiffness = 
-						Chaos::ConstraintSettings::SoftLinearStiffnessScale() * Profile.LinearLimit.Stiffness;
-					JointSettings.SoftLinearDamping = 
-						Chaos::ConstraintSettings::SoftLinearDampingScale() * Profile.LinearLimit.Damping;
-					JointSettings.SoftTwistStiffness = 
-						Chaos::ConstraintSettings::SoftAngularStiffnessScale() * Profile.TwistLimit.Stiffness;
-					JointSettings.SoftTwistDamping = 
-						Chaos::ConstraintSettings::SoftAngularDampingScale() * Profile.TwistLimit.Damping;
-					JointSettings.SoftSwingStiffness = 
-						Chaos::ConstraintSettings::SoftAngularStiffnessScale() * Profile.ConeLimit.Stiffness;
-					JointSettings.SoftSwingDamping = 
-						Chaos::ConstraintSettings::SoftAngularDampingScale() * Profile.ConeLimit.Damping;
-
-					JointSettings.LinearRestitution = Profile.LinearLimit.Restitution;
-					JointSettings.TwistRestitution = Profile.TwistLimit.Restitution;
-					JointSettings.SwingRestitution = Profile.ConeLimit.Restitution;
-
-					JointSettings.LinearContactDistance = Profile.LinearLimit.ContactDistance;
-					JointSettings.TwistContactDistance = Profile.TwistLimit.ContactDistance;
-					JointSettings.SwingContactDistance = Profile.ConeLimit.ContactDistance;
-
-					JointSettings.LinearDrivePositionTarget = Profile.LinearDrive.PositionTarget;
-					JointSettings.LinearDriveVelocityTarget = Profile.LinearDrive.VelocityTarget;
-					JointSettings.bLinearPositionDriveEnabled[0] = Profile.LinearDrive.XDrive.bEnablePositionDrive;
-					JointSettings.bLinearPositionDriveEnabled[1] = Profile.LinearDrive.YDrive.bEnablePositionDrive;
-					JointSettings.bLinearPositionDriveEnabled[2] = Profile.LinearDrive.ZDrive.bEnablePositionDrive;
-					JointSettings.bLinearVelocityDriveEnabled[0] = Profile.LinearDrive.XDrive.bEnableVelocityDrive;
-					JointSettings.bLinearVelocityDriveEnabled[1] = Profile.LinearDrive.YDrive.bEnableVelocityDrive;
-					JointSettings.bLinearVelocityDriveEnabled[2] = Profile.LinearDrive.ZDrive.bEnableVelocityDrive;
-					JointSettings.LinearDriveForceMode = EJointForceMode::Acceleration; // hardcoded!
-					JointSettings.LinearDriveStiffness = Chaos::ConstraintSettings::LinearDriveStiffnessScale() * FVec3(
-							Profile.LinearDrive.XDrive.Stiffness, 
-							Profile.LinearDrive.YDrive.Stiffness, 
-							Profile.LinearDrive.ZDrive.Stiffness);
-					JointSettings.LinearDriveDamping = Chaos::ConstraintSettings::LinearDriveDampingScale() * FVec3(
-							Profile.LinearDrive.XDrive.Damping, 
-							Profile.LinearDrive.YDrive.Damping, 
-							Profile.LinearDrive.ZDrive.Damping);
-					JointSettings.LinearDriveMaxForce[0] = Profile.LinearDrive.XDrive.MaxForce;
-					JointSettings.LinearDriveMaxForce[1] = Profile.LinearDrive.YDrive.MaxForce;
-					JointSettings.LinearDriveMaxForce[2] = Profile.LinearDrive.ZDrive.MaxForce;
-
-					JointSettings.AngularDrivePositionTarget = FQuat(Profile.AngularDrive.OrientationTarget);
-					JointSettings.AngularDriveVelocityTarget = Profile.AngularDrive.AngularVelocityTarget * UE_TWO_PI; // rev/s to rad/s
-
-					JointSettings.AngularDriveForceMode = EJointForceMode::Acceleration; // hardcoded!
-					if (Profile.AngularDrive.AngularDriveMode == EAngularDriveMode::SLERP)
-					{
-						JointSettings.AngularDriveStiffness = FVec3(
-							ConstraintSettings::AngularDriveStiffnessScale() * Profile.AngularDrive.SlerpDrive.Stiffness);
-						JointSettings.AngularDriveDamping = FVec3(
-							ConstraintSettings::AngularDriveDampingScale() * Profile.AngularDrive.SlerpDrive.Damping);
-						JointSettings.AngularDriveMaxTorque = FVec3(
-							Profile.AngularDrive.SlerpDrive.MaxForce);
-						JointSettings.bAngularSLerpPositionDriveEnabled = Profile.AngularDrive.SlerpDrive.bEnablePositionDrive;
-						JointSettings.bAngularSLerpVelocityDriveEnabled = Profile.AngularDrive.SlerpDrive.bEnableVelocityDrive;
-						JointSettings.bAngularTwistPositionDriveEnabled = false;
-						JointSettings.bAngularTwistVelocityDriveEnabled = false;
-						JointSettings.bAngularSwingPositionDriveEnabled = false;
-						JointSettings.bAngularSwingVelocityDriveEnabled = false;
-					}
-					else
-					{
-						JointSettings.AngularDriveStiffness = ConstraintSettings::AngularDriveStiffnessScale() * FVec3(
-							Profile.AngularDrive.TwistDrive.Stiffness, 
-							Profile.AngularDrive.SwingDrive.Stiffness, 
-							Profile.AngularDrive.SwingDrive.Stiffness);
-						JointSettings.AngularDriveDamping = ConstraintSettings::AngularDriveDampingScale() * FVec3(
-							Profile.AngularDrive.TwistDrive.Damping, 
-							Profile.AngularDrive.SwingDrive.Damping, 
-							Profile.AngularDrive.SwingDrive.Damping);
-						JointSettings.AngularDriveMaxTorque[0] = Profile.AngularDrive.TwistDrive.MaxForce;
-						JointSettings.AngularDriveMaxTorque[1] = Profile.AngularDrive.SwingDrive.MaxForce;
-						JointSettings.AngularDriveMaxTorque[2] = Profile.AngularDrive.SwingDrive.MaxForce;
-						JointSettings.bAngularSLerpPositionDriveEnabled = false;
-						JointSettings.bAngularSLerpVelocityDriveEnabled = false;
-						JointSettings.bAngularTwistPositionDriveEnabled = Profile.AngularDrive.TwistDrive.bEnablePositionDrive;
-						JointSettings.bAngularTwistVelocityDriveEnabled = Profile.AngularDrive.TwistDrive.bEnableVelocityDrive;
-						JointSettings.bAngularSwingPositionDriveEnabled = Profile.AngularDrive.SwingDrive.bEnablePositionDrive;
-						JointSettings.bAngularSwingVelocityDriveEnabled = Profile.AngularDrive.SwingDrive.bEnableVelocityDrive;
-					}
-
-					JointSettings.LinearBreakForce = Profile.bLinearBreakable ? 
-						Chaos::ConstraintSettings::LinearBreakScale() * Profile.LinearBreakThreshold : FLT_MAX;
-					JointSettings.LinearPlasticityLimit = Profile.bLinearPlasticity ? 
-						FMath::Clamp((float)Profile.LinearPlasticityThreshold, 0.0f, 1.0f) : FLT_MAX;
-					JointSettings.LinearPlasticityType = ConvertPlasticityType(Profile.LinearPlasticityType);
-					// JointSettings.LinearPlasticityInitialDistanceSquared = ; // What do we do with this?
-
-					JointSettings.AngularBreakTorque = Profile.bAngularBreakable ? 
-						Chaos::ConstraintSettings::AngularBreakScale() * Profile.AngularBreakThreshold : FLT_MAX;
-					JointSettings.AngularPlasticityLimit = Profile.bAngularPlasticity ? 
-						FMath::Clamp((float)Profile.AngularPlasticityThreshold, 0.0f, 1.0f) : FLT_MAX;;
-
-					JointSettings.ContactTransferScale = Profile.ContactTransferScale;
+					ImmediatePhysics_Chaos::UpdateJointSettingsFromConstraintProfile(Profile, JointSettings);
 
 					ConstraintHandle->SetSettings(JointSettings);
 				}
diff --git a/Engine/Source/Runtime/Engine/Private/PhysicsEngine/ImmediatePhysics/ImmediatePhysicsChaos/ImmediatePhysicsJointHandle_Chaos.cpp b/Engine/Source/Runtime/Engine/Private/PhysicsEngine/ImmediatePhysics/ImmediatePhysicsChaos/ImmediatePhysicsJointHandle_Chaos.cpp
--- a/Engine/Source/Runtime/Engine/Private/PhysicsEngine/ImmediatePhysics/ImmediatePhysicsChaos/ImmediatePhysicsJointHandle_Chaos.cpp
+++ b/Engine/Source/Runtime/Engine/Private/PhysicsEngine/ImmediatePhysics/ImmediatePhysicsChaos/ImmediatePhysicsJointHandle_Chaos.cpp
@@ -21,137 +21,191 @@ static_assert((int32)Chaos::EJointAngularConstraintIndex::Swing1 == 2, "Angular
 
 namespace ImmediatePhysics_Chaos
 {
-
-	void TransferJointSettings(FConstraintInstance* ConstraintInstance, Chaos::FPBDJointSettings& ConstraintSettings)
+	static Chaos::EJointMotionType ConvertToJointMotionType(ELinearConstraintMotion InType)
 	{
-		using namespace Chaos;
-
-		const FConstraintProfileProperties& Profile = ConstraintInstance->ProfileInstance;
-
-		ConstraintSettings.Stiffness = ConstraintSettings::JointStiffness();
-
-		ConstraintSettings.LinearMotionTypes =
-		{
-			static_cast<EJointMotionType>(ConstraintInstance->GetLinearXMotion()),
-			static_cast<EJointMotionType>(ConstraintInstance->GetLinearYMotion()),
-			static_cast<EJointMotionType>(ConstraintInstance->GetLinearZMotion()),
-		};
-		ConstraintSettings.LinearLimit = ConstraintInstance->GetLinearLimit();
-
-		ConstraintSettings.AngularMotionTypes[(int32)EJointAngularConstraintIndex::Twist] = static_cast<EJointMotionType>(ConstraintInstance->GetAngularTwistMotion());
-		ConstraintSettings.AngularMotionTypes[(int32)EJointAngularConstraintIndex::Swing1] = static_cast<EJointMotionType>(ConstraintInstance->GetAngularSwing1Motion());
-		ConstraintSettings.AngularMotionTypes[(int32)EJointAngularConstraintIndex::Swing2] = static_cast<EJointMotionType>(ConstraintInstance->GetAngularSwing2Motion());
-		ConstraintSettings.AngularLimits[(int32)EJointAngularConstraintIndex::Twist] = FMath::DegreesToRadians(ConstraintInstance->GetAngularTwistLimit());
-		ConstraintSettings.AngularLimits[(int32)EJointAngularConstraintIndex::Swing1] = FMath::DegreesToRadians(ConstraintInstance->GetAngularSwing1Limit());
-		ConstraintSettings.AngularLimits[(int32)EJointAngularConstraintIndex::Swing2] = FMath::DegreesToRadians(ConstraintInstance->GetAngularSwing2Limit());
-
-		ConstraintSettings.bProjectionEnabled = Profile.bEnableProjection;
-		ConstraintSettings.bShockPropagationEnabled = Profile.bEnableShockPropagation;
-		ConstraintSettings.bMassConditioningEnabled = Profile.bEnableMassConditioning;
-
-		ConstraintSettings.LinearProjection = Profile.bEnableProjection ? Profile.ProjectionLinearAlpha : 0.0f;
-		ConstraintSettings.AngularProjection = Profile.bEnableProjection ? Profile.ProjectionAngularAlpha : 0.0f;
-		ConstraintSettings.TeleportDistance = Profile.bEnableProjection ? Profile.ProjectionLinearTolerance : -1.0f;
-		ConstraintSettings.TeleportAngle = Profile.bEnableProjection ? FMath::DegreesToRadians(Profile.ProjectionAngularTolerance) : -1.0f;
-		ConstraintSettings.ShockPropagation = Profile.bEnableShockPropagation ? Profile.ShockPropagationAlpha : 0.0f;
-		ConstraintSettings.ParentInvMassScale = Profile.bParentDominates ? (FReal)0 : (FReal)1;
-
-		ConstraintSettings.bSoftLinearLimitsEnabled = ConstraintInstance->GetIsSoftLinearLimit();
-		ConstraintSettings.bSoftTwistLimitsEnabled = ConstraintInstance->GetIsSoftTwistLimit();
-		ConstraintSettings.bSoftSwingLimitsEnabled = ConstraintInstance->GetIsSoftSwingLimit();
-		ConstraintSettings.SoftLinearStiffness = Chaos::ConstraintSettings::SoftLinearStiffnessScale() * ConstraintInstance->GetSoftLinearLimitStiffness();
-		ConstraintSettings.SoftLinearDamping = Chaos::ConstraintSettings::SoftLinearDampingScale() * ConstraintInstance->GetSoftLinearLimitDamping();
-		ConstraintSettings.SoftTwistStiffness = Chaos::ConstraintSettings::SoftAngularStiffnessScale() * ConstraintInstance->GetSoftTwistLimitStiffness();
-		ConstraintSettings.SoftTwistDamping = Chaos::ConstraintSettings::SoftAngularDampingScale() * ConstraintInstance->GetSoftTwistLimitDamping();
-		ConstraintSettings.SoftSwingStiffness = Chaos::ConstraintSettings::SoftAngularStiffnessScale() * ConstraintInstance->GetSoftSwingLimitStiffness();
-		ConstraintSettings.SoftSwingDamping = Chaos::ConstraintSettings::SoftAngularDampingScale() * ConstraintInstance->GetSoftSwingLimitDamping();
-		ConstraintSettings.LinearSoftForceMode = (Chaos::ConstraintSettings::SoftLinearForceMode() == 0) ? EJointForceMode::Acceleration : EJointForceMode::Force;
-		ConstraintSettings.AngularSoftForceMode = (Chaos::ConstraintSettings::SoftAngularForceMode() == 0) ? EJointForceMode::Acceleration : EJointForceMode::Force;
-
-		if (!ConstraintSettings.bSoftLinearLimitsEnabled)
+		switch (InType)
 		{
-			ConstraintSettings.LinearRestitution = Profile.LinearLimit.Restitution;
-			ConstraintSettings.LinearContactDistance = Profile.LinearLimit.ContactDistance;
+		case ELinearConstraintMotion::LCM_Free: return Chaos::EJointMotionType::Free;
+		case ELinearConstraintMotion::LCM_Limited: return Chaos::EJointMotionType::Limited;
+		case ELinearConstraintMotion::LCM_Locked: return Chaos::EJointMotionType::Locked;
+		default: ensure(false); return Chaos::EJointMotionType::Locked;
 		}
-		if (!ConstraintSettings.bSoftTwistLimitsEnabled)
+	};
+
+	static Chaos::EJointMotionType ConvertToJointMotionType(EAngularConstraintMotion InType)
+	{
+		switch (InType)
 		{
-			ConstraintSettings.TwistRestitution = Profile.TwistLimit.Restitution;
-			ConstraintSettings.TwistContactDistance = FMath::DegreesToRadians(Profile.TwistLimit.ContactDistance);
+		case EAngularConstraintMotion::ACM_Free: return Chaos::EJointMotionType::Free;
+		case EAngularConstraintMotion::ACM_Limited: return Chaos::EJointMotionType::Limited;
+		case EAngularConstraintMotion::ACM_Locked: return Chaos::EJointMotionType::Locked;
+		default: ensure(false); return Chaos::EJointMotionType::Locked;
 		}
-		if (!ConstraintSettings.bSoftSwingLimitsEnabled)
+	};
+
+	static Chaos::EPlasticityType ConvertToPlasticityType(EConstraintPlasticityType InType)
+	{
+		switch (InType)
 		{
-			ConstraintSettings.SwingRestitution = Profile.ConeLimit.Restitution;
-			ConstraintSettings.SwingContactDistance = FMath::DegreesToRadians(Profile.ConeLimit.ContactDistance);
+		case EConstraintPlasticityType::CCPT_Free: return Chaos::EPlasticityType::Free;
+		case EConstraintPlasticityType::CCPT_Shrink: return Chaos::EPlasticityType::Shrink;
+		case EConstraintPlasticityType::CCPT_Grow: return Chaos::EPlasticityType::Grow;
+		default: return Chaos::EPlasticityType::Free;
 		}
+	}
 
-		ConstraintSettings.LinearDrivePositionTarget = Profile.LinearDrive.PositionTarget;
-		ConstraintSettings.LinearDriveVelocityTarget = Profile.LinearDrive.VelocityTarget;
-		ConstraintSettings.bLinearPositionDriveEnabled[0] = Profile.LinearDrive.XDrive.bEnablePositionDrive;
-		ConstraintSettings.bLinearPositionDriveEnabled[1] = Profile.LinearDrive.YDrive.bEnablePositionDrive;
-		ConstraintSettings.bLinearPositionDriveEnabled[2] = Profile.LinearDrive.ZDrive.bEnablePositionDrive;
-		ConstraintSettings.bLinearVelocityDriveEnabled[0] = Profile.LinearDrive.XDrive.bEnableVelocityDrive;
-		ConstraintSettings.bLinearVelocityDriveEnabled[1] = Profile.LinearDrive.YDrive.bEnableVelocityDrive;
-		ConstraintSettings.bLinearVelocityDriveEnabled[2] = Profile.LinearDrive.ZDrive.bEnableVelocityDrive;
-		ConstraintSettings.LinearDriveStiffness = Chaos::ConstraintSettings::LinearDriveStiffnessScale() * Chaos::FVec3(Profile.LinearDrive.XDrive.Stiffness, Profile.LinearDrive.YDrive.Stiffness, Profile.LinearDrive.ZDrive.Stiffness);
-		ConstraintSettings.LinearDriveDamping = Chaos::ConstraintSettings::LinearDriveDampingScale() * Chaos::FVec3(Profile.LinearDrive.XDrive.Damping, Profile.LinearDrive.YDrive.Damping, Profile.LinearDrive.ZDrive.Damping);
-		ConstraintSettings.LinearDriveForceMode = EJointForceMode::Acceleration;
-
-		ConstraintSettings.AngularDrivePositionTarget = FQuat(Profile.AngularDrive.OrientationTarget);
-		ConstraintSettings.AngularDriveVelocityTarget = Profile.AngularDrive.AngularVelocityTarget * 2.0f * UE_PI; // Rev/s to Rad/s
+	// Copies data from a profile into settings
+	void UpdateJointSettingsFromConstraintProfile(const FConstraintProfileProperties& Profile, Chaos::FPBDJointSettings& JointSettings)
+	{
+		using namespace Chaos;
 
+		JointSettings.Stiffness = ConstraintSettings::JointStiffness();
+		JointSettings.LinearProjection = Profile.bEnableProjection ? Profile.ProjectionLinearAlpha : 0.0f;
+		JointSettings.AngularProjection = Profile.bEnableProjection ? Profile.ProjectionAngularAlpha : 0.0f;
+		JointSettings.ShockPropagation = Profile.bEnableShockPropagation ? Profile.ShockPropagationAlpha : 0.0f;
+		JointSettings.TeleportDistance = Profile.bEnableProjection ? Profile.ProjectionLinearTolerance : -1.0f;
+		JointSettings.TeleportAngle = Profile.bEnableProjection ?
+			FMath::DegreesToRadians(Profile.ProjectionAngularTolerance) : -1.0f;
+		JointSettings.ParentInvMassScale = Profile.bParentDominates ? (FReal)0 : (FReal)1;
+
+		JointSettings.bCollisionEnabled = !Profile.bDisableCollision;
+		JointSettings.bProjectionEnabled = Profile.bEnableProjection;
+		JointSettings.bShockPropagationEnabled = Profile.bEnableShockPropagation;
+		JointSettings.bMassConditioningEnabled = Profile.bEnableMassConditioning;
+
+		JointSettings.LinearMotionTypes[0] = ConvertToJointMotionType(Profile.LinearLimit.XMotion);
+		JointSettings.LinearMotionTypes[1] = ConvertToJointMotionType(Profile.LinearLimit.YMotion);
+		JointSettings.LinearMotionTypes[2] = ConvertToJointMotionType(Profile.LinearLimit.ZMotion);
+
+		JointSettings.LinearLimit = Profile.LinearLimit.Limit;
+
+		// Order is twist, swing1, swing2 and in degrees
+		JointSettings.AngularMotionTypes[(int32)Chaos::EJointAngularConstraintIndex::Twist] = ConvertToJointMotionType(Profile.TwistLimit.TwistMotion);
+		JointSettings.AngularMotionTypes[(int32)Chaos::EJointAngularConstraintIndex::Swing1] = ConvertToJointMotionType(Profile.ConeLimit.Swing1Motion);
+		JointSettings.AngularMotionTypes[(int32)Chaos::EJointAngularConstraintIndex::Swing2] = ConvertToJointMotionType(Profile.ConeLimit.Swing2Motion);
+
+		JointSettings.AngularLimits[(int32)Chaos::EJointAngularConstraintIndex::Twist] = FMath::DegreesToRadians(Profile.TwistLimit.TwistLimitDegrees);
+		JointSettings.AngularLimits[(int32)Chaos::EJointAngularConstraintIndex::Swing1] = FMath::DegreesToRadians(Profile.ConeLimit.Swing1LimitDegrees);
+		JointSettings.AngularLimits[(int32)Chaos::EJointAngularConstraintIndex::Swing2] = FMath::DegreesToRadians(Profile.ConeLimit.Swing2LimitDegrees);
+
+		JointSettings.bSoftLinearLimitsEnabled = Profile.LinearLimit.bSoftConstraint;
+		JointSettings.bSoftTwistLimitsEnabled = Profile.TwistLimit.bSoftConstraint;
+		JointSettings.bSoftSwingLimitsEnabled = Profile.ConeLimit.bSoftConstraint;
+
+		JointSettings.LinearSoftForceMode = (ConstraintSettings::SoftLinearForceMode() == 0) ?
+			EJointForceMode::Acceleration : EJointForceMode::Force;
+		JointSettings.AngularSoftForceMode = (ConstraintSettings::SoftAngularForceMode() == 0) ?
+			EJointForceMode::Acceleration : EJointForceMode::Force;
+
+		JointSettings.SoftLinearStiffness =
+			Chaos::ConstraintSettings::SoftLinearStiffnessScale() * Profile.LinearLimit.Stiffness;
+		JointSettings.SoftLinearDamping =
+			Chaos::ConstraintSettings::SoftLinearDampingScale() * Profile.LinearLimit.Damping;
+		JointSettings.SoftTwistStiffness =
+			Chaos::ConstraintSettings::SoftAngularStiffnessScale() * Profile.TwistLimit.Stiffness;
+		JointSettings.SoftTwistDamping =
+			Chaos::ConstraintSettings::SoftAngularDampingScale() * Profile.TwistLimit.Damping;
+		JointSettings.SoftSwingStiffness =
+			Chaos::ConstraintSettings::SoftAngularStiffnessScale() * Profile.ConeLimit.Stiffness;
+		JointSettings.SoftSwingDamping =
+			Chaos::ConstraintSettings::SoftAngularDampingScale() * Profile.ConeLimit.Damping;
+
+		JointSettings.LinearRestitution = Profile.LinearLimit.Restitution;
+		JointSettings.TwistRestitution = Profile.TwistLimit.Restitution;
+		JointSettings.SwingRestitution = Profile.ConeLimit.Restitution;
+
+		JointSettings.LinearContactDistance = Profile.LinearLimit.ContactDistance;
+		JointSettings.TwistContactDistance = Profile.TwistLimit.ContactDistance;
+		JointSettings.SwingContactDistance = Profile.ConeLimit.ContactDistance;
+
+		JointSettings.LinearDrivePositionTarget = Profile.LinearDrive.PositionTarget;
+		JointSettings.LinearDriveVelocityTarget = Profile.LinearDrive.VelocityTarget;
+		JointSettings.bLinearPositionDriveEnabled[0] = Profile.LinearDrive.XDrive.bEnablePositionDrive;
+		JointSettings.bLinearPositionDriveEnabled[1] = Profile.LinearDrive.YDrive.bEnablePositionDrive;
+		JointSettings.bLinearPositionDriveEnabled[2] = Profile.LinearDrive.ZDrive.bEnablePositionDrive;
+		JointSettings.bLinearVelocityDriveEnabled[0] = Profile.LinearDrive.XDrive.bEnableVelocityDrive;
+		JointSettings.bLinearVelocityDriveEnabled[1] = Profile.LinearDrive.YDrive.bEnableVelocityDrive;
+		JointSettings.bLinearVelocityDriveEnabled[2] = Profile.LinearDrive.ZDrive.bEnableVelocityDrive;
+
+		JointSettings.LinearDriveForceMode = EJointForceMode::Acceleration; // hardcoded!
+		JointSettings.LinearDriveStiffness = Chaos::ConstraintSettings::LinearDriveStiffnessScale() * FVec3(
+			Profile.LinearDrive.XDrive.Stiffness,
+			Profile.LinearDrive.YDrive.Stiffness,
+			Profile.LinearDrive.ZDrive.Stiffness);
+		JointSettings.LinearDriveDamping = Chaos::ConstraintSettings::LinearDriveDampingScale() * FVec3(
+			Profile.LinearDrive.XDrive.Damping,
+			Profile.LinearDrive.YDrive.Damping,
+			Profile.LinearDrive.ZDrive.Damping);
+		JointSettings.LinearDriveMaxForce[0] = Profile.LinearDrive.XDrive.MaxForce;
+		JointSettings.LinearDriveMaxForce[1] = Profile.LinearDrive.YDrive.MaxForce;
+		JointSettings.LinearDriveMaxForce[2] = Profile.LinearDrive.ZDrive.MaxForce;
+
+		JointSettings.AngularDrivePositionTarget = FQuat(Profile.AngularDrive.OrientationTarget);
+		JointSettings.AngularDriveVelocityTarget = Profile.AngularDrive.AngularVelocityTarget * UE_TWO_PI; // rev/s to rad/s
+
+		JointSettings.AngularDriveForceMode = EJointForceMode::Acceleration; // hardcoded!
 		if (Profile.AngularDrive.AngularDriveMode == EAngularDriveMode::SLERP)
 		{
-			ConstraintSettings.AngularDriveStiffness = Chaos::ConstraintSettings::AngularDriveStiffnessScale() * FVec3(Profile.AngularDrive.SlerpDrive.Stiffness);
-			ConstraintSettings.AngularDriveDamping = Chaos::ConstraintSettings::AngularDriveDampingScale() * FVec3(Profile.AngularDrive.SlerpDrive.Damping);
-			ConstraintSettings.bAngularSLerpPositionDriveEnabled = Profile.AngularDrive.SlerpDrive.bEnablePositionDrive;
-			ConstraintSettings.bAngularSLerpVelocityDriveEnabled = Profile.AngularDrive.SlerpDrive.bEnableVelocityDrive;
+			JointSettings.AngularDriveStiffness = FVec3(
+				ConstraintSettings::AngularDriveStiffnessScale() * Profile.AngularDrive.SlerpDrive.Stiffness);
+			JointSettings.AngularDriveDamping = FVec3(
+				ConstraintSettings::AngularDriveDampingScale() * Profile.AngularDrive.SlerpDrive.Damping);
+			JointSettings.AngularDriveMaxTorque = FVec3(
+				Profile.AngularDrive.SlerpDrive.MaxForce);
+			JointSettings.bAngularSLerpPositionDriveEnabled = Profile.AngularDrive.SlerpDrive.bEnablePositionDrive;
+			JointSettings.bAngularSLerpVelocityDriveEnabled = Profile.AngularDrive.SlerpDrive.bEnableVelocityDrive;
+			JointSettings.bAngularTwistPositionDriveEnabled = false;
+			JointSettings.bAngularTwistVelocityDriveEnabled = false;
+			JointSettings.bAngularSwingPositionDriveEnabled = false;
+			JointSettings.bAngularSwingVelocityDriveEnabled = false;
 		}
 		else
 		{
-			ConstraintSettings.AngularDriveStiffness = Chaos::ConstraintSettings::AngularDriveStiffnessScale() * FVec3(Profile.AngularDrive.TwistDrive.Stiffness, Profile.AngularDrive.SwingDrive.Stiffness, Profile.AngularDrive.SwingDrive.Stiffness);
-			ConstraintSettings.AngularDriveDamping = Chaos::ConstraintSettings::AngularDriveDampingScale() * FVec3(Profile.AngularDrive.TwistDrive.Damping, Profile.AngularDrive.SwingDrive.Damping, Profile.AngularDrive.SwingDrive.Damping);
-			ConstraintSettings.bAngularTwistPositionDriveEnabled = Profile.AngularDrive.TwistDrive.bEnablePositionDrive;
-			ConstraintSettings.bAngularTwistVelocityDriveEnabled = Profile.AngularDrive.TwistDrive.bEnableVelocityDrive;
-			ConstraintSettings.bAngularSwingPositionDriveEnabled = Profile.AngularDrive.SwingDrive.bEnablePositionDrive;
-			ConstraintSettings.bAngularSwingVelocityDriveEnabled = Profile.AngularDrive.SwingDrive.bEnableVelocityDrive;
+			JointSettings.AngularDriveStiffness = ConstraintSettings::AngularDriveStiffnessScale() * FVec3(
+				Profile.AngularDrive.TwistDrive.Stiffness,
+				Profile.AngularDrive.SwingDrive.Stiffness,
+				Profile.AngularDrive.SwingDrive.Stiffness);
+			JointSettings.AngularDriveDamping = ConstraintSettings::AngularDriveDampingScale() * FVec3(
+				Profile.AngularDrive.TwistDrive.Damping,
+				Profile.AngularDrive.SwingDrive.Damping,
+				Profile.AngularDrive.SwingDrive.Damping);
+			JointSettings.AngularDriveMaxTorque[0] = Profile.AngularDrive.TwistDrive.MaxForce;
+			JointSettings.AngularDriveMaxTorque[1] = Profile.AngularDrive.SwingDrive.MaxForce;
+			JointSettings.AngularDriveMaxTorque[2] = Profile.AngularDrive.SwingDrive.MaxForce;
+			JointSettings.bAngularSLerpPositionDriveEnabled = false;
+			JointSettings.bAngularSLerpVelocityDriveEnabled = false;
+			JointSettings.bAngularTwistPositionDriveEnabled = Profile.AngularDrive.TwistDrive.bEnablePositionDrive;
+			JointSettings.bAngularTwistVelocityDriveEnabled = Profile.AngularDrive.TwistDrive.bEnableVelocityDrive;
+			JointSettings.bAngularSwingPositionDriveEnabled = Profile.AngularDrive.SwingDrive.bEnablePositionDrive;
+			JointSettings.bAngularSwingVelocityDriveEnabled = Profile.AngularDrive.SwingDrive.bEnableVelocityDrive;
 		}
-		ConstraintSettings.AngularDriveForceMode = EJointForceMode::Acceleration;
 
-		ConstraintSettings.LinearBreakForce = (Profile.bLinearBreakable) ? Chaos::ConstraintSettings::LinearBreakScale() * Profile.LinearBreakThreshold : FLT_MAX;
-		ConstraintSettings.LinearPlasticityLimit = (Profile.bLinearPlasticity) ? FMath::Clamp((float)Profile.LinearPlasticityThreshold, 0.f, 1.f) : FLT_MAX;
-		ConstraintSettings.AngularBreakTorque = (Profile.bAngularBreakable) ? Chaos::ConstraintSettings::AngularBreakScale() * Profile.AngularBreakThreshold : FLT_MAX;
-		ConstraintSettings.AngularPlasticityLimit = (Profile.bAngularPlasticity) ? Profile.AngularPlasticityThreshold : FLT_MAX;
+		JointSettings.LinearBreakForce = Profile.bLinearBreakable ?
+			Chaos::ConstraintSettings::LinearBreakScale() * Profile.LinearBreakThreshold : FLT_MAX;
+		JointSettings.LinearPlasticityLimit = Profile.bLinearPlasticity ?
+			FMath::Clamp((float)Profile.LinearPlasticityThreshold, 0.0f, 1.0f) : FLT_MAX;
+		JointSettings.LinearPlasticityType = ConvertToPlasticityType(Profile.LinearPlasticityType);
+		// JointSettings.LinearPlasticityInitialDistanceSquared = ; // What do we do with this?
+
+		JointSettings.AngularBreakTorque = Profile.bAngularBreakable ?
+			Chaos::ConstraintSettings::AngularBreakScale() * Profile.AngularBreakThreshold : FLT_MAX;
+		JointSettings.AngularPlasticityLimit = Profile.bAngularPlasticity ?
+			FMath::Clamp((float)Profile.AngularPlasticityThreshold, 0.0f, 1.0f) : FLT_MAX;;
 
-		ConstraintSettings.ContactTransferScale = 0.0f;
+		JointSettings.ContactTransferScale = Profile.ContactTransferScale;
 
 
 		// UE Disables Soft Limits when the Limit is less than some threshold. This is not necessary in Chaos but for now we also do it for parity's sake (See FLinearConstraint::UpdateLinearLimit_AssumesLocked).
-		if (ConstraintSettings.LinearLimit < RB_MinSizeToLockDOF)
+		if (JointSettings.LinearLimit < RB_MinSizeToLockDOF)
 		{
 			for (int32 AxisIndex = 0; AxisIndex < 3; ++AxisIndex)
 			{
-				if (ConstraintSettings.LinearMotionTypes[AxisIndex] == EJointMotionType::Limited)
+				if (JointSettings.LinearMotionTypes[AxisIndex] == EJointMotionType::Limited)
 				{
-					ConstraintSettings.LinearMotionTypes[AxisIndex] = EJointMotionType::Locked;
+					JointSettings.LinearMotionTypes[AxisIndex] = EJointMotionType::Locked;
 				}
 			}
 		}
-
-		// Disable Soft Limits when stiffness is extremely high
-		//const FReal MaxLinearStiffness = 10000;
-		//const FReal MaxAngularStiffness = 10000;
-		//if (ConstraintSettings.bSoftLinearLimitsEnabled && ConstraintSettings.SoftLinearStiffness > MaxLinearStiffness)
-		//{
-		//	ConstraintSettings.bSoftLinearLimitsEnabled = false;
-		//}
-		//if (ConstraintSettings.bSoftTwistLimitsEnabled && ConstraintSettings.SoftTwistStiffness > MaxAngularStiffness)
-		//{
-		//	ConstraintSettings.bSoftTwistLimitsEnabled = false;
-		//}
-		//if (ConstraintSettings.bSoftSwingLimitsEnabled && ConstraintSettings.SoftSwingStiffness > MaxAngularStiffness)
-		//{
-		//	ConstraintSettings.bSoftSwingLimitsEnabled = false;
-		//}
 	}
 
 	FJointHandle::FJointHandle(FChaosConstraintContainer* InConstraints, FConstraintInstance* ConstraintInstance, FActorHandle* Actor1, FActorHandle* Actor2)
@@ -165,7 +219,7 @@ namespace ImmediatePhysics_Chaos
 		if (ConstraintInstance != nullptr)
 		{
 			// BodyInstance/PhysX has the constraint locations in actor-space, but we need them in Center-of-Mass space
-			TransferJointSettings(ConstraintInstance, ConstraintSettings);
+			UpdateJointSettingsFromConstraintProfile(ConstraintInstance->ProfileInstance, ConstraintSettings);
 			FReal JointScale = ConstraintInstance->GetLastKnownScale();
 			ConstraintSettings.ConnectorTransforms[0] = FParticleUtilities::ActorLocalToParticleLocal(FGenericParticleHandle(Actor1->GetParticle()), ConstraintInstance->GetRefFrame(EConstraintFrame::Frame1));
 			ConstraintSettings.ConnectorTransforms[1] = FParticleUtilities::ActorLocalToParticleLocal(FGenericParticleHandle(Actor2->GetParticle()), ConstraintInstance->GetRefFrame(EConstraintFrame::Frame2));
diff --git a/Engine/Source/Runtime/Engine/Public/Physics/ImmediatePhysics/ImmediatePhysicsChaos/ImmediatePhysicsJointHandle_Chaos.h b/Engine/Source/Runtime/Engine/Public/Physics/ImmediatePhysics/ImmediatePhysicsChaos/ImmediatePhysicsJointHandle_Chaos.h
--- a/Engine/Source/Runtime/Engine/Public/Physics/ImmediatePhysics/ImmediatePhysicsChaos/ImmediatePhysicsJointHandle_Chaos.h
+++ b/Engine/Source/Runtime/Engine/Public/Physics/ImmediatePhysics/ImmediatePhysicsChaos/ImmediatePhysicsJointHandle_Chaos.h
@@ -14,6 +14,8 @@ namespace Chaos
 	class FPBDJointSettings;
 }
 
+struct FConstraintProfileProperties;
+
 namespace ImmediatePhysics_Chaos
 {
 	/** handle associated with a physics joint. This is the proper way to read/write to the physics simulation */
@@ -43,4 +45,8 @@ namespace ImmediatePhysics_Chaos
 		FChaosConstraintContainer* Constraints;
 		FChaosConstraintHandle* ConstraintHandle;
 	};
+
+	/** Transfers parameters from a constraint profile into the joint settings */
+	ENGINE_API void UpdateJointSettingsFromConstraintProfile(
+		const FConstraintProfileProperties& Profile, Chaos::FPBDJointSettings& JointSettings);
 }
