Fix FSkeletalMeshRenderData StaticRayTracingGeometry not being correctly initialized when using mesh streaming.
- calling FRayTracingGeometry::InitResource() led to a BLAS build being scheduled via GRayTracingGeometryManager, which was then submitted to RHI before stream-in was completed.

diff --git a/Engine/Source/Runtime/Engine/Private/SkeletalRenderStatic.cpp b/Engine/Source/Runtime/Engine/Private/SkeletalRenderStatic.cpp
--- a/Engine/Source/Runtime/Engine/Private/SkeletalRenderStatic.cpp
+++ b/Engine/Source/Runtime/Engine/Private/SkeletalRenderStatic.cpp
@@ -107,13 +107,14 @@ void FSkeletalMeshObjectStatic::InitResources(USkinnedMeshComponent* InMeshCompo
 							}
 							Initializer.Segments = GeometrySections;
 
-							RayTracingGeometry.SetInitializer(Initializer);
-
-							if (LODIndex >= SkelMeshRenderData->CurrentFirstLODIdx) // According to GetMeshElementsConditionallySelectable(), non-resident LODs should just be skipped
+							if (LODIndex < SkelMeshRenderData->CurrentFirstLODIdx) // According to GetMeshElementsConditionallySelectable(), non-resident LODs should just be skipped
 							{
-								RayTracingGeometry.InitResource(RHICmdList);
+								Initializer.Type = ERayTracingGeometryInitializerType::StreamingDestination;
 							}
 
+							RayTracingGeometry.SetInitializer(Initializer);
+							RayTracingGeometry.InitResource(RHICmdList);
+
 							bReferencedByStaticSkeletalMeshObjects_RenderThread = true;
 						}
 					);
diff --git a/Engine/Source/Runtime/Engine/Private/Streaming/SkeletalMeshUpdate.cpp b/Engine/Source/Runtime/Engine/Private/Streaming/SkeletalMeshUpdate.cpp
--- a/Engine/Source/Runtime/Engine/Private/Streaming/SkeletalMeshUpdate.cpp
+++ b/Engine/Source/Runtime/Engine/Private/Streaming/SkeletalMeshUpdate.cpp
@@ -74,6 +74,48 @@ void FSkeletalMeshStreamIn::FIntermediateBuffers::TransferBuffers(FSkeletalMeshL
 	LODResource.SkinWeightProfilesData.InitRHIForStreaming(AltSkinWeightVertexBuffers, Batcher);
 }
 
+#if RHI_RAYTRACING
+
+void FSkeletalMeshStreamIn::FIntermediateRayTracingGeometry::CreateFromCPUData(FRHICommandListBase& RHICmdList, FRayTracingGeometry& RayTracingGeometry)
+{
+	Initializer = RayTracingGeometry.Initializer;
+	Initializer.Type = ERayTracingGeometryInitializerType::StreamingSource;
+
+	if (RayTracingGeometry.RawData.Num())
+	{
+		check(Initializer.OfflineData == nullptr);
+		Initializer.OfflineData = &RayTracingGeometry.RawData;
+	}
+
+	static const auto CVarDebugForceRuntimeBLAS = IConsoleManager::Get().FindTConsoleVariableDataInt(TEXT("r.Raytracing.DebugForceRuntimeBLAS"));
+	const bool bDebugForceRuntimeBLAS = (!CVarDebugForceRuntimeBLAS) || (CVarDebugForceRuntimeBLAS->GetValueOnAnyThread() != 0);
+
+	if (bDebugForceRuntimeBLAS && Initializer.OfflineData != nullptr)
+	{
+		Initializer.OfflineData->Discard();
+		Initializer.OfflineData = nullptr;
+	}
+
+	RayTracingGeometryRHI = RHICmdList.CreateRayTracingGeometry(Initializer);
+	bRequiresBuild = Initializer.OfflineData == nullptr || RayTracingGeometryRHI->IsCompressed();
+}
+
+void FSkeletalMeshStreamIn::FIntermediateRayTracingGeometry::SafeRelease()
+{
+	Initializer = {};
+	RayTracingGeometryRHI.SafeRelease();
+}
+
+void FSkeletalMeshStreamIn::FIntermediateRayTracingGeometry::TransferRayTracingGeometry(FRayTracingGeometry& RayTracingGeometry, FRHIResourceReplaceBatcher& Batcher)
+{
+	RayTracingGeometry.InitRHIForStreaming(RayTracingGeometryRHI, Batcher);
+	RayTracingGeometry.SetRequiresBuild(bRequiresBuild);
+
+	SafeRelease();
+}
+
+#endif
+
 FSkeletalMeshStreamIn::FSkeletalMeshStreamIn(const USkeletalMesh* InMesh, EThreadType CreateResourcesThread)
 	: FSkeletalMeshUpdate(InMesh)
 	, CreateResourcesThread(CreateResourcesThread)
@@ -104,6 +146,14 @@ void FSkeletalMeshStreamIn::CreateBuffers(const FContext& Context)
 		{
 			FSkeletalMeshLODRenderData& LODResource = *Context.LODResourcesView[LODIndex];
 			IntermediateBuffersArray[LODIndex].CreateFromCPUData(*StreamingRHICmdList, LODResource);
+
+#if RHI_RAYTRACING
+			// Skip LODs that have their render data stripped
+			if (IsRayTracingEnabled() && Context.Mesh->GetSupportRayTracing() && LODResource.GetNumVertices() > 0 && LODResource.bReferencedByStaticSkeletalMeshObjects_RenderThread)
+			{
+				IntermediateRayTracingGeometry[LODIndex].CreateFromCPUData(*StreamingRHICmdList, LODResource.StaticRayTracingGeometry);
+			}
+#endif
 		}
 
 		// Use a scope to flush the batcher before updating CurrentFirstLODIdx
@@ -118,25 +168,6 @@ void FSkeletalMeshStreamIn::CreateBuffers(const FContext& Context)
 				IntermediateBuffersArray[LODIndex].TransferBuffers(LODResource, Batcher);
 			}
 		}
-
-#if RHI_RAYTRACING
-		// Must happen after the batched updates have been flushed
-		if (IsRayTracingAllowed())
-		{
-			for (int32 LODIndex = PendingFirstLODIdx; LODIndex < CurrentFirstLODIdx; ++LODIndex)
-			{
-				// Skip LODs that have their render data stripped
-				if (Context.RenderData->LODRenderData[LODIndex].GetNumVertices() > 0)
-				{
-					if (Context.RenderData->LODRenderData[LODIndex].bReferencedByStaticSkeletalMeshObjects_RenderThread)
-					{
-						ensure(!Context.RenderData->LODRenderData[LODIndex].StaticRayTracingGeometry.IsInitialized());
-						Context.RenderData->LODRenderData[LODIndex].StaticRayTracingGeometry.InitResource(*StreamingRHICmdList);
-					}
-				}
-			}
-		}
-#endif
 	}
 
 	StreamingRHICmdList->FinishRecording();
@@ -165,6 +196,42 @@ void FSkeletalMeshStreamIn::DoFinishUpdate(const FContext& Context)
 		StreamingRHICmdList = nullptr;
 	}
 
+#if RHI_RAYTRACING
+	if (IsRayTracingAllowed() && Context.Mesh->GetSupportRayTracing())
+	{
+		// Use a scope to flush the batcher before updating CurrentFirstLODIdx
+		{
+			FRHIResourceReplaceBatcher Batcher(FRHICommandListImmediate::Get(), GSkelMeshMaxNumResourceUpdatesPerBatch);
+			for (int32 LODIdx = PendingFirstLODIdx; LODIdx < CurrentFirstLODIdx; ++LODIdx)
+			{
+				FSkeletalMeshLODRenderData& LODResource = *Context.LODResourcesView[LODIdx];
+
+				if (LODResource.GetNumVertices() > 0 && LODResource.bReferencedByStaticSkeletalMeshObjects_RenderThread)
+				{
+					IntermediateRayTracingGeometry[LODIdx].TransferRayTracingGeometry(LODResource.StaticRayTracingGeometry, Batcher);
+				}
+			}
+		}
+
+		// Must happen after the batched updates have been flushed
+		for (int32 LODIndex = PendingFirstLODIdx; LODIndex < CurrentFirstLODIdx; ++LODIndex)
+		{
+			FSkeletalMeshLODRenderData& LODResource = *Context.LODResourcesView[LODIndex];
+
+			// Skip LODs that have their render data stripped
+			if (LODResource.GetNumVertices() > 0 && LODResource.bReferencedByStaticSkeletalMeshObjects_RenderThread)
+			{
+				// Under very rare circumstances that we switch ray tracing on/off right in the middle of streaming RayTracingGeometryRHI might not be valid.
+				if (IsRayTracingEnabled() && ensure(LODResource.StaticRayTracingGeometry.IsValid()))
+				{
+					LODResource.StaticRayTracingGeometry.RequestBuildIfNeeded(ERTAccelerationStructureBuildPriority::Normal);
+				}
+			}
+		}
+
+	}
+#endif
+
 	Context.RenderData->PendingFirstLODIdx = Context.RenderData->CurrentFirstLODIdx = ResourceState.LODCountToAssetFirstLODIdx(ResourceState.NumRequestedLODs);
 	MarkAsSuccessfullyFinished();
 }
diff --git a/Engine/Source/Runtime/Engine/Private/Streaming/SkeletalMeshUpdate.h b/Engine/Source/Runtime/Engine/Private/Streaming/SkeletalMeshUpdate.h
--- a/Engine/Source/Runtime/Engine/Private/Streaming/SkeletalMeshUpdate.h
+++ b/Engine/Source/Runtime/Engine/Private/Streaming/SkeletalMeshUpdate.h
@@ -12,6 +12,7 @@ SkeletalMeshUpdate.h: Helpers to stream in and out skeletal mesh LODs.
 #include "RenderAssetUpdate.h"
 #include "Rendering/SkinWeightVertexBuffer.h"
 #include "Serialization/BulkData.h"
+#include "RayTracingGeometry.h"
 
 /**
 * A context used to update or proceed with the next update step.
@@ -90,6 +91,23 @@ protected:
 		void TransferBuffers(FSkeletalMeshLODRenderData& LODResource, FRHIResourceReplaceBatcher& Batcher);
 	};
 
+#if RHI_RAYTRACING
+	struct FIntermediateRayTracingGeometry
+	{
+	private:
+		FRayTracingGeometryInitializer Initializer;
+		FRayTracingGeometryRHIRef RayTracingGeometryRHI;
+		bool bRequiresBuild = false;
+
+	public:
+		void CreateFromCPUData(FRHICommandListBase& RHICmdList, FRayTracingGeometry& RayTracingGeometry);
+
+		void SafeRelease();
+
+		void TransferRayTracingGeometry(FRayTracingGeometry& RayTracingGeometry, FRHIResourceReplaceBatcher& Batcher);
+	};
+#endif
+
 	/** Create buffers with new LOD data */
 	void CreateBuffers(const FContext& Context);
 
@@ -105,6 +123,10 @@ protected:
 	/** The intermediate buffers created in the update process. */
 	FIntermediateBuffers IntermediateBuffersArray[MAX_MESH_LOD_COUNT];
 
+#if RHI_RAYTRACING
+	FIntermediateRayTracingGeometry IntermediateRayTracingGeometry[MAX_MESH_LOD_COUNT];
+#endif
+
 	/** RHI command list used for creating buffers and replacing the streaming placeholders. Submitted in DoFinishUpdate */
 	FRHICommandList* StreamingRHICmdList = nullptr;
 
