Skeletal Mesh: Fix morphs getting disabled after geometry streaming re-initializes them.

If the morph vertex resource had been initialized and streaming then serialized the actual data in, FMorphTargetVertexInfoBuffers::Serialize would clear the bRHIInitialized flag. Then the subsequent streaming call to FSkeletalMeshLODRenderData::InitMorphResources would fail to release the earlier resource init, resulting in FRenderResource::InitResource effectively being a no-op, therefore leaving bRHIInitialized as false -- eventually causing FSkeletalMeshObjectGPUSkin::FSkeletalMeshObjectLOD::UpdateMorphVertexBufferGPU to fail.

Fixed a spelling mistake: RHIIntialized > RHIInitialized.

diff --git a/Engine/Plugins/Animation/MLDeformer/MLDeformerFramework/Source/MLDeformerFramework/Private/MLDeformerMorphModel.cpp b/Engine/Plugins/Animation/MLDeformer/MLDeformerFramework/Source/MLDeformerFramework/Private/MLDeformerMorphModel.cpp
--- a/Engine/Plugins/Animation/MLDeformer/MLDeformerFramework/Source/MLDeformerFramework/Private/MLDeformerMorphModel.cpp
+++ b/Engine/Plugins/Animation/MLDeformer/MLDeformerFramework/Source/MLDeformerFramework/Private/MLDeformerMorphModel.cpp
@@ -187,7 +187,7 @@ void UMLDeformerMorphModel::ClearMorphTargetSets()
 		if (GetMorphTargetSet(LOD).IsValid())
 		{
 			FMorphTargetVertexInfoBuffers& MorphBuffer = GetMorphTargetSet(LOD)->MorphBuffers;
-			if (MorphBuffer.IsRHIIntialized() && MorphBuffer.IsInitialized())
+			if (MorphBuffer.IsRHIInitialized() && MorphBuffer.IsInitialized())
 			{
 				ReleaseResourceAndFlush(&MorphBuffer);
 			}
diff --git a/Engine/Plugins/Animation/MLDeformer/MLDeformerFramework/Source/MLDeformerFrameworkEditor/Private/MLDeformerEditorModel.cpp b/Engine/Plugins/Animation/MLDeformer/MLDeformerFramework/Source/MLDeformerFrameworkEditor/Private/MLDeformerEditorModel.cpp
--- a/Engine/Plugins/Animation/MLDeformer/MLDeformerFramework/Source/MLDeformerFrameworkEditor/Private/MLDeformerEditorModel.cpp
+++ b/Engine/Plugins/Animation/MLDeformer/MLDeformerFramework/Source/MLDeformerFrameworkEditor/Private/MLDeformerEditorModel.cpp
@@ -2401,7 +2401,7 @@ namespace UE::MLDeformer
 		const int32 NumRenderVertices = RenderData->LODRenderData[LOD].GetNumVertices();
 
 		// Release any existing morph buffer data.
-		if (OutMorphBuffers.IsRHIIntialized() && OutMorphBuffers.IsInitialized())
+		if (OutMorphBuffers.IsRHIInitialized() && OutMorphBuffers.IsInitialized())
 		{
 			ReleaseResourceAndFlush(&OutMorphBuffers);
 		}
diff --git a/Engine/Plugins/Animation/MLDeformer/MLDeformerFramework/Source/MLDeformerFrameworkEditor/Private/MLDeformerMorphModelEditorModel.cpp b/Engine/Plugins/Animation/MLDeformer/MLDeformerFramework/Source/MLDeformerFrameworkEditor/Private/MLDeformerMorphModelEditorModel.cpp
--- a/Engine/Plugins/Animation/MLDeformer/MLDeformerFramework/Source/MLDeformerFrameworkEditor/Private/MLDeformerMorphModelEditorModel.cpp
+++ b/Engine/Plugins/Animation/MLDeformer/MLDeformerFramework/Source/MLDeformerFrameworkEditor/Private/MLDeformerMorphModelEditorModel.cpp
@@ -537,7 +537,7 @@ namespace UE::MLDeformer
 		const int32 NumRenderVertices = RenderData->LODRenderData[LOD].GetNumVertices();
 
 		// Release any existing morph buffer data.
-		if (OutMorphBuffers.IsRHIIntialized() && OutMorphBuffers.IsInitialized())
+		if (OutMorphBuffers.IsRHIInitialized() && OutMorphBuffers.IsInitialized())
 		{
 			ReleaseResourceAndFlush(&OutMorphBuffers);
 		}
diff --git a/Engine/Source/Runtime/Engine/Private/Rendering/MorphTargetVertexInfoBuffers.cpp b/Engine/Source/Runtime/Engine/Private/Rendering/MorphTargetVertexInfoBuffers.cpp
--- a/Engine/Source/Runtime/Engine/Private/Rendering/MorphTargetVertexInfoBuffers.cpp
+++ b/Engine/Source/Runtime/Engine/Private/Rendering/MorphTargetVertexInfoBuffers.cpp
@@ -16,7 +16,7 @@ void FMorphTargetVertexInfoBuffers::InitRHI(FRHICommandListBase& RHICmdList)
 	SCOPED_LOADTIMER(FFMorphTargetVertexInfoBuffers_InitRHI);
 
 	check(NumTotalBatches > 0);
-	check(!bRHIIntialized);
+	check(!bRHIInitialized);
 
 	const static FLazyName ClassName(TEXT("FMorphTargetVertexInfoBuffers"));
 
@@ -38,14 +38,14 @@ void FMorphTargetVertexInfoBuffers::InitRHI(FRHICommandListBase& RHICmdList)
 		bIsMorphCPUDataValid = false;
 	}
 
-	bRHIIntialized = true;
+	bRHIInitialized = true;
 }
 
 void FMorphTargetVertexInfoBuffers::ReleaseRHI()
 {
 	MorphDataBuffer.SafeRelease();
 	MorphDataSRV.SafeRelease();
-	bRHIIntialized = false;
+	bRHIInitialized = false;
 }
 
 uint32 FMorphTargetVertexInfoBuffers::GetMaximumThreadGroupSize()
@@ -121,7 +121,7 @@ void FMorphTargetVertexInfoBuffers::Serialize(FArchive& Ar)
 
 	if (Ar.IsLoading())
 	{
-		bRHIIntialized = false;
+		bRHIInitialized = false;
 		bIsMorphCPUDataValid = true;
 		bResourcesInitialized = true;
 		ValidateVertexBuffers(true);
@@ -174,7 +174,7 @@ private:
 
 void FMorphTargetVertexInfoBuffers::InitMorphResources(EShaderPlatform ShaderPlatform, const TArray<FSkelMeshRenderSection>& RenderSections, const TArray<UMorphTarget*>& MorphTargets, int NumVertices, int32 LODIndex, float TargetPositionErrorTolerance)
 {
-	check(!IsRHIIntialized());
+	check(!IsRHIInitialized());
 	check(!IsMorphResourcesInitialized());
 	check(!IsMorphCPUDataValid());
 
diff --git a/Engine/Source/Runtime/Engine/Private/SkeletalMeshLODRenderData.cpp b/Engine/Source/Runtime/Engine/Private/SkeletalMeshLODRenderData.cpp
--- a/Engine/Source/Runtime/Engine/Private/SkeletalMeshLODRenderData.cpp
+++ b/Engine/Source/Runtime/Engine/Private/SkeletalMeshLODRenderData.cpp
@@ -36,13 +36,6 @@ static FAutoConsoleVariableRef CVarStripSkeletalMeshLodsBelowMinLod(
 
 
 
-int32 GAllowSkinnedMorphDataStreaming = 0;
-static FAutoConsoleVariableRef CVarAllowSkinnedMorphDataStreaming(
-	TEXT("r.Skinned.AllowSkinnedMorphDataStreaming"),
-	GAllowSkinnedMorphDataStreaming,
-	TEXT("Call BeginInitResource when the morph resources are finish streaming (DoFinishUpdate)")
-);
-
 namespace
 {
 	struct FReverseOrderBitArraysBySetBits
@@ -225,8 +218,14 @@ FArchive& operator<<(FArchive& Ar, FSkelMeshRenderSection& S)
 
 void FSkeletalMeshLODRenderData::InitMorphResources()
 {
-	if (GAllowSkinnedMorphDataStreaming && !MorphTargetVertexInfoBuffers.IsRHIIntialized() && MorphTargetVertexInfoBuffers.IsMorphCPUDataValid() && MorphTargetVertexInfoBuffers.NumTotalBatches > 0)
+	if (!MorphTargetVertexInfoBuffers.IsRHIInitialized() && MorphTargetVertexInfoBuffers.IsMorphCPUDataValid() && MorphTargetVertexInfoBuffers.NumTotalBatches > 0)
 	{
+		// The morph target could have been loaded prior but gets streamed in again, so we have to release the resources here, otherwise
+		// FRenderResource::InitResource does nothing and leaves IsRHIInitialized() as false, which results in no morphs appearing. 
+		if (MorphTargetVertexInfoBuffers.IsInitialized())
+		{
+			BeginReleaseResource(&MorphTargetVertexInfoBuffers, &UE::RenderCommandPipe::SkeletalMesh);
+		}
 		BeginInitResource(&MorphTargetVertexInfoBuffers, &UE::RenderCommandPipe::SkeletalMesh);
 	}
 }
@@ -296,7 +295,7 @@ void FSkeletalMeshLODRenderData::InitResources(bool bNeedsVertexColors, int32 LO
 		MorphTargetVertexInfoBuffers.InitMorphResources(GMaxRHIShaderPlatform, RenderSections, Owner->GetMorphTargets(), StaticVertexBuffers.StaticMeshVertexBuffer.GetNumVertices(), LODIndex, SkeletalMeshLODInfo->MorphTargetPositionErrorTolerance);
 	}
 	
-	if (!MorphTargetVertexInfoBuffers.IsRHIIntialized() && MorphTargetVertexInfoBuffers.IsMorphCPUDataValid() && MorphTargetVertexInfoBuffers.NumTotalBatches > 0)
+	if (!MorphTargetVertexInfoBuffers.IsRHIInitialized() && MorphTargetVertexInfoBuffers.IsMorphCPUDataValid() && MorphTargetVertexInfoBuffers.NumTotalBatches > 0)
 	{
 		MorphTargetVertexInfoBuffers.SetOwnerName(OwnerName);
 		BeginInitResource(&MorphTargetVertexInfoBuffers, &UE::RenderCommandPipe::SkeletalMesh);
@@ -833,7 +832,7 @@ void FSkeletalMeshLODRenderData::SerializeStreamedData(FArchive& Ar, USkinnedAss
 					TargetMorphTargetVertexInfoBuffers = &MorphTargetVertexInfoBuffers;
 					TargetMorphTargetVertexInfoBuffers->InitMorphResources(MorphTargetShaderPlatform, RenderSections, MorphTargets, StaticVertexBuffers.StaticMeshVertexBuffer.GetNumVertices(), LODIdx, SkeletalMeshLODInfo->MorphTargetPositionErrorTolerance);
 				}
-				else if (MorphTargetVertexInfoBuffers.IsMorphCPUDataValid() && !MorphTargetVertexInfoBuffers.IsRHIIntialized())
+				else if (MorphTargetVertexInfoBuffers.IsMorphCPUDataValid() && !MorphTargetVertexInfoBuffers.IsRHIInitialized())
 				{
 					TargetMorphTargetVertexInfoBuffers = &MorphTargetVertexInfoBuffers;
 					check(TargetMorphTargetVertexInfoBuffers->IsMorphCPUDataValid());
diff --git a/Engine/Source/Runtime/Engine/Private/SkeletalRenderGPUSkin.cpp b/Engine/Source/Runtime/Engine/Private/SkeletalRenderGPUSkin.cpp
--- a/Engine/Source/Runtime/Engine/Private/SkeletalRenderGPUSkin.cpp
+++ b/Engine/Source/Runtime/Engine/Private/SkeletalRenderGPUSkin.cpp
@@ -1267,7 +1267,7 @@ void FSkeletalMeshObjectGPUSkin::FSkeletalMeshObjectLOD::UpdateMorphVertexBuffer
 			RHICmdList.ClearUAVUint(MorphVertexBuffer.GetUAV(), FUintVector4(0, 0, 0, 0));
 		}
 
-		if (MorphTargetVertexInfoBuffers.IsRHIIntialized() && (MorphTargetVertexInfoBuffers.GetNumMorphs() > 0))
+		if (MorphTargetVertexInfoBuffers.IsRHIInitialized() && (MorphTargetVertexInfoBuffers.GetNumMorphs() > 0))
 		{
 			{
 				SCOPED_DRAW_EVENTF(RHICmdList, MorphUpdateScatter, TEXT("Scatter"));
diff --git a/Engine/Source/Runtime/Engine/Public/Rendering/MorphTargetVertexInfoBuffers.h b/Engine/Source/Runtime/Engine/Public/Rendering/MorphTargetVertexInfoBuffers.h
--- a/Engine/Source/Runtime/Engine/Public/Rendering/MorphTargetVertexInfoBuffers.h
+++ b/Engine/Source/Runtime/Engine/Public/Rendering/MorphTargetVertexInfoBuffers.h
@@ -18,7 +18,7 @@ public:
 	ENGINE_API void InitMorphResources(EShaderPlatform ShaderPlatform, const TArray<FSkelMeshRenderSection>& RenderSections, const TArray<UMorphTarget*>& MorphTargets, int NumVertices, int32 LODIndex, float TargetPositionErrorTolerance);
 
 	inline bool IsMorphResourcesInitialized() const { return bResourcesInitialized; }
-	inline bool IsRHIIntialized() const { return bRHIIntialized; }
+	inline bool IsRHIInitialized() const { return bRHIInitialized; }
 	inline bool IsMorphCPUDataValid() const{ return bIsMorphCPUDataValid; }
 	
 	bool GetEmptyMorphCPUDataOnInitRHI() const { return bEmptyMorphCPUDataOnInitRHI; }
@@ -102,7 +102,7 @@ protected:
 
 	bool bIsMorphCPUDataValid = false;
 	bool bResourcesInitialized = false;
-	bool bRHIIntialized = false;
+	bool bRHIInitialized = false;
 	bool bEmptyMorphCPUDataOnInitRHI = true;
 
 	friend class FSkeletalMeshLODRenderData;
