-Initial experimental WorldPartition support for VolumetricLightMaps, must be enabled with the r.AllowStaticLightingInWorldPartitionMaps cvar
-Needs whole map in memory for proper computation
-Uses the currently implemented level streaming mechanism (streaming of highest detail bricks only, no dynamic stitching for correct filtering of bricks that switch neighboring levels)
-Contains a initial version of the static lighting builder for command line computation

diff --git a/Engine/Shaders/Private/VolumetricLightmapStreaming.usf b/Engine/Shaders/Private/VolumetricLightmapStreaming.usf
--- a/Engine/Shaders/Private/VolumetricLightmapStreaming.usf
+++ b/Engine/Shaders/Private/VolumetricLightmapStreaming.usf
@@ -35,7 +35,8 @@ void RemoveSubLevelBricksCS(uint3 BrickID : SV_DispatchThreadID)
 		SubLevelBrickPositions[BrickID.x * 3 + 2]
 	);
 	
-	uint4 IndirectionTextureOriginalValue = IndirectionTextureOriginalValues[BrickID.x];
+	// IndirectionTextureOriginalValues is stored as an FColor which is a bgra memory layout so we swizzle on read to get the correct values
+	uint4 IndirectionTextureOriginalValue = IndirectionTextureOriginalValues[BrickID.x].bgra;	
 	uint OriginalBrickLinearAddress = IndirectionTextureOriginalValue.x + IndirectionTextureOriginalValue.y * 256 + IndirectionTextureOriginalValue.z * 256 * 256;
 	IndirectionTexture[BrickPos] = uint4(ComputeBrickLayoutPosition(OriginalBrickLinearAddress + PersistentLevelBrickDataBaseOffset, uint3(256, 256, 256)), IndirectionTextureOriginalValue.w);
 }
@@ -124,4 +125,4 @@ void MoveWholeIndirectionTextureCS(uint3 IndirectionTextureCellID : SV_DispatchT
 	{
 		IndirectionTexture[IndirectionTextureCellID] = IndirectionTextureOriginalValue;
 	}
-}
+}
\ No newline at end of file
diff --git a/Engine/Source/Editor/UnrealEd/Private/Lightmass/ImportVolumetricLightmap.cpp b/Engine/Source/Editor/UnrealEd/Private/Lightmass/ImportVolumetricLightmap.cpp
--- a/Engine/Source/Editor/UnrealEd/Private/Lightmass/ImportVolumetricLightmap.cpp
+++ b/Engine/Source/Editor/UnrealEd/Private/Lightmass/ImportVolumetricLightmap.cpp
@@ -15,6 +15,7 @@
 #include "Logging/TokenizedMessage.h"
 #include "Logging/MessageLog.h"
 #include "RenderUtils.h"
+#include "WorldPartition/StaticLightingData/VolumetricLightmapGrid.h"
 
 #define LOCTEXT_NAMESPACE "ImportVolumetricLightmap"
 
@@ -996,6 +997,11 @@ bool bStitchDetailBricksWithLowDensityNeighbors = true;
 bool bCopyPaddingFromUniqueData = true;
 bool bCopyVolumeBorderFromInterior = true;
 
+void FLightmassProcessor::SetVolumetricLightMapImportMode(bool bUseVLMCellGrid)
+{
+	bSplitToVLMCellGrid = bUseVLMCellGrid;
+}
+
 void FLightmassProcessor::ImportVolumetricLightmap()
 {
 	const double StartTime = FPlatformTime::Seconds();
@@ -1079,7 +1085,10 @@ void FLightmassProcessor::ImportVolumetricLightmap()
 	const int32 NumBottomLevelBricksTrimmedByInterpolationError = TrimBricksByInterpolationError(BricksByDepth, VolumetricLightmapSettings);
 
 	const float MaximumBrickMemoryMb = System.GetWorld()->GetWorldSettings()->LightmassSettings.VolumetricLightmapMaximumBrickMemoryMb;
-	const int32 NumBottomLevelBricksTrimmedForMemoryLimit = TrimBricksForMemoryLimit(BricksByDepth, VolumetricLightmapSettings, CurrentLevelData.BrickData.GetMinimumVoxelSize(), MaximumBrickMemoryMb);
+
+	// avoid trimming for memory if we're going to output the bricks to streaming grid
+	// @todo_ow: Have a trim limit for base data + one limit per cell for the highest level bricks
+	const int32 NumBottomLevelBricksTrimmedForMemoryLimit = bSplitToVLMCellGrid ? 0 : TrimBricksForMemoryLimit(BricksByDepth, VolumetricLightmapSettings, CurrentLevelData.BrickData.GetMinimumVoxelSize(), MaximumBrickMemoryMb);
 
 	int32 BrickTextureLinearAllocator = 0;
 
@@ -1295,9 +1304,18 @@ void FLightmassProcessor::ImportVolumetricLightmap()
 
 				FGuid LevelGuid;
 
+				// assign to proper data container
 				if (CurrentDepth == VolumetricLightmapSettings.MaxRefinementLevels - 1)
 				{
-					LevelGuid = Brick.IntersectingLevelGuid;
+					if (bSplitToVLMCellGrid)
+					{
+						LevelGuid = System.LightingContext.GetLevelGuidForVLMBrick(Brick.IndirectionTexturePosition);
+						const_cast<FImportedVolumetricLightmapBrick&>(Brick).IntersectingLevelGuid = LevelGuid;	// change the level guid for this brick, I use a const cast since it preserves the intent everywhere else
+					}
+                    else
+                    {
+						LevelGuid = Brick.IntersectingLevelGuid;
+                    }
 				}
 				else
 				{
@@ -1329,8 +1347,7 @@ void FLightmassProcessor::ImportVolumetricLightmap()
 			SubLevelBrickLayoutDimension.Z = FMath::Min(SubLevelBrickTextureLinearAllocator, MaxBricksInLayoutOneDim);
 
 			{
-				UMapBuildDataRegistry* SubLevelRegistry = System.LightingContext.GetOrCreateRegistryForLevelGuid(Pair.Key);
-				FPrecomputedVolumetricLightmapData& SubLevelData = SubLevelRegistry->AllocateLevelPrecomputedVolumetricLightmapBuildData(System.LightingContext.GetLevelBuildDataID(Pair.Key));
+				FPrecomputedVolumetricLightmapData& SubLevelData =  System.LightingContext.GetOrCreateLevelPrecomputedVolumetricLightmapBuildData(Pair.Key);
 
 				SubLevelData.InitializeOnImport(FBox(VolumetricLightmapSettings.VolumeMin, VolumetricLightmapSettings.VolumeMin + VolumetricLightmapSettings.VolumeSize), BrickSize);
 				{
@@ -1388,8 +1405,7 @@ void FLightmassProcessor::ImportVolumetricLightmap()
 					LevelGuid = PersistentLevelGuid;
 				}
 
-				UMapBuildDataRegistry* SubLevelRegistry = System.LightingContext.GetOrCreateRegistryForLevelGuid(LevelGuid);
-				FPrecomputedVolumetricLightmapData& SubLevelData = *SubLevelRegistry->GetLevelPrecomputedVolumetricLightmapBuildData(System.LightingContext.GetLevelBuildDataID(LevelGuid));
+				FPrecomputedVolumetricLightmapData& SubLevelData =  System.LightingContext.GetOrCreateLevelPrecomputedVolumetricLightmapBuildData(LevelGuid);
 
 				const FIntVector BrickLayoutPosition = ComputeBrickLayoutPosition(BrickStartAllocation + BrickIndex, BrickLayoutDimensions) * PaddedBrickSize;
 				const FIntVector SubLevelBrickLayoutPosition = ComputeBrickLayoutPosition(SubLevelBrickAllocator[LevelGuid], SubLevelBrickLayoutDimensions[LevelGuid]) * PaddedBrickSize;
@@ -1457,8 +1473,7 @@ void FLightmassProcessor::ImportVolumetricLightmap()
 			CurrentLevelData.IndirectionTexture.Resize(TotalIndirectionTextureSize * IndirectionTextureDataStride);
 			BuildIndirectionTexture(BricksByDepth, VolumetricLightmapSettings, MaxBricksInLayoutOneDim, SubLevelBrickLayoutDimensions[FGuid()], IndirectionTextureDataStride, CurrentLevelData, true);
 
-			UMapBuildDataRegistry* SubLevelRegistry = System.LightingContext.GetOrCreateRegistryForLevelGuid(PersistentLevelGuid);
-			FPrecomputedVolumetricLightmapData& PersistentLevelData = *SubLevelRegistry->GetLevelPrecomputedVolumetricLightmapBuildData(System.LightingContext.GetLevelBuildDataID(PersistentLevelGuid));
+			FPrecomputedVolumetricLightmapData& PersistentLevelData =  System.LightingContext.GetOrCreateLevelPrecomputedVolumetricLightmapBuildData(PersistentLevelGuid);
 
 			PersistentLevelData.InitializeOnImport(FBox(VolumetricLightmapSettings.VolumeMin, VolumetricLightmapSettings.VolumeMin + VolumetricLightmapSettings.VolumeSize), BrickSize);
 			PersistentLevelData.IndirectionTexture = CurrentLevelData.IndirectionTexture;
@@ -1488,8 +1503,7 @@ void FLightmassProcessor::ImportVolumetricLightmap()
 
 					if (LevelGuid != PersistentLevelGuid)
 					{
-						UMapBuildDataRegistry* SubLevelRegistry = System.LightingContext.GetOrCreateRegistryForLevelGuid(LevelGuid);
-						FPrecomputedVolumetricLightmapData& SubLevelData = *SubLevelRegistry->GetLevelPrecomputedVolumetricLightmapBuildData(System.LightingContext.GetLevelBuildDataID(LevelGuid));
+						FPrecomputedVolumetricLightmapData& SubLevelData =  System.LightingContext.GetOrCreateLevelPrecomputedVolumetricLightmapBuildData(LevelGuid);
 
 						FIntVector IndirectionBrickOffset;
 						int32 IndirectionBrickSize;
@@ -1508,8 +1522,7 @@ void FLightmassProcessor::ImportVolumetricLightmap()
 			{
 				FGuid LevelGuid = Pair.Key;
 
-				UMapBuildDataRegistry* SubLevelRegistry = System.LightingContext.GetOrCreateRegistryForLevelGuid(LevelGuid);
-				FPrecomputedVolumetricLightmapData& SubLevelData = *SubLevelRegistry->GetLevelPrecomputedVolumetricLightmapBuildData(System.LightingContext.GetLevelBuildDataID(LevelGuid));
+				FPrecomputedVolumetricLightmapData& SubLevelData =  System.LightingContext.GetOrCreateLevelPrecomputedVolumetricLightmapBuildData(LevelGuid);
 
 				ConvertBGRA8ToRGBA8ForLayer(SubLevelData.BrickData.SkyBentNormal);
 
@@ -1538,8 +1551,7 @@ void FLightmassProcessor::ImportVolumetricLightmap()
 
 		for (auto Pair : SubLevelTotalBricks)
 		{
-			UMapBuildDataRegistry* SubLevelRegistry = System.LightingContext.GetOrCreateRegistryForLevelGuid(Pair.Key);
-			FPrecomputedVolumetricLightmapData& SubLevelData = *SubLevelRegistry->GetLevelPrecomputedVolumetricLightmapBuildData(System.LightingContext.GetLevelBuildDataID(Pair.Key));
+			FPrecomputedVolumetricLightmapData& SubLevelData =  System.LightingContext.GetOrCreateLevelPrecomputedVolumetricLightmapBuildData(Pair.Key);
 
 			SubLevelData.FinalizeImport();
 
@@ -1559,6 +1571,13 @@ void FLightmassProcessor::ImportVolumetricLightmap()
 				SubLevelData.IndirectionTextureOriginalValues.Num() * SubLevelData.IndirectionTextureOriginalValues.GetTypeSize();
 		}
 
+		if (bSplitToVLMCellGrid)
+		{
+			// Export the data			
+			FVolumetricLightMapGridDesc& GridDesc = *System.LightingContext.GetVolumetricLightMapGridDesc();
+			GridDesc.InitializeBulkData();
+		}
+
 		int32 TotalNumBricks = 0;
 
 		for (int32 CurrentDepth = 0; CurrentDepth < VolumetricLightmapSettings.MaxRefinementLevels; CurrentDepth++)
@@ -1610,8 +1629,18 @@ void FLightmassProcessor::ImportVolumetricLightmap()
 
 		for (auto Pair : SubLevelTotalBricks)
 		{
-			FString LevelName = FindLevel(Pair.Key)->GetOuter()->GetName();
-			if (FindLevel(Pair.Key)->IsPersistentLevel())
+			ULevel* Level = FindLevel(Pair.Key);
+			FString LevelName;
+			if (Level)
+			{
+				LevelName = Level->GetOuter()->GetName();
+			}
+			else
+			{				
+				LevelName = System.LightingContext.GetVolumetricLightMapGridDesc()->GetCellDesc(Pair.Key);
+			}
+
+			if (Level && Level->IsPersistentLevel())
 			{
 				UE_LOG(LogVolumetricLightmapImport, Log, TEXT("         %s \t\t\t %d bricks \t %.1f%% (Persistent Level)"), *LevelName, Pair.Value, 100.0f * Pair.Value / TotalNumBricks);
 			}
diff --git a/Engine/Source/Editor/UnrealEd/Private/Lightmass/Lightmass.h b/Engine/Source/Editor/UnrealEd/Private/Lightmass/Lightmass.h
--- a/Engine/Source/Editor/UnrealEd/Private/Lightmass/Lightmass.h
+++ b/Engine/Source/Editor/UnrealEd/Private/Lightmass/Lightmass.h
@@ -399,6 +399,8 @@ public:
 	{
 		bImportCompletedMappingsImmediately = bInImportCompletedMappingsImmediately;
 	}
+	
+	void SetVolumetricLightMapImportMode(bool bUseVLMCellGrid);
 
 	/** Exports everything but the materials */
 	void InitiateExport();
@@ -684,6 +686,8 @@ protected:
 	/** Queue of messages from the swarm callback, to be processed by the main thread */
 	TArray<FLightmassAlertMessage> SwarmCallbackMessages;
 
+	bool bSplitToVLMCellGrid = false;
+
 	/**
 	 *	Import all mappings that have been completed so far.
 	 *
diff --git a/Engine/Source/Editor/UnrealEd/Private/StaticLightingSystem/StaticLightingPrivate.h b/Engine/Source/Editor/UnrealEd/Private/StaticLightingSystem/StaticLightingPrivate.h
--- a/Engine/Source/Editor/UnrealEd/Private/StaticLightingSystem/StaticLightingPrivate.h
+++ b/Engine/Source/Editor/UnrealEd/Private/StaticLightingSystem/StaticLightingPrivate.h
@@ -269,7 +269,7 @@ public:
 	 * @param InOptions - The static lighting build options.
 	 * @param InContext - The context (world, lighting scenario, world subsection, data layers)  we wish to build the lighting for
 	 */
-	FStaticLightingSystem(const FLightingBuildOptions& InOptions, const FStaticLightingBuildContext& InContext);
+	FStaticLightingSystem(const FLightingBuildOptions& InOptions, FStaticLightingBuildContext&& InContext);
 	~FStaticLightingSystem();
 
 	bool CheckLightmassExecutableVersion();
diff --git a/Engine/Source/Editor/UnrealEd/Private/StaticLightingSystem/StaticLightingSystem.cpp b/Engine/Source/Editor/UnrealEd/Private/StaticLightingSystem/StaticLightingSystem.cpp
--- a/Engine/Source/Editor/UnrealEd/Private/StaticLightingSystem/StaticLightingSystem.cpp
+++ b/Engine/Source/Editor/UnrealEd/Private/StaticLightingSystem/StaticLightingSystem.cpp
@@ -50,6 +50,7 @@
 #include "EditorModes.h"
 #include "Dialogs/Dialogs.h"
 #include "WorldPartition/WorldPartition.h"
+#include "WorldPartition/StaticLightingData/VolumetricLightmapGrid.h"
 
 FSwarmDebugOptions GSwarmDebugOptions;
 
@@ -474,13 +475,13 @@ bool FStaticLightingManager::IsLightingBuildCurrentlyExporting() const
 	return ActiveStaticLightingSystem != NULL && ActiveStaticLightingSystem->IsAmortizedExporting();
 }
 
-FStaticLightingSystem::FStaticLightingSystem(const FLightingBuildOptions& InOptions, const FStaticLightingBuildContext& context)
+FStaticLightingSystem::FStaticLightingSystem(const FLightingBuildOptions& InOptions, FStaticLightingBuildContext&& context)
 	: Options(InOptions)
 	, bBuildCanceled(false)
 	, DeterministicIndex(0)
 	, NextVisibilityId(0)
 	, CurrentBuildStage(FStaticLightingSystem::NotRunning)
-	, LightingContext(context)
+	, LightingContext(MoveTemp(context))
 	, LightmassProcessor(NULL)
 {
 }
@@ -1368,6 +1369,14 @@ void FStaticLightingSystem::ApplyNewLightingData(bool bLightingSuccessful)
 			const bool bBuildLightingForLevel = Options.ShouldBuildLightingForLevel( Level );
 
 			UMapBuildDataRegistry* Registry = LightingContext.GetRegistryForLevel(Level);
+		
+			if (Level->IsPersistentLevel() && LightingContext.World->IsPartitionedWorld())
+			{
+				// Transfer the new VolumetricLightMapGrid to the PersistentLevel
+				Registry->SetVolumetricLightMapGridDesc(LightingContext.GetVolumetricLightMapGridDesc());
+				LightingContext.ReleaseVolumetricLightMapGridDesc();
+			}
+			
 
 			// Store off the quality of the lighting for the level if lighting was successful and we build lighting for this level.
 			if( bLightingSuccessful && bBuildLightingForLevel )
@@ -2030,6 +2039,12 @@ bool FStaticLightingSystem::CreateLightmassProcessor()
 	check(LightmassProcessor == NULL);
 	LightmassProcessor = new FLightmassProcessor(*this, Options.bDumpBinaryResults, Options.bOnlyBuildVisibility);
 	check(LightmassProcessor);
+
+	if (LightingContext.World->IsPartitionedWorld())
+	{		
+		LightmassProcessor->SetVolumetricLightMapImportMode(true);
+	}	
+
 	if (LightmassProcessor->IsSwarmConnectionIsValid() == false)
 	{
 		UE_LOG(LogStaticLightingSystem, Warning, TEXT("Failed to connect to Swarm."));
@@ -2079,16 +2094,31 @@ void FStaticLightingSystem::GatherScene()
 	{
 		LightmassExporter->SetLevelName(LightingContext.World->PersistentLevel->GetPathName());
 	}
-
+	
 	LightmassExporter->ClearImportanceVolumes();
-	for( TObjectIterator<ALightmassImportanceVolume> It ; It ; ++It )
+
+	if (!LightingContext.World->IsPartitionedWorld())
 	{
-		ALightmassImportanceVolume* LMIVolume = *It;
-		if (LightingContext.World->ContainsActor(LMIVolume) && IsValid(LMIVolume) && ShouldOperateOnLevel(LMIVolume->GetLevel()))
+		for( TObjectIterator<ALightmassImportanceVolume> It ; It ; ++It )
 		{
-			LightmassExporter->AddImportanceVolume(LMIVolume);
+			ALightmassImportanceVolume* LMIVolume = *It;
+			if (LightingContext.World->ContainsActor(LMIVolume) && IsValid(LMIVolume) && ShouldOperateOnLevel(LMIVolume->GetLevel()))
+			{
+				LightmassExporter->AddImportanceVolume(LMIVolume);
+			}			
 		}
 	}
+	else
+	{
+		// Ignore all user set ImportanceVolumes and add a LightMassImportanceVolume per VLM cell 
+		//@todo_ow: Possibly interesect user specified ALightmassImportanceVolume with world bounds to restrict area in which we compute static lighting
+		FBox ImportanceBounds(ForceInit);
+				
+		LightmassExporter->AddImportanceVolumeBoundingBox(LightingContext.GetVolumetricLightMapGridDesc()->GridBounds);
+		ImportanceBounds += LightingContext.GetVolumetricLightMapGridDesc()->GridBounds;
+
+		LightingContext.SetImportanceBounds(ImportanceBounds);
+	}
 
 	for( TObjectIterator<ALightmassCharacterIndirectDetailVolume> It ; It ; ++It )
 	{
diff --git a/Engine/Source/Editor/UnrealEd/Private/WorldPartition/WorldPartitionStaticLightingBuilder.cpp b/Engine/Source/Editor/UnrealEd/Private/WorldPartition/WorldPartitionStaticLightingBuilder.cpp
new file mode 100644
--- /dev/null
+++ b/Engine/Source/Editor/UnrealEd/Private/WorldPartition/WorldPartitionStaticLightingBuilder.cpp
@@ -0,0 +1,175 @@
+// Copyright Epic Games, Inc. All Rights Reserved.
+
+#include "WorldPartition/WorldPartitionStaticLightingBuilder.h"
+
+#include "CoreMinimal.h"
+#include "HAL/FileManager.h"
+#include "Logging/LogMacros.h"
+#include "Misc/ConfigCacheIni.h"
+#include "Misc/FileHelper.h"
+#include "Algo/ForEach.h"
+#include "UObject/SavePackage.h"
+#include "FileHelpers.h"
+
+#include "EngineUtils.h"
+#include "EngineModule.h"
+#include "SourceControlHelpers.h"
+
+#include "ISourceControlModule.h"
+#include "ISourceControlProvider.h"
+#include "SourceControlOperations.h"
+
+#include "Editor.h"
+
+#include "WorldPartition/WorldPartitionHelpers.h"
+
+#include "LightingBuildOptions.h"
+
+#include "HAL/IConsoleManager.h"
+#include "Engine/MapBuildDataRegistry.h"
+
+DEFINE_LOG_CATEGORY_STATIC(LogWorldPartitionStaticLightingBuilder, All, All);
+
+UWorldPartitionStaticLightingBuilder::UWorldPartitionStaticLightingBuilder(const FObjectInitializer& ObjectInitializer)
+	: Super(ObjectInitializer)	
+{
+	if (HasAnyFlags(RF_ClassDefaultObject))
+	{
+		return;
+	}
+
+	SourceControlHelper = nullptr;
+	
+	bBuildVLMOnly = FParse::Param(FCommandLine::Get(), TEXT("BuildVLM"));
+
+	QualityLevel = ELightingBuildQuality::Quality_Preview;
+
+	BuildOptions |= FParse::Param(FCommandLine::Get(), TEXT("Submit")) ? EWPStaticLightingBuildStep::WPSL_Submit : EWPStaticLightingBuildStep::None;
+	BuildOptions |= bBuildVLMOnly ?  (EWPStaticLightingBuildStep::WPSL_Build): EWPStaticLightingBuildStep::None;
+	
+
+	// Default behavior without any option is to build VLM only (will change when we add Lightmaps support)
+	if (BuildOptions == EWPStaticLightingBuildStep::None)
+	{
+		BuildOptions = EWPStaticLightingBuildStep::WPSL_Build; 
+		bBuildVLMOnly = true;
+	}
+
+	// Parse quality level and limit to valid values
+	int32 QualityLevelVal;
+	FParse::Value(FCommandLine::Get(), TEXT("QualityLevel="), QualityLevelVal);
+	QualityLevelVal = FMath::Clamp<int32>(QualityLevelVal, Quality_Preview, Quality_Production);
+	QualityLevel = (ELightingBuildQuality)QualityLevelVal;
+}
+
+bool UWorldPartitionStaticLightingBuilder::RequiresCommandletRendering() const
+{
+	// The Lightmass export process uses the renderer to generate some data so we need rendering
+	return true; 
+}
+
+bool UWorldPartitionStaticLightingBuilder::ShouldRunStep(const EWPStaticLightingBuildStep BuildStep) const
+{
+	return (BuildOptions & BuildStep) == BuildStep;
+}
+
+UWorldPartitionBuilder::ELoadingMode UWorldPartitionStaticLightingBuilder::GetLoadingMode() const 
+{
+	check(bBuildVLMOnly);
+	return ELoadingMode::EntireWorld;
+}
+
+bool UWorldPartitionStaticLightingBuilder::ValidateParams() const
+{
+	return true;
+}
+
+bool UWorldPartitionStaticLightingBuilder::PreWorldInitialization(UWorld* World, FPackageSourceControlHelper& PackageHelper)
+{
+	return ValidateParams();	
+}
+
+bool UWorldPartitionStaticLightingBuilder::PreRun(UWorld* World, FPackageSourceControlHelper& PackageHelper)
+{			
+	return true;
+}
+
+
+bool UWorldPartitionStaticLightingBuilder::RunInternal(UWorld* World, const FCellInfo& InCellInfo, FPackageSourceControlHelper& PackageHelper)
+{
+	WorldPartition = World->GetWorldPartition();
+	check(WorldPartition);
+
+	bool bRet = true;
+
+	SourceControlHelper = MakeUnique<FSourceControlHelper>(PackageHelper, ModifiedFiles);
+
+	if (ShouldRunStep(EWPStaticLightingBuildStep::WPSL_Build))
+	{
+		bRet = RunForVLM(World, InCellInfo, PackageHelper);		
+	}
+
+	if (bRet && ShouldRunStep(EWPStaticLightingBuildStep::WPSL_Submit))
+	{
+		bRet = Submit(World, PackageHelper);
+	}
+
+	SourceControlHelper.Reset();
+
+	return bRet;
+}
+
+bool UWorldPartitionStaticLightingBuilder::Submit(UWorld* World, FPackageSourceControlHelper& PackageHelper)
+{	
+	// Wait for pending async file writes before submitting
+	UPackage::WaitForAsyncFileWrites();
+
+	const FString ChangeDescription = FString::Printf(TEXT("Rebuilt static lighting for %s"), *World->GetPackage()->GetName());
+	return OnFilesModified(ModifiedFiles.GetAllFiles(), ChangeDescription);
+}
+
+bool UWorldPartitionStaticLightingBuilder::RunForVLM(UWorld* World, const FCellInfo& InCellInfo, FPackageSourceControlHelper& PackageHelper)
+{
+	bool bRet = true;
+	check (bBuildVLMOnly);
+
+	UE_LOG(LogWorldPartitionStaticLightingBuilder, Verbose, TEXT("Building Volumetric Lightmaps for %s"), *World->GetName());
+	 
+	// Invoke static lighting computation
+	FLightingBuildOptions LightingOptions;
+	LightingOptions.QualityLevel =  QualityLevel;
+
+	bool bLightingBuildFailed = false;
+
+	auto BuildFailedDelegate = [&bLightingBuildFailed, &World]() {
+		// WP-TODO!! better error message
+		UE_LOG(LogWorldPartitionStaticLightingBuilder, Error, TEXT("[REPORT] Failed building lighting for %s"), *World->GetName());
+		bLightingBuildFailed = true;
+	};
+
+	FDelegateHandle BuildFailedDelegateHandle = FEditorDelegates::OnLightingBuildFailed.AddLambda(BuildFailedDelegate);
+	
+	GEditor->BuildLighting(LightingOptions);
+	while (GEditor->IsLightingBuildCurrentlyRunning())
+	{
+		GEditor->UpdateBuildLighting();
+	}
+		
+	if (!bLightingBuildFailed)
+	{
+		// Save the AMapBuildData actors + MapBuildData we just updated
+		TArray<UPackage*> PackagesToSave;
+
+		PackagesToSave.Add(World->PersistentLevel->MapBuildData->GetPackage());
+	
+		UEditorLoadingAndSavingUtils::SavePackages(PackagesToSave, false);
+	}
+	else
+	{
+		bRet = false;
+	}
+
+	FEditorDelegates::OnLightingBuildFailed.Remove(BuildFailedDelegateHandle);
+		
+	return bRet;
+}
diff --git a/Engine/Source/Editor/UnrealEd/Public/WorldPartition/WorldPartitionStaticLightingBuilder.h b/Engine/Source/Editor/UnrealEd/Public/WorldPartition/WorldPartitionStaticLightingBuilder.h
new file mode 100644
--- /dev/null
+++ b/Engine/Source/Editor/UnrealEd/Public/WorldPartition/WorldPartitionStaticLightingBuilder.h
@@ -0,0 +1,58 @@
+// Copyright Epic Games, Inc. All Rights Reserved.
+
+#pragma once
+
+#include "Engine/EngineTypes.h"
+#include "Misc/EnumClassFlags.h"
+#include "WorldPartition/WorldPartitionBuilder.h"
+#include "WorldPartition/WorldPartitionBuilderHelpers.h"
+#include "WorldPartition/WorldPartitionHandle.h"
+
+#include "WorldPartitionStaticLightingBuilder.generated.h"
+
+
+class AMapBuildDataActor;
+
+enum class EWPStaticLightingBuildStep : uint8
+{
+	None		= 0,
+	WPSL_Build = 1 << 1,						// Build the static lighting by iterating over the map and associates the data actors with the map actors already present	
+	WPSL_Submit = 1 <<2,						// Optionally, submit results to source control
+};
+ENUM_CLASS_FLAGS(EWPStaticLightingBuildStep);
+
+UCLASS()
+class UWorldPartitionStaticLightingBuilder : public UWorldPartitionBuilder
+{
+	GENERATED_UCLASS_BODY()
+public:
+	// UWorldPartitionBuilder interface begin
+	virtual bool RequiresCommandletRendering() const override;
+	virtual ELoadingMode GetLoadingMode() const override;
+	virtual bool PreWorldInitialization(UWorld* World, FPackageSourceControlHelper& PackageHelper) override;
+
+protected:
+	virtual bool PreRun(UWorld* World, FPackageSourceControlHelper& PackageHelper) override;
+	virtual bool RunInternal(UWorld* World, const FCellInfo& InCellInfo, FPackageSourceControlHelper& PackageHelper) override;
+
+	// UWorldPartitionBuilder interface end
+
+	bool ValidateParams() const;
+	bool ShouldRunStep(const EWPStaticLightingBuildStep BuildStep) const;
+
+	bool RunForVLM(UWorld* World, const FCellInfo& InCellInfo, FPackageSourceControlHelper& PackageHelper);
+
+	bool Submit(UWorld* World, FPackageSourceControlHelper& PackageHelper);
+	
+
+private:
+	class UWorldPartition* WorldPartition;
+	TUniquePtr<FSourceControlHelper> SourceControlHelper;
+	 
+	// Options --
+	EWPStaticLightingBuildStep BuildOptions;	
+	bool bBuildVLMOnly;
+	ELightingBuildQuality QualityLevel;
+
+	FBuilderModifiedFiles ModifiedFiles;
+};
diff --git a/Engine/Source/Runtime/Core/Private/HAL/ConsoleManager.cpp b/Engine/Source/Runtime/Core/Private/HAL/ConsoleManager.cpp
--- a/Engine/Source/Runtime/Core/Private/HAL/ConsoleManager.cpp
+++ b/Engine/Source/Runtime/Core/Private/HAL/ConsoleManager.cpp
@@ -3756,6 +3756,15 @@ static TAutoConsoleVariable<int32> CVarAllowStaticLighting(
 		 "Games that only use dynamic lighting should set this to 0 to save some static lighting overhead."),
 	ECVF_ReadOnly | ECVF_RenderThreadSafe);
 
+static TAutoConsoleVariable<int32> CVarAllowStaticLightingInWorldPartitionMaps(
+	TEXT("r.AllowStaticLightingInWorldPartitionMaps"),
+	0,
+	TEXT("Whether to allow any static lighting to be generated and used in WorldPartition maps, like lightmaps and shadowmaps.\n"
+		 "Games that only use dynamic lighting should set this to 0 to save some static lighting overhead."),
+	ECVF_ReadOnly | ECVF_RenderThreadSafe);
+
+
+
 static TAutoConsoleVariable<int32> CVarNormalMaps(
 	TEXT("r.NormalMapsForStaticLighting"),
 	0,
diff --git a/Engine/Source/Runtime/Core/Public/UObject/FortniteMainBranchObjectVersion.h b/Engine/Source/Runtime/Core/Public/UObject/FortniteMainBranchObjectVersion.h
--- a/Engine/Source/Runtime/Core/Public/UObject/FortniteMainBranchObjectVersion.h
+++ b/Engine/Source/Runtime/Core/Public/UObject/FortniteMainBranchObjectVersion.h
@@ -429,6 +429,9 @@ struct FFortniteMainBranchObjectVersion
 		// Level Instance Property overrides
 		LevelInstancePropertyOverrides,
 
+		// Added FVolumetricLightMapGridDesc in MapBuildData
+		VolumetricLightMapGridDescSupport,
+
 		// -----<new versions can be added above this line>-------------------------------------------------
 		VersionPlusOne,
 		LatestVersion = VersionPlusOne - 1
diff --git a/Engine/Source/Runtime/Engine/Classes/Engine/Level.h b/Engine/Source/Runtime/Engine/Classes/Engine/Level.h
--- a/Engine/Source/Runtime/Engine/Classes/Engine/Level.h
+++ b/Engine/Source/Runtime/Engine/Classes/Engine/Level.h
@@ -538,6 +538,9 @@ public:
 	UPROPERTY()
 	FIntVector LightBuildLevelOffset;
 
+	/** The volumetric lightmap grid manager for this map */
+	class FVolumetricLightmapGridManager* VolumetricLightmapGridManager = nullptr;
+
 	/** 
 	 * Whether the level is a lighting scenario.  Lighting is built separately for each lighting scenario level with all other scenario levels hidden. 
 	 * Only one lighting scenario level should be visible at a time for correct rendering, and lightmaps from that level will be used on the rest of the world.
diff --git a/Engine/Source/Runtime/Engine/Classes/Engine/MapBuildDataRegistry.h b/Engine/Source/Runtime/Engine/Classes/Engine/MapBuildDataRegistry.h
--- a/Engine/Source/Runtime/Engine/Classes/Engine/MapBuildDataRegistry.h
+++ b/Engine/Source/Runtime/Engine/Classes/Engine/MapBuildDataRegistry.h
@@ -21,6 +21,7 @@
 class FPrecomputedLightVolumeData;
 class FPrecomputedVolumetricLightmapData;
 struct FAssetCompileData;
+struct FVolumetricLightMapGridDesc;
 
 struct FPerInstanceLightmapData
 {
@@ -389,6 +390,9 @@ public:
 	/** Returns the associated UUMapBuildDataRegistry for an Actor */
 	ENGINE_API static UMapBuildDataRegistry* Get(const AActor* Actor);
 	
+	FVolumetricLightMapGridDesc* GetVolumetricLightMapGridDesc() { return VolumetricLightMapGridDesc; } 
+	ENGINE_API void SetVolumetricLightMapGridDesc(FVolumetricLightMapGridDesc* GridDesc); 
+
 private:
 #if WITH_EDITOR
 	void HandleAssetPostCompileEvent(const TArray<FAssetCompileData>& CompiledAssets);
@@ -409,6 +413,8 @@ private:
 	TArray<FLightmapResourceCluster> LightmapResourceClusters;
 
 	FRenderCommandFence DestroyFence;
+
+	FVolumetricLightMapGridDesc*	VolumetricLightMapGridDesc;
 };
 
 extern ENGINE_API FUObjectAnnotationSparse<FMeshMapBuildLegacyData, true> GComponentsWithLegacyLightmaps;
diff --git a/Engine/Source/Runtime/Engine/Classes/GameFramework/WorldSettings.h b/Engine/Source/Runtime/Engine/Classes/GameFramework/WorldSettings.h
--- a/Engine/Source/Runtime/Engine/Classes/GameFramework/WorldSettings.h
+++ b/Engine/Source/Runtime/Engine/Classes/GameFramework/WorldSettings.h
@@ -160,9 +160,13 @@ struct FLightmassWorldInfoSettings
 	/** 
 	 * Maximum amount of memory to spend on Volumetric Lightmap Brick data.  High density bricks will be discarded until this limit is met, with bricks furthest from geometry discarded first.
 	 */
-	UPROPERTY(BlueprintReadWrite, EditAnywhere, Category=LightmassVolumeLighting, meta=(UIMin = "1", UIMax = "500"))
+	UPROPERTY(BlueprintReadWrite, EditAnywhere, Category=LightmassVolumeLighting, meta=(UIMin = "1", UIMax = "500"), meta = (EditConditionHides, EditCondition = "!bWorldPartition"))
 	float VolumetricLightmapMaximumBrickMemoryMb;
 
+	/**  Size of an Volumetric Lightmap high detail loading cell.  */
+	UPROPERTY(BlueprintReadWrite, EditAnywhere, Category = LightmassVolumeLighting, meta = (UIMin = "1600", UIMax = "102400"), meta = (EditConditionHides, EditCondition = "bWorldPartition"))
+	float VolumetricLightmapLoadingCellSize;
+
 	/** 
 	 * Controls how much smoothing should be done to Volumetric Lightmap samples during Spherical Harmonic de-ringing.  
 	 * Whenever highly directional lighting is stored in a Spherical Harmonic, a ringing artifact occurs which manifests as unexpected black areas on the opposite side.
@@ -199,6 +203,9 @@ struct FLightmassWorldInfoSettings
 	UPROPERTY(BlueprintReadWrite, EditAnywhere, Category=LightmassOcclusion)
 	float MaxOcclusionDistance;
 
+	UPROPERTY(transient)
+	bool bWorldPartition;
+
 	FLightmassWorldInfoSettings()
 		: StaticLightingLevelScale(1)
 		, NumIndirectLightingBounces(3)
@@ -217,6 +224,7 @@ struct FLightmassWorldInfoSettings
 		, bCompressLightmaps(true)
 		, VolumetricLightmapDetailCellSize(200)
 		, VolumetricLightmapMaximumBrickMemoryMb(30)
+		, VolumetricLightmapLoadingCellSize(3200)		
 		, VolumetricLightmapSphericalHarmonicSmoothing(.02f)
 		, VolumeLightSamplePlacementScale(1)
 		, DirectIlluminationOcclusionFraction(0.5f)
@@ -224,6 +232,7 @@ struct FLightmassWorldInfoSettings
 		, OcclusionExponent(1.0f)
 		, FullyOccludedSamplesFraction(1.0f)
 		, MaxOcclusionDistance(200.0f)
+		, bWorldPartition(false)
 	{
 	}
 };
@@ -556,6 +565,12 @@ public:
 	UPROPERTY(EditAnywhere, Category=Lightmass, AdvancedDisplay)
 	uint8 bForceNoPrecomputedLighting:1;
 
+	/** 
+	 * Force precomputed lighting to only use VolumetricLightmaps.
+	 */
+	UPROPERTY(EditAnywhere, Category=Lightmass, AdvancedDisplay)
+	uint8 bForceVolumetricLightmapsOnly:1;
+
 	/** when this flag is set, more time is allocated to background loading (replicated) */
 	UPROPERTY(replicated)
 	uint8 bHighPriorityLoading:1;
@@ -740,6 +755,11 @@ public:
 	struct FLightmassWorldInfoSettings LightmassSettings;
 #endif
 
+	/**  Range in which volumetric lightmaps will be loaded.  */
+	UPROPERTY(BlueprintReadWrite, EditAnywhere, Category = LightmassVolumeLighting, meta = (UIMin = "1600", UIMax = "102400"), meta = (EditConditionHides, EditCondition = "bWorldPartition"))
+	float VolumetricLightmapLoadingRange;
+
+
 	/************************************/
 	/** NANITE SETTINGS **/
 	UPROPERTY(EditAnywhere, config, replicated, ReplicatedUsing = OnRep_NaniteSettings, Category = Nanite)
diff --git a/Engine/Source/Runtime/Engine/Private/Components/PrimitiveComponent.cpp b/Engine/Source/Runtime/Engine/Private/Components/PrimitiveComponent.cpp
--- a/Engine/Source/Runtime/Engine/Private/Components/PrimitiveComponent.cpp
+++ b/Engine/Source/Runtime/Engine/Private/Components/PrimitiveComponent.cpp
@@ -419,6 +419,17 @@ PRAGMA_DISABLE_DEPRECATION_WARNINGS
 
 ELightmapType UPrimitiveComponent::GetLightmapType() const
 {
+	if (UWorld* World = GetWorld())	
+	{
+		if (AWorldSettings* WorldSettings = World->GetWorldSettings())
+		{
+			if (WorldSettings->bForceVolumetricLightmapsOnly)
+			{
+				return ELightmapType::ForceVolumetric;
+			}
+		}
+	}
+
 	return LightmapType;
 }
 
diff --git a/Engine/Source/Runtime/Engine/Private/Level.cpp b/Engine/Source/Runtime/Engine/Private/Level.cpp
--- a/Engine/Source/Runtime/Engine/Private/Level.cpp
+++ b/Engine/Source/Runtime/Engine/Private/Level.cpp
@@ -22,6 +22,7 @@ Level.cpp: Level-related functions
 #include "SceneInterface.h"
 #include "PrecomputedLightVolume.h"
 #include "PrecomputedVolumetricLightmap.h"
+#include "PrecomputedVolumetricLightmapStreaming.h"
 #include "Engine/MapBuildDataRegistry.h"
 #include "Components/LightComponent.h"
 #include "Model.h"
@@ -1481,6 +1482,9 @@ void ULevel::FinishDestroy()
 	delete PrecomputedVolumetricLightmap;
 	PrecomputedVolumetricLightmap = NULL;
 
+	delete VolumetricLightmapGridManager;
+	VolumetricLightmapGridManager = NULL;
+
 	Super::FinishDestroy();
 }
 
@@ -3311,6 +3315,11 @@ void ULevel::InitializeRenderingResources()
 			PrecomputedLightVolume->AddToScene(OwningWorld->Scene, EffectiveMapBuildData, LevelBuildDataId);
 		}
 
+		if (EffectiveMapBuildData && EffectiveMapBuildData->GetVolumetricLightMapGridDesc())
+		{
+			VolumetricLightmapGridManager = new FVolumetricLightmapGridManager(OwningWorld, EffectiveMapBuildData->GetVolumetricLightMapGridDesc());
+		}
+
 		if (!PrecomputedVolumetricLightmap->IsAddedToScene())
 		{
 			PrecomputedVolumetricLightmap->AddToScene(OwningWorld->Scene, EffectiveMapBuildData, LevelBuildDataId, IsPersistentLevel());
@@ -3327,6 +3336,11 @@ void ULevel::ReleaseRenderingResources()
 {
 	if (OwningWorld && FApp::CanEverRender())
 	{
+		if (VolumetricLightmapGridManager)
+		{
+			VolumetricLightmapGridManager->RemoveFromScene(OwningWorld->Scene);
+		}
+
 		if (PrecomputedLightVolume)
 		{
 			PrecomputedLightVolume->RemoveFromScene(OwningWorld->Scene);
diff --git a/Engine/Source/Runtime/Engine/Private/MapBuildData.cpp b/Engine/Source/Runtime/Engine/Private/MapBuildData.cpp
--- a/Engine/Source/Runtime/Engine/Private/MapBuildData.cpp
+++ b/Engine/Source/Runtime/Engine/Private/MapBuildData.cpp
@@ -21,6 +21,7 @@ MapBuildData.cpp
 #include "UObject/ReflectionCaptureObjectVersion.h"
 #include "UObject/UE5MainStreamObjectVersion.h"
 #include "UObject/UE5ReleaseStreamObjectVersion.h"
+#include "UObject/FortniteMainBranchObjectVersion.h"
 #include "ContentStreaming.h"
 #include "Components/ReflectionCaptureComponent.h"
 #include "Interfaces/ITargetPlatform.h"
@@ -34,7 +35,7 @@ MapBuildData.cpp
 #include "Engine/TextureCube.h"
 #include "Components/HierarchicalInstancedStaticMeshComponent.h"
 #include "UnrealEngine.h"
-
+#include "WorldPartition/StaticLightingData/VolumetricLightmapGrid.h"
 
 DECLARE_MEMORY_STAT(TEXT("Stationary Light Static Shadowmap"),STAT_StationaryLightBuildData,STATGROUP_MapBuildData);
 DECLARE_MEMORY_STAT(TEXT("Reflection Captures"),STAT_ReflectionCaptureBuildData,STATGROUP_MapBuildData);
@@ -438,6 +439,8 @@ UMapBuildDataRegistry::UMapBuildDataRegistry(const FObjectInitializer& ObjectIni
 {
 	LevelLightingQuality = Quality_MAX;
 	bSetupResourceClusters = false;
+	VolumetricLightMapGridDesc = nullptr;
+
 
 #if WITH_EDITOR
 	FAssetCompilingManager::Get().OnAssetPostCompileEvent().AddUObject(this, &ThisClass::HandleAssetPostCompileEvent);
@@ -525,6 +528,7 @@ void UMapBuildDataRegistry::Serialize(FArchive& Ar)
 	Ar.UsingCustomVersion(FReflectionCaptureObjectVersion::GUID);
 	Ar.UsingCustomVersion(FUE5ReleaseStreamObjectVersion::GUID);
 	Ar.UsingCustomVersion(FUE5MainStreamObjectVersion::GUID);
+	Ar.UsingCustomVersion(FFortniteMainBranchObjectVersion::GUID);
 
 	if (!StripFlags.IsAudioVisualDataStripped())
 	{
@@ -557,6 +561,24 @@ void UMapBuildDataRegistry::Serialize(FArchive& Ar)
 		{
 			Ar << SkyAtmosphereBuildData;
 		}
+
+		if (Ar.CustomVer(FFortniteMainBranchObjectVersion::GUID) >= FFortniteMainBranchObjectVersion::VolumetricLightMapGridDescSupport)
+		{			
+			bool bHasGrid = VolumetricLightMapGridDesc != nullptr;
+			Ar << bHasGrid;
+
+			if (bHasGrid)
+			{
+				// Create the grid when loading for the 1st time
+				if (!VolumetricLightMapGridDesc)
+				{
+					VolumetricLightMapGridDesc = new FVolumetricLightMapGridDesc();
+				}
+
+				FVolumetricLightMapGridDesc::StaticStruct()->SerializeItem(Ar, VolumetricLightMapGridDesc, nullptr);
+				VolumetricLightMapGridDesc->SerializeBulkData(Ar, this);
+			}
+		}
 	}
 }
 
@@ -748,6 +770,12 @@ FPrecomputedLightVolumeData* UMapBuildDataRegistry::GetLevelPrecomputedLightVolu
 
 FPrecomputedVolumetricLightmapData& UMapBuildDataRegistry::AllocateLevelPrecomputedVolumetricLightmapBuildData(const FGuid& LevelId)
 {
+	if (VolumetricLightMapGridDesc && VolumetricLightMapGridDesc->GetCell(LevelId) )
+	{
+		FPrecomputedVolumetricLightmapData* DataPtr = VolumetricLightMapGridDesc->GetOrCreatePrecomputedVolumetricLightmapBuildData(LevelId);		
+		return *DataPtr;	
+	}
+	
 	check(LevelId.IsValid());
 	MarkPackageDirty();
 	return *LevelPrecomputedVolumetricLightmapBuildData.Add(LevelId, new FPrecomputedVolumetricLightmapData());
@@ -761,6 +789,12 @@ void UMapBuildDataRegistry::AddLevelPrecomputedVolumetricLightmapBuildData(const
 
 const FPrecomputedVolumetricLightmapData* UMapBuildDataRegistry::GetLevelPrecomputedVolumetricLightmapBuildData(FGuid LevelId) const
 {
+	if (VolumetricLightMapGridDesc)
+	{
+		const FPrecomputedVolumetricLightmapData* DataPtr = VolumetricLightMapGridDesc->GetPrecomputedVolumetricLightmapBuildData(LevelId);		
+		return DataPtr;
+	}
+
 	const FPrecomputedVolumetricLightmapData* const * DataPtr = LevelPrecomputedVolumetricLightmapBuildData.Find(LevelId);
 
 	if (DataPtr)
@@ -773,6 +807,16 @@ const FPrecomputedVolumetricLightmapData* UMapBuildDataRegistry::GetLevelPrecomp
 
 FPrecomputedVolumetricLightmapData* UMapBuildDataRegistry::GetLevelPrecomputedVolumetricLightmapBuildData(FGuid LevelId)
 {
+	if (VolumetricLightMapGridDesc)
+	{
+		FPrecomputedVolumetricLightmapData* DataPtr = VolumetricLightMapGridDesc->GetPrecomputedVolumetricLightmapBuildData(LevelId);		
+		
+		if (DataPtr)
+		{
+			return DataPtr;
+		}
+	}
+
 	FPrecomputedVolumetricLightmapData** DataPtr = LevelPrecomputedVolumetricLightmapBuildData.Find(LevelId);
 
 	if (DataPtr)
@@ -1140,6 +1184,9 @@ void UMapBuildDataRegistry::EmptyLevelData(const TSet<FGuid>* ResourcesToKeep)
 			LevelPrecomputedVolumetricLightmapBuildData.Add(It.Key(), It.Value());
 		}
 	}
+	
+	delete VolumetricLightMapGridDesc;
+	VolumetricLightMapGridDesc = nullptr;
 
 	LightmapResourceClusters.Empty();
 }
@@ -1173,6 +1220,14 @@ UMapBuildDataRegistry* UMapBuildDataRegistry::Get(const AActor* Actor)
 	UWorld* World = OwnerLevel ? OwnerLevel->GetWorld() : nullptr;
 	UMapBuildDataRegistry* MapBuildData = nullptr;
 	
+	if (World && World->IsPartitionedWorld())
+	{
+		//@todo_ow: At current level of support there's no reason to return a ptr and force a look-up later on
+		//No lighting scenario support in WP maps
+		MapBuildData = World->PersistentLevel->MapBuildData;
+		return MapBuildData;
+	}
+
 	if (OwnerLevel && World)
 	{
 		ULevel* ActiveLightingScenario = World->GetActiveLightingScenario();
@@ -1190,6 +1245,12 @@ UMapBuildDataRegistry* UMapBuildDataRegistry::Get(const AActor* Actor)
 	return MapBuildData;
 }
 
+void UMapBuildDataRegistry::SetVolumetricLightMapGridDesc(FVolumetricLightMapGridDesc* GridDesc)
+{	 
+	delete VolumetricLightMapGridDesc;
+	VolumetricLightMapGridDesc = GridDesc;
+}
+
 FUObjectAnnotationSparse<FMeshMapBuildLegacyData, true> GComponentsWithLegacyLightmaps;
 FUObjectAnnotationSparse<FLevelLegacyMapBuildData, true> GLevelsWithLegacyBuildData;
 FUObjectAnnotationSparse<FLightComponentLegacyMapBuildData, true> GLightComponentsWithLegacyBuildData;
diff --git a/Engine/Source/Runtime/Engine/Private/PrecomputedVolumetricLightmap.cpp b/Engine/Source/Runtime/Engine/Private/PrecomputedVolumetricLightmap.cpp
--- a/Engine/Source/Runtime/Engine/Private/PrecomputedVolumetricLightmap.cpp
+++ b/Engine/Source/Runtime/Engine/Private/PrecomputedVolumetricLightmap.cpp
@@ -112,6 +112,7 @@ FArchive& operator<<(FArchive& Ar,FPrecomputedVolumetricLightmapData& Volume)
 {
 	Ar.UsingCustomVersion(FMobileObjectVersion::GUID);
 	Ar.UsingCustomVersion(FRenderingObjectVersion::GUID);
+	Ar.UsingCustomVersion(FUE5MainStreamObjectVersion::GUID);
 
 	Ar << Volume.Bounds;
 	Ar << Volume.IndirectionTextureDimensions;
@@ -674,21 +675,34 @@ void FPrecomputedVolumetricLightmap::AddToScene(FSceneInterface* Scene, UMapBuil
 		return;
 	}
 
-	// FIXME: temp fix for ordering issue between WorldContext.World()->InitWorld(); and GShaderCompilingManager->ProcessAsyncResults(false, true); in UnrealEngine.cpp
-	if (GShaderCompilingManager)
+	FPrecomputedVolumetricLightmapData* NewData = NULL;
+
+	if (Registry)
 	{
-		GShaderCompilingManager->ProcessAsyncResults(true, true);
+		NewData = Registry->GetLevelPrecomputedVolumetricLightmapBuildData(LevelBuildDataId);
 	}
 
-	check(!bAddedToScene);
+	if (NewData && Scene)
+	{		
+		AddToScene(Scene, Registry, NewData, bIsPersistentLevel);
+	}	
+}
 
-	FPrecomputedVolumetricLightmapData* NewData = NULL;
+void FPrecomputedVolumetricLightmap::AddToScene(FSceneInterface* Scene, UMapBuildDataRegistry* Registry, FPrecomputedVolumetricLightmapData* NewData, bool bIsPersistentLevel)
+{
+	if (!IsStaticLightingAllowed())
+	{
+		return;
+	}
 
-	if (Registry)
+	// FIXME: temp fix for ordering issue between WorldContext.World()->InitWorld(); and GShaderCompilingManager->ProcessAsyncResults(false, true); in UnrealEngine.cpp
+	if (GShaderCompilingManager)
 	{
-		NewData = Registry->GetLevelPrecomputedVolumetricLightmapBuildData(LevelBuildDataId);
+		GShaderCompilingManager->ProcessAsyncResults(true, true);
 	}
 
+	check(!bAddedToScene);
+
 	if (NewData && Scene)
 	{
 		bAddedToScene = true;
diff --git a/Engine/Source/Runtime/Engine/Private/PrecomputedVolumetricLightmapStreaming.cpp b/Engine/Source/Runtime/Engine/Private/PrecomputedVolumetricLightmapStreaming.cpp
new file mode 100644
--- /dev/null
+++ b/Engine/Source/Runtime/Engine/Private/PrecomputedVolumetricLightmapStreaming.cpp
@@ -0,0 +1,382 @@
+// Copyright Epic Games, Inc. All Rights Reserved.
+
+/*=============================================================================
+	PrecomputedVolumetricLightmap.cpp
+=============================================================================*/
+
+#include "PrecomputedVolumetricLightmapStreaming.h"
+#include "PrecomputedVolumetricLightmap.h"
+#include "Engine/World.h"
+#include "Engine/Level.h"
+#include "ContentStreaming.h"
+#include "WorldPartition/StaticLightingData/VolumetricLightmapGrid.h"
+#include "Serialization/VersionedArchive.h"
+#include "Serialization/MemoryReader.h"
+#include "GameFramework/WorldSettings.h"
+
+using FVersionedMemoryReaderView = TVersionedReader<FMemoryReaderView>;
+
+enum ETimedExecutionControl
+{
+	Continue, 
+	Restart, 
+	Stop,
+};
+
+template<class ItemCollectionType, class ExecuteFn>
+bool TimedExecution(ItemCollectionType& Items, float TimeLimit, const ExecuteFn& Execute) 		
+{
+	double EndTime = FPlatformTime::Seconds() + TimeLimit;	
+	if (TimeLimit == 0)
+	{
+		EndTime = FLT_MAX;
+	}
+	
+	bool bContinue = true; 
+	
+	while (bContinue)	
+	{
+		bContinue = false; 
+
+		for (auto& It : Items)
+		{
+			float ThisTimeLimit = static_cast<float>(EndTime - FPlatformTime::Seconds());
+
+			if (ThisTimeLimit < .001f) // one ms is the granularity of the platform event system
+			{
+				return false;
+			}
+
+			ETimedExecutionControl Control = Execute(ThisTimeLimit, It);
+
+			if (Control == ETimedExecutionControl::Restart)
+			{
+				bContinue = true;
+				break;
+			}
+			else if (Control == ETimedExecutionControl::Stop)
+			{
+				break;
+			}			
+		}
+	}
+		
+	return true;
+}
+
+class FVolumetricLightmapGridStreamingManager : public IStreamingManager
+{
+public:
+
+	friend class FVolumetricLightmapGridManager;
+	
+	FVolumetricLightmapGridManager* Owner;
+
+	FVolumetricLightmapGridStreamingManager(FVolumetricLightmapGridManager* InOwner)
+		: Owner(InOwner)
+	{
+		IStreamingManager::Get().AddStreamingManager(this);
+	}
+
+	~FVolumetricLightmapGridStreamingManager()
+	{
+		IStreamingManager::Get().RemoveStreamingManager(this);
+	}
+
+	virtual void Tick(float DeltaTime, bool bProcessEverything = false) override
+	{
+		
+	}
+
+	virtual void UpdateResourceStreaming(float DeltaTime, bool bProcessEverything = false) override
+	{
+		 int32 nbViews = IStreamingManager::Get().GetNumViews();
+
+		 //@todo_ow: Manage multiple views properly, add better render/world setting to determine extent / budget
+		 if (nbViews > 0)
+		 {			
+			FBox::FReal StreamDistance = Owner->World->GetWorldSettings()->VolumetricLightmapLoadingRange;
+
+			FStreamingViewInfo ViewInfo = IStreamingManager::Get().GetViewInformation(0);
+			FBox Bounds = FBox( ViewInfo.ViewOrigin - FVector(StreamDistance, StreamDistance, StreamDistance), ViewInfo.ViewOrigin  + FVector(StreamDistance, StreamDistance, StreamDistance));
+			Owner->UpdateBounds(Bounds);
+		 }
+	}
+
+	virtual int32 BlockTillAllRequestsFinished(float TimeLimit = 0.0f, bool bLogResults = false) override
+	{
+		return Owner->WaitForPendingRequest(TimeLimit);
+	}
+
+	virtual void CancelForcedResources() override {}
+	virtual void NotifyLevelChange() override {}
+	virtual void SetDisregardWorldResourcesForFrames(int32 NumFrames) override {}
+	virtual void AddLevel(class ULevel* Level) override {}
+	virtual void RemoveLevel(class ULevel* Level) override {}
+
+	virtual void NotifyLevelOffset(class ULevel* Level, const FVector& Offset) override { check(false); /* Unsupported */ } 
+	virtual int32 GetNumWantingResources() const override
+	{ 
+		return Owner->GetNumPendingRequests();
+	}
+};
+
+
+FVolumetricLightmapGridManager::FVolumetricLightmapGridManager(UWorld* InWorld, FVolumetricLightMapGridDesc* InGrid)
+	: World(InWorld)
+	, Registry(InWorld->PersistentLevel->MapBuildData)
+	, Grid(InGrid)
+{
+	StreamingManager = MakeUnique<FVolumetricLightmapGridStreamingManager>(this);
+}
+
+FVolumetricLightmapGridManager::~FVolumetricLightmapGridManager()
+{
+	check(LoadedCells.IsEmpty());
+}
+
+int32 FVolumetricLightmapGridManager::GetNumPendingRequests()
+{
+	return PendingCellRequests.Num();
+}
+
+int32 FVolumetricLightmapGridManager::WaitForPendingRequest(float TimeLimit)
+{
+	bool bTimedOut = TimedExecution(PendingCellRequests, TimeLimit, [this](float TimeLimit, CellRequest& Request) -> ETimedExecutionControl
+	{		
+		if (Request.IORequest && Request.IORequest->WaitCompletion(TimeLimit))
+		{	
+			ProcessRequests();
+			return ETimedExecutionControl::Restart; // Restart iteration once requests have been processded since PendingCellRequests is modified by ProcessRequests()
+		}
+
+		return ETimedExecutionControl::Continue;
+	});
+
+	return PendingCellRequests.Num();
+}
+
+
+void FVolumetricLightmapGridManager::ReleaseCellData(FVolumetricLightMapGridCell* GridCell, FSceneInterface* InScene)
+{
+	if (GridCell->Data)
+	{			
+		FPrecomputedVolumetricLightmapData* Data = GridCell->Data;
+		GridCell->Data = nullptr;
+
+		ENQUEUE_RENDER_COMMAND(DeleteVolumetricLightDataCommand)
+			([Data] (FRHICommandListBase&)
+		{
+			Data->ReleaseResource();
+			delete(Data);
+		});
+	}
+}
+
+void FVolumetricLightmapGridManager::RemoveFromScene(FSceneInterface* InScene)
+{
+	for (auto& It : LoadedCells)
+	{
+		FPrecomputedVolumetricLightmap* Lightmap = It.Value;
+		FVolumetricLightMapGridCell* GridCell = It.Key;
+
+		if (Lightmap)
+		{
+			Lightmap->RemoveFromScene(InScene);
+		}
+
+		ReleaseCellData(GridCell, InScene);
+	}
+
+	LoadedCells.Reset();
+
+	StreamingManager.Reset();	
+}
+
+IBulkDataIORequest* FVolumetricLightmapGridManager::RequestVolumetricLightMapCell(FVolumetricLightMapGridCell& Cell)
+{	
+	if (!Cell.BulkData.GetElementCount())
+	{
+		return nullptr;
+	}
+
+	FBulkDataIORequestCallBack RequestCallback = [&Cell](bool bWasCancelled, IBulkDataIORequest* IORequest)
+	{
+		if (!bWasCancelled)
+		{
+			void* Memory = IORequest->GetReadResults();
+			if (Memory)
+			{
+				FMemoryView MemoryView(Memory, IORequest->GetSize());
+
+				FVersionedMemoryReaderView FileDataAr(MemoryView, true);
+				FPrecomputedVolumetricLightmapData* Data = nullptr;	
+
+				FileDataAr << Data;
+
+				//@todo_ow: this is fine for now since the other thread checks for both IO completion and 
+				// the presence of the ptr, when we implement cancellation we'll need to have proper synchronization
+				check(Data);
+				check(!Cell.Data);
+				Cell.Data = Data;
+				
+				FMemory::Free(Memory);
+			}
+		}
+		else
+		{
+			check(!IORequest->GetReadResults());
+			check(!Cell.Data);
+		}
+	};
+
+	IBulkDataIORequest* IORequest = Cell.BulkData.CreateStreamingRequest(AIOP_Normal, &RequestCallback, nullptr);
+
+	if (!IORequest)
+	{
+		check(Cell.BulkData.IsBulkDataLoaded());
+		// For unsaved data we can't stream it, we need to do an immediate load
+		Grid->LoadVolumetricLightMapCell(Cell);
+	}
+
+	return IORequest;
+}
+
+void FVolumetricLightmapGridManager::UpdateBounds(const FBox& InBounds)
+{
+	//@todo_ow: Make a pass on this to minimize heap allocations
+	check(Grid);	
+
+	// @todo_ow: Add check to see if bounds changed enough? there's a whole bunch of logic around this we'll see if we do it later
+	TArray<FVolumetricLightMapGridCell*> IntersectingCells = Grid->GetIntersectingCells(InBounds, true);
+
+	// Build list of cells to add & remove	
+	TArray<FVolumetricLightMapGridCell*> CellsToRequest;
+	
+	TSet<FVolumetricLightMapGridCell*>  CellsToRemove;	
+	LoadedCells.GetKeys(CellsToRemove);
+
+	for (FVolumetricLightMapGridCell* Cell : IntersectingCells)
+	{
+		if (FPrecomputedVolumetricLightmap** LightMap = LoadedCells.Find(Cell))
+		{		
+			CellsToRemove.Remove(Cell);
+		}
+		else 
+		{
+			CellsToRequest.Add(Cell);
+		}
+	}
+
+	// Make the necessary IO requests for new cells
+	for (FVolumetricLightMapGridCell* Cell : CellsToRequest)
+	{
+		if (!PendingCellRequests.ContainsByPredicate([Cell](CellRequest& Request) { return Request.Cell == Cell; }))
+		{
+			CellRequest Request;
+			Request.Cell = Cell;
+			Request.IORequest = RequestVolumetricLightMapCell(*Cell);	
+			if (Request.IORequest)
+			{
+				Request.Status = CellRequest::Requested;				
+			}
+			else
+			{
+				Request.Status = CellRequest::Ready;
+			}
+
+			PendingCellRequests.Add(Request);			
+		}
+	}
+
+	// Since cells to remove are obtained from subtracting all IntersectingCells from the LoadedCells set, pending requested cells are never
+	// in CellsToRemove, so this bit is not useful for now. We could extract cells to remove from the pending rq
+	// @todo_ow: Handle removal of requested cells, an optimization for rarer/edgier cases so we'll do this later, plus this ones needs actual synchronization with the async request
+	/*for (FVolumetricLightMapGridCell* Cell : CellsToRemove)
+	{	
+		if (CellRequest* Request = PendingCellRequests.FindByPredicate([Cell](CellRequest& InRequest) { return InRequest.Cell == Cell; }))
+		{
+			check(Request->Status == CellRequest::Requested);
+			Request->Status = CellRequest::Cancelled;			
+			Request->IORequest->Cancel();
+			delete Request->IORequest;
+			Request->IORequest = nullptr;
+		}
+	}*/
+	
+	// Remove all unnecessary cells
+	for (FVolumetricLightMapGridCell* Cell : CellsToRemove)
+	{
+		FPrecomputedVolumetricLightmap* Lightmap = LoadedCells[Cell];
+		if (Lightmap)
+		{			
+			Lightmap->RemoveFromScene(World->Scene);						
+		}
+
+		LoadedCells.Remove(Cell);
+
+		ReleaseCellData(Cell, World->Scene);
+	}
+
+	// Update currently tracked bounds
+	Bounds = InBounds;
+
+	ProcessRequests();
+}
+
+int32 FVolumetricLightmapGridManager::ProcessRequests()
+{	
+	TArray<FVolumetricLightMapGridCell*> CellsToAdd;
+
+	// Process pending IO requests and move to add list if ready
+	//@todo_ow: Use an inline array and just mark the requests that need to be removed
+	TArray<CellRequest> UpdatedRequests;
+
+	for (CellRequest& Request : PendingCellRequests)
+	{
+		if (Request.Status == CellRequest::Ready)
+		{
+			check(Request.Cell->Data);
+			check(Request.IORequest == nullptr);
+			CellsToAdd.Add(Request.Cell);
+		}
+		else if (Request.Status == CellRequest::Requested)
+		{			
+			if (Request.IORequest->PollCompletion() && Request.Cell->Data)
+			{				
+				CellsToAdd.Add(Request.Cell);
+				
+				delete Request.IORequest;
+				Request.IORequest = nullptr;
+			}		
+			else
+			{
+				UpdatedRequests.Add(Request);
+			}
+		}
+		else
+		{
+			checkf(false, TEXT("Unexpected request status\n"));
+		}
+	}
+
+	PendingCellRequests = MoveTemp(UpdatedRequests);
+
+	// Add all necessary cells
+	for (FVolumetricLightMapGridCell* Cell : CellsToAdd)
+	{
+		FPrecomputedVolumetricLightmap* Lightmap = nullptr;
+		if (Cell->Data)
+		{	
+			Lightmap = new FPrecomputedVolumetricLightmap();
+			Lightmap->AddToScene(World->Scene, Registry, Cell->Data, false);
+		}
+		else
+		{
+			check(Cell->BulkData.GetElementCount() == 0);
+		}
+		LoadedCells.Add(Cell, Lightmap);
+	}
+
+	return PendingCellRequests.Num();
+}
\ No newline at end of file
diff --git a/Engine/Source/Runtime/Engine/Private/StaticLightingBuildContext.cpp b/Engine/Source/Runtime/Engine/Private/StaticLightingBuildContext.cpp
--- a/Engine/Source/Runtime/Engine/Private/StaticLightingBuildContext.cpp
+++ b/Engine/Source/Runtime/Engine/Private/StaticLightingBuildContext.cpp
@@ -5,8 +5,10 @@
 #include "Engine/Level.h"
 #include "Engine/World.h"
 #include "GameFramework/Actor.h"
+#include "WorldPartition/StaticLightingData/VolumetricLightmapGrid.h"
 #include "Engine/MapBuildDataRegistry.h"
 #include "WorldPartition/WorldPartition.h"
+#include "PrecomputedVolumetricLightmap.h"
 
 #if WITH_EDITOR
 
@@ -35,14 +37,47 @@ FStaticLightingBuildContext::FStaticLightingBuildContext(UWorld* InWorld, ULevel
 	
 	FGuid FirstGuid = FGuid(0, 0, 0, 0);
 	check(GetLevelForGuid(FirstGuid) == World->PersistentLevel);
+
+	if (InWorld->IsPartitionedWorld())
+	{
+		VolumetricLightMapGridDesc = new FVolumetricLightMapGridDesc;
+		VolumetricLightMapGridDesc->Initialize(InWorld, InWorld->GetWorldPartition()->GetRuntimeWorldBounds());
+	}	
+	else
+	{
+		VolumetricLightMapGridDesc = nullptr;
+	}
 }
 
-FStaticLightingBuildContext::FStaticLightingBuildContext(const FStaticLightingBuildContext& InFrom)
+FStaticLightingBuildContext::FStaticLightingBuildContext(FStaticLightingBuildContext&& InFrom)
 {
 	MapBuildDataRegistry = InFrom.MapBuildDataRegistry;
-	LevelGuids = InFrom.LevelGuids;
+	LevelGuids = MoveTemp(InFrom.LevelGuids);
 	World = InFrom.World;
 	LightingScenario = InFrom.LightingScenario;
+	ImportanceBounds = InFrom.ImportanceBounds;
+	LocalToGlobalIndirectionOffset = InFrom.LocalToGlobalIndirectionOffset;
+	VolumetricLightMapGridDesc = InFrom.VolumetricLightMapGridDesc;
+
+	InFrom.VolumetricLightMapGridDesc = nullptr;
+}
+
+FStaticLightingBuildContext::~FStaticLightingBuildContext()
+{
+	delete VolumetricLightMapGridDesc;	
+}
+
+void FStaticLightingBuildContext::SetImportanceBounds(const FBox& Bounds)
+{
+	// ImportanceBounds passed to Lightmass, may not encompass the whole world (for 
+	// example when using distributed VLM computations). 
+	// All indirections are local this value, so we recompute an indirection offset
+	// to be able to move our local results to world results	
+	check(VolumetricLightMapGridDesc);
+	ImportanceBounds = Bounds;
+
+	FVector Offset =  ImportanceBounds.Min - VolumetricLightMapGridDesc->GridBounds.Min;
+	LocalToGlobalIndirectionOffset = (FIntVector)(Offset / VolumetricLightMapGridDesc->BrickSize);
 }
 
 bool FStaticLightingBuildContext::ShouldIncludeActor(AActor* Actor) const
@@ -129,6 +164,45 @@ FGuid FStaticLightingBuildContext::GetLevelGuidForLevel(ULevel* Level) const
 	return *LevelGuids.FindKey(Level);
 }
 
+FPrecomputedVolumetricLightmapData& FStaticLightingBuildContext::GetOrCreateLevelPrecomputedVolumetricLightmapBuildData(const FGuid& LevelId) const
+{	
+	if (VolumetricLightMapGridDesc && LevelId.IsValid())
+	{
+		if (FVolumetricLightMapGridCell* Cell = VolumetricLightMapGridDesc->GetCell(LevelId))
+		{	
+			if (!Cell->Data)
+			{
+				Cell->Data = new FPrecomputedVolumetricLightmapData();
+			}
+			return *Cell->Data;
+		}
+	}
+
+	UMapBuildDataRegistry* Registry = GetOrCreateRegistryForLevelGuid(LevelId);
+	FGuid BuildDataID = GetLevelBuildDataID(LevelId);
+	if (FPrecomputedVolumetricLightmapData* Data = Registry->GetLevelPrecomputedVolumetricLightmapBuildData(BuildDataID))
+	{
+		return *Data;
+	}
+	
+	return Registry->AllocateLevelPrecomputedVolumetricLightmapBuildData(BuildDataID);
+}
+
+FGuid FStaticLightingBuildContext::GetLevelGuidForVLMBrick(const FIntVector& BrickCoordinates) const
+{
+	FVolumetricLightMapGridDesc& Desc = *VolumetricLightMapGridDesc;
+	FVector BrickInWorld = (FVector)(BrickCoordinates * Desc.BrickSize) + Desc.GridBounds.Min;
+	BrickInWorld += FVector(Desc.DetailCellSize/2, Desc.DetailCellSize/2, Desc.DetailCellSize/2);	// offset by half detail cell size to avoid being on the edge
+	
+	if (FVolumetricLightMapGridCell* Cell = Desc.GetCell(BrickInWorld))
+	{
+		return Desc.GetCellGuid(Cell->CellID);
+	}
+	
+	return FGuid();
+}
+
+
 FGuid FStaticLightingBuildContext::GetLevelGuidForActor(AActor* Actor) const
 {		
 	check(Actor);
diff --git a/Engine/Source/Runtime/Engine/Private/WorldPartition/StaticLightingData/VolumetricLightmapGrid.cpp b/Engine/Source/Runtime/Engine/Private/WorldPartition/StaticLightingData/VolumetricLightmapGrid.cpp
new file mode 100644
--- /dev/null
+++ b/Engine/Source/Runtime/Engine/Private/WorldPartition/StaticLightingData/VolumetricLightmapGrid.cpp
@@ -0,0 +1,197 @@
+// Copyright Epic Games, Inc. All Rights Reserved.
+
+#include "WorldPartition/StaticLightingData/VolumetricLightmapGrid.h"
+
+#include "Engine/World.h"
+#include "GameFramework/WorldSettings.h"
+#include "Misc/ConfigCacheIni.h"
+#include "PrecomputedVolumetricLightmap.h"
+
+#include "Serialization/VersionedArchive.h"
+#include "Misc/FileHelper.h"
+#include "Serialization/MemoryWriter.h"
+#include "Serialization/MemoryReader.h"
+#include "Serialization/BulkDataReader.h"
+#include "Serialization/BulkDataWriter.h"
+
+using FVersionedBulkDataReader = TVersionedReader<FBulkDataReader>;
+using FVersionedBulkDataWriter = TVersionedWriter<FBulkDataWriter>;
+
+#if WITH_EDITOR
+void FVolumetricLightMapGridDesc::Initialize(UWorld* InWorld, const FBox& InBounds)
+{
+	// Initialize to a new guid to unlink all data	
+	Guid = FGuid::NewGuid();
+
+	FVector::FReal VLMCellSize = InWorld->GetWorldSettings()->LightmassSettings.VolumetricLightmapLoadingCellSize;
+	FVector::FReal VLMDetailCellSize = InWorld->GetWorldSettings()->LightmassSettings.VolumetricLightmapDetailCellSize;
+
+	int VLMBrickSize = 4;
+	int VLMMaxRefinementLevels = 3;
+	
+	GConfig->GetInt(TEXT("DevOptions.VolumetricLightmaps"), TEXT("BrickSize"), VLMBrickSize, GLightmassIni);
+	GConfig->GetInt(TEXT("DevOptions.VolumetricLightmaps"), TEXT("MaxRefinementLevels"), VLMMaxRefinementLevels, GLightmassIni);
+
+	VLMMaxRefinementLevels = FMath::Clamp(VLMMaxRefinementLevels, 1, 6);
+
+	const int32 BrickSizeLog2 = FMath::FloorLog2(VLMBrickSize);
+	const int32 DetailCellsPerTopLevelBrick = 1 << (VLMMaxRefinementLevels * BrickSizeLog2);
+
+	// Adjust VLMCellSize to fit our requirements
+	// VLMCellSize must be an integer multiple of a brick in world units so that we can properly align bricks with cell transitions	
+	VLMCellSize = FMath::CeilToDouble(VLMCellSize / (VLMDetailCellSize*VLMBrickSize)) * (VLMDetailCellSize*VLMBrickSize);
+
+	FVector VLMCellMinExtent = FVector(0, 0, InBounds.Min.Z);
+	FVector VLMCellMaxExtent = FVector(VLMCellSize, VLMCellSize, InBounds.Max.Z);	
+
+	FIntVector2 VLMGridMin = FIntVector2(FMath::FloorToInt(InBounds.Min.X / VLMCellSize), FMath::FloorToInt(InBounds.Min.Y / VLMCellSize));
+	FIntVector2 VLMGridMax = FIntVector2(FMath::CeilToInt(InBounds.Max.X / VLMCellSize), FMath::CeilToInt(InBounds.Max.Y / VLMCellSize));
+		
+	// Keep the values used to generate the grid for proper usage of the data
+	DetailCellSize = VLMDetailCellSize;
+	BrickSize = VLMDetailCellSize * VLMBrickSize;
+	CellSize = VLMCellSize;	
+
+	GridBounds = FBox(FVector(VLMGridMin.X * VLMCellSize, VLMGridMin.Y * VLMCellSize, InBounds.Min.Z) ,
+					  FVector(VLMGridMax.X * VLMCellSize, VLMGridMax.Y * VLMCellSize, InBounds.Max.Z));
+
+	
+	// Initialize the cells
+	Cells.Reset();
+	Cells.SetNum((VLMGridMax.X - VLMGridMin.X) * (VLMGridMax.Y - VLMGridMin.Y));
+
+	uint32 CellID = 0;
+
+	for (int i = VLMGridMin.X; i < VLMGridMax.X; i++)
+	{
+		for (int j = VLMGridMin.Y; j < VLMGridMax.Y; j++)
+		{			
+			FVector CellPosition (i * VLMCellSize, j * VLMCellSize, 0);
+
+			Cells[CellID].Bounds = FBox(CellPosition + VLMCellMinExtent, CellPosition + VLMCellMaxExtent);
+			Cells[CellID].CellID = CellID;
+			CellID++;
+		}
+	}
+}
+#endif
+
+FVolumetricLightMapGridCell* FVolumetricLightMapGridDesc::GetCell(FGuid CellGuid) const
+{
+	check(IsValid());
+	
+	for (const FVolumetricLightMapGridCell& Cell : Cells)
+	{
+		if (GetCellGuid(Cell.CellID) == CellGuid)
+		{
+			return const_cast<FVolumetricLightMapGridCell*>(&Cell);
+		}
+	}	
+
+	return nullptr;
+}
+
+FVolumetricLightMapGridCell* FVolumetricLightMapGridDesc::GetCell(const FVector& InPosition) const
+{
+	check(IsValid());
+	
+	for (const FVolumetricLightMapGridCell& Cell : Cells)
+	{
+		if (Cell.Bounds.IsInsideOrOn(InPosition))		
+		{
+			return const_cast<FVolumetricLightMapGridCell*>(&Cell);
+		}
+	}	
+
+	return nullptr;
+}
+
+FGuid FVolumetricLightMapGridDesc::GetCellGuid(uint32 CellID) const
+{
+	return FGuid::Combine(Guid, FGuid(CellID, CellID, CellID, CellID));
+}
+
+FPrecomputedVolumetricLightmapData* FVolumetricLightMapGridDesc::GetPrecomputedVolumetricLightmapBuildData(FGuid LevelId) const
+{
+	if (FVolumetricLightMapGridCell* Cell = GetCell(LevelId))
+	{
+		return Cell->Data;
+	}
+
+	return nullptr;
+}
+	
+FPrecomputedVolumetricLightmapData* FVolumetricLightMapGridDesc::GetOrCreatePrecomputedVolumetricLightmapBuildData(FGuid LevelId)
+{
+	FVolumetricLightMapGridCell* Cell = GetCell(LevelId);
+	check(Cell); 
+
+	if (!Cell->Data)
+	{
+		Cell->Data = new FPrecomputedVolumetricLightmapData();
+	}
+
+	return Cell->Data;
+}
+
+TArray<FVolumetricLightMapGridCell*> FVolumetricLightMapGridDesc::GetIntersectingCells(const FBox& InBounds, bool bInWithData)
+{
+	TArray<FVolumetricLightMapGridCell*> IntersectingCells;
+
+	for (FVolumetricLightMapGridCell& Cell : Cells)
+	{
+		if (Cell.Bounds.Intersect(InBounds) && (Cell.BulkData.GetElementCount() || Cell.Data || !bInWithData) )
+		{
+			IntersectingCells.Add(&Cell);
+		}
+	}
+
+	return IntersectingCells;
+}
+
+void FVolumetricLightMapGridDesc::SerializeBulkData(FArchive& Ar, UObject* Owner)
+{	
+	for (FVolumetricLightMapGridCell& Cell : Cells)
+	{		
+		Cell.BulkData.Serialize(Ar, Owner);		
+	}	
+}
+
+#if WITH_EDITOR
+
+void FVolumetricLightMapGridDesc::InitializeBulkData()
+{
+	for (FVolumetricLightMapGridCell& Cell : Cells)
+	{
+		Cell.BulkData.SetBulkDataFlags(BULKDATA_Force_NOT_InlinePayload);
+
+		if (!Cell.Data)
+		{
+			Cell.BulkData.RemoveBulkData();
+			continue;
+		}
+			
+		FVersionedBulkDataWriter Ar(Cell.BulkData, true);
+		Ar << Cell.Data;
+	}
+}
+#endif
+
+void FVolumetricLightMapGridDesc::LoadVolumetricLightMapCell(FVolumetricLightMapGridCell& Cell)
+{
+	if (Cell.BulkData.GetElementCount())
+	{
+		FVersionedBulkDataReader Ar(Cell.BulkData, true);
+		Ar << Cell.Data;
+	}
+}
+
+FString	FVolumetricLightMapGridDesc::GetCellDesc(FGuid CellGuid) const
+{
+	if (FVolumetricLightMapGridCell* Cell = GetCell(CellGuid))
+	{
+		return FString::Printf(TEXT("Cell: %i (%s -> %s)"), Cell->CellID, *Cell->Bounds.Min.ToString(), *Cell->Bounds.Max.ToString());
+	}
+
+	return FString();		
+}
diff --git a/Engine/Source/Runtime/Engine/Private/WorldSettings.cpp b/Engine/Source/Runtime/Engine/Private/WorldSettings.cpp
--- a/Engine/Source/Runtime/Engine/Private/WorldSettings.cpp
+++ b/Engine/Source/Runtime/Engine/Private/WorldSettings.cpp
@@ -115,6 +115,8 @@ PRAGMA_ENABLE_DEPRECATION_WARNINGS
 	bPlaceCellsOnlyAlongCameraTracks = false;
 	VisibilityCellSize = 200;
 	VisibilityAggressiveness = VIS_LeastAggressive;
+	
+	VolumetricLightmapLoadingRange = 6400;
 
 #if WITH_EDITORONLY_DATA
 	bActorLabelEditable = false;
@@ -227,9 +229,22 @@ void AWorldSettings::SetWorldPartition(UWorldPartition* InWorldPartition)
 
 void AWorldSettings::ApplyWorldPartitionForcedSettings()
 {
+	const auto CVarAllowStaticLightingOnWPMaps = IConsoleManager::Get().FindTConsoleVariableDataInt(TEXT("r.AllowStaticLightingInWorldPartitionMaps"));
+
 	bEnableWorldComposition = false;
-	bForceNoPrecomputedLighting = true;
 	bPrecomputeVisibility = false;
+	if (CVarAllowStaticLightingOnWPMaps->GetValueOnAnyThread() != 0)
+	{
+		// leave static lighting options alone
+	}
+	else
+	{
+		bForceNoPrecomputedLighting = true;	
+	}
+#if WITH_EDITOR
+
+	LightmassSettings.bWorldPartition = true;
+#endif
 }
 
 float AWorldSettings::GetGravityZ() const
@@ -561,7 +576,7 @@ void AWorldSettings::PostLoad()
 #if WITH_EDITOR
 	if (WorldPartition)
 	{
-		// Force to re-apply WorldPartition restrictions on WorldSettings (in case they changed)
+		// Force to re-apply WorldPartition restrictions on WorldSettings (in case they changed)	
 		ApplyWorldPartitionForcedSettings();
 	}
 #endif// WITH_EDITOR
@@ -678,7 +693,9 @@ bool AWorldSettings::CanEditChange(const FProperty* InProperty) const
 
 			if (PropertyName == GET_MEMBER_NAME_STRING_CHECKED(FLightmassWorldInfoSettings, VolumetricLightmapDetailCellSize)
 				|| PropertyName == GET_MEMBER_NAME_STRING_CHECKED(FLightmassWorldInfoSettings, VolumetricLightmapMaximumBrickMemoryMb)
-				|| PropertyName == GET_MEMBER_NAME_STRING_CHECKED(FLightmassWorldInfoSettings, VolumetricLightmapSphericalHarmonicSmoothing))
+				|| PropertyName == GET_MEMBER_NAME_STRING_CHECKED(FLightmassWorldInfoSettings, VolumetricLightmapSphericalHarmonicSmoothing)
+				|| PropertyName == GET_MEMBER_NAME_STRING_CHECKED(FLightmassWorldInfoSettings, VolumetricLightmapLoadingCellSize)
+				|| PropertyName == GET_MEMBER_NAME_STRING_CHECKED(AWorldSettings, VolumetricLightmapLoadingRange))
 			{
 				return LightmassSettings.VolumeLightingMethod == VLM_VolumetricLightmap;
 			}
@@ -694,7 +711,6 @@ bool AWorldSettings::CanEditChange(const FProperty* InProperty) const
 			}
 		}
 		else if (PropertyName == GET_MEMBER_NAME_STRING_CHECKED(AWorldSettings, bEnableWorldComposition ) ||
-				 PropertyName == GET_MEMBER_NAME_STRING_CHECKED(AWorldSettings, bForceNoPrecomputedLighting ) ||
 				 PropertyName == GET_MEMBER_NAME_STRING_CHECKED(AWorldSettings, bPrecomputeVisibility))
 		{
 			return !IsPartitionedWorld();
diff --git a/Engine/Source/Runtime/Engine/Public/PrecomputedVolumetricLightmap.h b/Engine/Source/Runtime/Engine/Public/PrecomputedVolumetricLightmap.h
--- a/Engine/Source/Runtime/Engine/Public/PrecomputedVolumetricLightmap.h
+++ b/Engine/Source/Runtime/Engine/Public/PrecomputedVolumetricLightmap.h
@@ -223,6 +223,7 @@ public:
 	ENGINE_API ~FPrecomputedVolumetricLightmap();
 
 	ENGINE_API void AddToScene(class FSceneInterface* Scene, class UMapBuildDataRegistry* Registry, FGuid LevelBuildDataId, bool bIsPersistentLevel);
+	ENGINE_API void AddToScene(class FSceneInterface* Scene, class UMapBuildDataRegistry* Registry, FPrecomputedVolumetricLightmapData* NewData, bool bIsPersistentLevel);
 
 	ENGINE_API void RemoveFromScene(FSceneInterface* Scene);
 	
@@ -580,4 +581,4 @@ private:
 	int32 PaddedBrickSize;
 };
 
-extern ENGINE_API TGlobalResource<FVolumetricLightmapBrickAtlas> GVolumetricLightmapBrickAtlas;
+extern ENGINE_API TGlobalResource<FVolumetricLightmapBrickAtlas> GVolumetricLightmapBrickAtlas;
\ No newline at end of file
diff --git a/Engine/Source/Runtime/Engine/Public/PrecomputedVolumetricLightmapStreaming.h b/Engine/Source/Runtime/Engine/Public/PrecomputedVolumetricLightmapStreaming.h
new file mode 100644
--- /dev/null
+++ b/Engine/Source/Runtime/Engine/Public/PrecomputedVolumetricLightmapStreaming.h
@@ -0,0 +1,66 @@
+// Copyright Epic Games, Inc. All Rights Reserved.
+
+#pragma once
+
+#include "CoreMinimal.h"
+#include "Misc/Guid.h"
+
+class FSceneInterface;
+class UWorld;
+class IBulkDataIORequest;
+class FPrecomputedVolumetricLightmap;
+
+struct FVolumetricLightMapGridDesc;
+struct FVolumetricLightMapGridCell;
+class FVolumetricLightmapGridStreamingManager;
+class UMapBuildDataRegistry;
+
+class FVolumetricLightmapGridManager
+{	
+public:
+
+	FVolumetricLightmapGridManager(UWorld* InWorld, FVolumetricLightMapGridDesc* Grid);
+	~FVolumetricLightmapGridManager();
+
+	void UpdateBounds(const FBox& Bounds);
+	int32 ProcessRequests();
+	void RemoveFromScene(class FSceneInterface* InScene);
+	int32 WaitForPendingRequest(float Timelimit);
+	int32 GetNumPendingRequests();	
+
+private:
+	
+	friend class FVolumetricLightmapGridStreamingManager;
+
+	void ReleaseCellData(FVolumetricLightMapGridCell* GridCell, FSceneInterface* InScene);	
+
+	UWorld* World = nullptr;
+	UMapBuildDataRegistry* Registry = nullptr;
+	FVolumetricLightMapGridDesc* Grid = nullptr;
+	
+	struct CellRequest
+	{
+		enum EStatus
+		{
+			Created,
+			Requested,
+			Ready,
+			Cancelled,
+			Done
+		};
+
+		EStatus Status = Created;
+		IBulkDataIORequest* IORequest = nullptr;
+		FVolumetricLightMapGridCell* Cell = nullptr;			
+	};	
+
+	IBulkDataIORequest* RequestVolumetricLightMapCell(FVolumetricLightMapGridCell& Cell);
+	
+	TArray<CellRequest>	PendingCellRequests;
+	TMap<FVolumetricLightMapGridCell*, FPrecomputedVolumetricLightmap*> LoadedCells;	
+
+	FBox Bounds;
+
+	TUniquePtr<FVolumetricLightmapGridStreamingManager> StreamingManager;
+};
+
diff --git a/Engine/Source/Runtime/Engine/Public/StaticLightingBuildContext.h b/Engine/Source/Runtime/Engine/Public/StaticLightingBuildContext.h
--- a/Engine/Source/Runtime/Engine/Public/StaticLightingBuildContext.h
+++ b/Engine/Source/Runtime/Engine/Public/StaticLightingBuildContext.h
@@ -13,22 +13,31 @@ class ULevel;
 class UWorld;
 class AActor;
 struct FVolumetricLightMapGridDesc;
+class FPrecomputedVolumetricLightmapData;
 
 #if WITH_EDITOR
 class FStaticLightingBuildContext
 {	
 	mutable UMapBuildDataRegistry* MapBuildDataRegistry;
-		
-	TMap<FGuid, const TWeakObjectPtr<ULevel>> LevelGuids;	
+	
+	TMap<FGuid, const TWeakObjectPtr<ULevel>> LevelGuids;
+	
+	FBox ImportanceBounds;
+	FIntVector LocalToGlobalIndirectionOffset;
+	FVolumetricLightMapGridDesc* VolumetricLightMapGridDesc;
 
+	
 public:
 
 	UWorld*						World;
 	ULevel*						LightingScenario;
 
 	ENGINE_API FStaticLightingBuildContext(UWorld* InWorld, ULevel* InLightingScenario);
-	ENGINE_API FStaticLightingBuildContext(const FStaticLightingBuildContext& InFrom);
+	ENGINE_API FStaticLightingBuildContext(FStaticLightingBuildContext&& InFrom);
+	ENGINE_API ~FStaticLightingBuildContext();
 	
+	ENGINE_API void SetImportanceBounds(const FBox& Bounds);
+
 	ENGINE_API ULevel* GetLightingStorageLevel(ULevel* Level) const;
 
 	ENGINE_API bool ShouldIncludeActor(AActor* Actor) const;
@@ -41,11 +50,17 @@ public:
 	ENGINE_API UMapBuildDataRegistry* GetOrCreateRegistryForLevelGuid(const FGuid& Guid) const;
 	ENGINE_API FGuid GetLevelGuidForActor(AActor* Actor) const;
 	ENGINE_API FGuid GetPersistentLevelGuid() const;
-	ENGINE_API FGuid GetLevelGuidForLevel(ULevel* Level) const;	
+	ENGINE_API FGuid GetLevelGuidForLevel(ULevel* Level) const;
+	ENGINE_API FGuid GetLevelGuidForVLMBrick(const FIntVector& BrickCoordinates) const;
 	ENGINE_API const TWeakObjectPtr<ULevel> GetLevelForGuid(const FGuid& Guid) const;
 	ENGINE_API FGuid GetLevelBuildDataID(const FGuid& LevelGuid) const;
-	
-	ENGINE_API UMapBuildDataRegistry* GetRegistryForLevel(ULevel* Level) const;	
+
+	ENGINE_API FPrecomputedVolumetricLightmapData& GetOrCreateLevelPrecomputedVolumetricLightmapBuildData(const FGuid& LevelId) const;
+
+	ENGINE_API FVolumetricLightMapGridDesc* GetVolumetricLightMapGridDesc() const { return VolumetricLightMapGridDesc; }
+	void ReleaseVolumetricLightMapGridDesc() { VolumetricLightMapGridDesc = nullptr; }
+
+	ENGINE_API UMapBuildDataRegistry* GetRegistryForLevel(ULevel* Level) const;
 	ENGINE_API UMapBuildDataRegistry* GetOrCreateRegistryForActor(AActor* Actor) const;
 	ENGINE_API UMapBuildDataRegistry* GetOrCreateRegistryForLevel(ULevel* Level) const;	
 };
diff --git a/Engine/Source/Runtime/Engine/Public/WorldPartition/StaticLightingData/VolumetricLightmapGrid.h b/Engine/Source/Runtime/Engine/Public/WorldPartition/StaticLightingData/VolumetricLightmapGrid.h
new file mode 100644
--- /dev/null
+++ b/Engine/Source/Runtime/Engine/Public/WorldPartition/StaticLightingData/VolumetricLightmapGrid.h
@@ -0,0 +1,74 @@
+// Copyright Epic Games, Inc. All Rights Reserved.
+
+#pragma once
+
+#include "CoreMinimal.h"
+#include "Misc/Guid.h"
+#include "Serialization/BulkData.h"
+#include "VolumetricLightmapGrid.generated.h"
+
+
+class FPrecomputedVolumetricLightmapData;
+class UMapBuildDataRegistry;
+
+USTRUCT()
+struct FVolumetricLightMapGridCell
+{
+	GENERATED_USTRUCT_BODY()
+	
+	UPROPERTY()
+	FBox Bounds;
+
+	FPrecomputedVolumetricLightmapData* Data = nullptr;	
+
+	FByteBulkData		BulkData;
+	
+	UPROPERTY()
+	uint32 CellID;	
+};	
+
+USTRUCT()
+struct FVolumetricLightMapGridDesc
+{
+	GENERATED_USTRUCT_BODY()
+
+	UPROPERTY()
+	FBox GridBounds;
+	
+	UPROPERTY()
+	FGuid Guid;	
+
+	UPROPERTY()
+	float CellSize = 0.0f;
+
+	UPROPERTY()
+	float DetailCellSize = 0.0f;
+
+	UPROPERTY()
+	float BrickSize = 0.0f;
+
+	UPROPERTY()
+	TArray<FVolumetricLightMapGridCell>		Cells;
+	
+	bool IsValid() const { return Guid.IsValid(); }
+
+	ENGINE_API FGuid	GetCellGuid(uint32 CellID) const;
+	ENGINE_API FString	GetCellDesc(FGuid LevelId) const;
+
+	ENGINE_API FVolumetricLightMapGridCell* GetCell(FGuid LevelId) const;	
+	ENGINE_API FVolumetricLightMapGridCell* GetCell(const FVector& InPosition) const; 
+	
+	[[nodiscard]] ENGINE_API TArray<FVolumetricLightMapGridCell*> GetIntersectingCells(const FBox& InBounds, bool bInWithData);	
+	
+	ENGINE_API void LoadVolumetricLightMapCell(FVolumetricLightMapGridCell& Cell);
+
+	ENGINE_API FPrecomputedVolumetricLightmapData* GetPrecomputedVolumetricLightmapBuildData(FGuid LevelId) const;
+	ENGINE_API FPrecomputedVolumetricLightmapData* GetOrCreatePrecomputedVolumetricLightmapBuildData(FGuid LevelId);
+
+#if WITH_EDITOR
+	ENGINE_API void InitializeBulkData();
+	ENGINE_API void Initialize(UWorld* InWorld, const FBox& InBounds);
+#endif
+
+	void SerializeBulkData(FArchive& Ar, UObject* Owner);
+};
