- Enable lightweight memory capture by default

diff --git a/Engine/Source/Runtime/Engine/Private/Particles/ParticlePerfStats.cpp b/Engine/Source/Runtime/Engine/Private/Particles/ParticlePerfStats.cpp
--- a/Engine/Source/Runtime/Engine/Private/Particles/ParticlePerfStats.cpp
+++ b/Engine/Source/Runtime/Engine/Private/Particles/ParticlePerfStats.cpp
@@ -1117,7 +1117,7 @@ static FAutoConsoleVariable CVarWriteDetailedCSVStats(
 	ECVF_Default | ECVF_RenderThreadSafe
 );
 
-int32 GFXDetailedCSVMemorySMode = 0;
+int32 GFXDetailedCSVMemorySMode = 1;
 static FAutoConsoleVariableRef CVarFXDetailedCSVMemoryStats(
 	TEXT("fx.DetailedCSVStats.MemoryMode"),
 	GFXDetailedCSVMemorySMode,
