Skeletal Mesh: Don't convert bulk data on save if cooking.

The mesh bulk data won't be used anyway and just takes up unnecessary amount of memory.

diff --git a/Engine/Source/Runtime/Engine/Private/SkeletalMesh.cpp b/Engine/Source/Runtime/Engine/Private/SkeletalMesh.cpp
--- a/Engine/Source/Runtime/Engine/Private/SkeletalMesh.cpp
+++ b/Engine/Source/Runtime/Engine/Private/SkeletalMesh.cpp
@@ -1623,7 +1623,7 @@ void USkeletalMesh::Serialize( FArchive& Ar )
 		FSkinnedAssetCompilingManager::Get().FinishCompilation({ this });
 	}
 
-	if (Ar.IsSaving())
+	if (Ar.IsSaving() && !Ar.IsCooking())
 	{
 		// If saving out to disk, ensure that all source models have had their raw mesh bulk data converted to mesh description,
 		// since the bulk data won't be reloaded.
