Bump animation sequence DDC key to purge stale data due to retarget source changing impacting additive animations.

diff --git a/Engine/Source/Runtime/Engine/Classes/Animation/AnimCompress.h b/Engine/Source/Runtime/Engine/Classes/Animation/AnimCompress.h
--- a/Engine/Source/Runtime/Engine/Classes/Animation/AnimCompress.h
+++ b/Engine/Source/Runtime/Engine/Classes/Animation/AnimCompress.h
@@ -272,7 +272,7 @@ namespace UE
 			// This is a version string that mimics the old versioning scheme. If you
 			// want to bump this version, generate a new guid using VS->Tools->Create GUID and
 			// return it here. Ex.
-			static const FString AnimationCompressionVersionString = TEXT("6AB870F5E91D43739796EA2787A3FB07");
+			static const FString AnimationCompressionVersionString = TEXT("0439926D560447329623BE4394FA11A6");
 			
 			struct FAnimationCompressionMemorySummaryScope
 			{
