Small change to appease the static analysis

diff --git a/Engine/Source/Runtime/Engine/Private/PrecomputedVolumetricLightmapStreaming.cpp b/Engine/Source/Runtime/Engine/Private/PrecomputedVolumetricLightmapStreaming.cpp
--- a/Engine/Source/Runtime/Engine/Private/PrecomputedVolumetricLightmapStreaming.cpp
+++ b/Engine/Source/Runtime/Engine/Private/PrecomputedVolumetricLightmapStreaming.cpp
@@ -278,14 +278,14 @@ void FVolumetricLightmapGridManager::UpdateBounds(const FBox& InBounds)
 			Request.IORequest = RequestVolumetricLightMapCell(*Cell);	
 			if (Request.IORequest)
 			{
-				Request.Status = CellRequest::Requested;				
+				Request.Status = CellRequest::Requested;
 			}
 			else
 			{
 				Request.Status = CellRequest::Ready;
 			}
 
-			PendingCellRequests.Add(Request);			
+			PendingCellRequests.Add(Request);
 		}
 	}
 
@@ -310,7 +310,7 @@ void FVolumetricLightmapGridManager::UpdateBounds(const FBox& InBounds)
 		FPrecomputedVolumetricLightmap* Lightmap = LoadedCells[Cell];
 		if (Lightmap)
 		{			
-			Lightmap->RemoveFromScene(World->Scene);						
+			Lightmap->RemoveFromScene(World->Scene);
 		}
 
 		LoadedCells.Remove(Cell);
@@ -341,14 +341,15 @@ int32 FVolumetricLightmapGridManager::ProcessRequests()
 			CellsToAdd.Add(Request.Cell);
 		}
 		else if (Request.Status == CellRequest::Requested)
-		{			
+		{
 			if (Request.IORequest->PollCompletion() && Request.Cell->Data)
-			{				
+			{
 				CellsToAdd.Add(Request.Cell);
 				
 				delete Request.IORequest;
 				Request.IORequest = nullptr;
-			}		
+				//Request.Status = CellRequest::Ready;
+			}
 			else
 			{
 				UpdatedRequests.Add(Request);
