Small change to appease the static analysis (for real this time)

diff --git a/Engine/Source/Runtime/Engine/Private/PrecomputedVolumetricLightmapStreaming.cpp b/Engine/Source/Runtime/Engine/Private/PrecomputedVolumetricLightmapStreaming.cpp
--- a/Engine/Source/Runtime/Engine/Private/PrecomputedVolumetricLightmapStreaming.cpp
+++ b/Engine/Source/Runtime/Engine/Private/PrecomputedVolumetricLightmapStreaming.cpp
@@ -348,7 +348,7 @@ int32 FVolumetricLightmapGridManager::ProcessRequests()
 				
 				delete Request.IORequest;
 				Request.IORequest = nullptr;
-				//Request.Status = CellRequest::Ready;
+				Request.Status = CellRequest::Ready;
 			}
 			else
 			{
