#UE Fix non-static access for CVAR

diff --git a/Engine/Source/Runtime/Engine/Private/WorldSettings.cpp b/Engine/Source/Runtime/Engine/Private/WorldSettings.cpp
--- a/Engine/Source/Runtime/Engine/Private/WorldSettings.cpp
+++ b/Engine/Source/Runtime/Engine/Private/WorldSettings.cpp
@@ -229,7 +229,7 @@ void AWorldSettings::SetWorldPartition(UWorldPartition* InWorldPartition)
 
 void AWorldSettings::ApplyWorldPartitionForcedSettings()
 {
-	const auto CVarAllowStaticLightingOnWPMaps = IConsoleManager::Get().FindTConsoleVariableDataInt(TEXT("r.AllowStaticLightingInWorldPartitionMaps"));
+	static const auto CVarAllowStaticLightingOnWPMaps = IConsoleManager::Get().FindTConsoleVariableDataInt(TEXT("r.AllowStaticLightingInWorldPartitionMaps"));
 
 	bEnableWorldComposition = false;
 	bPrecomputeVisibility = false;
