Differentiate between whether asset is too big always or just too big right now

The old behavior is that we log out when an asset is bigger than the current limit, which is the minimum of a manually set limit and the currently available memory. The new behavior separately reports assets that violate the manually set memory limit. The rationale here is that users may cook assets that are wildly beyond what we could even deal with even if we had all the memory we accounted for available. These cases need separate treatment and reporting them separately makes it much simpler to triage them.

diff --git a/Engine/Source/Runtime/Engine/Private/AssetCompilingManager.cpp b/Engine/Source/Runtime/Engine/Private/AssetCompilingManager.cpp
--- a/Engine/Source/Runtime/Engine/Private/AssetCompilingManager.cpp
+++ b/Engine/Source/Runtime/Engine/Private/AssetCompilingManager.cpp
@@ -293,18 +293,35 @@ public:
 					DebugName = TEXT("No DebugName");
 				}
 
-				if ( NewRequiredMemory != GetDefaultMemoryPerAsset() )
+				const bool bHasMemoryEstimate = NextWork && NextWork->GetRequiredMemory() >= 0;
+				if (bHasMemoryEstimate)
 				{
-					// don't bother logging if estimate is the default (4GB)
-
-					UE_LOGFMT_NSLOC(LogAsyncCompilation, Display, "AsyncAssetCompilation", "MemoryLimitExceeded",
-						"BEWARE: AssetCompile memory estimate is greater than available, but we're running it [{TaskName}] anyway! "
-						"RequiredMemory = {TotalEstimatedMemory} MiB + {RequiredMemory} MiB, MemoryLimit = {MemoryLimit} MiB ",
-						("TaskName", DebugName),
-						("RequiredMemory", FString::SanitizeFloat(NewRequiredMemory / (1024 * 1024.f), 3)),
-						("TotalEstimatedMemory", FString::SanitizeFloat(TotalEstimatedMemory / (1024 * 1024.f), 3)),
-						("MemoryLimit", FString::SanitizeFloat(MemoryLimit / (1024 * 1024.f), 3))
-					);
+					const int64 HardMemoryLimit = GetHardMemoryLimit();
+					if (NewRequiredMemory > HardMemoryLimit && HardMemoryLimit > 0)
+					{
+						// Separately report the case if an asset was bigger than our manually set memory limit. Such assets are always too big,
+						// and don't just exceed the currently available memory.
+						UE_LOGFMT_NSLOC(LogAsyncCompilation, Display, "AsyncAssetCompilation", "HardMemoryLimitExceeded",
+							"BEWARE: AssetCompile memory estimate is greater than the hard memory limit, but we're running it [{TaskName}] anyway! "
+							"RequiredMemory = {TotalEstimatedMemory} MiB + {RequiredMemory} MiB, MemoryLimit = {MemoryLimit} MiB, HardMemoryLimit = {HardMemoryLimit} MiB",
+							("TaskName", DebugName),
+							("RequiredMemory", FString::SanitizeFloat(NewRequiredMemory / (1024 * 1024.f), 3)),
+							("TotalEstimatedMemory", FString::SanitizeFloat(TotalEstimatedMemory / (1024 * 1024.f), 3)),
+							("MemoryLimit", FString::SanitizeFloat(MemoryLimit / (1024 * 1024.f), 3)),
+							("HardMemoryLimit", FString::SanitizeFloat(HardMemoryLimit / (1024 * 1024.f), 3))
+						);
+					}
+					else
+					{
+						UE_LOGFMT_NSLOC(LogAsyncCompilation, Display, "AsyncAssetCompilation", "MemoryLimitExceeded",
+							"BEWARE: AssetCompile memory estimate is greater than available, but we're running it [{TaskName}] anyway! "
+							"RequiredMemory = {TotalEstimatedMemory} MiB + {RequiredMemory} MiB, MemoryLimit = {MemoryLimit} MiB ",
+							("TaskName", DebugName),
+							("RequiredMemory", FString::SanitizeFloat(NewRequiredMemory / (1024 * 1024.f), 3)),
+							("TotalEstimatedMemory", FString::SanitizeFloat(TotalEstimatedMemory / (1024 * 1024.f), 3)),
+							("MemoryLimit", FString::SanitizeFloat(MemoryLimit / (1024 * 1024.f), 3))
+						);
+					}
 				}
 
 				// @todo : ? pause the main thread? pause shader compilers? trigger a GC ?
