SplineComponent::GetSplinePointAt() now returns full FSplinePoint data, was missing the spline's Type property.

diff --git a/Engine/Source/Runtime/Engine/Private/Components/SplineComponent.cpp b/Engine/Source/Runtime/Engine/Private/Components/SplineComponent.cpp
--- a/Engine/Source/Runtime/Engine/Private/Components/SplineComponent.cpp
+++ b/Engine/Source/Runtime/Engine/Private/Components/SplineComponent.cpp
@@ -1062,6 +1062,7 @@ FSplinePoint USplineComponent::GetSplinePointAt(int32 PointIndex, ESplineCoordin
 
 	const FInterpCurvePointQuat& RotationPoint = GetRotationPointSafe(PointIndex);
 	const FRotator& Rotation = GetRotationAtSplineInputKey(RotationPoint.InVal, CoordinateSpace);
+	const ESplinePointType::Type SplinePointType = GetSplinePointType(PointIndex);
 
 	const FVector Scale = GetScaleAtSplinePoint(PointIndex);
 
@@ -1070,7 +1071,8 @@ FSplinePoint USplineComponent::GetSplinePointAt(int32 PointIndex, ESplineCoordin
 		SplinePoint.ArriveTangent,
 		SplinePoint.LeaveTangent,
 		Rotation,
-		Scale);
+		Scale,
+		SplinePointType);
 }
 
 FVector USplineComponent::GetLocationAtSplinePoint(int32 PointIndex, ESplineCoordinateSpace::Type CoordinateSpace) const
