Added TIsContiguousContainer specialization for TInlineComponentArray, so it can bind to TArrayView.

diff --git a/Engine/Source/Runtime/Engine/Classes/GameFramework/Actor.h b/Engine/Source/Runtime/Engine/Classes/GameFramework/Actor.h
--- a/Engine/Source/Runtime/Engine/Classes/GameFramework/Actor.h
+++ b/Engine/Source/Runtime/Engine/Classes/GameFramework/Actor.h
@@ -153,6 +153,12 @@ public:
 	TInlineComponentArray(const AActor* Actor, bool bIncludeFromChildActors = false);
 };
 
+template<class T, uint32 NumElements>
+struct TIsContiguousContainer<TInlineComponentArray<T, NumElements>>
+{
+	enum { Value = true };
+};
+
 /**
  * Actor is the base class for an Object that can be placed or spawned in a level.
  * Actors may contain a collection of ActorComponents, which can be used to control how actors move, how they are rendered, etc.
