Added default value for CellID (always overwriten later)

diff --git a/Engine/Source/Runtime/Engine/Public/WorldPartition/StaticLightingData/VolumetricLightmapGrid.h b/Engine/Source/Runtime/Engine/Public/WorldPartition/StaticLightingData/VolumetricLightmapGrid.h
--- a/Engine/Source/Runtime/Engine/Public/WorldPartition/StaticLightingData/VolumetricLightmapGrid.h
+++ b/Engine/Source/Runtime/Engine/Public/WorldPartition/StaticLightingData/VolumetricLightmapGrid.h
@@ -24,7 +24,7 @@ struct FVolumetricLightMapGridCell
 	FByteBulkData		BulkData;
 	
 	UPROPERTY()
-	uint32 CellID;	
+	uint32 CellID = -1;
 };	
 
 USTRUCT()
