Fixed HeterogeneousVolume component not picking up on changes of the SparseVolumeTexture parameter in the assigned material. Fixed by directly querying the assigned material for the SVT instead of querying the internally created MID. Also fixed an issue where duplicating a HV component broke playback because the MID was never created.

#rnx

diff --git a/Engine/Source/Runtime/Engine/Classes/Components/HeterogeneousVolumeComponent.h b/Engine/Source/Runtime/Engine/Classes/Components/HeterogeneousVolumeComponent.h
--- a/Engine/Source/Runtime/Engine/Classes/Components/HeterogeneousVolumeComponent.h
+++ b/Engine/Source/Runtime/Engine/Classes/Components/HeterogeneousVolumeComponent.h
@@ -119,6 +119,7 @@ private:
 	static USparseVolumeTexture* GetSparseVolumeTexture(UMaterialInterface* MaterialInterface, int32 ParameterIndex, FName* OutParamName = nullptr);
 	static UMaterialInstanceDynamic* CreateOrCastToMID(UMaterialInterface* MaterialInterface);
 	void OnSparseVolumeTextureChanged(const USparseVolumeTexture* SparseVolumeTexture);
+	UMaterialInterface* GetHeterogeneousVolumeMaterial() const; // Gets the UMaterialInterface* returned by GetMaterial(0), but returns nullptr if the material is incompatible with HeterogeneousVolumes.
 };
 
 /**
diff --git a/Engine/Source/Runtime/Engine/Private/Components/HeterogeneousVolumeComponent.cpp b/Engine/Source/Runtime/Engine/Private/Components/HeterogeneousVolumeComponent.cpp
--- a/Engine/Source/Runtime/Engine/Private/Components/HeterogeneousVolumeComponent.cpp
+++ b/Engine/Source/Runtime/Engine/Private/Components/HeterogeneousVolumeComponent.cpp
@@ -353,15 +353,9 @@ void UHeterogeneousVolumeComponent::PostLoad()
 {
 	Super::PostLoad();
 
-	UMaterialInterface* MaterialInterface = GetMaterial(0);
-	if (MaterialInterface)
+	MaterialInstanceDynamic = nullptr;
+	if (UMaterialInterface* MaterialInterface = GetHeterogeneousVolumeMaterial())
 	{
-		const UMaterial* Material = MaterialInterface->GetMaterial();
-		if (Material && Material->MaterialDomain == EMaterialDomain::MD_Volume)
-		{
-			MaterialInterface->CheckMaterialUsage(MATUSAGE_HeterogeneousVolumes);
-		}
-
 		MaterialInstanceDynamic = CreateOrCastToMID(MaterialInterface);
 	}
 }
@@ -435,6 +429,22 @@ void UHeterogeneousVolumeComponent::OnSparseVolumeTextureChanged(const USparseVo
 	MarkRenderStateDirty();
 }
 
+UMaterialInterface* UHeterogeneousVolumeComponent::GetHeterogeneousVolumeMaterial() const
+{
+	const uint32 MaterialIndex = 0;
+	UMaterialInterface* MaterialInterface = GetMaterial(MaterialIndex);
+	if (MaterialInterface)
+	{
+		const UMaterial* Material = MaterialInterface->GetMaterial();
+		if (Material && Material->MaterialDomain == EMaterialDomain::MD_Volume)
+		{
+			MaterialInterface->CheckMaterialUsage(MATUSAGE_HeterogeneousVolumes);
+			return MaterialInterface;
+		}
+	}
+	return nullptr;
+}
+
 #if WITH_EDITOR
 void UHeterogeneousVolumeComponent::PostEditChangeProperty(FPropertyChangedEvent& PropertyChangedEvent)
 {
@@ -448,28 +458,23 @@ void UHeterogeneousVolumeComponent::PostEditChangeProperty(FPropertyChangedEvent
 		PropertyName = PropertyChangedEvent.Property->GetFName();
 	}
 
-	if (PropertyName == GET_MEMBER_NAME_CHECKED(UHeterogeneousVolumeComponent, OverrideMaterials))
+	// When this component is copied/duplicated in the editor, PostEditChangeProperty() is called with a null PropertyChangedEvent, so we also 
+	// create the MID in that case. Otherwise the component will be copied but not play back because MaterialInstanceDynamic stays nullptr.
+	if (PropertyName == GET_MEMBER_NAME_CHECKED(UHeterogeneousVolumeComponent, OverrideMaterials) || PropertyChangedEvent.Property == nullptr)
 	{
 		MaterialInstanceDynamic = nullptr; // Reset internal MID. We either create a new one from the new material or leave it as null if the material was unset
-		uint32 MaterialIndex = 0;
-		UMaterialInterface* MaterialInterface = GetMaterial(MaterialIndex);
+		UMaterialInterface* MaterialInterface = GetHeterogeneousVolumeMaterial();
 		if (MaterialInterface)
 		{
-			const UMaterial* Material = MaterialInterface->GetMaterial();
-			if (Material && Material->MaterialDomain == EMaterialDomain::MD_Volume)
-			{
-				MaterialInterface->CheckMaterialUsage(MATUSAGE_HeterogeneousVolumes);
-			}
-
 			MaterialInstanceDynamic = CreateOrCastToMID(MaterialInterface);
 		}
-		OnSparseVolumeTextureChanged(GetSparseVolumeTexture(MaterialInstanceDynamic, SVTParameterIndex));
+		OnSparseVolumeTextureChanged(GetSparseVolumeTexture(MaterialInterface, SVTParameterIndex));
 	}
 
 	if (PropertyName == GET_MEMBER_NAME_CHECKED(UHeterogeneousVolumeComponent, VolumeResolution))
 	{
 		// Prevent resolution changes when using SVT
-		const USparseVolumeTexture* SparseVolumeTexture = GetSparseVolumeTexture(MaterialInstanceDynamic, SVTParameterIndex);
+		const USparseVolumeTexture* SparseVolumeTexture = GetSparseVolumeTexture(GetHeterogeneousVolumeMaterial(), SVTParameterIndex);
 		if (SparseVolumeTexture)
 		{
 			VolumeResolution = SparseVolumeTexture->GetVolumeResolution();
@@ -478,7 +483,7 @@ void UHeterogeneousVolumeComponent::PostEditChangeProperty(FPropertyChangedEvent
 
 	if (PropertyName == GET_MEMBER_NAME_CHECKED(UHeterogeneousVolumeComponent, Frame))
 	{
-		const USparseVolumeTexture* SparseVolumeTexture = GetSparseVolumeTexture(MaterialInstanceDynamic, SVTParameterIndex);
+		const USparseVolumeTexture* SparseVolumeTexture = GetSparseVolumeTexture(GetHeterogeneousVolumeMaterial(), SVTParameterIndex);
 		if (SparseVolumeTexture)
 		{
 			Frame = FMath::Clamp(Frame, StartFrame, EndFrame);
@@ -487,7 +492,7 @@ void UHeterogeneousVolumeComponent::PostEditChangeProperty(FPropertyChangedEvent
 
 	if (PropertyName == GET_MEMBER_NAME_CHECKED(UHeterogeneousVolumeComponent, StartFrame))
 	{
-		const USparseVolumeTexture* SparseVolumeTexture = GetSparseVolumeTexture(MaterialInstanceDynamic, SVTParameterIndex);
+		const USparseVolumeTexture* SparseVolumeTexture = GetSparseVolumeTexture(GetHeterogeneousVolumeMaterial(), SVTParameterIndex);
 		if (SparseVolumeTexture)
 		{
 			StartFrame = FMath::Clamp(StartFrame, 0, EndFrame);
@@ -497,7 +502,7 @@ void UHeterogeneousVolumeComponent::PostEditChangeProperty(FPropertyChangedEvent
 
 	if (PropertyName == GET_MEMBER_NAME_CHECKED(UHeterogeneousVolumeComponent, EndFrame))
 	{
-		const USparseVolumeTexture* SparseVolumeTexture = GetSparseVolumeTexture(MaterialInstanceDynamic, SVTParameterIndex);
+		const USparseVolumeTexture* SparseVolumeTexture = GetSparseVolumeTexture(GetHeterogeneousVolumeMaterial(), SVTParameterIndex);
 		if (SparseVolumeTexture)
 		{
 			const int32 FrameCount = SparseVolumeTexture->GetNumFrames();
@@ -523,8 +528,13 @@ void UHeterogeneousVolumeComponent::SetMaterial(int32 ElementIndex, UMaterialInt
 	Super::SetMaterial(ElementIndex, Material);
 	if (Material && ElementIndex == 0)
 	{
-		MaterialInstanceDynamic = CreateOrCastToMID(Material);
-		OnSparseVolumeTextureChanged(GetSparseVolumeTexture(MaterialInstanceDynamic, 0 /*SVTParameterIndex*/));
+		MaterialInstanceDynamic = nullptr; // Reset internal MID. We either create a new one from the new material or leave it as null if the material was unset
+		UMaterialInterface* MaterialInterface = GetHeterogeneousVolumeMaterial();
+		if (MaterialInterface)
+		{
+			MaterialInstanceDynamic = CreateOrCastToMID(MaterialInterface);
+		}
+		OnSparseVolumeTextureChanged(GetSparseVolumeTexture(MaterialInterface, 0 /*SVTParameterIndex*/));
 	}
 }
 
@@ -536,7 +546,7 @@ void UHeterogeneousVolumeComponent::TickComponent(float DeltaTime, ELevelTick Ti
 	{
 		const int32 SVTParameterIndex = 0;
 		FName SVTParameterName;
-		USparseVolumeTexture* SparseVolumeTexture = GetSparseVolumeTexture(MaterialInstanceDynamic, SVTParameterIndex, &SVTParameterName);
+		USparseVolumeTexture* SparseVolumeTexture = GetSparseVolumeTexture(GetHeterogeneousVolumeMaterial(), SVTParameterIndex, &SVTParameterName);
 
 #if WITH_EDITOR
 		// Detect an update to the material
