Iris
* Make sure PushModel support is not tied to ShouldUseIrisReplication config.
* Instead the gamenetdriver will directly disable or enable pushmodel handle creation if it initalized Iris or not.
* This enables push model to work in GameModes that don't use Iris when the default engine setting from ShouldUseIrisReplication is true
* Remove crashcontext ReplicationModel attribute when GameNetDriver is shutdown on clients or listen servers.

#rnx

diff --git a/Engine/Source/Runtime/Engine/Classes/Engine/NetDriver.h b/Engine/Source/Runtime/Engine/Classes/Engine/NetDriver.h
--- a/Engine/Source/Runtime/Engine/Classes/Engine/NetDriver.h
+++ b/Engine/Source/Runtime/Engine/Classes/Engine/NetDriver.h
@@ -2171,7 +2171,13 @@ private:
 	/** Used with FNetDelegates::OnSyncLoadDetected to log sync loads */
 	void ReportSyncLoad(const FNetSyncLoadReport& Report);
 
-	void UpdateCrashContext();
+	enum class ECrashContextUpdate
+	{
+		Default,
+		UpdateRepModel,
+		ClearRepModel,
+	};
+	void UpdateCrashContext(ECrashContextUpdate UpdateType=ECrashContextUpdate::Default);
 
 	void RemoveDestroyedGuidsByLevel(const ULevel* Level, const TArray<FNetworkGUID>& RemovedGUIDs);
 
diff --git a/Engine/Source/Runtime/Engine/Private/NetDriver.cpp b/Engine/Source/Runtime/Engine/Private/NetDriver.cpp
--- a/Engine/Source/Runtime/Engine/Private/NetDriver.cpp
+++ b/Engine/Source/Runtime/Engine/Private/NetDriver.cpp
@@ -1520,8 +1520,14 @@ bool UNetDriver::InitBase(bool bInitAsClient, FNetworkNotify* InNotify, const FU
 	
 	if (NetDriverDefinition == NAME_GameNetDriver)
 	{
-		static FString CrashContext_ReplicationModel = TEXT("ReplicationModel");
-		FGenericCrashContext::SetEngineData(CrashContext_ReplicationModel, *GetReplicationModelName());
+		UpdateCrashContext(ECrashContextUpdate::UpdateRepModel);
+
+#if WITH_PUSH_MODEL
+		// Disable PushModel globally when the GameNetDriver is running with Iris.
+		// Temp until we fully integrate Iris PushModel support.
+		const bool bAllowPushModelHandles = !IsUsingIrisReplication();
+		UEPushModelPrivate::SetHandleCreationAllowed(bAllowPushModelHandles);
+#endif
 	}
 
 	UE_LOG(LogNet, Log, TEXT("InitBase %s (NetDriverDefinition %s) using replication model %s"), *NetDriverName.ToString(), *NetDriverDefinition.ToString(), *GetReplicationModelName());
@@ -2086,7 +2092,16 @@ void UNetDriver::Shutdown()
 		AnalyticsAggregator.Reset();
 	}
 
-	UpdateCrashContext();
+	// Clear repmodel flags for clients or listen servers when shutting down since they are considered offline.
+	// For dedicated servers let's keep the flag in case we crash post-shutdown.
+	if (GetNetMode() != NM_DedicatedServer)
+	{
+		UpdateCrashContext(ECrashContextUpdate::ClearRepModel);
+	}
+	else
+	{
+		UpdateCrashContext(ECrashContextUpdate::Default);
+	}
 
 	NetworkMetricsDatabase->Reset();
 	NetworkMetricsListeners.Reset();
@@ -6227,22 +6242,35 @@ void UNetDriver::AddClientConnection(UNetConnection* NewConnection)
 	UpdateCrashContext();
 }
 
-void UNetDriver::UpdateCrashContext()
+void UNetDriver::UpdateCrashContext(ECrashContextUpdate UpdateType)
 {
-	if (NetDriverName == NAME_GameNetDriver)
+	if (NetDriverDefinition != NAME_GameNetDriver)
 	{
-		int32 NumClients = 0;
+		return;
+	}
+
+	int32 NumClients = 0;
 
-		for (const UNetConnection* Connection : ClientConnections)
+	for (const UNetConnection* Connection : ClientConnections)
+	{
+		if (Connection && !Connection->IsReplay())
 		{
-			if (Connection && !Connection->IsReplay())
-			{
-				++NumClients;
-			}
+			++NumClients;
 		}
+	}
 
-		static FString Attrib_NumClients = TEXT("NumClients");
-		FGenericCrashContext::SetEngineData(Attrib_NumClients, LexToString(NumClients));
+	static FString Attrib_NumClients = TEXT("NumClients");
+	FGenericCrashContext::SetEngineData(Attrib_NumClients, LexToString(NumClients));
+
+	static FString CrashContext_ReplicationModel = TEXT("ReplicationModel");
+
+	if (UpdateType == ECrashContextUpdate::UpdateRepModel)
+	{		
+		FGenericCrashContext::SetEngineData(CrashContext_ReplicationModel, *GetReplicationModelName());
+	}
+	else if (UpdateType == ECrashContextUpdate::ClearRepModel)
+	{
+		FGenericCrashContext::SetEngineData(CrashContext_ReplicationModel, TEXT(""));
 	}
 }
 
diff --git a/Engine/Source/Runtime/Engine/Private/RepLayout.cpp b/Engine/Source/Runtime/Engine/Private/RepLayout.cpp
--- a/Engine/Source/Runtime/Engine/Private/RepLayout.cpp
+++ b/Engine/Source/Runtime/Engine/Private/RepLayout.cpp
@@ -12,6 +12,7 @@
 #include "EngineStats.h"
 #include "Engine/PackageMapClient.h"
 #include "Engine/NetConnection.h"
+#include "Net/Core/PushModel/PushModel.h"
 #include "Net/Core/PushModel/PushModelMacros.h"
 #include "Net/NetworkProfiler.h"
 #include "Engine/ActorChannel.h"
@@ -219,14 +220,10 @@ namespace UE_RepLayout_Private
 	{
 		const int32 NumReplicatedProperties = InRepLayout->GetNumParents();
 
-#if UE_WITH_IRIS
 		// Implement shared PushModelids/NetHandles for Iris and PushModel to avoid conflicts when we mix systems - JIRA: UE-158304
-		const bool bShouldUseIrisReplication = UE::Net::ShouldUseIrisReplication();
-#else
-		const bool bShouldUseIrisReplication = false;
-#endif
+		const bool bAllowedToCreateHandles = UEPushModelPrivate::IsHandleCreationAllowed();
 
-		if (UEPushModelPrivate::IsPushModelEnabled() && !bShouldUseIrisReplication &&
+		if (UEPushModelPrivate::IsPushModelEnabled() && bAllowedToCreateHandles &&
 			NumReplicatedProperties > 0 &&
 			EnumHasAnyFlags(InRepLayout->GetFlags(), ERepLayoutFlags::FullPushSupport | ERepLayoutFlags::PartialPushSupport))
 		{
diff --git a/Engine/Source/Runtime/Experimental/Iris/Core/Private/Iris/ReplicationSystem/LegacyPushModel.cpp b/Engine/Source/Runtime/Experimental/Iris/Core/Private/Iris/ReplicationSystem/LegacyPushModel.cpp
--- a/Engine/Source/Runtime/Experimental/Iris/Core/Private/Iris/ReplicationSystem/LegacyPushModel.cpp
+++ b/Engine/Source/Runtime/Experimental/Iris/Core/Private/Iris/ReplicationSystem/LegacyPushModel.cpp
@@ -28,7 +28,7 @@ static FAutoConsoleVariableRef CVarIrisPushModelMode(
 void FNetHandleLegacyPushModelHelper::InitPushModel()
 {
 	bIsIrisPushModelForceEnabled = IrisPushModelMode == 1;
-	if (bIsIrisPushModelForceEnabled && ensureAlwaysMsgf(IS_PUSH_MODEL_ENABLED(), TEXT("Trying to force enable Iris push model support when push model is disabled, falling back to optional path. Set Net.IsPushModelEnabled true")))
+	if (bIsIrisPushModelForceEnabled && ensureMsgf(IS_PUSH_MODEL_ENABLED(), TEXT("Trying to force enable Iris push model support when push model is disabled, falling back to optional path. Set Net.IsPushModelEnabled true")))
 	{
 		// Push model is assumed to be enabled at all times. The cvar will not be checked.
 		UE_NET_SET_IRIS_MARK_PROPERTY_DIRTY_DELEGATE(UEPushModelPrivate::FIrisMarkPropertyDirty::CreateStatic(&FNetHandleLegacyPushModelHelper::MarkPropertyOwnerDirty));
diff --git a/Engine/Source/Runtime/Net/Core/Private/Net/Core/PushModel/PushModel.cpp b/Engine/Source/Runtime/Net/Core/Private/Net/Core/PushModel/PushModel.cpp
--- a/Engine/Source/Runtime/Net/Core/Private/Net/Core/PushModel/PushModel.cpp
+++ b/Engine/Source/Runtime/Net/Core/Private/Net/Core/PushModel/PushModel.cpp
@@ -387,6 +387,25 @@ namespace UEPushModelPrivate
 			return sizeof(*this) + Ar.GetMem();
 		}
 
+		bool HasHandlesAssigned() const
+		{
+			return !ObjectKeyToInternalId.IsEmpty();
+		}
+
+		void PrintAssignedHandles() const
+		{
+			UE_LOG(LogPushModel, Display, TEXT("Starting print of %u objects currently assigned a PushModel handle"), ObjectKeyToInternalId.Num());
+			for (auto It=ObjectKeyToInternalId.CreateConstIterator(); It; ++It)
+			{
+				const FObjectKey& Key = It.Key();
+				const FNetLegacyPushObjectId& Value = It.Value();
+
+				UObject* ObjPtr = Key.ResolveObjectPtrEvenIfUnreachable();
+				
+				UE_LOG(LogPushModel, Display, TEXT("Object: %s (%s) | PushObjectId: 0x%x"), *GetNameSafe(ObjPtr), IsValid(ObjPtr)?TEXT("Valid"):TEXT("Invalid"), Value);
+			}
+		}
+
 	private:
 		int32 NewObjectLookupPosition = 0;
 		TMap<FObjectKey, FNetLegacyPushObjectId> ObjectKeyToInternalId;
@@ -426,6 +445,12 @@ namespace UEPushModelPrivate
 		TEXT("Whether or not properties declared in Blueprints will be forced to used Push Model")
 	);
 
+	FAutoConsoleCommand PushModelPrintHandles(
+		TEXT("Net.PushModelPrintHandles"),
+		TEXT("Prints the list of replicated objects relevant to a specific connection"),
+		FConsoleCommandDelegate::CreateLambda([]() { PushObjectManager.PrintAssignedHandles(); })
+	);
+
 #if UE_WITH_IRIS
 	static FIrisMarkPropertyDirty IrisMarkPropertyDirtyDelegate;
 	static FIrisMarkPropertiesDirty IrisMarkPropertiesDirtyDelegate;
@@ -531,6 +556,24 @@ namespace UEPushModelPrivate
 
 		Ar.Logf(TEXT("  Push Model Memory: %u"), Count);
 	}
-}
+
+	bool bAllowHandleCreation = true;
+	bool IsHandleCreationAllowed()
+	{
+		return bAllowHandleCreation;
+	}
+
+	void SetHandleCreationAllowed(bool bAllow)
+	{
+		UE_LOG(LogPushModel, Log, TEXT("PushModel HandleCreation is now %s"), bAllow ? TEXT("enabled") : TEXT("disabled"));
+		bAllowHandleCreation = bAllow;
+
+		if (!bAllow && PushObjectManager.HasHandlesAssigned())
+		{
+			PushObjectManager.PrintAssignedHandles();
+			ensureMsgf(!PushObjectManager.HasHandlesAssigned(), TEXT("PushModel Handle creation is disabled but there are already handles allocated. See logs for details on the assigned handles"));
+		}
+	}
+} // end namespace UEPushModelPrivate
 
 #endif // WITH_PUSH_MODEL
\ No newline at end of file
diff --git a/Engine/Source/Runtime/Net/Core/Public/Net/Core/PushModel/PushModel.h b/Engine/Source/Runtime/Net/Core/Public/Net/Core/PushModel/PushModel.h
--- a/Engine/Source/Runtime/Net/Core/Public/Net/Core/PushModel/PushModel.h
+++ b/Engine/Source/Runtime/Net/Core/Public/Net/Core/PushModel/PushModel.h
@@ -348,6 +348,12 @@ namespace UEPushModelPrivate
 		return bMakeBpPropertiesPushModel;
 	}
 
+	/** Are we allowed to create new push model handles */
+	NETCORE_API bool IsHandleCreationAllowed();
+
+	/** Control if we are allowed to create pushmodel handles */
+	NETCORE_API void SetHandleCreationAllowed(bool bAllow);
+	
 	static FString ToString(const FNetPushObjectId Id)
 	{
 		return FString::Printf(TEXT("0x%" UINT64_X_FMT), Id.GetValue());
