Mitigation against editor viewport setting ortho zoom to float max value, which causes an ensure from rendering code due to precision loss.

diff --git a/Engine/Source/Editor/LevelEditor/Private/SLevelViewport.cpp b/Engine/Source/Editor/LevelEditor/Private/SLevelViewport.cpp
--- a/Engine/Source/Editor/LevelEditor/Private/SLevelViewport.cpp
+++ b/Engine/Source/Editor/LevelEditor/Private/SLevelViewport.cpp
@@ -152,16 +152,17 @@ namespace UE::SLevelViewport::Internal
 			World->EditorViews = PerUserEditorViews->LevelViewportsInfo;
 		}
 
-		if (World->EditorViews[LevelEditorViewportClient->ViewportType].CamOrthoZoom == 0.0f)
-		{
-			World->EditorViews[LevelEditorViewportClient->ViewportType].CamOrthoZoom = DEFAULT_ORTHOZOOM;
-		}
-
 		LevelEditorViewportClient->ResetCamera();
 
 		bool bInitializedOrthoViewport = false;
 		for (int32 ViewportType = 0; ViewportType < LVT_MAX; ViewportType++)
 		{
+			float& CamOrthoZoom = World->EditorViews[ViewportType].CamOrthoZoom;
+			if (CamOrthoZoom < MIN_ORTHOZOOM || CamOrthoZoom > MAX_ORTHOZOOM)
+			{
+				CamOrthoZoom = DEFAULT_ORTHOZOOM;
+			}
+
 			if (ViewportType == LVT_Perspective || !bInitializedOrthoViewport)
 			{
 				LevelEditorViewportClient->SetInitialViewTransform(
diff --git a/Engine/Source/Editor/UnrealEd/Private/EditorViewportClient.cpp b/Engine/Source/Editor/UnrealEd/Private/EditorViewportClient.cpp
--- a/Engine/Source/Editor/UnrealEd/Private/EditorViewportClient.cpp
+++ b/Engine/Source/Editor/UnrealEd/Private/EditorViewportClient.cpp
@@ -147,9 +147,6 @@ namespace EditorViewportClient
 	static const float LightRotSpeed = 0.22f;
 }
 
-// MIN_ORTHOZOOM defined in ULevelEditorViewportSettings
-#define MAX_ORTHOZOOM				MAX_FLT					/* Limit of 2D viewport zoom out */
-
 namespace OrbitConstants
 {
 	const float OrbitPanSpeed = 1.0f;
@@ -6511,7 +6508,7 @@ void FEditorViewportClient::DisableOverrideEngineShowFlags()
 
 float FEditorViewportClient::GetMinimumOrthoZoom() const
 {
-	return FMath::Max(GetDefault<ULevelEditorViewportSettings>()->MinimumOrthographicZoom, 1.0f);
+	return FMath::Max(GetDefault<ULevelEditorViewportSettings>()->MinimumOrthographicZoom, MIN_ORTHOZOOM);
 }
 
 ////////////////
diff --git a/Engine/Source/Editor/UnrealEd/Public/EditorViewportClient.h b/Engine/Source/Editor/UnrealEd/Public/EditorViewportClient.h
--- a/Engine/Source/Editor/UnrealEd/Public/EditorViewportClient.h
+++ b/Engine/Source/Editor/UnrealEd/Public/EditorViewportClient.h
@@ -239,6 +239,7 @@ public:
 	/** Set the ortho zoom amount */
 	void SetOrthoZoom( float InOrthoZoom )
 	{
+		ensure(InOrthoZoom >= MIN_ORTHOZOOM && InOrthoZoom <= MAX_ORTHOZOOM);
 		OrthoZoom = InOrthoZoom;
 	}
 
diff --git a/Engine/Source/Runtime/Core/Private/Math/DoubleFloat.cpp b/Engine/Source/Runtime/Core/Private/Math/DoubleFloat.cpp
--- a/Engine/Source/Runtime/Core/Private/Math/DoubleFloat.cpp
+++ b/Engine/Source/Runtime/Core/Private/Math/DoubleFloat.cpp
@@ -15,10 +15,8 @@ FMatrix CheckMatrixPrecision(const FMatrix& Matrix)
 	const double OriginX = FMath::Abs(Origin.X);
 	const double OriginY = FMath::Abs(Origin.Y);
 	const double OriginZ = FMath::Abs(Origin.Z);
-	if (OriginX > OriginMax || OriginY > OriginMax || OriginZ > OriginMax)
-	{
-		ensure(false);
-	}
+	ensureMsgf(OriginX <= OriginMax && OriginY <= OriginMax && OriginZ <= OriginMax, 
+		TEXT("Found precision loss while converting matrix to GPU format, verify the input transforms."));
 	return Matrix;
 }
 
diff --git a/Engine/Source/Runtime/Engine/Classes/Engine/World.h b/Engine/Source/Runtime/Engine/Classes/Engine/World.h
--- a/Engine/Source/Runtime/Engine/Classes/Engine/World.h
+++ b/Engine/Source/Runtime/Engine/Classes/Engine/World.h
@@ -338,7 +338,7 @@ struct ENGINE_API FLevelViewportInfo
 		{
 			I.CamUpdated = true;
 
-			if ( I.CamOrthoZoom == 0.f )
+			if ( I.CamOrthoZoom < MIN_ORTHOZOOM || I.CamOrthoZoom > MAX_ORTHOZOOM )
 			{
 				I.CamOrthoZoom = DEFAULT_ORTHOZOOM;
 			}
diff --git a/Engine/Source/Runtime/Engine/Private/World.cpp b/Engine/Source/Runtime/Engine/Private/World.cpp
--- a/Engine/Source/Runtime/Engine/Private/World.cpp
+++ b/Engine/Source/Runtime/Engine/Private/World.cpp
@@ -678,7 +678,7 @@ void UWorld::Serialize( FArchive& Ar )
 		{
 			ViewportInfo.CamUpdated = true;
 
-			if ( ViewportInfo.CamOrthoZoom == 0.f )
+			if ( ViewportInfo.CamOrthoZoom < MIN_ORTHOZOOM || ViewportInfo.CamOrthoZoom > MAX_ORTHOZOOM )
 			{
 				ViewportInfo.CamOrthoZoom = DEFAULT_ORTHOZOOM;
 			}
diff --git a/Engine/Source/Runtime/Engine/Public/EngineDefines.h b/Engine/Source/Runtime/Engine/Public/EngineDefines.h
--- a/Engine/Source/Runtime/Engine/Public/EngineDefines.h
+++ b/Engine/Source/Runtime/Engine/Public/EngineDefines.h
@@ -53,7 +53,9 @@
 #define UE_FLOAT_HUGE_DISTANCE		1048576.0				/* Maximum distance representable by a float whilst maintaining precision of at least 0.0625 units (1/16th of a cm) - Precision issues may occur for positions/distances represented by float types that exceed this value */
 #define UE_DOUBLE_HUGE_DISTANCE		562949953421312.0		/* Maximum distance representable by a double whilst maintaining precision of at least 0.0625 units (1/16th of a cm) - Precision issues may occur for positions/distances represented by double types that exceed this value */
 
-#define DEFAULT_ORTHOZOOM			10000.0						/* Default 2D viewport zoom */
+#define MIN_ORTHOZOOM				1.0							/* 2D ortho viewport zoom >= MIN_ORTHOZOOM */
+#define MAX_ORTHOZOOM				1e25						/* 2D ortho viewport zoom <= MAX_ORTHOZOOM */
+#define DEFAULT_ORTHOZOOM			10000.0						/* Default 2D ortho viewport zoom */
 #define DEFAULT_ORTHOWIDTH			1536.0f						/* Default 2D ortho viewport width */
 #define DEFAULT_ORTHONEARPLANE		-DEFAULT_ORTHOWIDTH/2.0f	/* Default 2D ortho viewport nearplane */
 #define DEFAULT_ORTHOFARPLANE		UE_OLD_WORLD_MAX			/* Default 2D ortho viewport farplane */
