remove references to the unused 'AudioDeviceModuleName' configuration parameter

 - originally removed in CL 25172786 for 5.3

#rnx

diff --git a/Engine/Source/Runtime/Engine/Private/AudioDeviceManager.cpp b/Engine/Source/Runtime/Engine/Private/AudioDeviceManager.cpp
--- a/Engine/Source/Runtime/Engine/Private/AudioDeviceManager.cpp
+++ b/Engine/Source/Runtime/Engine/Private/AudioDeviceManager.cpp
@@ -420,7 +420,7 @@ bool FAudioDeviceManager::CreateMainAudioDevice()
 
 		if (!MainAudioDeviceHandle)
 		{
-			UE_LOG(LogAudio, Display, TEXT("Main audio device could not be initialized. Please check the value for AudioDeviceModuleName and AudioMixerModuleName in [Platform]Engine.ini."));
+			UE_LOG(LogAudio, Display, TEXT("Main audio device could not be initialized. Please check the value for AudioMixerModuleName in [Platform]Engine.ini."));
 			return false;
 		}
 
@@ -469,7 +469,7 @@ FAudioDeviceHandle FAudioDeviceManager::CreateNewDevice(const FAudioDeviceParams
 	check(ContainerPtr);
 	if (!ContainerPtr->Device)
 	{
-		UE_LOG(LogAudio, Display, TEXT("Destroying Audio Device %d: could not be initialized. Check AudioDeviceModuleName and AudioMixerModuleName in [Platform]Engine.ini."), DeviceID);
+		UE_LOG(LogAudio, Display, TEXT("Destroying Audio Device %d: could not be initialized. Check AudioMixerModuleName in [Platform]Engine.ini."), DeviceID);
 
 		// Initializing the audio device failed. Remove the device container and return an empty handle.
 		Devices.Remove(DeviceID);
