[Material Editor] Updated UMaterialFunction::OnPostLoad() to mark the function as invalid when some of its expressions fail to load (e.g. due to unavailable expression types) to force the translation of materials containing active calls to it to fail.

Previously, materials with calls to functions missing expressions would still compile (potentially incorrectly); now, it fails translation and reports the error.

diff --git a/Engine/Source/Runtime/Engine/Classes/Materials/MaterialFunction.h b/Engine/Source/Runtime/Engine/Classes/Materials/MaterialFunction.h
--- a/Engine/Source/Runtime/Engine/Classes/Materials/MaterialFunction.h
+++ b/Engine/Source/Runtime/Engine/Classes/Materials/MaterialFunction.h
@@ -100,6 +100,10 @@ class UMaterialFunction : public UMaterialFunctionInterface
 	TEnumAsByte<enum EBlendMode> PreviewBlendMode = BLEND_Opaque;
 
 	class UMaterialGraph* MaterialGraph = nullptr;
+
+	/* Whether all expressions in the function loaded correctly. */
+	uint8 bAllExpressionsLoadedCorrectly : 1;
+
 private:
 	/** Transient flag used to track re-entrance in recursive functions like IsDependent. */
 	UPROPERTY(transient)
diff --git a/Engine/Source/Runtime/Engine/Private/Materials/MaterialExpressions.cpp b/Engine/Source/Runtime/Engine/Private/Materials/MaterialExpressions.cpp
--- a/Engine/Source/Runtime/Engine/Private/Materials/MaterialExpressions.cpp
+++ b/Engine/Source/Runtime/Engine/Private/Materials/MaterialExpressions.cpp
@@ -14899,6 +14899,7 @@ UMaterialFunction::UMaterialFunction(const FObjectInitializer& ObjectInitializer
 #if WITH_EDITORONLY_DATA
 	PreviewMaterial = nullptr;
 	ThumbnailInfo = nullptr;
+	bAllExpressionsLoadedCorrectly = true;
 #endif
 }
 
@@ -15120,28 +15121,50 @@ void UMaterialFunction::PostLoad()
 	}
 	UpdateDependentFunctionCandidates();
 
+	bAllExpressionsLoadedCorrectly = true;
+
 	if (GIsEditor && EditorOnly)
 	{
-		if (EditorOnly->ExpressionCollection.DebugExpressionInfos.Num() == EditorOnly->ExpressionCollection.Expressions.Num())
+		// We can display null expressions info if the DebugExpressionInfos was populated with data upon Serialize().
+		bool bDisplayNullExpressionInfo = EditorOnly->ExpressionCollection.DebugExpressionInfos.Num() == EditorOnly->ExpressionCollection.Expressions.Num();
+
+		// Go over all expressions in the collection and invalidate the material if a null expression is found. Then
+		// remove the null expression from the array.
+		for (int i = 0; i < EditorOnly->ExpressionCollection.Expressions.Num();)
 		{
-			for (int i = 0; i < EditorOnly->ExpressionCollection.Expressions.Num(); ++i)
+			UMaterialExpression* Expression = EditorOnly->ExpressionCollection.Expressions[i].Get();
+			if (Expression)
 			{
-				UMaterialExpression* Expression = EditorOnly->ExpressionCollection.Expressions[i].Get();
-				if (!Expression)
-				{
-					UE_LOG(LogMaterial, Log, TEXT("Expression in function expression collection with index %d was null. Expression Info: %s"), i, *EditorOnly->ExpressionCollection.DebugExpressionInfos[i]);
-					EditorOnly->ExpressionCollection.DebugExpressionInfos.RemoveAt(i);
-				}
+				++i;
+				continue;
 			}
+
+			// Mark this function as invalid. This will cause the material containing an active call to it to fail translation.
+			bAllExpressionsLoadedCorrectly = false;
+
+			if (bDisplayNullExpressionInfo)
+			{
+				UE_LOG(LogMaterial, Log, TEXT("Expression in function expression collection with index %d was null. Expression Info: %s"), i, *EditorOnly->ExpressionCollection.DebugExpressionInfos[i]);
+					
+				EditorOnly->ExpressionCollection.DebugExpressionInfos.RemoveAt(i);
+			}
+
+			EditorOnly->ExpressionCollection.Expressions.RemoveAt(i);
 		}
 
-		// Clean up any removed material expression classes	
-		if (EditorOnly->ExpressionCollection.Expressions.Remove(nullptr) != 0)
+		if (!bAllExpressionsLoadedCorrectly)
 		{
-			// Force this function to recompile because its expressions have changed
-			// Which means removing an expression class will cause the need for a resave of all materials affected
-			UE_LOG(LogMaterial, Log, TEXT("Please resave %s.  It is missing a material expression and this will cause any material using this function to recompile shaders each time it loads."), *GetFullName());
-			StateId = FGuid::NewGuid();
+			UE_LOG(LogMaterial, Warning, TEXT("Some expression in Material Function %s failed to load correctly. This will cause any material using this MF to fail translation. Please check open affected Material Function, make sure its expression graph is valid and resave it."), *GetFullName());
+			
+			// Mark this function as invalid. Translating a material containing a call to it will fail.
+			bAllExpressionsLoadedCorrectly = false;
+
+			// Dirty this function by deterministically changing its StateId.
+			static FGuid NotAllExpressionsLoadedCorrectlyToken(TEXT("6B9D300E-ED9D-4E4A-A141-05DE059B5704"));
+			StateId.A ^= NotAllExpressionsLoadedCorrectlyToken.A;
+			StateId.B ^= NotAllExpressionsLoadedCorrectlyToken.B;
+			StateId.C ^= NotAllExpressionsLoadedCorrectlyToken.C;
+			StateId.D ^= NotAllExpressionsLoadedCorrectlyToken.D;
 		}
 	}
 
@@ -16904,6 +16927,11 @@ int32 UMaterialExpressionMaterialFunctionCall::Compile(class FMaterialCompiler*
 		return Compiler->Errorf(TEXT("Material Functions with control flow are only supported with new HLSL translator"));
 	}
 
+	if (!MaterialFunction->GetBaseFunction()->bAllExpressionsLoadedCorrectly)
+	{
+		return Compiler->Errorf(TEXT("Called function is in an invalid state because some expressions didn't load correctly. Please open the affected function, review and correct the expression graph then save again."));
+	}
+
 	// Verify that all function inputs and outputs are in a valid state to be linked into this material for compiling
 	for (int32 i = 0; i < FunctionInputs.Num(); i++)
 	{
