Adds a project setting, r.GPUSkin.AlwaysUseDeformerForUnlimitedBoneInfluences, that allows you to enable Unlimited Bone Influences in a project without compiling extra shader permutations for GPU skinning. This saves runtime memory, disk space and shader compilation time.

When the setting is enabled, any mesh LODs using Unlimited Bone Influences that don't have a deformer assigned will use the DeformerGraph plugin's default deformer. This ensures that UBI meshes are always rendered with a deformer, and therefore the GPU skinning permutations for UBI aren't needed.

This change also adds a per-LOD setting that allows users to disable mesh deformers on a LOD, which could be useful for controlling performance, e.g. disabling an expensive deformer on lower LODs. Some changes to functions on USkinnedMeshComponent lay the foundations for having different deformers on different LODs as well.

diff --git a/Engine/Plugins/Animation/DeformerGraph/Source/OptimusSettings/Private/OptimusSettingsModule.cpp b/Engine/Plugins/Animation/DeformerGraph/Source/OptimusSettings/Private/OptimusSettingsModule.cpp
--- a/Engine/Plugins/Animation/DeformerGraph/Source/OptimusSettings/Private/OptimusSettingsModule.cpp
+++ b/Engine/Plugins/Animation/DeformerGraph/Source/OptimusSettings/Private/OptimusSettingsModule.cpp
@@ -84,8 +84,8 @@ void FOptimusSettingsModule::CacheDefaultMeshDeformers()
 void FOptimusSettingsModule::CacheDefaultMeshDeformers(UOptimusSettings const* InSettings)
 {
 	const bool bLoadDeformers = InSettings != nullptr && InSettings->DefaultMode != EOptimusDefaultDeformerMode::Never;
-	DefaultDeformer = bLoadDeformers && !InSettings->DefaultDeformer.IsNull() ? InSettings->DefaultDeformer.LoadSynchronous() : nullptr;
-	DefaultRecomputeTangentDeformer = bLoadDeformers && !InSettings->DefaultRecomputeTangentDeformer.IsNull() ? InSettings->DefaultRecomputeTangentDeformer.LoadSynchronous() : nullptr;
+	DefaultDeformer = TStrongObjectPtr<UMeshDeformer>(bLoadDeformers && !InSettings->DefaultDeformer.IsNull() ? InSettings->DefaultDeformer.LoadSynchronous() : nullptr);
+	DefaultRecomputeTangentDeformer = TStrongObjectPtr<UMeshDeformer>(bLoadDeformers && !InSettings->DefaultRecomputeTangentDeformer.IsNull() ? InSettings->DefaultRecomputeTangentDeformer.LoadSynchronous() : nullptr);
 }
 
 TObjectPtr<UMeshDeformer> FOptimusSettingsModule::GetDefaultMeshDeformer(FDefaultMeshDeformerSetup const& Setup)
@@ -93,15 +93,15 @@ TObjectPtr<UMeshDeformer> FOptimusSettingsModule::GetDefaultMeshDeformer(FDefaul
 	const UOptimusSettings* Settings = GetDefault<UOptimusSettings>();
 	if (Settings != nullptr && 
 		(Settings->DefaultMode == EOptimusDefaultDeformerMode::Always ||
-		(Settings->DefaultMode == EOptimusDefaultDeformerMode::SkinCacheOnly && Setup.bIsUsingSkinCache)))
+		(Settings->DefaultMode == EOptimusDefaultDeformerMode::OptIn && Setup.bIsRequestingDeformer)))
 	{
 		if (Setup.bIsRequestingRecomputeTangent && DefaultRecomputeTangentDeformer)
 		{
-			return DefaultRecomputeTangentDeformer;
+			return DefaultRecomputeTangentDeformer.Get();
 		}
 		if (DefaultDeformer)
 		{ 
-			return DefaultDeformer;
+			return DefaultDeformer.Get();
 		}
 	}
 	return nullptr;
diff --git a/Engine/Plugins/Animation/DeformerGraph/Source/OptimusSettings/Private/OptimusSettingsModule.h b/Engine/Plugins/Animation/DeformerGraph/Source/OptimusSettings/Private/OptimusSettingsModule.h
--- a/Engine/Plugins/Animation/DeformerGraph/Source/OptimusSettings/Private/OptimusSettingsModule.h
+++ b/Engine/Plugins/Animation/DeformerGraph/Source/OptimusSettings/Private/OptimusSettingsModule.h
@@ -4,6 +4,7 @@
 
 #include "Animation/MeshDeformerProvider.h"
 #include "Modules/ModuleInterface.h"
+#include "UObject/StrongObjectPtr.h"
 
 class UOptimusSettings;
 
@@ -21,6 +22,6 @@ private:
 	void CacheDefaultMeshDeformers();
 	void CacheDefaultMeshDeformers(UOptimusSettings const* InSettings);
 
-	TObjectPtr<UMeshDeformer> DefaultDeformer;
-	TObjectPtr<UMeshDeformer> DefaultRecomputeTangentDeformer;
+	TStrongObjectPtr<UMeshDeformer> DefaultDeformer;
+	TStrongObjectPtr<UMeshDeformer> DefaultRecomputeTangentDeformer;
 };
diff --git a/Engine/Plugins/Animation/DeformerGraph/Source/OptimusSettings/Public/OptimusSettings.h b/Engine/Plugins/Animation/DeformerGraph/Source/OptimusSettings/Public/OptimusSettings.h
--- a/Engine/Plugins/Animation/DeformerGraph/Source/OptimusSettings/Public/OptimusSettings.h
+++ b/Engine/Plugins/Animation/DeformerGraph/Source/OptimusSettings/Public/OptimusSettings.h
@@ -17,8 +17,8 @@ enum class EOptimusDefaultDeformerMode : uint8
 {
 	/** Never apply the default deformers. */
 	Never,
-	/** Only apply default deformers as replacement for the GPU Skin Cache. */
-	SkinCacheOnly,
+	/** Only apply default deformers if requested. */
+	OptIn,
 	/** Always apply the default deformers. */
 	Always,
 };
diff --git a/Engine/Source/Editor/Persona/Private/PersonaMeshDetails.cpp b/Engine/Source/Editor/Persona/Private/PersonaMeshDetails.cpp
--- a/Engine/Source/Editor/Persona/Private/PersonaMeshDetails.cpp
+++ b/Engine/Source/Editor/Persona/Private/PersonaMeshDetails.cpp
@@ -2675,9 +2675,17 @@ void FPersonaMeshDetails::CustomizeLODInfoSetingsDetails(IDetailLayoutBuilder& D
 	PreventAttributePostEditChange(WeightToPriortizeHandle, LODIndex, WeightToPriortizeHandle->GetProperty()->GetFName(), bForceComponentRefreshFalse);
 	DetailLayout.HideProperty(WeightToPriortizeHandle);
 
+	TSharedPtr<IPropertyHandle> AllowMeshDeformerHandle = LODInfoChild->GetChildHandle(GET_MEMBER_NAME_CHECKED(FSkeletalMeshLODInfo, bAllowMeshDeformer));
+	IDetailPropertyRow& AllowMeshDeformerRow = LODInfoGroup.AddPropertyRow(AllowMeshDeformerHandle->AsShared());
+	AllowMeshDeformerRow.IsEnabled(EnabledAttrib);
+	// PreventAttributePostEditChange is not called here because we do want PostEditChange to fire
+	// if this is changed, so that any mesh components will be re-registered and refresh their
+	// deformer setups.
+	DetailLayout.HideProperty(AllowMeshDeformerHandle);
+
 	const TArray<FName> HiddenProperties = { GET_MEMBER_NAME_CHECKED(FSkeletalMeshLODInfo, ReductionSettings), GET_MEMBER_NAME_CHECKED(FSkeletalMeshLODInfo, BakePose), GET_MEMBER_NAME_CHECKED(FSkeletalMeshLODInfo, BakePoseOverride), GET_MEMBER_NAME_CHECKED(FSkeletalMeshLODInfo, BonesToRemove),
 		GET_MEMBER_NAME_CHECKED(FSkeletalMeshLODInfo, BonesToPrioritize), GET_MEMBER_NAME_CHECKED(FSkeletalMeshLODInfo, SectionsToPrioritize), GET_MEMBER_NAME_CHECKED(FSkeletalMeshLODInfo, WeightOfPrioritization), GET_MEMBER_NAME_CHECKED(FSkeletalMeshLODInfo, ScreenSize),
-		GET_MEMBER_NAME_CHECKED(FSkeletalMeshLODInfo, LODHysteresis), GET_MEMBER_NAME_CHECKED(FSkeletalMeshLODInfo, BuildSettings) };
+		GET_MEMBER_NAME_CHECKED(FSkeletalMeshLODInfo, LODHysteresis), GET_MEMBER_NAME_CHECKED(FSkeletalMeshLODInfo, BuildSettings), GET_MEMBER_NAME_CHECKED(FSkeletalMeshLODInfo, bAllowMeshDeformer) };
 	for (uint32 ChildIndex = 0; ChildIndex < NumInfoChildren; ++ChildIndex)
 	{
 		TSharedRef<IPropertyHandle> LODInfoChildHandle = LODInfoChild->GetChildHandle(ChildIndex).ToSharedRef();
diff --git a/Engine/Source/Editor/UnrealEd/Private/ImportUtils/SkeletalMeshImportUtils.cpp b/Engine/Source/Editor/UnrealEd/Private/ImportUtils/SkeletalMeshImportUtils.cpp
--- a/Engine/Source/Editor/UnrealEd/Private/ImportUtils/SkeletalMeshImportUtils.cpp
+++ b/Engine/Source/Editor/UnrealEd/Private/ImportUtils/SkeletalMeshImportUtils.cpp
@@ -482,6 +482,7 @@ void SkeletalMesUtilsImpl::RestoreLODInfo(const TSharedPtr<const FExistingSkelMe
 	ImportedLODInfo.bAllowCPUAccess = ExistingLODInfo.bAllowCPUAccess;
 	ImportedLODInfo.bBuildHalfEdgeBuffers = ExistingLODInfo.bBuildHalfEdgeBuffers;
 	ImportedLODInfo.bSupportUniformlyDistributedSampling = ExistingLODInfo.bSupportUniformlyDistributedSampling;
+	ImportedLODInfo.bAllowMeshDeformer = ExistingLODInfo.bAllowMeshDeformer;
 }
 
 void SkeletalMeshImportUtils::ApplySkinning(USkeletalMesh* SkeletalMesh, FSkeletalMeshLODModel& SrcLODModel, FSkeletalMeshLODModel& DestLODModel)
diff --git a/Engine/Source/Runtime/Engine/Classes/Components/SkinnedMeshComponent.h b/Engine/Source/Runtime/Engine/Classes/Components/SkinnedMeshComponent.h
--- a/Engine/Source/Runtime/Engine/Classes/Components/SkinnedMeshComponent.h
+++ b/Engine/Source/Runtime/Engine/Classes/Components/SkinnedMeshComponent.h
@@ -187,6 +187,31 @@ struct FVertexOffsetUsage
 	int32 Usage = 0;
 };
 
+/** Specifies which mesh deformer should be used on each LOD of a mesh */
+struct FMeshDeformerSet
+{
+	/** Each deformer in this array should be used on at least one LOD */
+	TArray<TObjectPtr<UMeshDeformer>, TInlineAllocator<2>> Deformers;
+
+	/**
+	 * Indexed by mesh LOD. Each element is either a valid index into the Deformers array, or
+	 * INDEX_NONE to signify that no deformer should be used for this LOD.
+	 */
+	TArray<int8, TFixedAllocator<MAX_MESH_LOD_COUNT>> DeformerIndexForLOD;
+};
+
+/** Same as FMeshDeformerSet, except for mesh deformer instances */
+USTRUCT()
+struct FMeshDeformerInstanceSet
+{
+	GENERATED_BODY()
+
+	UPROPERTY()
+	TArray<TObjectPtr<UMeshDeformerInstance>> DeformerInstances;
+
+	TArray<int8, TFixedAllocator<MAX_MESH_LOD_COUNT>> InstanceIndexForLOD;
+};
+
 /** The map of external morph sets registered on the skinned mesh component. */
 using FExternalMorphSets = TMap<int32, TSharedPtr<FExternalMorphSet>>;
 
@@ -256,19 +281,38 @@ protected:
 	ENGINE_API void SetMeshDeformer(bool bInSetMeshDeformer, UMeshDeformer* InMeshDeformer);
 
 	/** Get the currently active MeshDeformer. This may come from the SkeletalMesh default or the Component override. */
-	ENGINE_API UMeshDeformer* GetActiveMeshDeformer() const;
+	ENGINE_API FMeshDeformerSet GetActiveMeshDeformers() const;
 
 	/** Object containing instance settings for the bound MeshDeformer. */
 	UPROPERTY(VisibleAnywhere, BlueprintReadOnly, Instanced, Category = "Deformer", meta = (DisplayName = "Deformer Settings", EditCondition = "MeshDeformerInstanceSettings!=nullptr", ShowOnlyInnerProperties))
 	TObjectPtr<UMeshDeformerInstanceSettings> MeshDeformerInstanceSettings;
 
+#if WITH_EDITORONLY_DATA
+	UE_DEPRECATED(5.4, "Replaced by MeshDeformerInstances. Call GetMeshDeformerInstance() to get the same behavior as reading MeshDeformerInstance.")
+	UPROPERTY(Transient, BlueprintReadOnly, BlueprintGetter = GetMeshDeformerInstance, Category = "Deformer", meta = (DeprecatedProperty, DeprecationMessage = "Use the GetMeshDeformerInstance function instead"))
+	TObjectPtr<UMeshDeformerInstance> MeshDeformerInstance_DEPRECATED;
+#endif // WITH_EDITORONLY_DATA
+
 	/** Object containing state for the bound MeshDeformer. */
-	UPROPERTY(Transient, BlueprintReadOnly, BlueprintGetter = GetMeshDeformerInstance, Category = "Deformer")
-	TObjectPtr<UMeshDeformerInstance> MeshDeformerInstance;
+	UPROPERTY(Transient)
+	FMeshDeformerInstanceSet MeshDeformerInstances;
+
+private:
+	/** Initializes the MeshDeformerInstances property */
+	void CreateMeshDeformerInstances(const FMeshDeformerSet& DeformerSet);
 
 public:
-	UFUNCTION(BlueprintGetter)
-	UMeshDeformerInstance* GetMeshDeformerInstance() const { return MeshDeformerInstance; }
+	UFUNCTION(BlueprintPure, Category = "Components|SkinnedMesh")
+	UMeshDeformerInstance* GetMeshDeformerInstance() const;
+
+	/** 
+	 * Gets the MeshDeformer for the given LOD.
+	 * 
+	 * Returns null if the LOD doesn't use a deformer or if the index is out of range.
+	 * 
+	 * This function takes GetMeshDeformerMaxLOD() into account, so there's no need to call both.
+	 */
+	UMeshDeformerInstance* GetMeshDeformerInstanceForLOD(int32 LODIndex) const;
 
 	/** Max LOD at which to update or apply the MeshDeformer. */
 	ENGINE_API int32 GetMeshDeformerMaxLOD() const;
@@ -1052,7 +1096,7 @@ public:
 
 	ENGINE_API bool IsSkinCacheAllowed(int32 LodIdx) const;
 
-	bool HasMeshDeformer() const { return GetActiveMeshDeformer() != nullptr; }
+	bool HasMeshDeformer() const { return GetActiveMeshDeformers().Deformers.Num() > 0; }
 
 	bool GetForceUpdateDynamicDataImmediately() const { return bForceUpdateDynamicDataImmediately; }
 	void SetForceUpdateDynamicDataImmediately(bool bForceUpdateImmediately) { bForceUpdateDynamicDataImmediately = bForceUpdateImmediately; }
diff --git a/Engine/Source/Runtime/Engine/Classes/Engine/RendererSettings.h b/Engine/Source/Runtime/Engine/Classes/Engine/RendererSettings.h
--- a/Engine/Source/Runtime/Engine/Classes/Engine/RendererSettings.h
+++ b/Engine/Source/Runtime/Engine/Classes/Engine/RendererSettings.h
@@ -1251,6 +1251,12 @@ class URendererSettings : public UDeveloperSettings
 		ConfigRestartRequired = true))
 		uint32 bUseUnlimitedBoneInfluences : 1;
 		
+	UPROPERTY(config, EditAnywhere, Category = Skinning, meta = (
+		ConsoleVariable = "r.GPUSkin.AlwaysUseDeformerForUnlimitedBoneInfluences",
+		ToolTip = "Any mesh LODs using Unlimited Bone Influences will always be rendered with a Mesh Deformer. This reduces the number of shader permutations needed for skeletal mesh materials, saving memory at the cost of performance. Has no effect if either Unlimited Bone Influences or Deformer Graph is disabled.",
+		ConfigRestartRequired = true))
+		uint32 bAlwaysUseDeformerForUnlimitedBoneInfluences : 1;
+		
 	UPROPERTY(config, EditAnywhere, Category = Skinning, meta = (
 		ConsoleVariable = "r.GPUSkin.UnlimitedBoneInfluencesThreshold", DisplayName = "Unlimited Bone Influences Threshold",
 		ToolTip = "When Unlimited Bone Influence is enabled, it still uses a fixed bone inflence buffer until the max bone influence of a mesh exceeds this value"))
diff --git a/Engine/Source/Runtime/Engine/Classes/Engine/SkeletalMeshLODSettings.h b/Engine/Source/Runtime/Engine/Classes/Engine/SkeletalMeshLODSettings.h
--- a/Engine/Source/Runtime/Engine/Classes/Engine/SkeletalMeshLODSettings.h
+++ b/Engine/Source/Runtime/Engine/Classes/Engine/SkeletalMeshLODSettings.h
@@ -113,6 +113,10 @@ struct FSkeletalMeshLODGroupSettings
 	/** The optimization settings to use for the respective LOD level */
 	UPROPERTY(EditAnywhere, Category = Reduction)
 	FSkeletalMeshOptimizationSettings ReductionSettings;
+
+	/** Whether a Mesh Deformer applied to the mesh asset or Skinned Mesh Component should be used on this LOD or not */
+	UPROPERTY(EditAnywhere, Category = LODSetting)
+	bool bAllowMeshDeformer = true;
 };
 
 UCLASS(config = Engine, defaultconfig, BlueprintType, MinimalAPI)
diff --git a/Engine/Source/Runtime/Engine/Classes/Engine/SkinnedAssetCommon.h b/Engine/Source/Runtime/Engine/Classes/Engine/SkinnedAssetCommon.h
--- a/Engine/Source/Runtime/Engine/Classes/Engine/SkinnedAssetCommon.h
+++ b/Engine/Source/Runtime/Engine/Classes/Engine/SkinnedAssetCommon.h
@@ -190,6 +190,10 @@ struct FSkeletalMeshLODInfo
 	 */
 	UPROPERTY(EditAnywhere, Category = SkeletalMeshLODInfo)
 	uint8 bBuildHalfEdgeBuffers : 1;
+
+	/** Whether a Mesh Deformer applied to the mesh asset or Skinned Mesh Component should be used on this LOD or not */
+	UPROPERTY(EditAnywhere, Category = SkeletalMeshLODInfo)
+	uint8 bAllowMeshDeformer : 1;
 	
 	/** List of vertex attributes to include for rendering and what type they should be */
 	UPROPERTY(EditAnywhere, Category = SkeletalMeshLODInfo, AdvancedDisplay, EditFixedSize, Meta=(NoResetToDefault))
@@ -228,6 +232,7 @@ struct FSkeletalMeshLODInfo
 		, bHasPerLODVertexColors(false)
 		, bAllowCPUAccess(false)
 		, bBuildHalfEdgeBuffers(false)
+		, bAllowMeshDeformer(true)
 		, bSupportUniformlyDistributedSampling(false)
 #if WITH_EDITORONLY_DATA
 		, bImportWithBaseMesh(false)
diff --git a/Engine/Source/Runtime/Engine/Private/Components/SkinnedMeshComponent.cpp b/Engine/Source/Runtime/Engine/Private/Components/SkinnedMeshComponent.cpp
--- a/Engine/Source/Runtime/Engine/Private/Components/SkinnedMeshComponent.cpp
+++ b/Engine/Source/Runtime/Engine/Private/Components/SkinnedMeshComponent.cpp
@@ -25,6 +25,7 @@
 #include "SkeletalMeshSceneProxy.h"
 #include "Animation/MeshDeformer.h"
 #include "Animation/MeshDeformerInstance.h"
+#include "Animation/MeshDeformerProvider.h"
 #include "AnimationRuntime.h"
 #include "BoneWeights.h"
 #include "Animation/SkinWeightProfileManager.h"
@@ -699,8 +700,7 @@ void USkinnedMeshComponent::OnRegister()
 
 	InvalidateCachedBounds();
 
-	UMeshDeformer* ActiveMeshDeformer = GetActiveMeshDeformer();
-	MeshDeformerInstance = ActiveMeshDeformer != nullptr ? ActiveMeshDeformer->CreateInstance(this, MeshDeformerInstanceSettings) : nullptr;
+	CreateMeshDeformerInstances(GetActiveMeshDeformers());
 
 	RefreshExternalMorphTargetWeights();
 
@@ -718,16 +718,15 @@ void USkinnedMeshComponent::OnUnregister()
 		AnimUpdateRateParams = nullptr;
 	}
 
-	MeshDeformerInstance = nullptr;
+	CreateMeshDeformerInstances(FMeshDeformerSet());
 }
 
 void USkinnedMeshComponent::BeginPlay()
 {
 	Super::BeginPlay();
 
-	// Recreate the deformer instance to collect any changes for component bindings during startup.
-	UMeshDeformer* ActiveMeshDeformer = GetActiveMeshDeformer();
-	MeshDeformerInstance = ActiveMeshDeformer != nullptr ? ActiveMeshDeformer->CreateInstance(this, MeshDeformerInstanceSettings) : nullptr;
+	// Recreate the deformer instances to collect any changes for component bindings during startup.
+	CreateMeshDeformerInstances(GetActiveMeshDeformers());
 }
 
 const FExternalMorphWeightData& USkinnedMeshComponent::GetExternalMorphWeights(int32 LOD) const
@@ -1066,12 +1065,12 @@ void USkinnedMeshComponent::CreateRenderState_Concurrent(FRegisterComponentConte
 		// scene proxy update of material usage based on active morphs
 		UpdateMorphMaterialUsageOnProxy();
 	}
- 
- 	if (MeshDeformerInstance)
+
+ 	if (UMeshDeformerInstance* MeshDeformerInstance = GetMeshDeformerInstance())
  	{
 		MeshDeformerInstance->AllocateResources();
 		
-			// Enqueue immediate execution of work here to ensure that we have some deformer outputs written for the next frame.
+		// Enqueue immediate execution of work here to ensure that we have some deformer outputs written for the next frame.
 		UMeshDeformerInstance::FEnqueueWorkDesc Desc;
 		Desc.Scene = GetScene();
 		Desc.ExecutionGroup = UMeshDeformerInstance::ExecutionGroup_Immediate;
@@ -1084,7 +1083,7 @@ void USkinnedMeshComponent::DestroyRenderState_Concurrent()
 {
 	Super::DestroyRenderState_Concurrent();
 
-	if (MeshDeformerInstance)
+	if (UMeshDeformerInstance* MeshDeformerInstance = GetMeshDeformerInstance())
 	{
 		MeshDeformerInstance->ReleaseResources();
 	}
@@ -1192,7 +1191,8 @@ void USkinnedMeshComponent::SendRenderDynamicData_Concurrent()
 			UpdateMorphMaterialUsageOnProxy();
 		}
 
-		if (MeshDeformerInstance != nullptr && UseLOD <= GetMeshDeformerMaxLOD())
+		UMeshDeformerInstance* DeformerInstanceForLOD = GetMeshDeformerInstanceForLOD(UseLOD);
+		if (DeformerInstanceForLOD)
 		{
 			UMeshDeformerInstance::FEnqueueWorkDesc Desc;
 			Desc.Scene = GetScene();
@@ -1200,7 +1200,7 @@ void USkinnedMeshComponent::SendRenderDynamicData_Concurrent()
 			// Fallback is to reset the passthrough vertex factory if the deformer fails to run.
 			Desc.FallbackDelegate.BindLambda([MeshObjectPtr = MeshObject, UseLOD]() { FSkeletalMeshDeformerHelpers::ResetVertexFactoryBufferOverrides(MeshObjectPtr, UseLOD); });
 
-			MeshDeformerInstance->EnqueueWork(Desc);
+			DeformerInstanceForLOD->EnqueueWork(Desc);
 		}
 	}
 }
@@ -1278,7 +1278,12 @@ void USkinnedMeshComponent::PostEditChangeProperty(FPropertyChangedEvent& Proper
 		if (Property->GetFName() == GET_MEMBER_NAME_CHECKED(USkinnedMeshComponent, MeshDeformer) ||
 			Property->GetFName() == GET_MEMBER_NAME_CHECKED(USkinnedMeshComponent, bSetMeshDeformer))
 		{
-			UMeshDeformer* ActiveMeshDeformer = GetActiveMeshDeformer();
+			const FMeshDeformerSet ActiveDeformers = GetActiveMeshDeformers();
+
+			// Only one deformer is supported for now
+			check(ActiveDeformers.Deformers.Num() <= 1);
+			UMeshDeformer* ActiveMeshDeformer = ActiveDeformers.Deformers.Num() > 0 ? ActiveDeformers.Deformers[0] : nullptr;
+
 			MeshDeformerInstanceSettings = ActiveMeshDeformer ? ActiveMeshDeformer->CreateSettingsInstance(this) : nullptr;
 		}
 	}
@@ -2044,7 +2049,7 @@ bool USkinnedMeshComponent::IsSkinCacheAllowed(int32 LodIdx) const
 	static const IConsoleVariable* CVarDefaultGPUSkinCacheBehavior = IConsoleManager::Get().FindConsoleVariable(TEXT("r.SkinCache.DefaultBehavior"));
 	const bool bGlobalDefault = CVarDefaultGPUSkinCacheBehavior && ESkinCacheDefaultBehavior(CVarDefaultGPUSkinCacheBehavior->GetInt()) == ESkinCacheDefaultBehavior::Inclusive;
 
-	if (MeshDeformerInstance != nullptr)
+	if (GetMeshDeformerInstance() != nullptr)
 	{
 		// Disable skin cache if a mesh deformer is in use.
 		// Any animation buffers are expected to be owned by the MeshDeformer.
@@ -2211,9 +2216,8 @@ void USkinnedMeshComponent::SetSkinnedAssetAndUpdate(USkinnedAsset* InSkinnedAss
 
 		PrecachePSOs();
 
-		// Re-init the MeshDeformer which might come from the SkelMesh.
-		UMeshDeformer* ActiveMeshDeformer = GetActiveMeshDeformer();
-		MeshDeformerInstance = ActiveMeshDeformer != nullptr ? ActiveMeshDeformer->CreateInstance(this, MeshDeformerInstanceSettings) : nullptr;
+		// Re-init the MeshDeformers which might come from the SkelMesh.
+		CreateMeshDeformerInstances(GetActiveMeshDeformers());
 	}
 
 	// Update external weight array sizes.
@@ -2242,17 +2246,136 @@ USkinnedAsset* USkinnedMeshComponent::GetSkinnedAsset() const
 	PRAGMA_ENABLE_DEPRECATION_WARNINGS
 }
 
-UMeshDeformer* USkinnedMeshComponent::GetActiveMeshDeformer() const
+FMeshDeformerSet USkinnedMeshComponent::GetActiveMeshDeformers() const
 {
+	const FSkeletalMeshRenderData* RenderData = GetSkeletalMeshRenderData();
+
+	if (!RenderData
+		|| !GetSkinnedAsset()
+		|| !GetScene())
+	{
+		// In order to select the right deformer, there needs to be a valid mesh and valid render data.
+		return FMeshDeformerSet();
+	}
+
+	UMeshDeformer* ActiveDeformer = nullptr;
 	if (bSetMeshDeformer)
 	{
-		return MeshDeformer;
+		ActiveDeformer = MeshDeformer;
 	}
-	if (GetSkinnedAsset())
+	else
 	{
-		return GetSkinnedAsset()->GetDefaultMeshDeformer();
+		ActiveDeformer = GetSkinnedAsset()->GetDefaultMeshDeformer();
+	}
+
+	const bool bIsDeformerRequestedByUser = ActiveDeformer != nullptr;
+
+	// If there's no user-specified deformer, find out if we need to set a default deformer for
+	// Unlimited Bone Influences.
+	const bool bIsDeformerRequiredForUBI = FGPUBaseSkinVertexFactory::GetAlwaysUseDeformerForUnlimitedBoneInfluences(GetScene()->GetShaderPlatform());
+	if (bIsDeformerRequiredForUBI && !ActiveDeformer)
+	{
+		bool bMeshUsesUBI = false;
+		for (const FSkeletalMeshLODRenderData& LODRenderData : RenderData->LODRenderData)
+		{
+			if (LODRenderData.SkinWeightVertexBuffer.GetBoneInfluenceType() == UnlimitedBoneInfluence)
+			{
+				bMeshUsesUBI = true;
+				break;
+			}
+		}
+
+		if (bMeshUsesUBI)
+		{
+			static IMeshDeformerProvider* MeshDeformerProvider = IMeshDeformerProvider::Get();
+
+			if (MeshDeformerProvider)
+			{
+				IMeshDeformerProvider::FDefaultMeshDeformerSetup Setup;
+				Setup.bIsRequestingDeformer = true;
+
+				ActiveDeformer = MeshDeformerProvider->GetDefaultMeshDeformer(Setup);
+				if (!ActiveDeformer)
+				{
+					UE_LOG(LogSkinnedMeshComp, Error, TEXT("Failed to fetch default mesh deformer for a mesh that requires it: %s. ")
+						TEXT("Ensure DefaultDeformer is set to a valid Deformer Graph asset in the project settings"), *GetSkinnedAsset()->GetPathName());
+				}
+			}
+		}
+	}
+
+	if (!ActiveDeformer)
+	{
+		return FMeshDeformerSet();
+	}
+
+	// Currently only one active deformer is supported
+	FMeshDeformerSet Result;
+	Result.Deformers.Empty(1);
+	Result.Deformers.Add(ActiveDeformer);
+
+	const TArray<FSkeletalMeshLODInfo>& MeshLODInfoArray = GetSkinnedAsset()->GetLODInfoArray();
+	const int32 MaxLOD = GetMeshDeformerMaxLOD();
+
+	// Every entry of this array will be written by the loop below
+	Result.DeformerIndexForLOD.Empty(RenderData->LODRenderData.Num());
+	Result.DeformerIndexForLOD.AddUninitialized(RenderData->LODRenderData.Num());
+
+	for (int32 Index = 0; Index < RenderData->LODRenderData.Num(); Index++)
+	{
+		const FSkeletalMeshLODRenderData& LODRenderData = RenderData->LODRenderData[Index];
+		const bool bRequiredForUBI = bIsDeformerRequiredForUBI && LODRenderData.SkinWeightVertexBuffer.GetBoneInfluenceType() == UnlimitedBoneInfluence;
+
+		const bool bAllowedByMaxLOD = Index <= MaxLOD;
+		// There should be a LODInfo entry for this LOD, but if not, default to allowing the deformer
+		const bool bAllowedByLODInfo = !MeshLODInfoArray.IsValidIndex(Index) || MeshLODInfoArray[Index].bAllowMeshDeformer;
+
+		const bool bDeformerEnabledForThisLOD = bRequiredForUBI || (bIsDeformerRequestedByUser && bAllowedByMaxLOD && bAllowedByLODInfo);
+		Result.DeformerIndexForLOD[Index] = bDeformerEnabledForThisLOD ? 0 : INDEX_NONE;
+	}
+
+	return Result;
+}
+
+void USkinnedMeshComponent::CreateMeshDeformerInstances(const FMeshDeformerSet& DeformerSet)
+{
+	// Only one deformer is supported for now
+	check(DeformerSet.Deformers.Num() <= 1);
+	UMeshDeformer* ActiveMeshDeformer = DeformerSet.Deformers.Num() > 0 ? DeformerSet.Deformers[0] : nullptr;
+
+	UMeshDeformerInstance* MeshDeformerInstance = ActiveMeshDeformer ? ActiveMeshDeformer->CreateInstance(this, MeshDeformerInstanceSettings) : nullptr;
+	
+	MeshDeformerInstances = FMeshDeformerInstanceSet();
+	if (MeshDeformerInstance)
+	{
+		MeshDeformerInstances.DeformerInstances.Empty(1);
+		MeshDeformerInstances.DeformerInstances.Add(MeshDeformerInstance);
+
+		MeshDeformerInstances.InstanceIndexForLOD = DeformerSet.DeformerIndexForLOD;
 	}
-	return nullptr;
+}
+
+UMeshDeformerInstance* USkinnedMeshComponent::GetMeshDeformerInstance() const
+{
+	return MeshDeformerInstances.DeformerInstances.Num() > 0 ? MeshDeformerInstances.DeformerInstances[0] : nullptr;
+}
+
+UMeshDeformerInstance* USkinnedMeshComponent::GetMeshDeformerInstanceForLOD(int32 LODIndex) const
+{
+	if (!MeshDeformerInstances.InstanceIndexForLOD.IsValidIndex(LODIndex))
+	{
+		return nullptr;
+	}
+	
+	const int8 InstanceIndex = MeshDeformerInstances.InstanceIndexForLOD[LODIndex];
+	if (InstanceIndex == INDEX_NONE)
+	{
+		// Don't use a deformer for this LOD
+		return nullptr;
+	}
+
+	check(MeshDeformerInstances.DeformerInstances.IsValidIndex(InstanceIndex));
+	return MeshDeformerInstances.DeformerInstances[InstanceIndex];
 }
 
 void USkinnedMeshComponent::SetMeshDeformer(bool bInSetMeshDeformer, UMeshDeformer* InMeshDeformer)
@@ -2260,9 +2383,15 @@ void USkinnedMeshComponent::SetMeshDeformer(bool bInSetMeshDeformer, UMeshDeform
 	bSetMeshDeformer = bInSetMeshDeformer;
 	MeshDeformer = InMeshDeformer;
 
-	UMeshDeformer* ActiveMeshDeformer = GetActiveMeshDeformer();
+	const FMeshDeformerSet ActiveDeformers = GetActiveMeshDeformers();
+
+	// Only one deformer is supported for now
+	check(ActiveDeformers.Deformers.Num() <= 1);
+	UMeshDeformer* ActiveMeshDeformer = ActiveDeformers.Deformers.Num() > 0 ? ActiveDeformers.Deformers[0] : nullptr;
+
 	MeshDeformerInstanceSettings = ActiveMeshDeformer ? ActiveMeshDeformer->CreateSettingsInstance(this) : nullptr;
-	MeshDeformerInstance = ActiveMeshDeformer ? ActiveMeshDeformer->CreateInstance(this, MeshDeformerInstanceSettings) : nullptr;
+
+	CreateMeshDeformerInstances(ActiveDeformers);
 
 	MarkRenderDynamicDataDirty();
 }
diff --git a/Engine/Source/Runtime/Engine/Private/GPUSkinVertexFactory.cpp b/Engine/Source/Runtime/Engine/Private/GPUSkinVertexFactory.cpp
--- a/Engine/Source/Runtime/Engine/Private/GPUSkinVertexFactory.cpp
+++ b/Engine/Source/Runtime/Engine/Private/GPUSkinVertexFactory.cpp
@@ -5,6 +5,7 @@
 =============================================================================*/
 
 #include "GPUSkinVertexFactory.h"
+#include "Animation/MeshDeformerProvider.h"
 #include "MeshBatch.h"
 #include "GPUSkinCache.h"
 #include "MeshDrawShaderBindings.h"
@@ -61,6 +62,13 @@ static FAutoConsoleVariableRef CVarUnlimitedBoneInfluencesThreshold(
 	TEXT("Unlimited Bone Influences Threshold to use unlimited bone influences buffer if r.GPUSkin.UnlimitedBoneInfluences is enabled. Should be unsigned int. Cannot be changed at runtime."),
 	ECVF_ReadOnly);
 
+static bool GCVarAlwaysUseDeformerForUnlimitedBoneInfluences = false;
+static FAutoConsoleVariableRef CVarAlwaysUseDeformerForUnlimitedBoneInfluences(
+	TEXT("r.GPUSkin.AlwaysUseDeformerForUnlimitedBoneInfluences"),
+	GCVarAlwaysUseDeformerForUnlimitedBoneInfluences,
+	TEXT("Any meshes using Unlimited Bone Influences will always be rendered with a Mesh Deformer. This reduces the number of shader permutations needed for skeletal mesh materials, saving memory at the cost of performance. Has no effect if either Unlimited Bone Influences or Deformer Graph is disabled. Cannot be changed at runtime."),
+	ECVF_ReadOnly);
+
 static TAutoConsoleVariable<bool> CVarMobileEnableCloth(
 	TEXT("r.Mobile.EnableCloth"),
 	true,
@@ -486,6 +494,25 @@ int32 FGPUBaseSkinVertexFactory::GetBoneInfluenceLimitForAsset(int32 AssetProvid
 	return MAX_TOTAL_INFLUENCES;
 }
 
+bool FGPUBaseSkinVertexFactory::GetAlwaysUseDeformerForUnlimitedBoneInfluences(EShaderPlatform Platform)
+{
+	auto InnerFunc = [](EShaderPlatform Platform)
+	{
+		FShaderPlatformCachedIniValue<bool> UseDeformerForUBICVar(TEXT("r.GPUSkin.AlwaysUseDeformerForUnlimitedBoneInfluences"));
+		const IMeshDeformerProvider* MeshDeformerProvider = IMeshDeformerProvider::Get();
+
+		return MeshDeformerProvider && MeshDeformerProvider->IsSupported(Platform) && UseDeformerForUBICVar.Get(Platform);
+	};
+
+#if WITH_EDITOR
+	return InnerFunc(Platform);
+#else
+	// This value can't change at runtime in a non-editor build, so it's safe to cache.
+	static const bool bCachedResult = InnerFunc(Platform);
+	return bCachedResult;
+#endif
+}
+
 void FGPUBaseSkinVertexFactory::SetData(const FGPUSkinDataType* InData)
 {
 	SetData(FRHICommandListExecutor::GetImmediateCommandList(), InData);
@@ -535,8 +562,14 @@ TGlobalResource<FBoneBufferPool> FGPUBaseSkinVertexFactory::BoneBufferPool;
 template <GPUSkinBoneInfluenceType BoneInfluenceType>
 bool TGPUSkinVertexFactory<BoneInfluenceType>::ShouldCompilePermutation(const FVertexFactoryShaderPermutationParameters& Parameters)
 {
-	static FShaderPlatformCachedIniValue<int32> PerPlatformCVar(TEXT("r.GPUSkin.UnlimitedBoneInfluences"));
-	const bool bUnlimitedBoneInfluences = BoneInfluenceType == UnlimitedBoneInfluence && PerPlatformCVar.Get(Parameters.Platform) != 0;
+	static FShaderPlatformCachedIniValue<int32> UBICVar(TEXT("r.GPUSkin.UnlimitedBoneInfluences"));
+	const bool bUseUBI = UBICVar.Get(Parameters.Platform) != 0;
+
+	static FShaderPlatformCachedIniValue<bool> UseDeformerForUBICVar(TEXT("r.GPUSkin.AlwaysUseDeformerForUnlimitedBoneInfluences"));
+	const bool bUseDeformerForUBI = UseDeformerForUBICVar.Get(Parameters.Platform);
+		
+	// Compile the shader for UBI if UBI is enabled and we're not forcing the use of a deformer for all UBI meshes
+	const bool bUnlimitedBoneInfluences = BoneInfluenceType == UnlimitedBoneInfluence && bUseUBI && !bUseDeformerForUBI;
 
 	return ShouldWeCompileGPUSkinVFShaders(Parameters.Platform, Parameters.MaterialParameters.FeatureLevel) &&
 		  (((Parameters.MaterialParameters.bIsUsedWithSkeletalMesh || Parameters.MaterialParameters.bIsUsedWithMorphTargets) && (BoneInfluenceType != UnlimitedBoneInfluence || bUnlimitedBoneInfluences)) 
diff --git a/Engine/Source/Runtime/Engine/Private/SkeletalMeshLODSettings.cpp b/Engine/Source/Runtime/Engine/Private/SkeletalMeshLODSettings.cpp
--- a/Engine/Source/Runtime/Engine/Private/SkeletalMeshLODSettings.cpp
+++ b/Engine/Source/Runtime/Engine/Private/SkeletalMeshLODSettings.cpp
@@ -58,6 +58,7 @@ bool USkeletalMeshLODSettings::SetLODSettingsToMesh(USkeletalMesh* InMesh, int32
 		LODInfo->ScreenSize = Setting.ScreenSize;
 		LODInfo->LODHysteresis = Setting.LODHysteresis;
 		LODInfo->WeightOfPrioritization = Setting.WeightOfPrioritization;
+		LODInfo->bAllowMeshDeformer = Setting.bAllowMeshDeformer;
 		// if we have available bake pose
 		// it's possible for skeleton to be null if this happens in the middle of importing
 		// so if skeleton is null, we allow copy (the GetBakePose will check correct skeleton when get it)
@@ -291,6 +292,8 @@ int32 USkeletalMeshLODSettings::SetLODSettingsFromMesh(USkeletalMesh* InMesh)
 				NewFilter.BoneName = LODInfo->BonesToRemove[BoneIndex].BoneName;
 				Setting.BoneList.Add(NewFilter);
 			}
+
+			Setting.bAllowMeshDeformer = LODInfo->bAllowMeshDeformer;
 		}
 
 		return NumSettings;
diff --git a/Engine/Source/Runtime/Engine/Private/SkeletalRenderGPUSkin.cpp b/Engine/Source/Runtime/Engine/Private/SkeletalRenderGPUSkin.cpp
--- a/Engine/Source/Runtime/Engine/Private/SkeletalRenderGPUSkin.cpp
+++ b/Engine/Source/Runtime/Engine/Private/SkeletalRenderGPUSkin.cpp
@@ -5,6 +5,7 @@
 =============================================================================*/
 
 #include "SkeletalRenderGPUSkin.h"
+#include "Animation/MeshDeformerInstance.h"
 #include "Components/SkeletalMeshComponent.h"
 #include "RenderUtils.h"
 #include "SkeletalRender.h"
@@ -2076,7 +2077,7 @@ void FSkeletalMeshObjectGPUSkin::GetUsedVertexFactoryData(
 	InitGPUSkinVertexFactoryComponents(&GPUSkinDataType, VertexBuffers, nullptr /*FGPUBaseSkinVertexFactory*/);
 
 	bool bIsSkinCacheAllowed = SkinnedMeshComponent ? SkinnedMeshComponent->IsSkinCacheAllowed(LODIndex) : false;
-	bool bHasMeshDeformer = SkinnedMeshComponent ? SkinnedMeshComponent->GetMeshDeformerInstance() != nullptr && LODIndex <= SkinnedMeshComponent->GetMeshDeformerMaxLOD() : false;
+	bool bHasMeshDeformer = SkinnedMeshComponent && SkinnedMeshComponent->GetMeshDeformerInstanceForLOD(LODIndex);
 	if (bIsSkinCacheAllowed || bHasMeshDeformer)
 	{
 		VertexFactoryDataList.AddUnique(FPSOPrecacheVertexFactoryData(&FGPUSkinPassthroughVertexFactory::StaticType));
@@ -2639,7 +2640,7 @@ void FDynamicSkelMeshObjectDataGPUSkin::InitDynamicSkelMeshObjectDataGPUSkin(
 	}
 
 	bIsSkinCacheAllowed = InMeshComponent->IsSkinCacheAllowed(InLODIndex);
-	bHasMeshDeformer = InMeshComponent->GetMeshDeformerInstance() != nullptr && LODIndex <= InMeshComponent->GetMeshDeformerMaxLOD();
+	bHasMeshDeformer = InMeshComponent->GetMeshDeformerInstanceForLOD(LODIndex) != nullptr;
 	bForceUpdateDynamicDataImmediately = InMeshComponent->GetForceUpdateDynamicDataImmediately();
 	// Force immediate update when using mesh deformer.
 	bForceUpdateDynamicDataImmediately |= bHasMeshDeformer;
diff --git a/Engine/Source/Runtime/Engine/Public/GPUSkinVertexFactory.h b/Engine/Source/Runtime/Engine/Public/GPUSkinVertexFactory.h
--- a/Engine/Source/Runtime/Engine/Public/GPUSkinVertexFactory.h
+++ b/Engine/Source/Runtime/Engine/Public/GPUSkinVertexFactory.h
@@ -380,6 +380,12 @@ public:
 	 */
 	ENGINE_API static int32 GetBoneInfluenceLimitForAsset(int32 AssetProvidedLimit, const ITargetPlatform* TargetPlatform = nullptr);
 
+	/**
+	 * Returns true if mesh LODs with Unlimited Bone Influences must always be rendered using a
+	 * Mesh Deformer for the given shader platform.
+	 */
+	ENGINE_API static bool GetAlwaysUseDeformerForUnlimitedBoneInfluences(EShaderPlatform Platform);
+
 	/** Morph vertex factory functions */
 	virtual void UpdateMorphVertexStream(const class FMorphVertexBuffer* MorphVertexBuffer) {}
 	virtual const class FMorphVertexBuffer* GetMorphVertexBuffer(bool bPrevious) const { return nullptr; }
diff --git a/Engine/Source/Runtime/RenderCore/Public/Animation/MeshDeformerProvider.h b/Engine/Source/Runtime/RenderCore/Public/Animation/MeshDeformerProvider.h
--- a/Engine/Source/Runtime/RenderCore/Public/Animation/MeshDeformerProvider.h
+++ b/Engine/Source/Runtime/RenderCore/Public/Animation/MeshDeformerProvider.h
@@ -32,7 +32,7 @@ public:
 	/** Structure for passing to GetDefaultMeshDeformer(). */
 	struct FDefaultMeshDeformerSetup
 	{
-		bool bIsUsingSkinCache = false;
+		bool bIsRequestingDeformer = false;
 		bool bIsRequestingRecomputeTangent = false;
 	};
 
