VirtualTextureDataBuilder : remove slow and unnecessary alpha detection code
(which was also a dupe of FImageCore)
use the already computed bKnowAlphaTransparency
note that this has very slightly different alpha logic so can be different in very rare cases
dump DDC key for all VT

diff --git a/Engine/Source/Runtime/Engine/Private/TextureDerivedData.cpp b/Engine/Source/Runtime/Engine/Private/TextureDerivedData.cpp
--- a/Engine/Source/Runtime/Engine/Private/TextureDerivedData.cpp
+++ b/Engine/Source/Runtime/Engine/Private/TextureDerivedData.cpp
@@ -69,7 +69,8 @@ static TAutoConsoleVariable<int32> CVarTexturesCookToDerivedDataReferences(
 // This GUID is mixed into DDC version for virtual textures only, this allows updating DDC version for VT without invalidating DDC for all textures
 // This is useful during development, but once large numbers of VT are present in shipped content, it will have the same problem as TEXTURE_DERIVEDDATA_VER
 // This is put in the DDC1 key but NOT in the DDC2 key
-#define TEXTURE_VT_DERIVEDDATA_VER	TEXT("7C16439390E24F1F9468894FB4D4BC54")
+// VT key bumped 02-27-2024 for Alpha change
+#define TEXTURE_VT_DERIVEDDATA_VER	TEXT("7C16439390E24F1F9468894FB4D4BC55")
 
 // TEXTURE_DDC_STB_IMAGE_RESIZE_VERSION should change whenever the stb_image_resize2.h version number changes
 //	*if* it is a version change that changes output
diff --git a/Engine/Source/Runtime/Engine/Private/VT/VirtualTextureDataBuilder.cpp b/Engine/Source/Runtime/Engine/Private/VT/VirtualTextureDataBuilder.cpp
--- a/Engine/Source/Runtime/Engine/Private/VT/VirtualTextureDataBuilder.cpp
+++ b/Engine/Source/Runtime/Engine/Private/VT/VirtualTextureDataBuilder.cpp
@@ -444,9 +444,6 @@ bool FVirtualTextureDataBuilder::Build(FTextureSourceData& InSourceData, FTextur
 				LayerData.GammaSpace = EGammaSpace::Linear;
 			}
 
-			// LayerData.bHasAlpha will be updated after the Build to detect if there is actual alpha in the layers
-			LayerData.bHasAlpha = BuildSettingsForLayer.bForceAlphaChannel;
-
 			LayerData.FormatName = FImageCoreUtils::ConvertToUncompressedTextureFormatName(LayerData.ImageFormat);
 			LayerData.PixelFormat = FImageCoreUtils::GetPixelFormatForRawImageFormat(LayerData.ImageFormat);
 			LayerData.SourceFormat = FImageCoreUtils::ConvertToTextureSourceFormat(LayerData.ImageFormat);
@@ -455,22 +452,32 @@ bool FVirtualTextureDataBuilder::Build(FTextureSourceData& InSourceData, FTextur
 			// (VT physical tiles generally not power-of-2 after adding border)
 			// Must match TextureDerivedData.cpp
 			LayerData.TextureFormatName = UE::TextureBuildUtilities::TextureFormatRemovePlatformPrefixFromName(BuildSettingsForLayer.TextureFormatName);
-
-			// bHasAlpha was previously set to true if bForceAlphaChannel
-			// if it's false and not bForceNoAlphaChannel, we scan each block for alpha
-			// if alpha is in any tile of any block, it gets enabled for all so they have consistent pixel format
-			if (!LayerData.bHasAlpha && !BuildSettingsForLayer.bForceNoAlphaChannel)
+			
+			if ( BuildSettingsForLayer.bKnowAlphaTransparency )
+			{
+				// bKnowAlphaTransparency includes all Force actions
+				LayerData.bHasAlpha = BuildSettingsForLayer.bHasTransparentAlpha;
+			}
+			else if ( BuildSettingsForLayer.bForceNoAlphaChannel ) // note the order of operations! ( ForceNo takes precedence )
+			{
+				LayerData.bHasAlpha = false;
+			}
+			else if ( BuildSettingsForLayer.bForceAlphaChannel )
 			{
-				// Note that check is a bit wrong to do at this point, because it does require all blocks to be present in the memory
-				// This should've been stored somewhere way earlier, maybe at import process
+				LayerData.bHasAlpha = true;
+			}
+			else
+			{
+				// alpha detection was not previously done
+				//	must do it now on all blocks
+				// (this is hard to hit; bKnowAlphaTransparency is almost always true now)
+
 				for (int32 BlockIndex = 0; BlockIndex < NumBlocks; ++BlockIndex)
 				{
 					const TArray<FImage>& SourceMips = InSourceData.Blocks[BlockIndex].MipsPerLayer[LayerIndex];
 					if (!SourceMips.IsEmpty())
 					{
-						// @todo Oodle : use FImageCore::DetectAlphaChannel instead
-						//	VT_DetectAlphaChannel is slow and not threaded and unnecessary code dupe
-						LayerData.bHasAlpha = VT_DetectAlphaChannel(SourceMips[0]);
+						LayerData.bHasAlpha = FImageCore::DetectAlphaChannel(SourceMips[0]);
 						if (LayerData.bHasAlpha)
 						{
 							break;
@@ -1417,84 +1424,4 @@ void FVirtualTextureDataBuilder::BuildMipTails()
 }
 #endif // 0
 
-bool FVirtualTextureDataBuilder::VT_DetectAlphaChannel(const FImage &Image)
-{
-	TRACE_CPUPROFILER_EVENT_SCOPE(VT.DetectAlphaChannel);
-
-	// note : VT DetectAlphaChannel slightly different than the same function in TextureCompressorModule
-	//	  could factor them out and share
-	//		technically could change output so may need a ddc key bump and verify
-
-	if (Image.Format == ERawImageFormat::BGRA8)
-	{
-		const FColor* SrcColors = (&Image.AsBGRA8()[0]);
-		const FColor* LastColor = SrcColors + (Image.SizeX * Image.SizeY * Image.NumSlices);
-		while (SrcColors < LastColor)
-		{
-			if (SrcColors->A < 255)
-			{
-				return true;
-			}
-			++SrcColors;
-		}
-		return false;
-	}
-	else if (Image.Format == ERawImageFormat::RGBA16F)
-	{
-		const FFloat16Color* SrcColors = (&Image.AsRGBA16F()[0]);
-		const FFloat16Color* LastColor = SrcColors + (Image.SizeX * Image.SizeY * Image.NumSlices);
-		while (SrcColors < LastColor)
-		{
-			if (SrcColors->A < (1.0f - UE_SMALL_NUMBER))
-			{
-				return true;
-			}
-			++SrcColors;
-		}
-		return false;
-	}
-	else if (Image.Format == ERawImageFormat::RGBA32F)
-	{
-		const FLinearColor* SrcColors = (&Image.AsRGBA32F()[0]);
-		const FLinearColor* LastColor = SrcColors + (Image.SizeX * Image.SizeY * Image.NumSlices);
-		while (SrcColors < LastColor)
-		{
-			// this comparison matches to would happen if RGBA32F would be converted to BGRA8 format
-			if (SrcColors->QuantizeRound().A < 255)
-			{
-				return true;
-			}
-			++SrcColors;
-		}
-		return false;
-	}
-	else if (Image.Format == ERawImageFormat::RGBA16)
-	{
-		const uint16* SrcColors = (&Image.AsRGBA16()[0]);
-		const uint16* LastColor = SrcColors + (Image.SizeX * Image.SizeY * Image.NumSlices);
-		while (SrcColors < LastColor)
-		{
-			if (SrcColors[3] < 65535)
-			{
-				return true;
-			}
-			SrcColors += 4;
-		}
-		return false;
-	}
-	else if (Image.Format == ERawImageFormat::G16 ||
-			 Image.Format == ERawImageFormat::G8 ||
-			 Image.Format == ERawImageFormat::BGRE8 ||
-			 Image.Format == ERawImageFormat::R16F ||
-			 Image.Format == ERawImageFormat::R32F)
-	{
-		return false;
-	}
-	else
-	{
-		check(false);
-		return true;
-	}
-}
-
 #endif // WITH_EDITOR
diff --git a/Engine/Source/Runtime/Engine/Private/VT/VirtualTextureDataBuilder.h b/Engine/Source/Runtime/Engine/Private/VT/VirtualTextureDataBuilder.h
--- a/Engine/Source/Runtime/Engine/Private/VT/VirtualTextureDataBuilder.h
+++ b/Engine/Source/Runtime/Engine/Private/VT/VirtualTextureDataBuilder.h
@@ -149,6 +149,4 @@ private:
 
 	const FString& DebugTexturePathName;
 	int32 ChunkDumpIndex = 0;
-
-	static bool VT_DetectAlphaChannel(const FImage &image);
 };
\ No newline at end of file
