Fixing issue with missing overlap results for objects lacking a primitive component due to uniqueness collision. The FPhysicsObjectHandle is now part of the uniqueness check.

#tests  multi overlap query tested for multiple physics object with null primitive components. The logic for both when `ExpectedSize >= GNumOverlapsRequiredForTMap` is true and fals has been verified.

diff --git a/Engine/Source/Runtime/Engine/Private/Collision/CollisionConversions.cpp b/Engine/Source/Runtime/Engine/Private/Collision/CollisionConversions.cpp
--- a/Engine/Source/Runtime/Engine/Private/Collision/CollisionConversions.cpp
+++ b/Engine/Source/Runtime/Engine/Private/Collision/CollisionConversions.cpp
@@ -22,27 +22,28 @@
 // Used to place overlaps into a TMap when deduplicating them
 struct FOverlapKey
 {
-	UPrimitiveComponent* Component;
-	int32 ComponentIndex;
+	const Chaos::FPhysicsObjectHandle PhysicsObject;
+	const UPrimitiveComponent* Component;
+	const int32 ComponentIndex;
 
-	FOverlapKey(UPrimitiveComponent* InComponent, int32 InComponentIndex)
-		: Component(InComponent)
-		, ComponentIndex(InComponentIndex)
+	FOverlapKey(const FOverlapResult& InResult)
+		: PhysicsObject(InResult.PhysicsObject)
+		, Component(InResult.GetComponent())
+		, ComponentIndex(InResult.ItemIndex)
 	{
 	}
 
 	friend bool operator==(const FOverlapKey& X, const FOverlapKey& Y)
 	{
-		return (X.ComponentIndex == Y.ComponentIndex) && (X.Component == Y.Component);
+		return  (X.PhysicsObject == Y.PhysicsObject) && (X.ComponentIndex == Y.ComponentIndex) && (X.Component == Y.Component);
 	}
 };
 
 uint32 GetTypeHash(const FOverlapKey& Key)
 {
-	return GetTypeHash(Key.Component) ^ GetTypeHash(Key.ComponentIndex);
+	return GetTypeHash(Key.PhysicsObject) ^ GetTypeHash(Key.Component) ^ GetTypeHash(Key.ComponentIndex);
 }
 
-
 extern int32 CVarShowInitialOverlaps;
 
 // Forward declare, I don't want to move the entire function right now or we lose change history.
@@ -645,7 +646,7 @@ static void AddUniqueOverlap(TArray<FOverlapResult>& OutOverlaps, const FOverlap
 	{
 		FOverlapResult& Overlap = OutOverlaps[TestIdx];
 
-		if (Overlap.ItemIndex == NewOverlap.ItemIndex && Overlap.Component == NewOverlap.Component)
+		if (FOverlapKey(Overlap) == FOverlapKey(NewOverlap))
 		{
 			// These should be the same if the component matches!
 			checkSlow(Overlap.OverlapObjectHandle == NewOverlap.OverlapObjectHandle);
@@ -704,7 +705,7 @@ bool ConvertOverlapResultsImp(int32 NumOverlaps, THitOverlap* OverlapResults, co
 		for (int32 ExistingIndex = 0; ExistingIndex < OutOverlaps.Num(); ++ExistingIndex)
 		{
 			const FOverlapResult& ExistingOverlap = OutOverlaps[ExistingIndex];
-			OverlapMap.Add(FOverlapKey(ExistingOverlap.Component.Get(), ExistingOverlap.ItemIndex), ExistingIndex + 1);
+			OverlapMap.Add(FOverlapKey(ExistingOverlap), ExistingIndex + 1);
 		}
 
 		for (int32 PResultIndex = 0; PResultIndex < NumOverlaps; ++PResultIndex)
@@ -728,7 +729,7 @@ bool ConvertOverlapResultsImp(int32 NumOverlaps, THitOverlap* OverlapResults, co
 			//TODO: this doesn't de-duplicate if FOverlapResult has no component - like in PT for example. If we care about de-duplication we should use the PT handle if needed
 			//Question: why do we de-duplicate? This seems expensive and is not consistent with traces. The user could handle this at their level anyway (at the cost of a bit of transient memory here)
 			//TODO: make sure it doesn't de-duplicate null results
-			int32& DestinationIndex = OverlapMap.FindOrAdd(FOverlapKey(NewOverlap.Component.Get(), NewOverlap.ItemIndex));
+			int32& DestinationIndex = OverlapMap.FindOrAdd(FOverlapKey(NewOverlap));
 			if (DestinationIndex == 0)
 			{
 				DestinationIndex = OutOverlaps.Add(NewOverlap) + 1;
