Fixed typo in define.

diff --git a/Engine/Source/Runtime/Engine/Classes/Particles/ParticlePerfStats.h b/Engine/Source/Runtime/Engine/Classes/Particles/ParticlePerfStats.h
--- a/Engine/Source/Runtime/Engine/Classes/Particles/ParticlePerfStats.h
+++ b/Engine/Source/Runtime/Engine/Classes/Particles/ParticlePerfStats.h
@@ -15,7 +15,7 @@
 	#if (WITH_UNREAL_DEVELOPER_TOOLS || WITH_UNREAL_TARGET_DEVELOPER_TOOLS || (!UE_BUILD_SHIPPING) || WITH_GLOBAL_RUNTIME_FX_BUDGET)
 	#define WITH_PARTICLE_PERF_STATS 1
 	#else
-		WITH_PARTICLE_PERF_STATS 0
+	#define	WITH_PARTICLE_PERF_STATS 0
 	#endif
 #else
 	#if !WITH_PARTICLE_PERF_STATS
