TextureDerivedDataTask.cpp : insights scope for IsBulkDataLoaded
this can be a stall waiting on the BulkData mutex if it is serializing to the undo buffer on the main thread

diff --git a/Engine/Source/Runtime/Engine/Private/TextureDerivedDataTask.cpp b/Engine/Source/Runtime/Engine/Private/TextureDerivedDataTask.cpp
--- a/Engine/Source/Runtime/Engine/Private/TextureDerivedDataTask.cpp
+++ b/Engine/Source/Runtime/Engine/Private/TextureDerivedDataTask.cpp
@@ -1726,7 +1726,15 @@ static bool DDC1_LoadAndValidateTextureData(
 	TRACE_CPUPROFILER_EVENT_SCOPE(Texture.DDC1_LoadAndValidateTextureData);
 
 	bool bHasTextureSourceMips = false;
-	if (TextureData.IsValid() && Texture.Source.IsBulkDataLoaded())
+	bool bNeedsGetSourceMips;
+
+	{
+	// this can be a stall waiting on the BulkData mutex if it is serializing to the undo buffer on the main thread :
+	TRACE_CPUPROFILER_EVENT_SCOPE(Texture.IsBulkDataLoaded);
+	bNeedsGetSourceMips = TextureData.IsValid() && Texture.Source.IsBulkDataLoaded();
+	}
+
+	if ( bNeedsGetSourceMips )
 	{
 		TRACE_CPUPROFILER_EVENT_SCOPE(GetSourceMips);
 		TextureData.GetSourceMips(Texture.Source, ImageWrapper);
