Deprecating old FVisualLogger::GetEntryToWrite in favor of a new, static FVisualLogger::GetEntryToWrite for consistent Engine Timestamps in Editor.

diff --git a/Engine/Source/Runtime/AIModule/Private/Navigation/CrowdFollowingComponent.cpp b/Engine/Source/Runtime/AIModule/Private/Navigation/CrowdFollowingComponent.cpp
--- a/Engine/Source/Runtime/AIModule/Private/Navigation/CrowdFollowingComponent.cpp
+++ b/Engine/Source/Runtime/AIModule/Private/Navigation/CrowdFollowingComponent.cpp
@@ -709,8 +709,7 @@ void LogPathPartHelper(AActor* LogOwner, FNavMeshPath* NavMeshPath, int32 StartI
 	const FVector CorridorOffset = NavigationDebugDrawing::PathOffset * 1.25f;
 	int32 NumAreaMark = 1;
 
-	FVisualLogEntry* Snapshot = VisualLogger.GetEntryToWrite(LogOwner, LogOwner->GetWorld()->GetTimeSeconds());
-	if (Snapshot)
+	if (FVisualLogEntry* Snapshot = FVisualLogger::GetEntryToWrite(LogOwner, LogCrowdFollowing))
 	{
 		NavMesh->BeginBatchQuery();
 
diff --git a/Engine/Source/Runtime/AIModule/Private/Navigation/PathFollowingComponent.cpp b/Engine/Source/Runtime/AIModule/Private/Navigation/PathFollowingComponent.cpp
--- a/Engine/Source/Runtime/AIModule/Private/Navigation/PathFollowingComponent.cpp
+++ b/Engine/Source/Runtime/AIModule/Private/Navigation/PathFollowingComponent.cpp
@@ -197,8 +197,7 @@ void UPathFollowingComponent::LogPathHelper(const AActor* LogOwner, FNavigationP
 	{
 		const FVector PathEnd = *InLogPath->GetPathPointLocation(InLogPath->GetPathPoints().Num() - 1);
 
-		FVisualLogEntry* Entry = Vlog.GetEntryToWrite(LogOwner, LogOwner->GetWorld()->TimeSeconds);
-		if (Entry)
+		if (FVisualLogEntry* Entry = FVisualLogger::GetEntryToWrite(LogOwner, LogPathFollowing))
 		{
 			InLogPath->DescribeSelfToVisLog(Entry);
 			if (LogGoalActor)
diff --git a/Engine/Source/Runtime/Engine/Private/VisualLogger/VisualLogger.cpp b/Engine/Source/Runtime/Engine/Private/VisualLogger/VisualLogger.cpp
--- a/Engine/Source/Runtime/Engine/Private/VisualLogger/VisualLogger.cpp
+++ b/Engine/Source/Runtime/Engine/Private/VisualLogger/VisualLogger.cpp
@@ -221,21 +221,38 @@ bool FVisualLogger::CheckVisualLogInputInternal(const UObject* Object, const FNa
 		return false;
 	}
 
-	*OutWorld = GetWorldForVisualLogger(Object);
-	if (!ensure(*OutWorld != nullptr))
+	if (OutWorld)
 	{
-		return false;
+		*OutWorld = GetWorldForVisualLogger(Object);
+		if (!ensure(*OutWorld != nullptr))
+		{
+			return false;
+		}
 	}
 
-	*OutCurrentEntry = VisualLogger.GetEntryToWrite(Object, VisualLogger.GetTimeStampForObject(Object));
-	if (*OutCurrentEntry == nullptr)
+	if (OutCurrentEntry)
 	{
-		return false;
+		*OutCurrentEntry = VisualLogger.GetEntryToWrite(Object, VisualLogger.GetTimeStampForObject(Object));
+		if (*OutCurrentEntry == nullptr)
+		{
+			return false;
+		}
 	}
 
 	return true;
 }
 
+FVisualLogEntry* FVisualLogger::GetEntryToWrite(const UObject* LogOwner, const FLogCategoryBase& LogCategory)
+{
+	FVisualLogEntry* LogEntry;
+	if (CheckVisualLogInputInternal(LogOwner, LogCategory.GetCategoryName(), ELogVerbosity::Log, /*OutWorld*/ nullptr, &LogEntry))
+	{
+		return LogEntry;
+	}
+
+	return nullptr;
+}
+
 double FVisualLogger::GetTimeStampForObject(const UObject* Object) const
 {
 	// Licensees can write their own synchronized clock that can work across the network
@@ -764,8 +781,6 @@ FVisualLogger::FVisualLogger()
 	});
 
 #if WITH_EDITOR
-	// When PIE is Starting, we want to reset the EditorBaseTimeStamp on the very first PIE World Instance (but not the other instances).
-	// This gives us a pseudo-world time (that starts at zero) while still being synchronized across multiple PIE instances.
 	PIEStartedHandle = FWorldDelegates::OnPIEStarted.AddLambda(
 		[this](UGameInstance*) {
 			// We only want to do this if we are using unique log names, otherwise it makes more sense
@@ -809,7 +824,9 @@ void FVisualLogger::OnDataReset()
 		const UWorld* EditorWorld = EditorEngine->GetEditorWorldContext().World();
 		if (EditorWorld)
 		{
-			EditorOnly::EditorBaseTimeStamp = EditorWorld->TimeSeconds;
+			// Reset the base timestamp to zero so the next (aka first) log entry will set this variable to determine the global offset in GetTimeStampForObject
+			// This ensures that when you start recording, the first entry appears to be at 0.0 (computed as TimeStamp - EditorBaseTimeStamp)
+			EditorOnly::EditorBaseTimeStamp = 0.0;
 		}
 	}
 #endif
diff --git a/Engine/Source/Runtime/Engine/Public/VisualLogger/VisualLogger.h b/Engine/Source/Runtime/Engine/Public/VisualLogger/VisualLogger.h
--- a/Engine/Source/Runtime/Engine/Public/VisualLogger/VisualLogger.h
+++ b/Engine/Source/Runtime/Engine/Public/VisualLogger/VisualLogger.h
@@ -742,7 +742,13 @@ public:
 	/** Returns  current entry for given TimeStamp or creates another one  but first it serialize previous 
 	 *	entry as completed to vislog devices. Use VisualLogger::DontCreate to get current entry without serialization
 	 *	@note this function can return null */
+	UE_DEPRECATED_FORGAME(5.4, "Use the static GetEntryToWrite instead because this TimeStamp is inconsistent across multiple instances (or threads in Editor).  This function will be made private/protected.")
 	ENGINE_API FVisualLogEntry* GetEntryToWrite(const UObject* Object, double TimeStamp, ECreateIfNeeded ShouldCreate = ECreateIfNeeded::Create);
+	/** Returns  the current (or new) entry for the given object; alternatively nullptr if we aren't allowed to vlog with the given parameters.
+	 * @param LogOwner - The UObject (typically an AActor) we are going to write log entries about.  This becomes the row in the Visual Logger timeline.
+	 * @param LogCategory - The LogCategory we are logging about.  This function will only return a valid log entry if the visual logging is enabled for the category.
+	 */
+	[[nodiscard]] static ENGINE_API FVisualLogEntry* GetEntryToWrite(const UObject* LogOwner, const FLogCategoryBase& LogCategory);
 	/** Retrieves last used entry for given UObject
 	 *	@note this function can return null */
 	ENGINE_API FVisualLogEntry* GetLastEntryForObject(const UObject* Object);
