Chaos
- initial overlap depenetration default value set to -1 which means infinity (was 0)
- fix initial overlap velocity for dynamic versus static if the dynamic has no override
- add PrimitiveComponent blueprint methods to set the max depenetration velocity
- add more coments to describe the feature

When the new system was added, the default Depentration Velocity was set to 0 which is probably what is wanted most of the time, but would lead to a behaviour change in existing projects. This change puts things back to how they were unless the project setting CollisionInitialOverlapDepenetrationVelocity or the Body setting MaxDepenetrationVelocity is changed.

diff --git a/Engine/Source/Runtime/Engine/Classes/Components/PrimitiveComponent.h b/Engine/Source/Runtime/Engine/Classes/Components/PrimitiveComponent.h
--- a/Engine/Source/Runtime/Engine/Classes/Components/PrimitiveComponent.h
+++ b/Engine/Source/Runtime/Engine/Classes/Components/PrimitiveComponent.h
@@ -2682,6 +2682,24 @@ public:
 	/** Returns the calculated mass in kg. This is not 100% exactly the mass physx will calculate, but it is very close ( difference < 0.1kg ). */
 	ENGINE_API virtual float CalculateMass(FName BoneName = NAME_None);
 
+	/**
+	 * The maximum velocity used to depenetrate this object from others when spawned or teleported with initial overlaps (does not affect overlaps as a result of normal movement).
+	 * A value of zero will allow objects that are spawned overlapping to go to sleep without moving rather than pop out of each other. E.g., use zero if you spawn dynamic rocks
+	 * partially embedded in the ground and want them to be interactive but not pop out of the ground when touched.
+	 * A negative value means that the config setting CollisionInitialOverlapDepenetrationVelocity will be used.
+	*/
+	UFUNCTION(BlueprintCallable, Category = "Physics")
+	ENGINE_API virtual float GetMaxDepenetrationVelocity(FName BoneName = NAME_None);
+
+	/**
+	 * The maximum velocity used to depenetrate this object from others when spawned or teleported with initial overlaps (does not affect overlaps as a result of normal movement).
+	 * A value of zero will allow objects that are spawned overlapping to go to sleep without moving rather than pop out of each other. E.g., use zero if you spawn dynamic rocks
+	 * partially embedded in the ground and want them to be interactive but not pop out of the ground when touched.
+	 * A negative value means that the config setting CollisionInitialOverlapDepenetrationVelocity will be used.
+	*/
+	UFUNCTION(BlueprintCallable, Category = "Physics")
+	ENGINE_API virtual void SetMaxDepenetrationVelocity(FName BoneName = NAME_None, float InMaxDepenetrationVelocity = -1.0f);
+
 	/** Set whether this component should use Continuous Collision Detection */
 	UFUNCTION(BlueprintCallable, Category = "Physics")
 	ENGINE_API virtual void SetUseCCD(bool InUseCCD, FName BoneName = NAME_None);
diff --git a/Engine/Source/Runtime/Engine/Classes/PhysicsEngine/BodyInstance.h b/Engine/Source/Runtime/Engine/Classes/PhysicsEngine/BodyInstance.h
--- a/Engine/Source/Runtime/Engine/Classes/PhysicsEngine/BodyInstance.h
+++ b/Engine/Source/Runtime/Engine/Classes/PhysicsEngine/BodyInstance.h
@@ -514,10 +514,11 @@ private:
 
 protected:
 	/** 
-	 * The maximum velocity used to depenetrate this object from others when spawned with initial overlaps or teleports (does not affect overlaps as a result of normal movement).
-	 * A value of zero will allow objects that are spawned overlapping to go to sleep as they are rather than pop out of each other.
+	 * The maximum velocity used to depenetrate this object from others when spawned or teleported with initial overlaps (does not affect overlaps as a result of normal movement).
+	 * A value of zero will allow objects that are spawned overlapping to go to sleep without moving rather than pop out of each other. E.g., use zero if you spawn dynamic rocks 
+	 * partially embedded in the ground and want them to be interactive but not pop out of the ground when touched.
 	 * A negative value is equivalent to bOverrideMaxDepenetrationVelocity = false, meaning use the project setting.
-	 * When set, this overrides the CollisionInitialOverlapDepenetrationVelocity project setting (and not the MaxDepenetrationVelocity solver setting that will be deprecated)
+	 * This overrides the CollisionInitialOverlapDepenetrationVelocity project setting on a per-body basis (and not the MaxDepenetrationVelocity solver setting that will be deprecated).
 	*/
 	UPROPERTY(EditAnywhere, AdvancedDisplay, BlueprintReadOnly, Category = Physics, meta = (editcondition = "bOverrideMaxDepenetrationVelocity", ClampMin = "0.0", UIMin = "0.0"))
 	float MaxDepenetrationVelocity;
diff --git a/Engine/Source/Runtime/Engine/Private/PrimitiveComponentPhysics.cpp b/Engine/Source/Runtime/Engine/Private/PrimitiveComponentPhysics.cpp
--- a/Engine/Source/Runtime/Engine/Private/PrimitiveComponentPhysics.cpp
+++ b/Engine/Source/Runtime/Engine/Private/PrimitiveComponentPhysics.cpp
@@ -614,6 +614,29 @@ float UPrimitiveComponent::CalculateMass(FName)
 	return 0.0f;
 }
 
+float UPrimitiveComponent::GetMaxDepenetrationVelocity(FName BoneName)
+{
+	FBodyInstance* BI = GetBodyInstance(BoneName);
+	if (BI)
+	{
+		return BI->GetMaxDepenetrationVelocity();
+	}
+	// @todo: add PhysicsObject support
+
+	// Negative means the config default value will be used
+	return -1.0f;
+}
+
+void UPrimitiveComponent::SetMaxDepenetrationVelocity(FName BoneName, float InMaxDepenetrationVelocity)
+{
+	FBodyInstance* BI = GetBodyInstance(BoneName);
+	if (BI)
+	{
+		BI->SetMaxDepenetrationVelocity(InMaxDepenetrationVelocity);
+	}
+	// @todo: add PhysicsObject support
+}
+
 void UPrimitiveComponent::SetUseCCD(bool bInUseCCD, FName BoneName)
 {
 	FBodyInstance* BI = GetBodyInstance(BoneName);
diff --git a/Engine/Source/Runtime/Experimental/Chaos/Private/Chaos/Collision/PBDCollisionConstraint.cpp b/Engine/Source/Runtime/Experimental/Chaos/Private/Chaos/Collision/PBDCollisionConstraint.cpp
--- a/Engine/Source/Runtime/Experimental/Chaos/Private/Chaos/Collision/PBDCollisionConstraint.cpp
+++ b/Engine/Source/Runtime/Experimental/Chaos/Private/Chaos/Collision/PBDCollisionConstraint.cpp
@@ -221,7 +221,7 @@ namespace Chaos
 		, SavedManifoldPoints()
 		, ManifoldPoints()
 		, MinInitialPhi(0)
-		, InitialOverlapDepenetrationVelocity(0)
+		, InitialOverlapDepenetrationVelocity(-1.0f)
 		, CCDTimeOfImpact(0)
 		, SolverBodies{ nullptr, nullptr }
 		, CCDEnablePenetration(0)
@@ -262,7 +262,7 @@ namespace Chaos
 		, SavedManifoldPoints()
 		, ManifoldPoints()
 		, MinInitialPhi(0)
-		, InitialOverlapDepenetrationVelocity(0)
+		, InitialOverlapDepenetrationVelocity(-1.0f)
 		, CCDTimeOfImpact(0)
 		, SolverBodies{ nullptr, nullptr }
 		, CCDEnablePenetration(0)
diff --git a/Engine/Source/Runtime/Experimental/Chaos/Private/Chaos/Collision/PBDCollisionContainerSolver.cpp b/Engine/Source/Runtime/Experimental/Chaos/Private/Chaos/Collision/PBDCollisionContainerSolver.cpp
--- a/Engine/Source/Runtime/Experimental/Chaos/Private/Chaos/Collision/PBDCollisionContainerSolver.cpp
+++ b/Engine/Source/Runtime/Experimental/Chaos/Private/Chaos/Collision/PBDCollisionContainerSolver.cpp
@@ -76,7 +76,7 @@ namespace Chaos
 		FPBDCollisionConstraint* Constraint, 
 		const int32 ConstraintPointIndex, 
 		const FRealSingle Dt, 
-		const FRealSingle MaxDepentrationVelocity,
+		const FRealSingle MaxDepenetrationVelocity,
 		const FRealSingle MaxPushOut,
 		const FConstraintSolverBody& Body0,
 		const FConstraintSolverBody& Body1)
@@ -160,7 +160,7 @@ namespace Chaos
 		// Initial Phi for initial-overlap depenetration.
 		// If we have an initial contact, calculate the initial overlap. This will get saved in SetSolverResults
 		FRealSingle WorldContactInitialPhi = 0;
-		if ((MaxDepentrationVelocity >= 0) && CVars::bChaos_Collision_EnableInitialDepenetration)
+		if ((MaxDepenetrationVelocity >= 0) && CVars::bChaos_Collision_EnableInitialDepenetration)
 		{
 			if (ManifoldPoint.Flags.bInitialContact)
 			{
@@ -183,8 +183,8 @@ namespace Chaos
 			{
 				// This is a pre-existing manifold point with some initial penetration to resolve.
 				// If we are currently penetrating less than the inital overlap, reduce the initial overlap
-				// Also resolve initial overlap over time by reducing allowed penetration by MaxDepentrationVelocity
-				WorldContactInitialPhi = FMath::Max(ManifoldPoint.InitialPhi + MaxDepentrationVelocity * Dt, WorldContactDeltaNormal);
+				// Also resolve initial overlap over time by reducing allowed penetration by MaxDepenetrationVelocity
+				WorldContactInitialPhi = FMath::Max(ManifoldPoint.InitialPhi + MaxDepenetrationVelocity * Dt, WorldContactDeltaNormal);
 			}
 
 			// InitialPhi is only for tracking penetration - cannot be positive
diff --git a/Engine/Source/Runtime/Experimental/Chaos/Public/Chaos/Collision/PBDCollisionSolverSettings.h b/Engine/Source/Runtime/Experimental/Chaos/Public/Chaos/Collision/PBDCollisionSolverSettings.h
--- a/Engine/Source/Runtime/Experimental/Chaos/Public/Chaos/Collision/PBDCollisionSolverSettings.h
+++ b/Engine/Source/Runtime/Experimental/Chaos/Public/Chaos/Collision/PBDCollisionSolverSettings.h
@@ -28,7 +28,8 @@ namespace Chaos
 		// A value of zero means unlimited.
 		FReal MaxPushOutVelocity;
 
-		// The speed at which initially-overlapping objects depentrate
+		// The speed at which initially-overlapping objects depentrate.
+		// This value is used when particles do not provide an override for their MaxDepenetrationVelocity.
 		FRealSingle DepenetrationVelocity;
 
 		// How many of the position iterations should run static/dynamic friction
diff --git a/Engine/Source/Runtime/Experimental/Chaos/Public/Chaos/PBDCollisionConstraints.h b/Engine/Source/Runtime/Experimental/Chaos/Public/Chaos/PBDCollisionConstraints.h
--- a/Engine/Source/Runtime/Experimental/Chaos/Public/Chaos/PBDCollisionConstraints.h
+++ b/Engine/Source/Runtime/Experimental/Chaos/Public/Chaos/PBDCollisionConstraints.h
@@ -238,7 +238,19 @@ public:
 
 	void SetDepenetrationVelocity(const FRealSingle InVel)
 	{
-		SolverSettings.DepenetrationVelocity = FMath::Max(InVel, FRealSingle(0));
+		// The user can specify a very large number up to float_max or any negative number to mean "infinity".
+		// However, we don't use float_max for infinity because we want to be able to perform simple math on
+		// it without numeric limit issues (search MaxDepenetrationVelocity).
+		constexpr float MaxDepenetrationVelocity = 1e10f;	// [cm/s] almost speed of light :)
+
+		if (InVel >= 0.0f)
+		{
+			SolverSettings.DepenetrationVelocity = FMath::Min(InVel, MaxDepenetrationVelocity);
+		}
+		else
+		{
+			SolverSettings.DepenetrationVelocity = MaxDepenetrationVelocity;
+		}
 	}
 
 	void SetPositionFrictionIterations(const int32 InNumIterations)
diff --git a/Engine/Source/Runtime/Experimental/Chaos/Public/Chaos/PBDRigidsEvolutionGBF.h b/Engine/Source/Runtime/Experimental/Chaos/Public/Chaos/PBDRigidsEvolutionGBF.h
--- a/Engine/Source/Runtime/Experimental/Chaos/Public/Chaos/PBDRigidsEvolutionGBF.h
+++ b/Engine/Source/Runtime/Experimental/Chaos/Public/Chaos/PBDRigidsEvolutionGBF.h
@@ -64,7 +64,7 @@ namespace Chaos
 		static constexpr FRealSingle DefaultCollisionMarginMax = 10.0f;
 		static constexpr FRealSingle DefaultCollisionCullDistance = 3.0f;
 		static constexpr FRealSingle DefaultCollisionMaxPushOutVelocity = 1000.0f;
-		static constexpr FRealSingle DefaultCollisionDepenetrationVelocity = 0.0f;
+		static constexpr FRealSingle DefaultCollisionDepenetrationVelocity = -1.0f;
 		static constexpr int32 DefaultRestitutionThreshold = 1000;
 
 		CHAOS_API FPBDRigidsEvolutionGBF(
diff --git a/Engine/Source/Runtime/Experimental/Chaos/Public/Chaos/ParticleHandle.h b/Engine/Source/Runtime/Experimental/Chaos/Public/Chaos/ParticleHandle.h
--- a/Engine/Source/Runtime/Experimental/Chaos/Public/Chaos/ParticleHandle.h
+++ b/Engine/Source/Runtime/Experimental/Chaos/Public/Chaos/ParticleHandle.h
@@ -2398,7 +2398,7 @@ public:
 		{
 			return RigidHandle->InitialOverlapDepenetrationVelocity();
 		}
-		return 0;
+		return -1.0f;	// Use config
 	}
 
 	FRealSingle SleepThresholdMultiplier() const
diff --git a/Engine/Source/Runtime/Experimental/Chaos/Public/ChaosSolverConfiguration.h b/Engine/Source/Runtime/Experimental/Chaos/Public/ChaosSolverConfiguration.h
--- a/Engine/Source/Runtime/Experimental/Chaos/Public/ChaosSolverConfiguration.h
+++ b/Engine/Source/Runtime/Experimental/Chaos/Public/ChaosSolverConfiguration.h
@@ -63,10 +63,13 @@ struct FChaosSolverConfiguration
 	UPROPERTY(EditAnywhere, Category = "SolverConfiguration|Collision", meta = (ClampMin = "0.0"))
 	float CollisionMaxPushOutVelocity;
 
-	// If two bodies start off in overaapping each other, they will depentrate at this speed when they wake.
-	// NOTE: this is the setting for the new Initial Penetration Handling system. This is currently enabled with the 
-	// cvar p.Chaos.PBDCollisionSolver.EnableInitialDepenetration but will eventually be on by default.
-	UPROPERTY(EditAnywhere, Category = "SolverConfiguration|Collision", meta = (ClampMin = "0.0"))
+	// If two bodies start off in overlapping each other, they will depentrate at this speed when they wake.
+	// If set to a large value, initially-overlapping objects will tend to "explode" apart at a speed that depends on the
+	// overlap amount and the timestep (this is the original, previously untunable behaviour). If set to zero, 
+	// initially-overlapping objects will remain stationary and go to sleep until acted on by some other object or force.
+	// A negative value (-1) disables the feature and is equivalent to infinity.
+	// This property can be overridden per Body (see FBodyInstance::MaxDepenetrationVelocity)
+	UPROPERTY(EditAnywhere, Category = "SolverConfiguration|Collision")
 	float CollisionInitialOverlapDepenetrationVelocity;
 
 	UPROPERTY(EditAnywhere, Category = "SolverConfiguration|Clustering")
diff --git a/Engine/Source/Runtime/Experimental/ChaosVisualDebugger/Public/DataWrappers/ChaosVDCollisionDataWrappers.h b/Engine/Source/Runtime/Experimental/ChaosVisualDebugger/Public/DataWrappers/ChaosVDCollisionDataWrappers.h
--- a/Engine/Source/Runtime/Experimental/ChaosVisualDebugger/Public/DataWrappers/ChaosVDCollisionDataWrappers.h
+++ b/Engine/Source/Runtime/Experimental/ChaosVisualDebugger/Public/DataWrappers/ChaosVDCollisionDataWrappers.h
@@ -320,7 +320,7 @@ struct CHAOSVDRUNTIME_API FChaosVDConstraint
 	float MinInitialPhi = 0.f;
 
 	UPROPERTY(VisibleAnywhere, Category = ConstraintData)
-	float InitialOverlapDepenetrationVelocity = 0.f;
+	float InitialOverlapDepenetrationVelocity = -1.f;
 	
 	UPROPERTY(VisibleAnywhere, Category=ConstraintData)
 	float CCDTimeOfImpact = 0.f;
