[NetPhysics] Fix crash from PopulateReplicationCache() where it accessed a pointer to an object in garbage collection.

diff --git a/Engine/Source/Runtime/Engine/Public/Physics/Experimental/PhysScene_Chaos.h b/Engine/Source/Runtime/Engine/Public/Physics/Experimental/PhysScene_Chaos.h
--- a/Engine/Source/Runtime/Engine/Public/Physics/Experimental/PhysScene_Chaos.h
+++ b/Engine/Source/Runtime/Engine/Public/Physics/Experimental/PhysScene_Chaos.h
@@ -266,7 +266,7 @@ public:
 			, bValidStateCached(false)
 		{}
 
-		TObjectPtr<UPrimitiveComponent> GetRootComponent()	{ return RootComponent; }
+		UPrimitiveComponent* GetRootComponent()	{ return RootComponent.Get(); }
 		FRigidBodyState& GetState()	{ return State; }
 		void SetAccessTime(Chaos::FReal Time) { AccessTime = Time; }
 		Chaos::FReal GetAccessTime() { return AccessTime; }
@@ -274,7 +274,7 @@ public:
 		bool IsCached() { return bValidStateCached; }
 
 	private:
-		TObjectPtr<UPrimitiveComponent> RootComponent;
+		TWeakObjectPtr<UPrimitiveComponent> RootComponent;
 		Chaos::FReal AccessTime;
 		bool bValidStateCached;
 		FRigidBodyState State;
