UE-206314 - Crash preloading InheritableComponentHandler with editor async loading

We're hitting this assert because UInheritableComponentHandler::PreloadAll is calling ConditionalPostLoad. However, it's calling this function because it's attempting execute the fixup for component template names. This can be easily worked around if we simply extract the body of UInheritableComponentHandler::PostLoad and put it in a new function, FixComponentTemplateNames, that both PostLoad and PreloadAll can call.

#rnx

diff --git a/Engine/Source/Runtime/Engine/Classes/Engine/InheritableComponentHandler.h b/Engine/Source/Runtime/Engine/Classes/Engine/InheritableComponentHandler.h
--- a/Engine/Source/Runtime/Engine/Classes/Engine/InheritableComponentHandler.h
+++ b/Engine/Source/Runtime/Engine/Classes/Engine/InheritableComponentHandler.h
@@ -126,7 +126,6 @@ public:
 
 	//~ Begin UObject Interface
 	ENGINE_API virtual void Serialize(FArchive& Ar) override;
-	ENGINE_API virtual void PostLoad() override;
 #if WITH_EDITOR
 	ENGINE_API virtual EDataValidationResult IsDataValid(class FDataValidationContext& Context) const override;
 #endif	// WITH_EDITOR
@@ -165,6 +164,8 @@ private:
 	/** Helper method used to assist with fixing up component template names at load time. */
 	ENGINE_API void FixComponentTemplateName(UActorComponent* ComponentTemplate, FName NewName);
 	
+	void FixComponentTemplateNames();
+
 	/** All component records */
 	UPROPERTY()
 	TArray<FComponentOverrideRecord> Records;
diff --git a/Engine/Source/Runtime/Engine/Private/InheritableComponentHandler.cpp b/Engine/Source/Runtime/Engine/Private/InheritableComponentHandler.cpp
--- a/Engine/Source/Runtime/Engine/Private/InheritableComponentHandler.cpp
+++ b/Engine/Source/Runtime/Engine/Private/InheritableComponentHandler.cpp
@@ -25,10 +25,8 @@ void UInheritableComponentHandler::Serialize(FArchive& Ar)
 	Ar.UsingCustomVersion(FBlueprintsObjectVersion::GUID);
 }
 
-void UInheritableComponentHandler::PostLoad()
+void UInheritableComponentHandler::FixComponentTemplateNames()
 {
-	Super::PostLoad();
-	
 #if WITH_EDITOR
 	if (!GIsDuplicatingClassForReinstancing)
 	{
@@ -514,7 +512,7 @@ void UInheritableComponentHandler::PreloadAll()
 	// this are not clear to me, but that the logic exists and occasionally
 	// runs means we should run it ASAP - otherwise we may use the wrong
 	// archetype on construction.
-	ConditionalPostLoad();
+	FixComponentTemplateNames();
 }
 
 FComponentKey UInheritableComponentHandler::FindKey(const FName VariableName) const
