GPU skin cache:  Fix for crash when recreating FSkeletalMeshObjectGPUSkin render state when ray tracing is enabled.  Need to update the GPUSkin pointer in the skin cache entry immediately when moving the entry to the new mesh object, to avoid issues where PendingProcessRTGeometryEntries are processed using a stale pointer to the previous deleted mesh object.  Also need to properly delete the ray tracing GPU skin cache entry if it's no longer needed due to a change in the configuration of the mesh object during recreate.

#rnx

diff --git a/Engine/Source/Runtime/Engine/Private/GPUSkinCache.cpp b/Engine/Source/Runtime/Engine/Private/GPUSkinCache.cpp
--- a/Engine/Source/Runtime/Engine/Private/GPUSkinCache.cpp
+++ b/Engine/Source/Runtime/Engine/Private/GPUSkinCache.cpp
@@ -1726,12 +1726,6 @@ bool FGPUSkinCache::ProcessEntry(
 
 	if (InOutEntry)
 	{
-		// If re-creating, the FSkeletalMeshObjectGPUSkin will have changed, update it here so the "IsValid" call below succeeds
-		if (bRecreating)
-		{
-			InOutEntry->GPUSkin = Skin;
-		}
-
 		// If the LOD changed, the entry has to be invalidated
 		if (!InOutEntry->IsValid(Skin, LODIndex))
 		{
@@ -2328,6 +2322,14 @@ void FGPUSkinCache::UpdateSkinWeightBuffer(FGPUSkinCacheEntry* Entry)
 	}
 }
 
+void FGPUSkinCache::SetEntryGPUSkin(FGPUSkinCacheEntry* Entry, FSkeletalMeshObjectGPUSkin* Skin)
+{
+	if (Entry)
+	{
+		Entry->GPUSkin = Skin;
+	}
+}
+
 void FGPUSkinCache::CVarSinkFunction()
 {
 	int32 NewGPUSkinCacheValue = CVarEnableGPUSkinCache.GetValueOnAnyThread() != 0;
diff --git a/Engine/Source/Runtime/Engine/Private/SkeletalRenderGPUSkin.cpp b/Engine/Source/Runtime/Engine/Private/SkeletalRenderGPUSkin.cpp
--- a/Engine/Source/Runtime/Engine/Private/SkeletalRenderGPUSkin.cpp
+++ b/Engine/Source/Runtime/Engine/Private/SkeletalRenderGPUSkin.cpp
@@ -261,6 +261,10 @@ FSkeletalMeshObjectGPUSkin::FSkeletalMeshObjectGPUSkin(USkinnedMeshComponent* In
 				SkinCacheEntry = PreviousMeshObject->SkinCacheEntry;
 				SkinCacheEntryForRayTracing = PreviousMeshObject->SkinCacheEntryForRayTracing;
 
+				// patch entries to point to new GPUSkin
+				FGPUSkinCache::SetEntryGPUSkin(SkinCacheEntry, this);
+				FGPUSkinCache::SetEntryGPUSkin(SkinCacheEntryForRayTracing, this);
+
 				PreviousMeshObject->SkinCacheEntry = nullptr;
 				PreviousMeshObject->SkinCacheEntryForRayTracing = nullptr;
 			}
@@ -356,7 +360,7 @@ void FSkeletalMeshObjectGPUSkin::ReleaseResources()
 		BeginReleaseResource(&RayTracingGeometry, &UE::RenderCommandPipe::SkeletalMesh);
 	}
 
-	// Only enqueue when intialized
+	// Only enqueue when initialized
 	if (RayTracingUpdateQueue != nullptr || RayTracingDynamicVertexBuffer.NumBytes > 0)
 	{
 		ENQUEUE_RENDER_COMMAND(ReleaseRayTracingDynamicVertexBuffer)(UE::RenderCommandPipe::SkeletalMesh,
@@ -636,9 +640,27 @@ void FSkeletalMeshObjectGPUSkin::ProcessUpdatedDynamicData(EGPUSkinCacheEntryMod
 
 	// Immediately release any stale entry if we've recently switched to a LOD level that disallows skin cache
 	// This saves memory and avoids confusing ShouldUseSeparateSkinCacheEntryForRayTracing() which checks SkinCacheEntry == nullptr
-	if (!bGPUSkinCacheEnabled && SkinCacheEntry)
+	if (!bGPUSkinCacheEnabled)
 	{
-		FGPUSkinCache::Release(SkinCacheEntry);
+#if RHI_RAYTRACING
+		if (Mode == EGPUSkinCacheEntryMode::Raster)
+#endif
+		{
+			if (SkinCacheEntry)
+			{
+				FGPUSkinCache::Release(SkinCacheEntry);
+			}
+		}
+#if RHI_RAYTRACING
+		else
+		{
+			check(Mode == EGPUSkinCacheEntryMode::RayTracing);
+			if (SkinCacheEntryForRayTracing)
+			{
+				FGPUSkinCache::Release(SkinCacheEntryForRayTracing);
+			}
+		}
+#endif
 	}
 
 	// We need to clear the external morph buffers when the weights are zero.
diff --git a/Engine/Source/Runtime/Engine/Public/GPUSkinCache.h b/Engine/Source/Runtime/Engine/Public/GPUSkinCache.h
--- a/Engine/Source/Runtime/Engine/Public/GPUSkinCache.h
+++ b/Engine/Source/Runtime/Engine/Public/GPUSkinCache.h
@@ -129,6 +129,7 @@ public:
 	ENGINE_API ~FGPUSkinCache();
 
 	static void UpdateSkinWeightBuffer(FGPUSkinCacheEntry* Entry);
+	static void SetEntryGPUSkin(FGPUSkinCacheEntry* Entry, FSkeletalMeshObjectGPUSkin* Skin);
 
 	bool ProcessEntry(
 		EGPUSkinCacheEntryMode Mode,
