PR #7692: Navigation isn't correctly updating on a mesh component when SetStaticMesh is called if collision is disabled

diff --git a/Engine/Source/Runtime/Engine/Private/Components/StaticMeshComponent.cpp b/Engine/Source/Runtime/Engine/Private/Components/StaticMeshComponent.cpp
--- a/Engine/Source/Runtime/Engine/Private/Components/StaticMeshComponent.cpp
+++ b/Engine/Source/Runtime/Engine/Private/Components/StaticMeshComponent.cpp
@@ -2331,6 +2331,8 @@ bool UStaticMeshComponent::SetStaticMesh(UStaticMesh* NewMesh)
 	// Since we have new mesh, we need to update bounds
 	UpdateBounds();
 
+	FNavigationSystem::UpdateComponentData(*this);
+
 	// Mark cached material parameter names dirty
 	MarkCachedMaterialParameterNameIndicesDirty();
 
