Skelmesh: Use the alt skin import data to reconstruct alt skin on mesh description. In some cases the narrowed skin weights don't seem to map back correctly.

#rnx

diff --git a/Engine/Source/Runtime/Core/Private/UObject/UE5ReleaseStreamObjectVersion.cpp b/Engine/Source/Runtime/Core/Private/UObject/UE5ReleaseStreamObjectVersion.cpp
--- a/Engine/Source/Runtime/Core/Private/UObject/UE5ReleaseStreamObjectVersion.cpp
+++ b/Engine/Source/Runtime/Core/Private/UObject/UE5ReleaseStreamObjectVersion.cpp
@@ -13,7 +13,7 @@ TMap<FGuid, FGuid> FUE5ReleaseStreamObjectVersion::GetSystemGuids()
 	SystemGuids.Add(DevGuids.NANITE_DERIVEDDATA_VER, FGuid("AC88CFBDDA614A9C8488F64D2DAF699D"));
 	SystemGuids.Add(DevGuids.NIAGARASHADERMAP_DERIVEDDATA_VER, FGuid("6BE35B6FACB34568970120F9BC8DAB80"));
 	SystemGuids.Add(DevGuids.Niagara_LatestScriptCompileVersion, FGuid("43A777C54EB24E9894525B8D04529F23"));
-	SystemGuids.Add(DevGuids.SkeletalMeshDerivedDataVersion, FGuid("3DEA297061DA4BB991A649D0B2F08F18"));
+	SystemGuids.Add(DevGuids.SkeletalMeshDerivedDataVersion, FGuid("BA91E71F642E4813A1B74B80D0448928"));
 	SystemGuids.Add(DevGuids.STATICMESH_DERIVEDDATA_VER, FGuid("7A556175518D458AA8786ED69AB8DDE7"));
 
 	return SystemGuids;
diff --git a/Engine/Source/Runtime/Engine/Private/SkeletalMeshLODModel.cpp b/Engine/Source/Runtime/Engine/Private/SkeletalMeshLODModel.cpp
--- a/Engine/Source/Runtime/Engine/Private/SkeletalMeshLODModel.cpp
+++ b/Engine/Source/Runtime/Engine/Private/SkeletalMeshLODModel.cpp
@@ -1580,6 +1580,38 @@ void FSkeletalMeshLODModel::GetMeshDescription(const USkeletalMesh *InSkeletalMe
 		FSkinWeightsVertexAttributesRef SkinWeightAttribute = MeshAttributes.GetVertexSkinWeights(SkinWeightProfileInfo.Key);
 		const FImportedSkinWeightProfileData& SkinWeightProfileData = SkinWeightProfileInfo.Value;
 
+#if 1
+		TMultiMap<int32, int32> InfluenceMap;
+		for (int32 Index = 0; Index < SkinWeightProfileData.SourceModelInfluences.Num(); Index++)
+		{
+			const SkeletalMeshImportData::FVertInfluence& Influence = SkinWeightProfileData.SourceModelInfluences[Index];
+			InfluenceMap.Add(Influence.VertIndex, Index);
+		}
+
+		TArray<FBoneIndexType> BoneIndexes;
+		TArray<float> BoneWeights;
+		TArray<int32> InfluenceIndexes;
+		for (int32 Index = 0; Index < VertexIDs.Num(); Index++)
+		{
+			InfluenceIndexes.Reset();
+			InfluenceMap.MultiFind(Index, InfluenceIndexes);
+
+			if (!InfluenceIndexes.IsEmpty())
+			{
+				BoneIndexes.Reset();
+				BoneWeights.Reset();
+				for (int32 InfluenceIndex: InfluenceIndexes)
+				{
+					const SkeletalMeshImportData::FVertInfluence& Influence = SkinWeightProfileData.SourceModelInfluences[InfluenceIndex];
+					BoneIndexes.Add(Influence.BoneIndex);
+					BoneWeights.Add(Influence.Weight);
+				}
+
+				FBoneWeights Weights = FBoneWeights::Create(BoneIndexes.GetData(), BoneWeights.GetData(), BoneIndexes.Num());
+				SkinWeightAttribute.Set(VertexIDs[Index], Weights);
+			}
+		}
+#else
 		check(SkinWeightProfileData.SkinWeights.Num() == NumVertices);
 		
 		for (int32 Index = 0; Index < SkinWeightProfileData.SkinWeights.Num(); Index++)
@@ -1591,6 +1623,7 @@ void FSkeletalMeshLODModel::GetMeshDescription(const USkeletalMesh *InSkeletalMe
 			FBoneWeights Weights = FBoneWeights::Create(RawSkinWeights.InfluenceBones, RawSkinWeights.InfluenceWeights);
 			SkinWeightAttribute.Set(VertexID, Weights);
 		}
+#endif
 	}
 	
 	// Set Bone Attributes
