Fix MapBuildData access for Brushes (the component isn't on an Actor so the new path failed)
Fix for automated tests "Lights_Static.IndirectLighting.Off/On"

diff --git a/Engine/Source/Runtime/Engine/Classes/Engine/MapBuildDataRegistry.h b/Engine/Source/Runtime/Engine/Classes/Engine/MapBuildDataRegistry.h
--- a/Engine/Source/Runtime/Engine/Classes/Engine/MapBuildDataRegistry.h
+++ b/Engine/Source/Runtime/Engine/Classes/Engine/MapBuildDataRegistry.h
@@ -389,6 +389,9 @@ public:
 
 	/** Returns the associated UUMapBuildDataRegistry for an Actor */
 	ENGINE_API static UMapBuildDataRegistry* Get(const AActor* Actor);
+
+	/** Returns the associated UUMapBuildDataRegistry for a pair of Level/World */
+	ENGINE_API static UMapBuildDataRegistry* Get(ULevel* OwnerLevel, UWorld* World);
 	
 	FVolumetricLightMapGridDesc* GetVolumetricLightMapGridDesc() { return VolumetricLightMapGridDesc; } 
 	ENGINE_API void SetVolumetricLightMapGridDesc(FVolumetricLightMapGridDesc* GridDesc); 
diff --git a/Engine/Source/Runtime/Engine/Private/Components/ModelComponent.cpp b/Engine/Source/Runtime/Engine/Private/Components/ModelComponent.cpp
--- a/Engine/Source/Runtime/Engine/Private/Components/ModelComponent.cpp
+++ b/Engine/Source/Runtime/Engine/Private/Components/ModelComponent.cpp
@@ -55,10 +55,14 @@ FModelElement::~FModelElement()
 const FMeshMapBuildData* FModelElement::GetMeshMapBuildData() const
 {
 	check(Component);
+	ULevel* OwnerLevel = Cast<ULevel>(Component->GetModel()->GetOuter());
 
-	if (UMapBuildDataRegistry* MapBuildData = UMapBuildDataRegistry::Get(Component))
+	if (OwnerLevel && OwnerLevel->OwningWorld)
 	{
-		return MapBuildData->GetMeshBuildData(MapBuildDataId);		
+		if (UMapBuildDataRegistry* MapBuildData = UMapBuildDataRegistry::Get(OwnerLevel, OwnerLevel->OwningWorld))
+		{
+			return MapBuildData->GetMeshBuildData(MapBuildDataId);		
+		}
 	}
 	
 	return NULL;
diff --git a/Engine/Source/Runtime/Engine/Private/MapBuildData.cpp b/Engine/Source/Runtime/Engine/Private/MapBuildData.cpp
--- a/Engine/Source/Runtime/Engine/Private/MapBuildData.cpp
+++ b/Engine/Source/Runtime/Engine/Private/MapBuildData.cpp
@@ -1217,17 +1217,22 @@ UMapBuildDataRegistry* UMapBuildDataRegistry::Get(const UActorComponent* Compone
 UMapBuildDataRegistry* UMapBuildDataRegistry::Get(const AActor* Actor)
 {
 	ULevel* OwnerLevel = Actor->GetLevel();
-	UWorld* World = OwnerLevel ? OwnerLevel->GetWorld() : nullptr;
-	UMapBuildDataRegistry* MapBuildData = nullptr;
+	UWorld* World = OwnerLevel ? OwnerLevel->GetWorld() : nullptr;	
 	
 	if (World && World->IsPartitionedWorld())
 	{
 		//@todo_ow: At current level of support there's no reason to return a ptr and force a look-up later on
 		//No lighting scenario support in WP maps
-		MapBuildData = World->PersistentLevel->MapBuildData;
-		return MapBuildData;
+		return World->PersistentLevel->MapBuildData;		
 	}
 
+	return Get(OwnerLevel, World);
+}
+
+UMapBuildDataRegistry* UMapBuildDataRegistry::Get(ULevel* OwnerLevel, UWorld* World)
+{
+	UMapBuildDataRegistry* MapBuildData = nullptr;
+
 	if (OwnerLevel && World)
 	{
 		ULevel* ActiveLightingScenario = World->GetActiveLightingScenario();
