[EditorPerf]
  - Enable shader job cache in -game

#rnx

diff --git a/Engine/Source/Runtime/Engine/Private/ShaderCompiler/ShaderCompiler.cpp b/Engine/Source/Runtime/Engine/Private/ShaderCompiler/ShaderCompiler.cpp
--- a/Engine/Source/Runtime/Engine/Private/ShaderCompiler/ShaderCompiler.cpp
+++ b/Engine/Source/Runtime/Engine/Private/ShaderCompiler/ShaderCompiler.cpp
@@ -248,7 +248,7 @@ bool IsShaderJobCacheDDCEnabled()
 #endif
 
 	// For now we only support the editor and not commandlets like the cooker.
-	if (GIsEditor && (!IsRunningCommandlet() || bForceAllowShaderCompilerJobCache))
+	if ((GIsEditor || IsRunningGame()) && (!IsRunningCommandlet() || bForceAllowShaderCompilerJobCache))
 	{
 		// job cache itself must be enabled first
 		return GShaderCompilerJobCache && CVarJobCacheDDC.GetValueOnAnyThread();
