Replace check with if in UActorDescContainerInstance::OnActorDescUpdating/Updated since UActorDescContainerInstance with an Instancing Context do not listen for OnActorDescAdded which means they might not find newly added instances when OnActorDescUpdating/Updated is called

#rnx

diff --git a/Engine/Source/Runtime/Engine/Private/WorldPartition/ActorDescContainerInstance.cpp b/Engine/Source/Runtime/Engine/Private/WorldPartition/ActorDescContainerInstance.cpp
--- a/Engine/Source/Runtime/Engine/Private/WorldPartition/ActorDescContainerInstance.cpp
+++ b/Engine/Source/Runtime/Engine/Private/WorldPartition/ActorDescContainerInstance.cpp
@@ -534,42 +534,50 @@ void UActorDescContainerInstance::OnActorDescRemoved(FWorldPartitionActorDesc* I
 
 void UActorDescContainerInstance::OnActorDescUpdating(FWorldPartitionActorDesc* InActorDesc)
 {
-	TUniquePtr<FWorldPartitionActorDescInstance>* ActorDescInstance = GetActorDescInstancePtr(InActorDesc->GetGuid());
-	check(ActorDescInstance && ActorDescInstance->IsValid());
-		
-	if (bCreateChildContainerHierarchy && ActorDescInstance->Get()->IsChildContainerInstance())
+	// It is possible for a container instance with an instancing context to listen to updates but not newly added actors
+	// so it might not find the actor if it was newly added (through the edit of a different instance of the same Level Instance)
+	// See UActorDescContainerInstance::RegisterDelegates
+	if (TUniquePtr<FWorldPartitionActorDescInstance>* ActorDescInstance = GetActorDescInstancePtr(InActorDesc->GetGuid()))
 	{
-		ActorDescInstance->Get()->UnregisterChildContainerInstance();
-	}
+		check(ActorDescInstance->IsValid());
 
-	if (UWorldPartition* WorldPartition = GetOuterWorldPartition())
-	{
-		WorldPartition->OnActorDescInstanceUpdating(ActorDescInstance->Get());
-	}
+		if (bCreateChildContainerHierarchy && ActorDescInstance->Get()->IsChildContainerInstance())
+		{
+			ActorDescInstance->Get()->UnregisterChildContainerInstance();
+		}
 
-	OnActorDescInstanceUpdatingEvent.Broadcast(ActorDescInstance->Get());
+		if (UWorldPartition* WorldPartition = GetOuterWorldPartition())
+		{
+			WorldPartition->OnActorDescInstanceUpdating(ActorDescInstance->Get());
+		}
+
+		OnActorDescInstanceUpdatingEvent.Broadcast(ActorDescInstance->Get());
+	}
 }
 
 void UActorDescContainerInstance::OnActorDescUpdated(FWorldPartitionActorDesc* InActorDesc)
 {
-	TUniquePtr<FWorldPartitionActorDescInstance>* ActorDescInstance = GetActorDescInstancePtr(InActorDesc->GetGuid());
-	check(ActorDescInstance && ActorDescInstance->IsValid());
-	
-	// Update instance desc
-	ActorDescInstance->Get()->UpdateActorDesc(InActorDesc);
-
-	// Re-register container
-	if (bCreateChildContainerHierarchy && ActorDescInstance->Get()->IsChildContainerInstance())
+	// See comment in UActorDescContainerInstance::OnActorDescUpdating
+	if (TUniquePtr<FWorldPartitionActorDescInstance>* ActorDescInstance = GetActorDescInstancePtr(InActorDesc->GetGuid()))
 	{
-		ActorDescInstance->Get()->RegisterChildContainerInstance();
-	}
+		check(ActorDescInstance->IsValid());
 
-	if (UWorldPartition* WorldPartition = GetOuterWorldPartition())
-	{
-		WorldPartition->OnActorDescInstanceUpdated(ActorDescInstance->Get());
-	}
+		// Update instance desc
+		ActorDescInstance->Get()->UpdateActorDesc(InActorDesc);
 
-	OnActorDescInstanceUpdatedEvent.Broadcast(ActorDescInstance->Get());
+		// Re-register container
+		if (bCreateChildContainerHierarchy && ActorDescInstance->Get()->IsChildContainerInstance())
+		{
+			ActorDescInstance->Get()->RegisterChildContainerInstance();
+		}
+
+		if (UWorldPartition* WorldPartition = GetOuterWorldPartition())
+		{
+			WorldPartition->OnActorDescInstanceUpdated(ActorDescInstance->Get());
+		}
+
+		OnActorDescInstanceUpdatedEvent.Broadcast(ActorDescInstance->Get());
+	}
 }
 
 #endif
