Removed unused fields and accessors in FInstancedStaticMeshSceneProxyDesc. They were added during development but subsequently made redundant in later changes. FInstancedStaticMeshSceneProxyDesc is new in 5.4

#rnx

diff --git a/Engine/Source/Runtime/Engine/Private/InstancedStaticMesh.cpp b/Engine/Source/Runtime/Engine/Private/InstancedStaticMesh.cpp
--- a/Engine/Source/Runtime/Engine/Private/InstancedStaticMesh.cpp
+++ b/Engine/Source/Runtime/Engine/Private/InstancedStaticMesh.cpp
@@ -1297,21 +1297,15 @@ void FInstancedStaticMeshSceneProxyDesc::InitializeFrom(UInstancedStaticMeshComp
 {
 	FStaticMeshSceneProxyDesc::InitializeFrom(InComponent);
 
-	PerInstanceSMData = InComponent->PerInstanceSMData;
 	InstanceDataSceneProxy = InComponent->GetOrCreateInstanceDataSceneProxy();
-	PerInstanceSMCustomData = InComponent->PerInstanceSMCustomData;
 #if WITH_EDITOR
 	SelectedInstances = InComponent->SelectedInstances;
 #endif
-	InstanceReorderTable = InComponent->InstanceReorderTable; 
-	PerInstancePrevTransform = InComponent->PerInstancePrevTransform;
 
 	InstanceStartCullDistance = InComponent->InstanceStartCullDistance ;
 	InstanceEndCullDistance = InComponent->InstanceEndCullDistance;
 
 	InComponent->GetInstancesMinMaxScale(MinScale, MaxScale);
-	NumCustomDataFloats = InComponent->NumCustomDataFloats;
-	TranslatedInstanceSpaceOrigin = InComponent->GetTranslatedInstanceSpaceOrigin();	
 	InstanceLODDistanceScale = InComponent->InstanceLODDistanceScale;
 
 	bUseGpuLodSelection = InComponent->bUseGpuLodSelection;
diff --git a/Engine/Source/Runtime/Engine/Public/InstancedStaticMeshSceneProxyDesc.h b/Engine/Source/Runtime/Engine/Public/InstancedStaticMeshSceneProxyDesc.h
--- a/Engine/Source/Runtime/Engine/Public/InstancedStaticMeshSceneProxyDesc.h
+++ b/Engine/Source/Runtime/Engine/Public/InstancedStaticMeshSceneProxyDesc.h
@@ -13,21 +13,15 @@ struct FInstancedStaticMeshSceneProxyDesc : public FStaticMeshSceneProxyDesc
 	ENGINE_API FInstancedStaticMeshSceneProxyDesc(UInstancedStaticMeshComponent*);
 	void InitializeFrom(UInstancedStaticMeshComponent*);
 
-	TArrayView<const FInstancedStaticMeshInstanceData> PerInstanceSMData;	
 	TSharedPtr<FISMCInstanceDataSceneProxy, ESPMode::ThreadSafe> InstanceDataSceneProxy;
-	TArrayView<const float> PerInstanceSMCustomData;
 #if WITH_EDITOR
 	TBitArray<>	SelectedInstances;
 #endif
-	TArrayView<const int32> InstanceReorderTable; 
-	TArrayView<const FMatrix> PerInstancePrevTransform;
 
 	int32 InstanceStartCullDistance = 0;
 	int32 InstanceEndCullDistance = 0;
 	FVector MinScale = FVector(0);
 	FVector MaxScale = FVector(0);
-	int32 NumCustomDataFloats = 0;
-	FVector TranslatedInstanceSpaceOrigin = FVector(0);
 	float InstanceLODDistanceScale = 1.0f;
 
 	bool bUseGpuLodSelection = false;
@@ -37,44 +31,4 @@ struct FInstancedStaticMeshSceneProxyDesc : public FStaticMeshSceneProxyDesc
 		InMinScale = MinScale;
 		InMaxScale = MaxScale;
 	}
-
-	FVector GetTranslatedInstanceSpaceOrigin() const
-	{
-		return TranslatedInstanceSpaceOrigin;	
-	}
-
-	int32 GetInstanceCount() const
-	{
-		return PerInstanceSMData.Num();
-	}
-
-	int32 GetRenderIndex(int32 InInstanceIndex) const 
-	{ 
-		return InstanceReorderTable.IsValidIndex(InInstanceIndex) ? InstanceReorderTable[InInstanceIndex] : InInstanceIndex; 
-	}
-
-	void GetInstanceTransform(int32 InInstanceIndex, FTransform& OutInstanceTransform) const
-	{	
-		if (!PerInstanceSMData.IsValidIndex(InInstanceIndex))
-		{
-			return;
-		}
-
-		const FInstancedStaticMeshInstanceData& InstanceData = PerInstanceSMData[InInstanceIndex];
-
-		OutInstanceTransform = FTransform(InstanceData.Transform);		
-	}
-
-	bool GetInstancePrevTransform(int32 InInstanceIndex, FTransform& OutInstanceTransform) const
-	{
-		if (!PerInstancePrevTransform.IsValidIndex(InInstanceIndex))
-		{
-			return false;
-		}
-
-		const FMatrix& InstanceData = PerInstancePrevTransform[InInstanceIndex];
-
-		OutInstanceTransform = FTransform(InstanceData);
-		return true;
-	}
 };
\ No newline at end of file
