Replacement for LocalPlayer ortho view info logic to improve frame time resolution.

Works for most use cases as it no longer compensates for when a view target is not specified, however so far with test cases this hasn't been a significant issue with the forced test edge cases as the view still resolves ok, so unless a use case comes up, this works well enough for now.

diff --git a/Engine/Source/Runtime/Engine/Private/LocalPlayer.cpp b/Engine/Source/Runtime/Engine/Private/LocalPlayer.cpp
--- a/Engine/Source/Runtime/Engine/Private/LocalPlayer.cpp
+++ b/Engine/Source/Runtime/Engine/Private/LocalPlayer.cpp
@@ -688,29 +688,6 @@ void ULocalPlayer::GetViewPoint(FMinimalViewInfo& OutViewInfo) const
 			OutViewInfo = PlayerController->PlayerCameraManager->GetCameraCacheView();
 			OutViewInfo.FOV = PlayerController->PlayerCameraManager->GetFOVAngle();
 			PlayerController->GetPlayerViewPoint(/*out*/ OutViewInfo.Location, /*out*/ OutViewInfo.Rotation);
-			
-			bool bCameraFound = false;
-			if (const AActor* Target = PlayerController->GetViewTarget())
-			{
-				if (const UCameraComponent* Camera = Target->GetComponentByClass<const UCameraComponent>())
-				{
-					OutViewInfo.bAutoCalculateOrthoPlanes = Camera->bAutoCalculateOrthoPlanes;
-					OutViewInfo.SetOrthoCameraArmLengthFromOwnerLocation(Target->GetActorLocation());
-					bCameraFound = true;
-				}
-			}
-
-			if (!bCameraFound)
-			{
-				OutViewInfo.bAutoCalculateOrthoPlanes = PlayerController->PlayerCameraManager->bAutoCalculateOrthoPlanes;
-				if(OutViewInfo.bAutoCalculateOrthoPlanes)
-				{
-					if (const APawn* Owner = PlayerController->GetPawn())
-					{
-						OutViewInfo.SetOrthoCameraArmLengthFromOwnerLocation(Owner->GetActorLocation());
-					}
-				}
-			}
 		}
 		else
 		{
diff --git a/Engine/Source/Runtime/Engine/Private/PlayerCameraManager.cpp b/Engine/Source/Runtime/Engine/Private/PlayerCameraManager.cpp
--- a/Engine/Source/Runtime/Engine/Private/PlayerCameraManager.cpp
+++ b/Engine/Source/Runtime/Engine/Private/PlayerCameraManager.cpp
@@ -366,13 +366,6 @@ void APlayerCameraManager::UpdateViewTarget(FTViewTarget& OutVT, float DeltaTime
 	OutVT.POV.ProjectionMode = bIsOrthographic ? ECameraProjectionMode::Orthographic : ECameraProjectionMode::Perspective;
 	OutVT.POV.PostProcessBlendWeight = 1.0f;
 	OutVT.POV.bAutoCalculateOrthoPlanes = bAutoCalculateOrthoPlanes;
-	if (bAutoCalculateOrthoPlanes)
-	{
-		if (OutVT.Target)
-		{
-			OutVT.POV.SetOrthoCameraArmLengthFromOwnerLocation(OutVT.Target->GetActorLocation());
-		}
-	}
 
 	bool bDoNotApplyModifiers = false;
 
@@ -458,6 +451,10 @@ void APlayerCameraManager::UpdateViewTarget(FTViewTarget& OutVT, float DeltaTime
 
 	// Synchronize the actor with the view target results
 	SetActorLocationAndRotation(OutVT.POV.Location, OutVT.POV.Rotation, false);
+	if (bAutoCalculateOrthoPlanes && OutVT.Target)
+	{
+		OutVT.POV.SetOrthoCameraArmLengthFromOwnerLocation(OutVT.Target->GetActorLocation());
+	}
 
 	UpdateCameraLensEffects(OutVT);
 }
