Fix crash when null SubClass passed in
#rnx

diff --git a/Engine/Source/Runtime/Engine/Private/Subsystems/SubsystemCollection.cpp b/Engine/Source/Runtime/Engine/Private/Subsystems/SubsystemCollection.cpp
--- a/Engine/Source/Runtime/Engine/Private/Subsystems/SubsystemCollection.cpp
+++ b/Engine/Source/Runtime/Engine/Private/Subsystems/SubsystemCollection.cpp
@@ -243,14 +243,16 @@ void FSubsystemCollectionBase::Deinitialize()
 
 USubsystem* FSubsystemCollectionBase::InitializeDependency(TSubclassOf<USubsystem> SubsystemClass)
 {
-	UE_LOG(LogSubsystemCollection, VeryVerbose, TEXT("Attempting to initialize subsystem dependency (%s)"), *SubsystemClass->GetName());
-
 	USubsystem* Subsystem = nullptr;
-	if (ensureMsgf(SubsystemClass, TEXT("Attempting to add invalid subsystem as dependancy."))
-		&& ensureMsgf(bPopulating, TEXT("InitializeDependancy() should only be called from System USubsystem::Initialization() implementations."))
-		&& ensureMsgf(SubsystemClass->IsChildOf(BaseType), TEXT("ClassType (%s) must be a subclass of BaseType(%s)."), *SubsystemClass->GetName(), *BaseType->GetName()))
+	if (ensureMsgf(SubsystemClass, TEXT("Attempting to add invalid subsystem as dependancy.")))
 	{
-		Subsystem = AddAndInitializeSubsystem(SubsystemClass);
+		UE_LOG(LogSubsystemCollection, VeryVerbose, TEXT("Attempting to initialize subsystem dependency (%s)"), *SubsystemClass->GetName());
+
+		if (ensureMsgf(bPopulating, TEXT("InitializeDependancy() should only be called from System USubsystem::Initialization() implementations."))
+			&& ensureMsgf(SubsystemClass->IsChildOf(BaseType), TEXT("ClassType (%s) must be a subclass of BaseType(%s)."), *SubsystemClass->GetName(), *BaseType->GetName()))
+		{
+			Subsystem = AddAndInitializeSubsystem(SubsystemClass);
+		}
 	}
 
 	UE_CLOG(!Subsystem, LogSubsystemCollection, Log, TEXT("Failed to initialize subsystem dependency (%s)"), *SubsystemClass->GetName());
