Add some replication testing utilities & a test case for renaming dynamic actors on servers.

UE::Net::FTestWorldInstance and and UE::Net::FTestWorlds enable writing synchronous test cases within the scope of one function that create full UWorlds and utilize the whole networking & replication path. Useful for testing features that rely on worlds, actors, streaming levels, etc. The worlds created are PIE worlds and they currently only work in the editor.

Added an automation test FRenameReplicatedActor_ServerOnlyDynamicOuterChangeTest using these new testing tools to test renaming dynamic actors on a server.

diff --git a/Engine/Source/Runtime/Engine/Private/GameInstance.cpp b/Engine/Source/Runtime/Engine/Private/GameInstance.cpp
--- a/Engine/Source/Runtime/Engine/Private/GameInstance.cpp
+++ b/Engine/Source/Runtime/Engine/Private/GameInstance.cpp
@@ -300,7 +300,8 @@ FGameInstancePIEResult UGameInstance::InitializeForPlayInEditor(int32 PIEInstanc
 
 	WorldContext->OwningGameInstance = this;
 	
-	const FString WorldPackageName = EditorEngine->EditorWorld->GetOutermost()->GetName();
+	UWorld* EditorWorld = EditorEngine->GetEditorWorldContext().World();
+	const FString WorldPackageName = EditorWorld->GetOutermost()->GetName();
 
 	// Establish World Context for PIE World
 	WorldContext->LastURL.Map = WorldPackageName;
@@ -315,7 +316,7 @@ FGameInstancePIEResult UGameInstance::InitializeForPlayInEditor(int32 PIEInstanc
 	if (Params.NetMode == EPlayNetMode::PIE_Client)
 	{
 		// We are going to connect, so just load an empty world
-		NewWorld = EditorEngine->CreatePIEWorldFromEntry(*WorldContext, EditorEngine->EditorWorld, PIEMapName);
+		NewWorld = EditorEngine->CreatePIEWorldFromEntry(*WorldContext, EditorWorld, PIEMapName);
 	}
 	else
 	{
@@ -326,14 +327,14 @@ FGameInstancePIEResult UGameInstance::InitializeForPlayInEditor(int32 PIEInstanc
 			UWorld* WorldToDuplicate = Cast<UWorld>(TargetWorld.TryLoad());
 			if (WorldToDuplicate)
 			{
-				WorldToDuplicate->ChangeFeatureLevel(EditorEngine->EditorWorld->GetFeatureLevel(), false);
+				WorldToDuplicate->ChangeFeatureLevel(EditorWorld->GetFeatureLevel(), false);
 				NewWorld = EditorEngine->CreatePIEWorldByDuplication(*WorldContext, WorldToDuplicate, PIEMapName);
 			}
 		}
 		else
 		{
 			// Standard PIE path: just duplicate the EditorWorld
-			NewWorld = EditorEngine->CreatePIEWorldByDuplication(*WorldContext, EditorEngine->EditorWorld, PIEMapName);
+			NewWorld = EditorEngine->CreatePIEWorldByDuplication(*WorldContext, EditorWorld, PIEMapName);
 		}
 
 		// Duplication can result in unreferenced objects, so indicate that we should do a GC pass after initializing the world context
