Cluster union : Add flag to disable damage from collision on cluster union

diff --git a/Engine/Source/Runtime/Engine/Classes/PhysicsEngine/ClusterUnionComponent.h b/Engine/Source/Runtime/Engine/Classes/PhysicsEngine/ClusterUnionComponent.h
--- a/Engine/Source/Runtime/Engine/Classes/PhysicsEngine/ClusterUnionComponent.h
+++ b/Engine/Source/Runtime/Engine/Classes/PhysicsEngine/ClusterUnionComponent.h
@@ -236,6 +236,13 @@ public:
 
 	ENGINE_API bool IsAnchored() const;
 
+	UFUNCTION(BlueprintCallable, Category = "Cluster Union")
+	ENGINE_API void SetEnableDamageFromCollision(bool bValue);
+
+	/** Whether or not collisions against this geometry collection will apply strain which could cause the geometry collection to fracture. */
+	UPROPERTY(EditAnywhere, BlueprintReadWrite, BlueprintSetter = SetEnableDamageFromCollision, Category = "Cluster Union")
+	bool bEnableDamageFromCollision;
+
 	/** Structure that stores an object key and its raw ptr, so it can be accessed without going trough
 	* the global uobjects array (so it could be garbage, therefore it has to be used withing the same frame)
 	*/
@@ -387,7 +394,7 @@ private:
 	// instead of automatically inheriting it from its children particles.
 	UPROPERTY(EditAnywhere, Category = "Cluster Union")
 	int32 GravityGroupIndexOverride;
-	
+
 	// Also keep track of which actors we are clustering and their components. We make modifications on
 	// actors that get clustered so we need to make sure we undo those changes only once all its clustered
 	// components are removed from the cluster.
diff --git a/Engine/Source/Runtime/Engine/Private/PhysicsEngine/ClusterUnionComponent.cpp b/Engine/Source/Runtime/Engine/Private/PhysicsEngine/ClusterUnionComponent.cpp
--- a/Engine/Source/Runtime/Engine/Private/PhysicsEngine/ClusterUnionComponent.cpp
+++ b/Engine/Source/Runtime/Engine/Private/PhysicsEngine/ClusterUnionComponent.cpp
@@ -113,6 +113,7 @@ UClusterUnionComponent::UClusterUnionComponent(const FObjectInitializer& ObjectI
 	bVisualizeComponent = true;
 #endif
 	GravityGroupIndexOverride = INDEX_NONE;
+	bEnableDamageFromCollision = true;
 }
 
 FPhysScene_Chaos* UClusterUnionComponent::GetChaosScene() const
@@ -638,6 +639,16 @@ ENGINE_API bool UClusterUnionComponent::IsAnchored() const
 	return PhysicsProxy->IsAnchored_External();
 }
 
+void UClusterUnionComponent::SetEnableDamageFromCollision(bool bValue)
+{
+	bEnableDamageFromCollision = bValue;
+
+	if (PhysicsProxy)
+	{
+		PhysicsProxy->SetEnableStrainOnCollision_External(bValue);
+	}
+}
+
 bool UClusterUnionComponent::IsAuthority() const
 {
 	if (bUseLocalRoleForAuthorityCheck)
@@ -687,6 +698,7 @@ void UClusterUnionComponent::OnCreatePhysicsState()
 	// TODO: Expose these parameters via the component.
 	Chaos::FClusterCreationParameters Parameters{ 0.3f, 100, false, false };
 	Parameters.ConnectionMethod = Chaos::FClusterCreationParameters::EConnectionMethod::PointImplicit;
+	Parameters.bEnableStrainOnCollision = bEnableDamageFromCollision;
 
 	FChaosUserData::Set<UPrimitiveComponent>(&PhysicsUserData, this);
 
diff --git a/Engine/Source/Runtime/Experimental/Chaos/Private/Chaos/PBDRigidClustering.cpp b/Engine/Source/Runtime/Experimental/Chaos/Private/Chaos/PBDRigidClustering.cpp
--- a/Engine/Source/Runtime/Experimental/Chaos/Private/Chaos/PBDRigidClustering.cpp
+++ b/Engine/Source/Runtime/Experimental/Chaos/Private/Chaos/PBDRigidClustering.cpp
@@ -2249,6 +2249,12 @@ namespace Chaos
 				{
 					if (const FClusterUnionPhysicsProxy* ClusterUnionProxy = GetConcreteProxy<FClusterUnionPhysicsProxy>(Cluster))
 					{
+						// check if this cluster union accept damage from collision 
+						if (!ClusterUnionProxy->GetEnableStrainOnCollision_Internal())
+						{
+							return;
+						}
+
 						// At the moment, we don't want to apply strains to children of ClusterUnions, we want instead
 						// to apply the strains to GRANDchildren of ClusterUnions.
 						FPBDRigidParticleHandle* ClosestChild = FindClosestChild(Cluster, ContactWorldLocation);
diff --git a/Engine/Source/Runtime/Experimental/Chaos/Private/PhysicsProxy/ClusterUnionPhysicsProxy.cpp b/Engine/Source/Runtime/Experimental/Chaos/Private/PhysicsProxy/ClusterUnionPhysicsProxy.cpp
--- a/Engine/Source/Runtime/Experimental/Chaos/Private/PhysicsProxy/ClusterUnionPhysicsProxy.cpp
+++ b/Engine/Source/Runtime/Experimental/Chaos/Private/PhysicsProxy/ClusterUnionPhysicsProxy.cpp
@@ -110,6 +110,8 @@ namespace Chaos
 
 		bIsInitializedOnPhysicsThread = true;
 
+		bEnableStrainOnCollision_Internal = ClusterParameters.bEnableStrainOnCollision;
+
 		FPBDRigidsEvolutionGBF* Evolution = RigidsSolver->GetEvolution();
 		if (!ensure(Evolution))
 		{
@@ -804,6 +806,19 @@ namespace Chaos
 		);
 	}
 
+	void FClusterUnionPhysicsProxy::SetEnableStrainOnCollision_External(bool bEnable)
+	{
+		if (Solver && ensure(Particle_External))
+		{
+			Solver->EnqueueCommandImmediate(
+				[this, bEnable]() mutable
+				{
+					bEnableStrainOnCollision_Internal = bEnable;
+				}
+			);
+		}
+	}
+
 	void FClusterUnionPhysicsProxy::ChangeMainParticleStatus_External(const TArray<FPhysicsObjectHandle>& Objects, bool bIsMain)
 	{
 		if (!Solver)
diff --git a/Engine/Source/Runtime/Experimental/Chaos/Public/Chaos/ClusterCreationParameters.h b/Engine/Source/Runtime/Experimental/Chaos/Public/Chaos/ClusterCreationParameters.h
--- a/Engine/Source/Runtime/Experimental/Chaos/Public/Chaos/ClusterCreationParameters.h
+++ b/Engine/Source/Runtime/Experimental/Chaos/Public/Chaos/ClusterCreationParameters.h
@@ -34,6 +34,7 @@ namespace Chaos
 			, Chaos::TPBDRigidClusteredParticleHandle<Chaos::FReal,3>* ClusterParticleHandleIn = nullptr
 			, const FVec3& ScaleIn = FVec3::OneVector
 			, bool bIsAnchoredIn = false
+			, bool bInEnableStrainOnCollision = true
 		)
 			: CoillisionThicknessPercent(CoillisionThicknessPercentIn)
 			, MaxNumConnections(MaxNumConnectionsIn)
@@ -46,6 +47,7 @@ namespace Chaos
 			, ClusterParticleHandle(ClusterParticleHandleIn)
 			, Scale(ScaleIn)
 			, bIsAnchored(bIsAnchoredIn)
+			, bEnableStrainOnCollision(bInEnableStrainOnCollision)
 		{}
 
 		FReal CoillisionThicknessPercent;
@@ -59,5 +61,6 @@ namespace Chaos
 		Chaos::FPBDRigidClusteredParticleHandle* ClusterParticleHandle;
 		FVec3 Scale;
 		bool bIsAnchored;
+		bool bEnableStrainOnCollision;
 	};
 }
diff --git a/Engine/Source/Runtime/Experimental/Chaos/Public/PhysicsProxy/ClusterUnionPhysicsProxy.h b/Engine/Source/Runtime/Experimental/Chaos/Public/PhysicsProxy/ClusterUnionPhysicsProxy.h
--- a/Engine/Source/Runtime/Experimental/Chaos/Public/PhysicsProxy/ClusterUnionPhysicsProxy.h
+++ b/Engine/Source/Runtime/Experimental/Chaos/Public/PhysicsProxy/ClusterUnionPhysicsProxy.h
@@ -146,8 +146,13 @@ namespace Chaos
 
 		void ForceSetGeometryChildParticles_External(TArray<FExternalParticle*>&& InParticles) { GeometryChildParticles_External = InParticles; }
 
+		CHAOS_API void SetEnableStrainOnCollision_External(bool bEnable);
+		CHAOS_API bool GetEnableStrainOnCollision_Internal() const { return bEnableStrainOnCollision_Internal; }
+
 	private:
-		bool bIsInitializedOnPhysicsThread = false;
+		bool bIsInitializedOnPhysicsThread: 1 = false;
+		bool bEnableStrainOnCollision_Internal : 1 = true;
+
 		FClusterCreationParameters ClusterParameters;
 		const FClusterUnionInitData InitData;
 
