Added status log category.
Added log status to multiple places.
Added log vision category template

diff --git a/Engine/Source/Runtime/Engine/Classes/Engine/Engine.h b/Engine/Source/Runtime/Engine/Classes/Engine/Engine.h
--- a/Engine/Source/Runtime/Engine/Classes/Engine/Engine.h
+++ b/Engine/Source/Runtime/Engine/Classes/Engine/Engine.h
@@ -28,6 +28,7 @@
 #include "Templates/UniqueObj.h"
 #include "Containers/Ticker.h"
 #include "DynamicRenderScaling.h"
+#include "Misc/StatusLog.h"
 #include "Engine.generated.h"
 
 #define WITH_DYNAMIC_RESOLUTION (!UE_SERVER)
@@ -2176,6 +2177,7 @@ public:
 	/** Called by internal engine systems after a travel failure has occurred */
 	void BroadcastTravelFailure(UWorld* InWorld, ETravelFailure::Type FailureType, const FString& ErrorString = TEXT(""))
 	{
+		UE_LOGSTATUS(Warning, TEXT("Travel failed, type: %s, reason: \"%s\""), *UEnum::GetValueAsString(FailureType), *ErrorString);
 		TravelFailureEvent.Broadcast(InWorld, FailureType, ErrorString);
 	}
 
diff --git a/Engine/Source/Runtime/Engine/Private/Misc/StatusLog.cpp b/Engine/Source/Runtime/Engine/Private/Misc/StatusLog.cpp
new file mode 100644
--- /dev/null
+++ b/Engine/Source/Runtime/Engine/Private/Misc/StatusLog.cpp
@@ -0,0 +1,10 @@
+// Copyright Epic Games, Inc. All Rights Reserved.
+
+
+#include "Misc/StatusLog.h"
+
+DEFINE_LOG_CATEGORY(LogGlobalStatus);
+
+
+
+
diff --git a/Engine/Source/Runtime/Engine/Private/PendingNetGame.cpp b/Engine/Source/Runtime/Engine/Private/PendingNetGame.cpp
--- a/Engine/Source/Runtime/Engine/Private/PendingNetGame.cpp
+++ b/Engine/Source/Runtime/Engine/Private/PendingNetGame.cpp
@@ -206,6 +206,8 @@ void UPendingNetGame::TravelCompleted(UEngine* Engine, FWorldContext& Context)
 	// Send join.
 	Context.PendingNetGame->SendJoin();
 	Context.PendingNetGame->NetDriver = NULL;
+
+	UE_LOGSTATUS(Log, TEXT("Pending net game travel completed"));
 }
 
 EAcceptConnection::Type UPendingNetGame::NotifyAcceptingConnection()
diff --git a/Engine/Source/Runtime/Engine/Private/UnrealEngine.cpp b/Engine/Source/Runtime/Engine/Private/UnrealEngine.cpp
--- a/Engine/Source/Runtime/Engine/Private/UnrealEngine.cpp
+++ b/Engine/Source/Runtime/Engine/Private/UnrealEngine.cpp
@@ -14615,6 +14615,8 @@ EBrowseReturnVal::Type UEngine::Browse( FWorldContext& WorldContext, FURL URL, F
 	Error = TEXT("");
 	WorldContext.TravelURL = TEXT("");
 
+	UE_LOGSTATUS(Log, TEXT("Started Browse: \"%s\""), *URL.ToString());
+
 	if (WorldContext.World() && WorldContext.World()->GetNetDriver())
 	{
 		const TCHAR* InTickErrorString = TEXT("Attempting to call UEngine::Browse and destroy the net driver while the net driver is ticking. Instead try using UWorld::ServerTravel or APlayerController::ClientTravel.");
@@ -15657,6 +15659,7 @@ bool UEngine::LoadMap( FWorldContext& WorldContext, FURL URL, class UPendingNetG
 	}
 
 	// Successfully started local level.
+	UE_LOGSTATUS(Log, TEXT("Load map complete %s"), *URL.Map);
 	return true;
 }
 
diff --git a/Engine/Source/Runtime/Engine/Public/Misc/StatusLog.h b/Engine/Source/Runtime/Engine/Public/Misc/StatusLog.h
new file mode 100644
--- /dev/null
+++ b/Engine/Source/Runtime/Engine/Public/Misc/StatusLog.h
@@ -0,0 +1,23 @@
+// Copyright Epic Games, Inc. All Rights Reserved.
+
+/*=============================================================================
+	StatusLog.h: high level status logging helper:
+		should be minimal and track important game state changing events, travel, game ending failures, etc
+		verbosity should be minimal, goal is 20 lines per match / round
+=============================================================================*/
+
+#pragma once
+
+#include "CoreGlobals.h"
+
+ENGINE_API DECLARE_LOG_CATEGORY_EXTERN(LogGlobalStatus, Verbose, All);
+
+
+#define UE_LOGSTATUS(Verbosity, STRING, ...) UE_LOG(LogGlobalStatus, Verbosity, TEXT("%s ") STRING, ANSI_TO_TCHAR(__FUNCTION__), ##__VA_ARGS__ )
+/* UE_LOGFMTSTATUS 
+ * can only be used as such will cause compilation errors if not specifying format arguments completely
+ * UE_LOGFMTSTATUS(Log, "format string {FormatArgument1} {FormatArguement2}", ("FormatArgument1", FormatArgument1Value), ("FormatArgument2", FormatArgument2Value));
+ */
+#define UE_LOGFMTSTATUS(Verbosity, STRING, ...) UE_LOGFMT(LogGlobalStatus, Verbosity, "{FUNCTIONNAME} " STRING, ("FUNCTIONNAME", __FUNCTION__), ##__VA_ARGS__)
+
+
