#UE Fix cook warning about frequent access of a CVAR

diff --git a/Engine/Source/Runtime/Engine/Private/GPUSkinVertexFactory.cpp b/Engine/Source/Runtime/Engine/Private/GPUSkinVertexFactory.cpp
--- a/Engine/Source/Runtime/Engine/Private/GPUSkinVertexFactory.cpp
+++ b/Engine/Source/Runtime/Engine/Private/GPUSkinVertexFactory.cpp
@@ -498,7 +498,7 @@ bool FGPUBaseSkinVertexFactory::GetAlwaysUseDeformerForUnlimitedBoneInfluences(E
 {
 	auto InnerFunc = [](EShaderPlatform Platform)
 	{
-		FShaderPlatformCachedIniValue<bool> UseDeformerForUBICVar(TEXT("r.GPUSkin.AlwaysUseDeformerForUnlimitedBoneInfluences"));
+		static FShaderPlatformCachedIniValue<bool> UseDeformerForUBICVar(TEXT("r.GPUSkin.AlwaysUseDeformerForUnlimitedBoneInfluences"));
 		const IMeshDeformerProvider* MeshDeformerProvider = IMeshDeformerProvider::Get();
 
 		return MeshDeformerProvider && MeshDeformerProvider->IsSupported(Platform) && UseDeformerForUBICVar.Get(Platform);
