Fix erroneous parsing of 'listtypes' for ODSC which was preventing the SingleShader command to happen
The standard listtypes commands in the editor works fine, only ODSC was affected

diff --git a/Engine/Source/Runtime/Engine/Private/ShaderCompiler/ShaderCompiler.cpp b/Engine/Source/Runtime/Engine/Private/ShaderCompiler/ShaderCompiler.cpp
--- a/Engine/Source/Runtime/Engine/Private/ShaderCompiler/ShaderCompiler.cpp
+++ b/Engine/Source/Runtime/Engine/Private/ShaderCompiler/ShaderCompiler.cpp
@@ -8844,7 +8844,7 @@ namespace
 				OutMaterialsToLoad.Add(It->GetPathName());
 			}
 		}
-		else if (FCString::Stricmp(*CmdName, TEXT("listtypes")))
+		else if (FCString::Stricmp(*CmdName, TEXT("listtypes")) == 0)
 		{
 			ListAllShaderTypes();
 		}
