NetAnalytics
* Added ability to add game specific attributes to the NetAnalytics framework
* To add those game attributes, Implementations of FNetAnalyticsData must call AppendGameInstanceAttributes before calling SendAttributes.
* Add or modify your game attributes with the  NetDriver::SetNetAnalyticsAttributes function
* NetDriver will add some default attributes (NetDriverName, NetDriverDefinition, ReplicationModel, NetMode) when net.AddNetDriverInfoToNetAnalytics is true

#[review]-31918845

diff --git a/Engine/Plugins/Compression/OodleNetwork/Source/Private/OodleNetworkAnalytics.cpp b/Engine/Plugins/Compression/OodleNetwork/Source/Private/OodleNetworkAnalytics.cpp
--- a/Engine/Plugins/Compression/OodleNetwork/Source/Private/OodleNetworkAnalytics.cpp
+++ b/Engine/Plugins/Compression/OodleNetwork/Source/Private/OodleNetworkAnalytics.cpp
@@ -237,7 +237,7 @@ void FOodleNetAnalyticsData::SendAnalytics()
 		static const FString EZAttrib_NumOodleNetworkHandlersCompressionEnabled = TEXT("NumOodleNetworkHandlersCompressionEnabled");
 		static const FString EZAttrib_NumOodleNetworkHandlersCompressionDisabled = TEXT("NumOodleNetworkHandlersCompressionDisabled");
 
-		const TArray<FAnalyticsEventAttribute> EventAttributes = MakeAnalyticsEventAttributeArray(
+		TArray<FAnalyticsEventAttribute> EventAttributes = MakeAnalyticsEventAttributeArray(
 			EZAttrib_InCompressedNum, InCompressedNum,
 			EZAttrib_InNotCompressedNum, InNotCompressedNum,
 			EZAttrib_InCompressedWithOverheadLengthTotal, InCompressedWithOverheadLengthTotal,
@@ -274,6 +274,8 @@ void FOodleNetAnalyticsData::SendAnalytics()
 			EZAttrib_NumOodleNetworkHandlersCompressionDisabled, NumOodleNetworkHandlersCompressionDisabled
 		);
 
+		Aggregator->AppendGameInstanceAttributes(EventAttributes);
+
 		AnalyticsProvider->RecordEvent(GetAnalyticsEventName(), EventAttributes);
 	}
 }
diff --git a/Engine/Source/Runtime/Engine/Classes/Engine/NetDriver.h b/Engine/Source/Runtime/Engine/Classes/Engine/NetDriver.h
--- a/Engine/Source/Runtime/Engine/Classes/Engine/NetDriver.h
+++ b/Engine/Source/Runtime/Engine/Classes/Engine/NetDriver.h
@@ -1214,15 +1214,28 @@ public:
 	TMap<FNetworkGUID, TUniquePtr<FActorDestructionInfo>>	DestroyedStartupOrDormantActors;
 
 private:
+
 	/** Tracks the network guids in DestroyedStartupOrDormantActors above, but keyed on the streaming level name. */
 	TMap<FName, TSet<FNetworkGUID>> DestroyedStartupOrDormantActorsByLevel;
 
+	/** Cached list of analytic attributes that can be appended via FNetAnalyticsAggregator::AppendGameInstanceAttributes */
+	TMap<FString, FString> CachedNetAnalyticsAttributes;
+
 public:
+
 	const TSet<FNetworkGUID>& GetDestroyedStartupOrDormantActors(const FName& LevelName)
 	{
 		return DestroyedStartupOrDormantActorsByLevel.FindOrAdd(LevelName);
 	}
 
+	/** 
+	 * Add or overwrite an analytics attribute that will be appended via FNetAnalyticsAggregator::AppendGameInstanceAttributes
+	 * Useful to add game specific attributes to your analytics.
+	 * @param AttributeKey The key of the attribute. Stored in a case-insensitive map
+	 * @param AttributeValue Value of the attribute to store. Only supports strings.
+	 */
+	void SetNetAnalyticsAttributes(const FString& AttributeKey, const FString& AttributeValue) { CachedNetAnalyticsAttributes.Add(AttributeKey, AttributeValue); }
+
 	/** The server adds an entry into this map for every startup actor that has been renamed, and will
 	 *  always map from current name to original name
 	 */
@@ -2203,6 +2216,9 @@ private:
 	void PostDispatchSendUpdate();
 #endif
 
+	/** Tell the registered NetAnalytics to send their analytics via the provider */
+	void SendNetAnalytics();
+
 	/** Description of the replication model used by this Driver (RepGraph, Iris or Generic) */
 	FString GetReplicationModelName() const;
 
diff --git a/Engine/Source/Runtime/Engine/Private/NetDriver.cpp b/Engine/Source/Runtime/Engine/Private/NetDriver.cpp
--- a/Engine/Source/Runtime/Engine/Private/NetDriver.cpp
+++ b/Engine/Source/Runtime/Engine/Private/NetDriver.cpp
@@ -5,6 +5,8 @@
 =============================================================================*/
 
 #include "Engine/NetDriver.h"
+
+#include "AnalyticsEventAttribute.h"
 #include "Engine/GameInstance.h"
 #include "Engine/ServerStatReplicator.h"
 #include "Misc/App.h"
@@ -409,6 +411,11 @@ static FAutoConsoleVariableRef CVarNetResetAckStatePostSeamlessTravel(
 	TEXT("If 1, the server will reset the ack state of the package map after seamless travel. Increases bandwidth usage, but may resolve some issues with GUIDs not being available on clients after seamlessly traveling."),
 	ECVF_Default);
 
+static TAutoConsoleVariable<bool> CVarAddNetDriverInfoToNetAnalytics(
+	TEXT("net.AddNetDriverInfoToNetAnalytics"),
+	false,
+	TEXT("Automatically add NetDriver information to the NetAnalytics cache"));
+
 namespace UE::Net
 {
 	static FString GRequiredEncryptionNetDriverDefNames_Internal = TEXT("all");
@@ -2097,11 +2104,7 @@ void UNetDriver::Shutdown()
 	// End NetTrace session for this instance
 	UE_NET_TRACE_END_SESSION(GetNetTraceId());
 
-	if (AnalyticsAggregator.IsValid())
-	{
-		AnalyticsAggregator->SendAnalytics();
-		AnalyticsAggregator.Reset();
-	}
+	SendNetAnalytics();
 
 	// Clear repmodel flags for clients or listen servers when shutting down since they are considered offline.
 	// For dedicated servers let's keep the flag in case we crash post-shutdown.
@@ -2125,6 +2128,36 @@ bool UNetDriver::IsServer() const
 	return ServerConnection == NULL;
 }
 
+void UNetDriver::SendNetAnalytics()
+{
+	if (!AnalyticsAggregator.IsValid())
+	{
+		return;
+	}
+
+	// Add the default NetDriver information if requested
+	if (CVarAddNetDriverInfoToNetAnalytics.GetValueOnAnyThread())
+	{
+		SetNetAnalyticsAttributes(TEXT("NetDriverName"), NetDriverName.ToString());
+		SetNetAnalyticsAttributes(TEXT("NetDriverDefinition"), NetDriverDefinition.ToString());
+		SetNetAnalyticsAttributes(TEXT("ReplicationModel"), GetReplicationModelName());
+		SetNetAnalyticsAttributes(TEXT("NetMode"), ToString(GetNetMode()));
+	}
+
+	auto GameAttributes = [this](TArray<FAnalyticsEventAttribute>& OutAttributes)
+	{
+		OutAttributes.Reserve(OutAttributes.Num() + CachedNetAnalyticsAttributes.Num());
+		for (auto It=CachedNetAnalyticsAttributes.CreateConstIterator(); It; ++It)
+		{
+			OutAttributes.Emplace(FAnalyticsEventAttribute(It.Key(), It.Value()));
+		}
+	};
+
+	AnalyticsAggregator->SetAnalyticsAppender(GameAttributes);
+	AnalyticsAggregator->SendAnalytics();
+	AnalyticsAggregator.Reset();
+}
+
 EEngineNetworkRuntimeFeatures UNetDriver::GetNetworkRuntimeFeatures() const
 {
 	EEngineNetworkRuntimeFeatures NetDriverFeatures = EEngineNetworkRuntimeFeatures::None;
diff --git a/Engine/Source/Runtime/Net/Core/Private/Net/Core/Analytics/NetAnalytics.cpp b/Engine/Source/Runtime/Net/Core/Private/Net/Core/Analytics/NetAnalytics.cpp
--- a/Engine/Source/Runtime/Net/Core/Private/Net/Core/Analytics/NetAnalytics.cpp
+++ b/Engine/Source/Runtime/Net/Core/Private/Net/Core/Analytics/NetAnalytics.cpp
@@ -142,3 +142,16 @@ void FNetAnalyticsAggregator::SendAnalytics()
 		bSentAnalytics = true;
 	}
 }
+
+void FNetAnalyticsAggregator::SetAnalyticsAppender(FAnalyticsAppender AppenderFunction)
+{
+	AnalyticsAppender = AppenderFunction;
+}
+
+void FNetAnalyticsAggregator::AppendGameInstanceAttributes(TArray<FAnalyticsEventAttribute>& OutAttributes)
+{
+	if (AnalyticsAppender)
+	{
+		AnalyticsAppender(OutAttributes);
+	}
+}
\ No newline at end of file
diff --git a/Engine/Source/Runtime/Net/Core/Public/Net/Core/Analytics/NetAnalytics.h b/Engine/Source/Runtime/Net/Core/Public/Net/Core/Analytics/NetAnalytics.h
--- a/Engine/Source/Runtime/Net/Core/Public/Net/Core/Analytics/NetAnalytics.h
+++ b/Engine/Source/Runtime/Net/Core/Public/Net/Core/Analytics/NetAnalytics.h
@@ -45,6 +45,7 @@
 // Forward Declarations
 class FNetAnalyticsAggregator;
 class IAnalyticsProvider;
+struct FAnalyticsEventAttribute;
 
 
 // Defines
@@ -341,6 +342,11 @@ public:
 	 */
 	NETCORE_API void InitConfig();
 
+	/**
+	 * Adds game and engine specific attributes to the passed in array.
+	 */
+	NETCORE_API void AppendGameInstanceAttributes(TArray<FAnalyticsEventAttribute>& OutAttributes);
+
 	/**
 	 * Tells the analytics data holders to finish aggregating their analytics data, and to dispatch it.
 	 * Only called once, at NetDriver shutdown.
@@ -374,6 +380,10 @@ public:
 		return NetDriverName;
 	}
 
+public:
+
+	using FAnalyticsAppender = TFunction<void(TArray<FAnalyticsEventAttribute>&)>;	
+	void SetAnalyticsAppender(FAnalyticsAppender AppenderFunction);
 
 private:
 	/** The analytics provider we are aggregating data for */
@@ -391,6 +401,9 @@ private:
 	/** Maps analytics data holder names, to a config value specifying whether that data holder is enabled or not */
 	TMap<FName, bool> AnalyticsDataConfigMap;
 
+	/** Optional functor that can be used to append game specific attributes to your analytics */
+	FAnalyticsAppender AnalyticsAppender;
+
 	/** Whether or not analytics was already sent */
 	bool bSentAnalytics;
 };
