World Partition: fix failed bounds check on Unhash because of transformed editor bounds

Partially revert CL 31253792 so that EditorHash for Level Instances remains in local container space.

Bug was that an unsaved moved Level Instance would fail check on Unhash when another instance of the same level would commit an actor change

#rnx

diff --git a/Engine/Source/Runtime/Engine/Private/WorldPartition/LoaderAdapter/LoaderAdapterSpatial.cpp b/Engine/Source/Runtime/Engine/Private/WorldPartition/LoaderAdapter/LoaderAdapterSpatial.cpp
--- a/Engine/Source/Runtime/Engine/Private/WorldPartition/LoaderAdapter/LoaderAdapterSpatial.cpp
+++ b/Engine/Source/Runtime/Engine/Private/WorldPartition/LoaderAdapter/LoaderAdapterSpatial.cpp
@@ -31,9 +31,12 @@ void ILoaderAdapterSpatial::HandleIntersectingContainer(UWorldPartition* InWorld
 		.SetIncludeSpatiallyLoadedActors(bIncludeSpatiallyLoadedActors)
 		.SetIncludeNonSpatiallyLoadedActors(bIncludeNonSpatiallyLoadedActors);
 
-	InWorldPartition->EditorHash->ForEachIntersectingActor(InBoundingBox, [this, InWorldPartition, &InOperation](FWorldPartitionActorDescInstance* ActorDescInstance)
+	const FTransform InstanceTransform = InWorldPartition->GetInstanceTransform();
+	const FBox LocalBoundingBox = InBoundingBox.InverseTransformBy(InstanceTransform);
+	InWorldPartition->EditorHash->ForEachIntersectingActor(LocalBoundingBox, [this, InWorldPartition, &InstanceTransform, &InOperation](FWorldPartitionActorDescInstance* ActorDescInstance)
 	{
-		if (Intersect(ActorDescInstance->GetEditorBounds()))
+		const FBox WorldActorEditorBounds = ActorDescInstance->GetLocalEditorBounds().TransformBy(InstanceTransform);
+		if (Intersect(WorldActorEditorBounds))
 		{
 			FWorldPartitionHandle ActorHandle(InWorldPartition, ActorDescInstance->GetGuid());
 			InOperation(ActorHandle);
diff --git a/Engine/Source/Runtime/Engine/Private/WorldPartition/WorldPartitionEditorSpatialHash.cpp b/Engine/Source/Runtime/Engine/Private/WorldPartition/WorldPartitionEditorSpatialHash.cpp
--- a/Engine/Source/Runtime/Engine/Private/WorldPartition/WorldPartitionEditorSpatialHash.cpp
+++ b/Engine/Source/Runtime/Engine/Private/WorldPartition/WorldPartitionEditorSpatialHash.cpp
@@ -69,14 +69,14 @@ void UWorldPartitionEditorSpatialHash::Tick(float DeltaSeconds)
 
 			for (const FWorldPartitionHandle& ActorHandle : Cell->Actors)
 			{
-				RuntimeBounds += ActorHandle->GetEditorBounds();
+				RuntimeBounds += ActorHandle->GetLocalEditorBounds();
 			}
 		}
 
 		NonSpatialBounds.Init();
 		for (const FWorldPartitionHandle& ActorHandle : AlwaysLoadedCell->Actors)
 		{
-			NonSpatialBounds += ActorHandle->GetEditorBounds();
+			NonSpatialBounds += ActorHandle->GetLocalEditorBounds();
 		}
 
 		const int32 OldLevel = GetLevelForBox(EditorBounds);
@@ -128,7 +128,7 @@ void UWorldPartitionEditorSpatialHash::HashActor(FWorldPartitionHandle& InActorH
 
 	const FWorldPartitionActorDescInstance* ActorDescInstance = *InActorHandle;
 	const bool bConsiderActorSpatiallyLoaded = ActorDescInstance->GetIsSpatiallyLoaded();
-	const FBox ActorBounds = bConsiderActorSpatiallyLoaded ? ActorDescInstance->GetEditorBounds() : FBox(ForceInit);
+	const FBox ActorBounds = bConsiderActorSpatiallyLoaded ? ActorDescInstance->GetLocalEditorBounds() : FBox(ForceInit);
 
 #if DO_CHECK
 	check(!HashedActors.Contains(ActorDescInstance->GetGuid()));
@@ -138,7 +138,7 @@ void UWorldPartitionEditorSpatialHash::HashActor(FWorldPartitionHandle& InActorH
 	if (!bConsiderActorSpatiallyLoaded)
 	{
 		AlwaysLoadedCell->Actors.Add(InActorHandle);
-		NonSpatialBounds += ActorDescInstance->GetEditorBounds();
+		NonSpatialBounds += ActorDescInstance->GetLocalEditorBounds();
 	}
 	else
 	{
@@ -231,7 +231,7 @@ void UWorldPartitionEditorSpatialHash::HashActor(FWorldPartitionHandle& InActorH
 			}
 		}
 
-		RuntimeBounds += ActorDescInstance->GetEditorBounds();
+		RuntimeBounds += ActorDescInstance->GetLocalEditorBounds();
 	}
 }
 
@@ -241,7 +241,7 @@ void UWorldPartitionEditorSpatialHash::UnhashActor(FWorldPartitionHandle& InActo
 
 	const FWorldPartitionActorDescInstance* ActorDescInstance = *InActorHandle;
 	const bool bConsiderActorSpatiallyLoaded = ActorDescInstance->GetIsSpatiallyLoaded();
-	const FBox ActorBounds = bConsiderActorSpatiallyLoaded ? ActorDescInstance->GetEditorBounds() : FBox(ForceInit);
+	const FBox ActorBounds = bConsiderActorSpatiallyLoaded ? ActorDescInstance->GetLocalEditorBounds() : FBox(ForceInit);
 
 #if DO_CHECK
 	const FBox OldActorBounds = HashedActors.FindAndRemoveChecked(ActorDescInstance->GetGuid());
@@ -318,7 +318,7 @@ int32 UWorldPartitionEditorSpatialHash::ForEachIntersectingActor(const FBox& Box
 
 					if (!bWasAlreadyInSet)
 					{
-						if (Box.Intersect(ActorHandle->GetEditorBounds()))
+						if (Box.Intersect(ActorHandle->GetLocalEditorBounds()))
 						{
 							InOperation(*ActorHandle);
 						}
@@ -334,7 +334,7 @@ int32 UWorldPartitionEditorSpatialHash::ForEachIntersectingActor(const FBox& Box
 		{
 			if (ActorHandle.IsValid())
 			{
-				if (Box.Intersect(ActorHandle->GetEditorBounds()))
+				if (Box.Intersect(ActorHandle->GetLocalEditorBounds()))
 				{
 					bool bWasAlreadyInSet;
 					IntersectedActors.Add(ActorHandle->GetGuid(), &bWasAlreadyInSet);
diff --git a/Engine/Source/Runtime/Engine/Public/WorldPartition/WorldPartitionActorDescInstance.h b/Engine/Source/Runtime/Engine/Public/WorldPartition/WorldPartitionActorDescInstance.h
--- a/Engine/Source/Runtime/Engine/Public/WorldPartition/WorldPartitionActorDescInstance.h
+++ b/Engine/Source/Runtime/Engine/Public/WorldPartition/WorldPartitionActorDescInstance.h
@@ -28,6 +28,8 @@ class FWorldPartitionActorDescInstance
 	friend class UWorldPartition;
 	friend class FWorldPartitionLoadingContext;
 	friend class IWorldPartitionActorLoaderInterface;
+	friend class UWorldPartitionEditorSpatialHash;
+	friend class ILoaderAdapterSpatial;
 
 protected:
 	ENGINE_API FWorldPartitionActorDescInstance();
@@ -117,6 +119,8 @@ public:
 	ENGINE_API const FText& GetUnloadedReason() const;
 		
 protected:
+	ENGINE_API virtual FBox GetLocalEditorBounds() const { return ActorDesc->GetEditorBounds(); }
+
 	ENGINE_API bool StartAsyncLoad();
 	ENGINE_API void FlushAsyncLoad() const;
 	ENGINE_API void MarkUnload();
