Fix for CitySample VT build failure - removal of sanity check since it appears to be wrong in some cases.

diff --git a/Engine/Source/Runtime/Engine/Private/VT/VirtualTextureDataBuilder.cpp b/Engine/Source/Runtime/Engine/Private/VT/VirtualTextureDataBuilder.cpp
--- a/Engine/Source/Runtime/Engine/Private/VT/VirtualTextureDataBuilder.cpp
+++ b/Engine/Source/Runtime/Engine/Private/VT/VirtualTextureDataBuilder.cpp
@@ -352,15 +352,15 @@ bool FVirtualTextureBuilderDerivedInfo::InitializeFromBuildSettings(const FTextu
 	//	  or is it only on the actual output data size?
 	//		(note that SizeX is the VT canvas size, not a pixel count, when you have UDIM blocks where not all tiles are present)
 	//		(eg. see "bigoffsets" test case)
-	if ( (int64)SizeX * SizeY > (1ULL<<32) )
-	{
-		UE_LOG(LogVirtualTexturing,Warning,TEXT("InitializeFromBuildSettings failed : total pixel count over 4G "
-			"(%d x %d = %lld) [%s]"),
-			SizeX,SizeY,(int64)SizeX * SizeY,
-			*InSourceData.TextureFullName);
-
-		return false;
-	}
+	//if ( (int64)SizeX * SizeY > (1ULL<<32) )
+	//{
+	//	UE_LOG(LogVirtualTexturing,Warning,TEXT("InitializeFromBuildSettings failed : total pixel count over 4G "
+	//		"(%d x %d = %lld) [%s]"),
+	//		SizeX,SizeY,(int64)SizeX * SizeY,
+	//		*InSourceData.TextureFullName);
+
+	//	return false;
+	//}
 
 	const uint32 Size = FMath::Max(SizeX, SizeY);
 
