the AudioStreamCache needs to keep track of the number of time "Force Inline" sound waves get added now. Due to other changes I'm making, it's possible that the same sound can get added more then once by different assets.

The contract before was that you should only ever add one and remove one. Now you need to make sure you're removing as many times as you add.

#tests PIE

diff --git a/Engine/Source/Runtime/Engine/Private/AudioStreamingCache.cpp b/Engine/Source/Runtime/Engine/Private/AudioStreamingCache.cpp
--- a/Engine/Source/Runtime/Engine/Private/AudioStreamingCache.cpp
+++ b/Engine/Source/Runtime/Engine/Private/AudioStreamingCache.cpp
@@ -1026,25 +1026,47 @@ void FAudioChunkCache::AddForceInlineSoundWave(const FSoundWaveProxyPtr& SoundWa
 	{
 		return;
 	}
-
+	
 	FName Format = SoundWave->GetRuntimeFormat();
 	FByteBulkData* Data = SoundWave->GetCompressedData(Format);
-	ForceInlineMemoryCounterBytes += Data ? Data->GetBulkDataSize() : 0;
-
-	const uint64 MemoryUsageBytes = GetCurrentMemoryUsageBytes();
-	if (TrimCacheWhenOverBudgetCVar != 0 && MemoryUsageBytes > MemoryLimitBytes)
+	int64 MemoryCount = Data ? Data->GetBulkDataSize() : 0;
+	int32 RefCount = 0;
 	{
-		uint64 MemoryToTrim = 0;
-		if (MemoryLimitTrimPercentageCVar > 0.0f)
+		FScopeLock Lock(&SoundWaveMemoryTrackerCritSec);
+		FSoundWaveMemoryTracker& Tracker = SoundWaveTracker.FindOrAdd(SoundWave);
+		checkf(Tracker.RefCount >= 0, TEXT("AudioStreamCache::AddForceInlineSoundWave: ref count for Added sound wave is negative!: %s"), *SoundWave->GetFName().ToString());
+		RefCount = ++Tracker.RefCount;
+		if (RefCount == 1)
 		{
-			MemoryToTrim = MemoryLimitBytes * FMath::Min(MemoryLimitTrimPercentageCVar, 1.0f);
+			// set the tracker memory count to that of the sound wave
+			Tracker.MemoryCount = MemoryCount;
 		}
 		else
 		{
-			MemoryToTrim = MemoryUsageBytes - MemoryLimitBytes;
+			// use the memory count set by the tracker.
+			MemoryCount = Tracker.MemoryCount;
 		}
+	}
 
-		TrimMemory(MemoryToTrim, true);
+	// we only increment memory count for the first time the sound wave is added
+	if (RefCount == 1)
+	{
+		ForceInlineMemoryCounterBytes += MemoryCount;
+		const uint64 MemoryUsageBytes = GetCurrentMemoryUsageBytes();
+		if (TrimCacheWhenOverBudgetCVar != 0 && MemoryUsageBytes > MemoryLimitBytes)
+		{
+			uint64 MemoryToTrim = 0;
+			if (MemoryLimitTrimPercentageCVar > 0.0f)
+			{
+				MemoryToTrim = MemoryLimitBytes * FMath::Min(MemoryLimitTrimPercentageCVar, 1.0f);
+			}
+			else
+			{
+				MemoryToTrim = MemoryUsageBytes - MemoryLimitBytes;
+			}
+
+			TrimMemory(MemoryToTrim, true);
+		}
 	}
 }
 
@@ -1059,9 +1081,30 @@ void FAudioChunkCache::RemoveForceInlineSoundWave(const FSoundWaveProxyPtr& Soun
 		return;
 	}
 
-	FName Format = SoundWave->GetRuntimeFormat();
-	FByteBulkData* Data = SoundWave->GetCompressedData(Format);
-	ForceInlineMemoryCounterBytes -= Data ? Data->GetBulkDataSize() : 0;
+	int64 MemoryCount = 0;
+	int32 RefCount = 0;
+
+	// scope lock
+	{
+		FScopeLock Lock(&SoundWaveMemoryTrackerCritSec);
+		SoundWaveMemoryTrackerCritSec.Lock();
+
+		FSoundWaveMemoryTracker* Tracker = SoundWaveTracker.Find(SoundWave);
+		checkf(Tracker != nullptr, TEXT("AudioStreamCache::RemoveForceInlineSoundWave: Attempted to remove SoundWave that was never added, or has already been removed: %s"), *SoundWave->GetFName().ToString());
+		checkf(Tracker->RefCount > 0, TEXT("AudioStreamCache::RemoveForceInlineSoundWve: Attempted to remove SoundWave that has a ref count of zero or less. Something has gone horribly wrong: %s"), *SoundWave->GetFName().ToString());
+		MemoryCount = Tracker->MemoryCount;
+		RefCount = --Tracker->RefCount;
+		// use the memory count we cached from the last sound wave add
+		if (RefCount == 0)
+		{
+			SoundWaveTracker.Remove(SoundWave);
+		}
+	}
+	
+	if (RefCount == 0)
+	{
+		ForceInlineMemoryCounterBytes -= MemoryCount;
+	}
 }
 
 void FAudioChunkCache::AddMemoryCountedFeature(const FAudioStreamCacheMemoryHandle& Feature)
diff --git a/Engine/Source/Runtime/Engine/Public/AudioStreamingCache.h b/Engine/Source/Runtime/Engine/Public/AudioStreamingCache.h
--- a/Engine/Source/Runtime/Engine/Public/AudioStreamingCache.h
+++ b/Engine/Source/Runtime/Engine/Public/AudioStreamingCache.h
@@ -420,6 +420,16 @@ private:
 	// Map that USoundWaves, FSoundWaveProxys, FChunkKeys, and FAudioChunkHandles can use to
 	// quickly lookup where their chunks are currently stored in the cache
 	TMap<FChunkKey, uint64> CacheLookupIdMap;
+
+	struct FSoundWaveMemoryTracker
+	{
+		int32 RefCount = 0;
+		int64 MemoryCount = 0;
+	};
+	// A map to look up the number of times a sound wave has been added for memory tracking
+	// as well as the memory added for the "latest" addition
+	TMap<FSoundWaveProxyPtr, FSoundWaveMemoryTracker> SoundWaveTracker;
+	FCriticalSection SoundWaveMemoryTrackerCritSec;
 	 
 	// This struct is used for logging cache misses.
 	struct FCacheMissInfo
