World Partition Optimization: avoid using IWorldPartitionCell interface functions from UWorldPartitionRuntimeCell code to favorize inlining rather than virtual functions.

diff --git a/Engine/Source/Runtime/Engine/Private/Actor.cpp b/Engine/Source/Runtime/Engine/Private/Actor.cpp
--- a/Engine/Source/Runtime/Engine/Private/Actor.cpp
+++ b/Engine/Source/Runtime/Engine/Private/Actor.cpp
@@ -6383,7 +6383,7 @@ bool AActor::HasDataLayers() const
 {
 	if (const IWorldPartitionCell* Cell = GetWorldPartitionRuntimeCell())
 	{
-		return Cell->HasDataLayers();
+		return Cell->GetDataLayers().Num() > 0;
 	}
 
 #if WITH_EDITOR
diff --git a/Engine/Source/Runtime/Engine/Private/WorldPartition/WorldPartitionRuntimeCell.cpp b/Engine/Source/Runtime/Engine/Private/WorldPartition/WorldPartitionRuntimeCell.cpp
--- a/Engine/Source/Runtime/Engine/Private/WorldPartition/WorldPartitionRuntimeCell.cpp
+++ b/Engine/Source/Runtime/Engine/Private/WorldPartition/WorldPartitionRuntimeCell.cpp
@@ -136,7 +136,7 @@ int32 UWorldPartitionRuntimeCell::SortCompare(const UWorldPartitionRuntimeCell*
 bool UWorldPartitionRuntimeCell::IsDebugShown() const
 {
 	return FWorldPartitionDebugHelper::IsDebugStreamingStatusShown(GetStreamingStatus()) &&
-	       FWorldPartitionDebugHelper::AreDebugDataLayersShown(GetDataLayers()) &&
+	       FWorldPartitionDebugHelper::AreDebugDataLayersShown(GetDataLayersInline()) &&
 		   (FWorldPartitionDebugHelper::CanDrawContentBundles() || !ContentBundleID.IsValid()) &&
 			RuntimeCellData->IsDebugShown();
 }
@@ -216,7 +216,7 @@ EDataLayerRuntimeState UWorldPartitionRuntimeCell::GetCellEffectiveWantedState()
 TArray<const UDataLayerInstance*> UWorldPartitionRuntimeCell::GetDataLayerInstances() const
 {
 	const UDataLayerManager* DataLayerManager = HasDataLayers() ? GetDataLayerManager() : nullptr;
-	return DataLayerManager ? DataLayerManager->GetDataLayerInstances(GetDataLayers()) : TArray<const UDataLayerInstance*>();
+	return DataLayerManager ? DataLayerManager->GetDataLayerInstances(GetDataLayersInline()) : TArray<const UDataLayerInstance*>();
 }
 
 const UExternalDataLayerInstance* UWorldPartitionRuntimeCell::GetExternalDataLayerInstance() const
@@ -239,7 +239,7 @@ bool UWorldPartitionRuntimeCell::HasContentBundle() const
 
 bool UWorldPartitionRuntimeCell::ContainsDataLayer(const UDataLayerInstance* DataLayerInstance) const
 {
-	return GetDataLayers().Contains(DataLayerInstance->GetDataLayerFName());
+	return GetDataLayersInline().Contains(DataLayerInstance->GetDataLayerFName());
 }
 
 FName UWorldPartitionRuntimeCell::GetLevelPackageName() const
diff --git a/Engine/Source/Runtime/Engine/Public/WorldPartition/WorldPartitionRuntimeCell.h b/Engine/Source/Runtime/Engine/Public/WorldPartition/WorldPartitionRuntimeCell.h
--- a/Engine/Source/Runtime/Engine/Public/WorldPartition/WorldPartitionRuntimeCell.h
+++ b/Engine/Source/Runtime/Engine/Public/WorldPartition/WorldPartitionRuntimeCell.h
@@ -272,7 +272,7 @@ class UWorldPartitionRuntimeCell : public UObject, public IWorldPartitionCell, p
 	ENGINE_API virtual bool ContainsDataLayer(const UDataLayerAsset* DataLayerAsset) const override;
 	ENGINE_API virtual bool ContainsDataLayer(const UDataLayerInstance* DataLayerInstance) const override;
 	ENGINE_API virtual bool HasContentBundle() const override;
-	virtual const TArray<FName>& GetDataLayers() const override { return DataLayers.ToArray(); }
+	virtual const TArray<FName>& GetDataLayers() const override { return GetDataLayersInline(); }
 	virtual FName GetExternalDataLayer() const override { return DataLayers.GetExternalDataLayer(); }
 	virtual bool HasAnyDataLayer(const TSet<FName>& InDataLayers) const override
 	{
@@ -286,6 +286,7 @@ class UWorldPartitionRuntimeCell : public UObject, public IWorldPartitionCell, p
 
 	ENGINE_API UDataLayerManager* GetDataLayerManager() const;
 	ENGINE_API EDataLayerRuntimeState GetCellEffectiveWantedState() const;
+	FORCEINLINE bool HasDataLayers() const { return !DataLayers.IsEmpty(); }
 
 	void SetBlockOnSlowLoading(bool bInBlockOnSlowLoading) { bBlockOnSlowLoading = bInBlockOnSlowLoading; }
 	bool GetBlockOnSlowLoading() const { return bBlockOnSlowLoading; }
@@ -352,6 +353,8 @@ class UWorldPartitionRuntimeCell : public UObject, public IWorldPartitionCell, p
 #endif
 
 protected:
+	FORCEINLINE const TArray<FName>& GetDataLayersInline() const { return DataLayers.ToArray(); }
+
 	//@todo_ow: Implement ServerOnlyVisible and refactor ClientOnlyVisible.
 	//          Instead of this function, server would not not wait for client level visibility 
 	//          for server-only visible cells.
diff --git a/Engine/Source/Runtime/Engine/Public/WorldPartition/WorldPartitionRuntimeCellInterface.h b/Engine/Source/Runtime/Engine/Public/WorldPartition/WorldPartitionRuntimeCellInterface.h
--- a/Engine/Source/Runtime/Engine/Public/WorldPartition/WorldPartitionRuntimeCellInterface.h
+++ b/Engine/Source/Runtime/Engine/Public/WorldPartition/WorldPartitionRuntimeCellInterface.h
@@ -28,8 +28,6 @@ public:
 	virtual bool ContainsDataLayer(const UDataLayerAsset* DataLayerAsset) const = 0;
 	/** Returns whether the cell data layers referenced the provided data layer instance or not. */
 	virtual bool ContainsDataLayer(const UDataLayerInstance* DataLayerInstance) const = 0;
-	/** Returns whether the cell content is associated with data layers or not. */
-	bool HasDataLayers() const { return !GetDataLayers().IsEmpty(); }
 	/** Returns whether the cell content is associated with content bundle or not. */
 	virtual bool HasContentBundle() const = 0;
 	/** Returns the cell content associated data layers. */
