Add explicit capabilities for drums or guitar devices

#rnx

diff --git a/Engine/Source/Runtime/Engine/Classes/GameFramework/InputSettings.h b/Engine/Source/Runtime/Engine/Classes/GameFramework/InputSettings.h
--- a/Engine/Source/Runtime/Engine/Classes/GameFramework/InputSettings.h
+++ b/Engine/Source/Runtime/Engine/Classes/GameFramework/InputSettings.h
@@ -360,6 +360,12 @@ namespace EHardwareDeviceSupportedFeatures
 		/** This device can track the orientation in world space, such as a gyroscope */
 		Orientation			= 0x00002000,
 
+		/** This device has the capabilities of a guitar (whammy bar, tilt, etc) */
+		Guitar				= 0x00004000,
+
+		/** This device has the capabilities of drums (symbols, foot pedal, etc) */
+		Drums				= 0x00008000,
+
 		/** Some custom flags that can be used in your game if you have custom hardware! */
 		CustomA				= 0x01000000,
 		CustomB				= 0x02000000,
