[StaticMeshComponent] removed registration on mesh prebuild delegate recently added since it was not thread safe (i.e. async loading)
The issue this was trying to fix can no longer be reproduced.

diff --git a/Engine/Source/Runtime/Engine/Classes/Components/StaticMeshComponent.h b/Engine/Source/Runtime/Engine/Classes/Components/StaticMeshComponent.h
--- a/Engine/Source/Runtime/Engine/Classes/Components/StaticMeshComponent.h
+++ b/Engine/Source/Runtime/Engine/Classes/Components/StaticMeshComponent.h
@@ -131,7 +131,6 @@ private:
 #if WITH_EDITOR
 	/** Used to track down when StaticMesh has been modified for notification purpose */
 	class UStaticMesh* KnownStaticMesh = nullptr;
-	FDelegateHandle PreMeshBuildDelegateHandle;
 #endif
 	ENGINE_API void NotifyIfStaticMeshChanged();
 
diff --git a/Engine/Source/Runtime/Engine/Private/Components/StaticMeshComponent.cpp b/Engine/Source/Runtime/Engine/Private/Components/StaticMeshComponent.cpp
--- a/Engine/Source/Runtime/Engine/Private/Components/StaticMeshComponent.cpp
+++ b/Engine/Source/Runtime/Engine/Private/Components/StaticMeshComponent.cpp
@@ -661,9 +661,6 @@ void UStaticMeshComponent::NotifyIfStaticMeshChanged()
 		// Remove delegates from our previous mesh 
 		if (KnownStaticMesh)
 		{
-			KnownStaticMesh->OnPreMeshBuild().Remove(PreMeshBuildDelegateHandle);
-			PreMeshBuildDelegateHandle.Reset();
-
 			if (bMipLevelCallbackRegistered)
 			{
 				KnownStaticMesh->RemoveMipLevelChangeCallback(this);
@@ -673,22 +670,6 @@ void UStaticMeshComponent::NotifyIfStaticMeshChanged()
 
 		KnownStaticMesh = StaticMesh;
 
-		// Add delegate on our new mesh
-		if (KnownStaticMesh)
-		{
-			PreMeshBuildDelegateHandle = KnownStaticMesh->OnPreMeshBuild().AddWeakLambda(this, [this](UStaticMesh*)
-				{
-					// Our associated mesh is compiling so we are no longer navigation relevant.
-					// This will invalidate any pending add to the octree and dirty tiles until compilation completes.
-					// Note that we force the value instead of calling IsNavigationRelevant since the IsCompiling
-					// condition will not return the proper value in all code paths where OnPreMeshBuild delegate is called.
-					// (e.g. async task created right after calling OnPreMeshBuild)
-					// Member 'bNavigationRelevant' will be updated by PostStaticMeshCompilation.
-					bNavigationRelevant = false;
-					FNavigationSystem::UpdateComponentData(*this);
-				});
-		}
-
 		FObjectCacheEventSink::NotifyStaticMeshChanged_Concurrent(GetStaticMeshComponentInterface());
 
 		// Update this component streaming data.
@@ -1853,11 +1834,6 @@ void UStaticMeshComponent::BeginDestroy()
 #if WITH_EDITOR
 	// The object cache needs to be notified when we're getting destroyed
 	FObjectCacheEventSink::NotifyStaticMeshChanged_Concurrent(GetStaticMeshComponentInterface());
-	
-	if (PreMeshBuildDelegateHandle.IsValid())
-	{
-		KnownStaticMesh->OnPreMeshBuild().Remove(PreMeshBuildDelegateHandle);
-	}
 #endif // WITH_EDITOR
 }
 
