World Partition Debug Display: only show content bundles with actual content.

diff --git a/Engine/Source/Runtime/Engine/Private/WorldPartition/ContentBundle/ContentBundle.cpp b/Engine/Source/Runtime/Engine/Private/WorldPartition/ContentBundle/ContentBundle.cpp
--- a/Engine/Source/Runtime/Engine/Private/WorldPartition/ContentBundle/ContentBundle.cpp
+++ b/Engine/Source/Runtime/Engine/Private/WorldPartition/ContentBundle/ContentBundle.cpp
@@ -104,6 +104,11 @@ bool FContentBundle::IsValid() const
 	return GetDescriptor()->IsValid();
 }
 
+bool FContentBundle::HasContent() const
+{
+	return !!ExternalStreamingObject;
+}
+
 #if WITH_EDITOR
 void FContentBundle::InitializeForPIE()
 {
diff --git a/Engine/Source/Runtime/Engine/Private/WorldPartition/ContentBundle/ContentBundleEditor.cpp b/Engine/Source/Runtime/Engine/Private/WorldPartition/ContentBundle/ContentBundleEditor.cpp
--- a/Engine/Source/Runtime/Engine/Private/WorldPartition/ContentBundle/ContentBundleEditor.cpp
+++ b/Engine/Source/Runtime/Engine/Private/WorldPartition/ContentBundle/ContentBundleEditor.cpp
@@ -161,9 +161,12 @@ void FContentBundleEditor::DoRemoveContent()
 
 bool FContentBundleEditor::IsValid() const
 {
-	bool bIsValid = true;
+	return true;
+}
 
-	return bIsValid;
+bool FContentBundleEditor::HasContent() const
+{
+	return !!ExternalStreamingObject;
 }
 
 bool FContentBundleEditor::AddActor(AActor* InActor)
diff --git a/Engine/Source/Runtime/Engine/Private/WorldPartition/ContentBundle/ContentBundleWorldSubsystem.cpp b/Engine/Source/Runtime/Engine/Private/WorldPartition/ContentBundle/ContentBundleWorldSubsystem.cpp
--- a/Engine/Source/Runtime/Engine/Private/WorldPartition/ContentBundle/ContentBundleWorldSubsystem.cpp
+++ b/Engine/Source/Runtime/Engine/Private/WorldPartition/ContentBundle/ContentBundleWorldSubsystem.cpp
@@ -109,7 +109,10 @@ void UContentBundleManager::DrawContentBundlesStatus(const UWorld* InWorld, UCan
 
 				(*ContentBundleContainer)->ForEachContentBundle([&Offset, &Canvas, &Pos, &MaxTextWidth](FContentBundleBase* ContentBundle)
 				{
-					FWorldPartitionDebugHelper::DrawLegendItem(Canvas, *FString::Printf(TEXT("%s [%s]"), *ContentBundle->GetDisplayName(), *UContentBundleDescriptor::GetContentBundleCompactString(ContentBundle->GetDescriptor()->GetGuid())), GEngine->GetSmallFont(), ContentBundle->GetDebugColor(), FColor::White, Pos, &MaxTextWidth);
+					if (ContentBundle->HasContent())
+					{
+						FWorldPartitionDebugHelper::DrawLegendItem(Canvas, *FString::Printf(TEXT("%s [%s]"), *ContentBundle->GetDisplayName(), *UContentBundleDescriptor::GetContentBundleCompactString(ContentBundle->GetDescriptor()->GetGuid())), GEngine->GetSmallFont(), ContentBundle->GetDebugColor(), FColor::White, Pos, &MaxTextWidth);
+					}
 				});
 			}
 
diff --git a/Engine/Source/Runtime/Engine/Public/WorldPartition/ContentBundle/ContentBundle.h b/Engine/Source/Runtime/Engine/Public/WorldPartition/ContentBundle/ContentBundle.h
--- a/Engine/Source/Runtime/Engine/Public/WorldPartition/ContentBundle/ContentBundle.h
+++ b/Engine/Source/Runtime/Engine/Public/WorldPartition/ContentBundle/ContentBundle.h
@@ -19,8 +19,8 @@ public:
 
 	//~ Begin IContentBundle Interface
 	ENGINE_API virtual void AddReferencedObjects(FReferenceCollector& Collector) override;
-
 	ENGINE_API virtual bool IsValid() const override;
+	ENGINE_API virtual bool HasContent() const override;
 	//~ End IContentBundle Interface
 
 protected:
diff --git a/Engine/Source/Runtime/Engine/Public/WorldPartition/ContentBundle/ContentBundleBase.h b/Engine/Source/Runtime/Engine/Public/WorldPartition/ContentBundle/ContentBundleBase.h
--- a/Engine/Source/Runtime/Engine/Public/WorldPartition/ContentBundle/ContentBundleBase.h
+++ b/Engine/Source/Runtime/Engine/Public/WorldPartition/ContentBundle/ContentBundleBase.h
@@ -24,6 +24,7 @@ public:
 	ENGINE_API virtual void AddReferencedObjects(FReferenceCollector& Collector);
 
 	virtual bool IsValid() const = 0;
+	virtual bool HasContent() const = 0;
 
 	ENGINE_API UWorld* GetInjectedWorld() const;
 	ENGINE_API const FString& GetDisplayName() const;
diff --git a/Engine/Source/Runtime/Engine/Public/WorldPartition/ContentBundle/ContentBundleEditor.h b/Engine/Source/Runtime/Engine/Public/WorldPartition/ContentBundle/ContentBundleEditor.h
--- a/Engine/Source/Runtime/Engine/Public/WorldPartition/ContentBundle/ContentBundleEditor.h
+++ b/Engine/Source/Runtime/Engine/Public/WorldPartition/ContentBundle/ContentBundleEditor.h
@@ -33,8 +33,8 @@ public:
 
 	//~ Begin IContentBundle Interface
 	ENGINE_API virtual void AddReferencedObjects(FReferenceCollector& Collector) override;
-
 	ENGINE_API virtual bool IsValid() const override;
+	ENGINE_API virtual bool HasContent() const override;
 	//~ End IContentBundle Interface
 
 	const FGuid& GetTreeItemID() const { return TreeItemID; }
