Default texture is only used in editor.

diff --git a/Engine/Source/Runtime/Engine/Private/Texture2D.cpp b/Engine/Source/Runtime/Engine/Private/Texture2D.cpp
--- a/Engine/Source/Runtime/Engine/Private/Texture2D.cpp
+++ b/Engine/Source/Runtime/Engine/Private/Texture2D.cpp
@@ -225,6 +225,7 @@ int64 FTexture2DMipMap::StoreInDerivedDataCache(const FStringView InKey, const F
 /**
  * Get the optimal placeholder to use during texture compilation
  */ 
+#if WITH_EDITOR
 static UTexture2D* GetDefaultTexture2D(const UTexture2D* Texture)
 {
 	static TStrongObjectPtr<UTexture2D> CheckerboardTexture;
@@ -287,6 +288,7 @@ static UTexture2D* GetDefaultTexture2D(const UTexture2D* Texture)
 
 	return CheckerboardTexture.Get();
 }
+#endif
 
 FTexturePlatformData** UTexture2D::GetRunningPlatformData()
 {
