World Partition Optimization: avoid computing cell to source twice and remove a superfluous vector normalization in UWorldPartitionRuntimeCellData::AppendStreamingSourceInfo.

diff --git a/Engine/Source/Runtime/Engine/Private/WorldPartition/WorldPartitionRuntimeCellData.cpp b/Engine/Source/Runtime/Engine/Private/WorldPartition/WorldPartitionRuntimeCellData.cpp
--- a/Engine/Source/Runtime/Engine/Private/WorldPartition/WorldPartitionRuntimeCellData.cpp
+++ b/Engine/Source/Runtime/Engine/Private/WorldPartition/WorldPartitionRuntimeCellData.cpp
@@ -41,14 +41,20 @@ void UWorldPartitionRuntimeCellData::AppendStreamingSourceInfo(const FWorldParti
 		check(CachedSourceInfoEpoch == StreamingSourceCacheEpoch);
 	}
 
+	// Compute cosine angle from cell to source direction ratio
+	const FVector CellToSource = SourceShape.GetCenter() - ContentBounds.GetClosestPointTo(SourceShape.GetCenter());
+	const double CellToSourceSquareDistance = CellToSource.SizeSquared();
+	const FVector CellToSourceNormal = FMath::IsNearlyZero(CellToSourceSquareDistance) ? FVector::ZeroVector : (CellToSource * FMath::InvSqrt(CellToSourceSquareDistance));
+	const FVector SourceAxis = Source.bUseVelocityContributionToCellsSorting ? FVector(SourceShape.GetAxis() + Source.Velocity).GetSafeNormal() : SourceShape.GetAxis();
+	const float SourceCosAngle = ContentBounds.IsInsideOrOn(SourceShape.GetCenter()) ? -1.0f : (SourceAxis | CellToSourceNormal);
+	const float SourceCosAngleRatio = SourceCosAngle * 0.5f + 0.5f;
+
 	CachedMinSourcePriority = FMath::Min((uint8)Source.Priority, CachedMinSourcePriority);
 
 	if (Source.bBlockOnSlowLoading)
 	{
 		bCachedWasRequestedByBlockingSource = true;
 
-		const FVector CellToSource = SourceShape.GetCenter() - ContentBounds.GetClosestPointTo(SourceShape.GetCenter());
-		const double CellToSourceSquareDistance = CellToSource.SizeSquared();
 		CachedMinSquareDistanceToBlockingSource = FMath::Min(CellToSourceSquareDistance, CachedMinSquareDistanceToBlockingSource);
 
 		const double BlockOnSlowStreamingRatio = FMath::Sqrt(CachedMinSquareDistanceToBlockingSource) / SourceShape.GetRadius();
@@ -58,12 +64,6 @@ void UWorldPartitionRuntimeCellData::AppendStreamingSourceInfo(const FWorldParti
 	// Compute square distance from cell to source ratio
 	const double SoureDistanceRatio = FMath::Clamp(ContentBounds.ComputeSquaredDistanceToPoint(SourceShape.GetCenter()) / FMath::Square(SourceShape.GetRadius()), 0.0f, 1.0f);
 
-	// Compute cosine angle from cell to source direction ratio
-	const FVector CellToSource = SourceShape.GetCenter() - ContentBounds.GetClosestPointTo(SourceShape.GetCenter());
-	const FVector SourceAxis = FVector(SourceShape.GetAxis() + Source.Velocity * (Source.bUseVelocityContributionToCellsSorting ? 1.0f : 0.0f)).GetSafeNormal();
-	const float SourceCosAngle = ContentBounds.IsInsideOrOn(SourceShape.GetCenter()) ? -1.0f : (SourceAxis | CellToSource.GetSafeNormal());
-	const float SourceCosAngleRatio = SourceCosAngle * 0.5f + 0.5f;
-
 	// Compute final cell priority for this source
 	const double SortingPriority = SoureDistanceRatio * SourceCosAngleRatio;
 
