Level Instance: fix a crash when reusing an already loaded level, regression introduced in 31717976.

diff --git a/Engine/Source/Runtime/Engine/Private/LevelInstance/LevelInstanceLevelStreaming.cpp b/Engine/Source/Runtime/Engine/Private/LevelInstance/LevelInstanceLevelStreaming.cpp
--- a/Engine/Source/Runtime/Engine/Private/LevelInstance/LevelInstanceLevelStreaming.cpp
+++ b/Engine/Source/Runtime/Engine/Private/LevelInstance/LevelInstanceLevelStreaming.cpp
@@ -523,11 +523,12 @@ void ULevelStreamingLevelInstance::OnLevelLoadedChanged(ULevel* InLevel)
 		}
 #endif
 		check(InLevel == NewLoadedLevel);
-		check(!NewLoadedLevel->bAlreadyMovedActors);
-		
-		AWorldSettings* WorldSettings = NewLoadedLevel->GetWorldSettings();
-		check(WorldSettings);
-		LevelTransform = FTransform(WorldSettings->LevelInstancePivotOffset) * LevelTransform;
+		if (!NewLoadedLevel->bAlreadyMovedActors)
+		{
+			AWorldSettings* WorldSettings = NewLoadedLevel->GetWorldSettings();
+			check(WorldSettings);
+			LevelTransform = FTransform(WorldSettings->LevelInstancePivotOffset) * LevelTransform;
+		}
 
 		if (ULevelInstanceSubsystem* LevelInstanceSubsystem = GetWorld()->GetSubsystem<ULevelInstanceSubsystem>())
 		{
