Flush resources when drawing with minimized windows

#tests Standalone memory leak fix confirmed

diff --git a/Engine/Source/Editor/UnrealEd/Classes/Editor/EditorEngine.h b/Engine/Source/Editor/UnrealEd/Classes/Editor/EditorEngine.h
--- a/Engine/Source/Editor/UnrealEd/Classes/Editor/EditorEngine.h
+++ b/Engine/Source/Editor/UnrealEd/Classes/Editor/EditorEngine.h
@@ -1775,12 +1775,6 @@ public:
 	 */
 	UNREALED_API bool IsAnyViewportRealtime();
 
-
-	/**
-	 * @return true if all windows are hidden (including minimized)                                                         
-	 */
-	UNREALED_API bool AreAllWindowsHidden() const;
-
 	/**
 	 *	Returns pointer to a temporary render target.
 	 *	If it has not already been created, does so here.
diff --git a/Engine/Source/Editor/UnrealEd/Private/EditorEngine.cpp b/Engine/Source/Editor/UnrealEd/Private/EditorEngine.cpp
--- a/Engine/Source/Editor/UnrealEd/Private/EditorEngine.cpp
+++ b/Engine/Source/Editor/UnrealEd/Private/EditorEngine.cpp
@@ -6535,27 +6535,6 @@ bool UEditorEngine::ShouldThrottleCPUUsage() const
 	return bShouldThrottle;
 }
 
-bool UEditorEngine::AreAllWindowsHidden() const
-{
-	if (!FSlateApplication::IsInitialized())
-	{
-		return true;
-	}
-	const TArray< TSharedRef<SWindow> > AllWindows = FSlateApplication::Get().GetInteractiveTopLevelWindows();
-
-	bool bAllHidden = true;
-	for( const TSharedRef<SWindow>& Window : AllWindows )
-	{
-		if( !Window->IsWindowMinimized() && Window->IsVisible() )
-		{
-			bAllHidden = false;
-			break;
-		}
-	}
-
-	return bAllHidden;
-}
-
 AActor* UEditorEngine::AddActor(ULevel* InLevel, UClass* Class, const FTransform& Transform, bool bSilent, EObjectFlags InObjectFlags, bool bSelectActor)
 {
 	check( Class );
diff --git a/Engine/Source/Runtime/Engine/Classes/Engine/Engine.h b/Engine/Source/Runtime/Engine/Classes/Engine/Engine.h
--- a/Engine/Source/Runtime/Engine/Classes/Engine/Engine.h
+++ b/Engine/Source/Runtime/Engine/Classes/Engine/Engine.h
@@ -2454,6 +2454,11 @@ public:
 	 */
 	ENGINE_API virtual bool ShouldThrottleCPUUsage() const;
 
+	/**
+	 * @return true if all windows are minimized or hidden (Per OS definition)
+	 */
+	ENGINE_API bool AreAllWindowsHidden() const;
+
 public:
 	/** 
 	 * Return a reference to the GamePlayers array. 
diff --git a/Engine/Source/Runtime/Engine/Private/GameEngine.cpp b/Engine/Source/Runtime/Engine/Private/GameEngine.cpp
--- a/Engine/Source/Runtime/Engine/Private/GameEngine.cpp
+++ b/Engine/Source/Runtime/Engine/Private/GameEngine.cpp
@@ -100,6 +100,13 @@ static FAutoConsoleVariableRef CVarDoAsyncEndOfFrameTasks(
 	TEXT("Experimental option to run various things concurrently with the HUD render.")
 	);
 
+static int32 GMinimizedSyncDrawToGPU = 1;
+static FAutoConsoleVariableRef CVarMinimizedSyncDrawToGPU(
+	TEXT("tick.MinimizedSyncDrawToGPU"),
+	GMinimizedSyncDrawToGPU,
+	TEXT("True means we will wait for GPU idle when minimized. Prevents mem leaks due to CPU issuing draws faster than GPU processes when minimized.")
+);
+
 bool ParseResolution(const TCHAR* InResolution, uint32& OutX, uint32& OutY, int32& WindowMode);
 
 /** Benchmark results to the log */
@@ -1894,8 +1901,21 @@ void UGameEngine::Tick( float DeltaSeconds, bool bIdleMode )
 			// Render everything.
 			RedrawViewports();
 
-			// Some tasks can only be done once we finish all scenes/viewports
-			GetRendererModule().PostRenderAllViewports();
+			// CPU/GPU synchronization is achieved by calling EndDrawingViewport. If no viewports are updated (because the game is hidden),
+			// we need to explicitly wait for the GPU to finish here, to prevent the CPU from submitting work faster than the GPU
+			// can process it, which leads to an unbounded accumulation of resources.
+			if (GMinimizedSyncDrawToGPU && AreAllWindowsHidden())
+			{
+				ENQUEUE_RENDER_COMMAND(SubmitAndBlockUntilGPUIdle_MinimizedRealtime)([](FRHICommandListImmediate& RHICmdList)
+				{
+					RHICmdList.BlockUntilGPUIdle();
+				});
+			}
+			else
+			{
+				// Some tasks can only be done once we finish all scenes/viewports
+				GetRendererModule().PostRenderAllViewports();
+			}
 		}
 		else
 		{
diff --git a/Engine/Source/Runtime/Engine/Private/UnrealEngine.cpp b/Engine/Source/Runtime/Engine/Private/UnrealEngine.cpp
--- a/Engine/Source/Runtime/Engine/Private/UnrealEngine.cpp
+++ b/Engine/Source/Runtime/Engine/Private/UnrealEngine.cpp
@@ -78,6 +78,7 @@ UnrealEngine.cpp: Implements the UEngine class and helpers.
 #include "UObject/GarbageCollectionHistory.h"
 #include "UObject/ReachabilityAnalysis.h"
 #include "Widgets/Input/SButton.h"
+#include "Widgets/SWindow.h"
 #include "Engine/EngineCustomTimeStep.h"
 #include "Engine/LevelStreamingPersistent.h"
 #include "Engine/ObjectReferencer.h"
@@ -12208,6 +12209,27 @@ bool UEngine::ShouldThrottleCPUUsage() const
 	return false;
 }
 
+bool UEngine::AreAllWindowsHidden() const
+{
+	if (!FSlateApplication::IsInitialized())
+	{
+		return true;
+	}
+	const TArray<TSharedRef<SWindow>> AllWindows = FSlateApplication::Get().GetInteractiveTopLevelWindows();
+
+	bool bAllHidden = true;
+	for (const TSharedRef<SWindow>& Window : AllWindows)
+	{
+		if (!Window->IsWindowMinimized() && Window->IsVisible())
+		{
+			bAllHidden = false;
+			break;
+		}
+	}
+
+	return bAllHidden;
+}
+
 #if !UE_BUILD_SHIPPING
 /**
 *	Renders warnings about the level that should be addressed prior to shipping
