Skeletal Mesh: When converting from import model to mesh description, don't ditch all normals/tangents even if some are broken. Instead, re-create the broken normals/tangents as needed.

#rnx

diff --git a/Engine/Source/Runtime/Engine/Private/Rendering/SkeletalMeshLODImporterData.cpp b/Engine/Source/Runtime/Engine/Private/Rendering/SkeletalMeshLODImporterData.cpp
--- a/Engine/Source/Runtime/Engine/Private/Rendering/SkeletalMeshLODImporterData.cpp
+++ b/Engine/Source/Runtime/Engine/Private/Rendering/SkeletalMeshLODImporterData.cpp
@@ -2,8 +2,6 @@
 
 #include "Rendering/SkeletalMeshLODImporterData.h"
 
-#include "UObject/Package.h"
-
 #if WITH_EDITOR
 
 #include "Algo/AnyOf.h"
@@ -1904,7 +1902,7 @@ static void CopySkinWeightsToAttribute(
 	}
 }
 
-bool FSkeletalMeshImportData::GetMeshDescription(const USkeletalMesh* InSkeletalMesh, FMeshDescription& OutMeshDescription) const
+bool FSkeletalMeshImportData::GetMeshDescription(const USkeletalMesh* InSkeletalMesh, int32 InLODIndex, FMeshDescription& OutMeshDescription) const
 {
 	using namespace UE::AnimationCore;
 	
@@ -2202,33 +2200,6 @@ bool FSkeletalMeshImportData::GetMeshDescription(const USkeletalMesh* InSkeletal
 		MaterialGroups.Add(PolygonGroupID);
 	}
 
-	// These are not always honest.
-	bool bHaveValidNormals = bHasNormals;
-	bool bHaveValidTangents = bHasNormals;
-	
-	for (int32 TriangleIndex = 0; TriangleIndex < Faces.Num(); TriangleIndex++)
-	{
-		const SkeletalMeshImportData::FTriangle &Triangle = Faces[TriangleIndex];
-
-		for (int32 Corner = 0; Corner < 3; Corner++)
-		{
-			if (!Triangle.TangentZ[Corner].IsNormalized())
-			{
-				bHaveValidNormals = false;
-			}
-			if (!Triangle.TangentX[Corner].IsNormalized() ||
-			    !Triangle.TangentY[Corner].IsNormalized())
-			{
-				bHaveValidTangents = false;
-			}
-
-			if (!bHaveValidNormals && !bHaveValidTangents)
-			{
-				break;
-			}
-		}
-	}
-
 	for (int32 TriangleIndex = 0; TriangleIndex < Faces.Num(); TriangleIndex++)
 	{
 		const SkeletalMeshImportData::FTriangle &Triangle = Faces[TriangleIndex];
@@ -2254,16 +2225,16 @@ bool FSkeletalMeshImportData::GetMeshDescription(const USkeletalMesh* InSkeletal
 					VertexInstanceUVs.Set(VertexInstanceID, UVIndex, Wedge.UVs[UVIndex]);
 				}
 
-				if (bHaveValidNormals)
+				if (bHasNormals)
 				{
 					VertexInstanceNormals.Set(VertexInstanceID, Triangle.TangentZ[Corner]);
 				}
-				if (bHaveValidTangents)
+				if (bHasTangents)
 				{
 					VertexInstanceTangents.Set(VertexInstanceID, Triangle.TangentX[Corner]);
 
 					// We can only divine the bi-tangent sign if the normal is also given. 
-					if (bHaveValidNormals)
+					if (bHasNormals)
 					{
 						VertexInstanceBinormalSigns.Set(VertexInstanceID,
 							((Triangle.TangentZ[Corner] ^ Triangle.TangentX[Corner]) | Triangle.TangentY[Corner]) < 0 ? -1.0f : 1.0f);
@@ -2393,7 +2364,41 @@ bool FSkeletalMeshImportData::GetMeshDescription(const USkeletalMesh* InSkeletal
 		}
 	}
 
-	FSkeletalMeshOperations::ConvertSmoothGroupToHardEdges(FaceSmoothingMasks, OutMeshDescription);
+	// We don't need to do this when retrieving alt skin weights.
+	if (InSkeletalMesh)
+	{
+		FSkeletalMeshOperations::ConvertSmoothGroupToHardEdges(FaceSmoothingMasks, OutMeshDescription);
+		FSkeletalMeshOperations::ValidateAndFixData(OutMeshDescription, *InSkeletalMesh->GetOuter()->GetPathName());
+
+		// This is required by compute normals to function correctly.
+		FSkeletalMeshOperations::ComputeTriangleTangentsAndNormals(OutMeshDescription, UE_SMALL_NUMBER, *InSkeletalMesh->GetOuter()->GetPathName());
+
+		EComputeNTBsFlags ComputeNTBsOptions = EComputeNTBsFlags::None; 
+		const FSkeletalMeshLODInfo* LODInfo = InSkeletalMesh->GetLODInfo(InLODIndex);
+		if (ensure(LODInfo))
+		{
+			if (LODInfo->BuildSettings.bComputeWeightedNormals)
+			{
+				ComputeNTBsOptions |= EComputeNTBsFlags::WeightedNTBs;
+			}
+			if (LODInfo->BuildSettings.bUseMikkTSpace)
+			{
+				ComputeNTBsOptions |= EComputeNTBsFlags::UseMikkTSpace;
+			}
+			if (LODInfo->BuildSettings.bRemoveDegenerates)
+			{
+				ComputeNTBsOptions |= EComputeNTBsFlags::IgnoreDegenerateTriangles;
+			}
+		}
+		else
+		{
+			// Default to all the things.
+			ComputeNTBsOptions = EComputeNTBsFlags::WeightedNTBs | EComputeNTBsFlags::UseMikkTSpace | EComputeNTBsFlags::IgnoreDegenerateTriangles; 
+		}
+
+		// This only recomputes broken normals/tangents (nan/zero)
+		FSkeletalMeshOperations::ComputeTangentsAndNormals(OutMeshDescription, ComputeNTBsOptions);
+	}
 
 	// Convert the MeshInfo data.
 	if (!MeshInfos.IsEmpty())
@@ -2495,7 +2500,7 @@ void FSkeletalMeshImportData::CopySkinWeightsToMeshDescription(
 
 	// The topologies don't match, proceed as above.
 	FMeshDescription AlternateInfluenceMesh;
-	InSkinWeightMesh.GetMeshDescription(nullptr, AlternateInfluenceMesh);
+	InSkinWeightMesh.GetMeshDescription(nullptr, 0, AlternateInfluenceMesh);
 	
 	FSkeletalMeshOperations::CopySkinWeightAttributeFromMesh(
 	AlternateInfluenceMesh, OutMeshDescription, NAME_None, InSkinWeightName, &AltMeshBoneToBaseBoneMap);   
@@ -2551,7 +2556,11 @@ static void ConvertHardEdgesToSmoothMasks(
 		CurrentSmoothMask <<= 1;
 		if (CurrentSmoothMask == 0)
 		{
-			CurrentSmoothMask = 1;
+			// If we exhausted all available bits, then thunk to the more complete algorithm. For reasons unknown at this time, it doesn't generate
+			// nice smooth groups for some simpler test objects. For more complex input products it does a decent job though.
+			OutSmoothMasks.SetNumZeroed(InMeshDescription.Triangles().Num());
+			FStaticMeshOperations::ConvertHardEdgesToSmoothGroup(InMeshDescription, OutSmoothMasks);
+			break;
 		}
 	}
 }
@@ -2727,20 +2736,23 @@ FSkeletalMeshImportData FSkeletalMeshImportData::CreateFromMeshDescription(const
 		}
 	}
 
-	
-	bool bHaveValidNormals = true;
-	bool bHaveValidTangents = true;
+	bool bHaveValidNormals = false;
+	bool bHaveValidTangents = false;
 	for (FVertexInstanceID VertexInstanceID: InMeshDescription.VertexInstances().GetElementIDs())
 	{
-		if (!VertexInstanceNormals.Get(VertexInstanceID).IsNormalized())
+		const FVector3f Normal = VertexInstanceNormals.Get(VertexInstanceID);
+		if (!Normal.IsNearlyZero(UE_SMALL_NUMBER) && !Normal.ContainsNaN())
 		{
-			bHaveValidNormals = false;
+			bHaveValidNormals = true;
 		}
-		if (!VertexInstanceTangents.Get(VertexInstanceID).IsNormalized())
+		const FVector3f Tangent = VertexInstanceTangents.Get(VertexInstanceID);
+		if (!Tangent.IsNearlyZero(UE_SMALL_NUMBER) && !Tangent.ContainsNaN())
 		{
-			bHaveValidTangents = false;
+			bHaveValidTangents = true;
 		}
-		if (!bHaveValidNormals && !bHaveValidTangents)
+
+		// If we found any valid normals/tangents, we can stop now.
+		if (bHaveValidNormals && bHaveValidTangents)
 		{
 			break;
 		}
@@ -2797,6 +2809,8 @@ FSkeletalMeshImportData FSkeletalMeshImportData::CreateFromMeshDescription(const
 			}
 			else
 			{
+				// The normal/tangent computation during conversion to render data will automatically regenerate any degenerate normals.
+				// Same for the tangent/binormal below.
 				Face.TangentZ[Corner] = FVector3f::ZeroVector;
 			}
 
@@ -2808,6 +2822,7 @@ FSkeletalMeshImportData FSkeletalMeshImportData::CreateFromMeshDescription(const
 			else
 			{
 				Face.TangentX[Corner] = FVector3f::ZeroVector;
+				Face.TangentY[Corner] = FVector3f::ZeroVector;
 			}
 
 			const int32 WedgeIndex = VertexInstanceID.GetValue();
diff --git a/Engine/Source/Runtime/Engine/Private/SkeletalMesh.cpp b/Engine/Source/Runtime/Engine/Private/SkeletalMesh.cpp
--- a/Engine/Source/Runtime/Engine/Private/SkeletalMesh.cpp
+++ b/Engine/Source/Runtime/Engine/Private/SkeletalMesh.cpp
@@ -2585,7 +2585,7 @@ void USkeletalMesh::SaveLODImportedData(const int32 LODIndex, const FSkeletalMes
 	
 	FMeshDescription MeshDescription;
 
-	if (InMesh.GetMeshDescription(nullptr, MeshDescription))
+	if (InMesh.GetMeshDescription(nullptr, LODIndex, MeshDescription))
 	{
 		CreateMeshDescription(LODIndex, MoveTemp(MeshDescription));
 		CommitMeshDescription(LODIndex);
@@ -3277,7 +3277,7 @@ void USkeletalMesh::BeginPostLoadInternal(FSkinnedAssetPostLoadContext& Context)
 						}
 
 						FMeshDescription MeshDescription;
-						if (SkeletalMeshImportData.GetMeshDescription(this, MeshDescription))
+						if (SkeletalMeshImportData.GetMeshDescription(this, LODIndex, MeshDescription))
 						{
 							CreateMeshDescription(LODIndex, MoveTemp(MeshDescription));
 							CommitMeshDescription(LODIndex);
diff --git a/Engine/Source/Runtime/Engine/Private/SkeletalMeshSourceModel.cpp b/Engine/Source/Runtime/Engine/Private/SkeletalMeshSourceModel.cpp
--- a/Engine/Source/Runtime/Engine/Private/SkeletalMeshSourceModel.cpp
+++ b/Engine/Source/Runtime/Engine/Private/SkeletalMeshSourceModel.cpp
@@ -329,7 +329,7 @@ void FSkeletalMeshSourceModel::ConvertRawMeshToMeshDescriptionBulkData()
 	}
 	
 	FMeshDescription& MeshDescription = MeshDescriptionBulkData->CreateMeshDescription()->GetMeshDescription(); 
-	ImportData.GetMeshDescription(GetOwner(), MeshDescription);
+	ImportData.GetMeshDescription(GetOwner(), RawMeshBulkDataLODIndex, MeshDescription);
 
 	UpdateCachedMeshStatistics(&MeshDescription);
 
diff --git a/Engine/Source/Runtime/Engine/Public/Rendering/SkeletalMeshLODImporterData.h b/Engine/Source/Runtime/Engine/Public/Rendering/SkeletalMeshLODImporterData.h
--- a/Engine/Source/Runtime/Engine/Public/Rendering/SkeletalMeshLODImporterData.h
+++ b/Engine/Source/Runtime/Engine/Public/Rendering/SkeletalMeshLODImporterData.h
@@ -434,7 +434,7 @@ public:
 	 * Returns a mesh description from the import data. If logging on failures is required, pass in a pointer 
 	 * to the owning skeletal mesh. Otherwise leave as a \c nullptr. 
 	 */
-	ENGINE_API bool GetMeshDescription(const USkeletalMesh* InSkeletalMesh, FMeshDescription& OutMeshDescription) const;
+	ENGINE_API bool GetMeshDescription(const USkeletalMesh* InSkeletalMesh, int32 InLODIndex, FMeshDescription& OutMeshDescription) const;
 
 	/**
 	 * @note MeshDescription always contains color, normal and tangent data by default. Therefore, while iterating
