Clean up some ray tracing logic in FInstancedStaticMeshSceneProxy

diff --git a/Engine/Source/Runtime/Engine/Classes/Engine/InstancedStaticMesh.h b/Engine/Source/Runtime/Engine/Classes/Engine/InstancedStaticMesh.h
--- a/Engine/Source/Runtime/Engine/Classes/Engine/InstancedStaticMesh.h
+++ b/Engine/Source/Runtime/Engine/Classes/Engine/InstancedStaticMesh.h
@@ -475,7 +475,7 @@ public:
 
 	ENGINE_API virtual void GetDynamicRayTracingInstances(struct FRayTracingMaterialGatheringContext& Context, TArray<FRayTracingInstance>& OutRayTracingInstances) final override;
 
-	ENGINE_API void SetupRayTracingDynamicInstances(int32 NumDynamicInstances, int32 LOD);
+	ENGINE_API void SetupRayTracingDynamicInstances(int32 NumDynamicInstances, int32 LODIndex);
 
 #endif
 
@@ -551,7 +551,7 @@ protected:
 
 	TArray<FRayTracingDynamicData> RayTracingDynamicData;
 
-	int32 CachedRayTracingLOD;
+	int32 CachedRayTracingLODIndex;
 #endif
 
 	/** Common path for the Get*MeshElement functions */
diff --git a/Engine/Source/Runtime/Engine/Private/InstancedStaticMesh.cpp b/Engine/Source/Runtime/Engine/Private/InstancedStaticMesh.cpp
--- a/Engine/Source/Runtime/Engine/Private/InstancedStaticMesh.cpp
+++ b/Engine/Source/Runtime/Engine/Private/InstancedStaticMesh.cpp
@@ -1327,7 +1327,7 @@ FInstancedStaticMeshSceneProxy::FInstancedStaticMeshSceneProxy(const FInstancedS
 #endif
 	,	InstanceLODDistanceScale(InProxyDesc.InstanceLODDistanceScale)
 #if RHI_RAYTRACING
-	,	CachedRayTracingLOD(-1)
+	,	CachedRayTracingLODIndex(-1)
 #endif
 	,	StaticMeshBounds(StaticMesh->GetBounds())
 	,	InstanceDataSceneProxy(InProxyDesc.InstanceDataSceneProxy)
@@ -1612,25 +1612,28 @@ void FInstancedStaticMeshSceneProxy::GetDynamicRayTracingInstances(struct FRayTr
 		return;
 	}
 
-	uint32 MinAllowedLOD = FMath::Clamp<int32>(CVarRayTracingInstancedStaticMeshesMinLOD.GetValueOnRenderThread(), 0, RenderData->LODResources.Num() - 1);
-	uint32 LOD = FMath::Max<uint32>(MinAllowedLOD, GetCurrentFirstLODIdx_RenderThread());
+	const uint32 InstanceCount = GetInstanceDataHeader().NumInstances;
 
-	if (!RenderData->LODResources[LOD].RayTracingGeometry.IsInitialized())
+	if (InstanceCount == 0u)
 	{
 		return;
 	}
 
-	const uint32 InstanceCount = GetInstanceDataHeader().NumInstances;
+	const int32 MinAllowedLODIndex = FMath::Clamp<int32>(CVarRayTracingInstancedStaticMeshesMinLOD.GetValueOnRenderThread(), 0, RenderData->LODResources.Num() - 1);
 
-	if (InstanceCount == 0u)
+	int32 LODIndex = FMath::Max<int32>(MinAllowedLODIndex, GetCurrentFirstLODIdx_RenderThread());
+
+	FRayTracingGeometry* RayTracingGeometry = &RenderData->LODResources[LODIndex].RayTracingGeometry;
+
+	if (!RayTracingGeometry->IsInitialized())
 	{
 		return;
 	}
 
 	// TODO: Select different LOD when current LOD is still requested for build?
-	if (RenderData->LODResources[LOD].RayTracingGeometry.HasPendingBuildRequest())
+	if (RayTracingGeometry->HasPendingBuildRequest())
 	{
-		RenderData->LODResources[LOD].RayTracingGeometry.BoostBuildPriority();
+		RayTracingGeometry->BoostBuildPriority();
 		return;
 	}
 
@@ -1641,11 +1644,12 @@ void FInstancedStaticMeshSceneProxy::GetDynamicRayTracingInstances(struct FRayTr
 	};
 
 	//setup a 'template' for the instance first, so we aren't duplicating work
-	//#dxr_todo: when multiple LODs are used, template needs to be an array of templates, probably best initialized on-demand via a lamda
+	//#dxr_todo: when multiple LODs are used, template needs to be an array of templates, probably best initialized on-demand via a lambda
 	FRayTracingInstance RayTracingInstanceTemplate;
+	RayTracingInstanceTemplate.Geometry = RayTracingGeometry;
+
 	FRayTracingInstance RayTracingWPOInstanceTemplate;  //template for evaluating the WPO instances into the world
 	FRayTracingInstance RayTracingWPODynamicTemplate;   //template for simulating the WPO instances
-	RayTracingInstanceTemplate.Geometry = &RenderData->LODResources[LOD].RayTracingGeometry;
 
 	// Which index holds the reference to the particular simulated instance
 	TArray<uint32> ActiveInstances;
@@ -1662,39 +1666,43 @@ void FInstancedStaticMeshSceneProxy::GetDynamicRayTracingInstances(struct FRayTr
 
 	if (bHasWorldPositionOffset)
 	{
-		int32 SectionCount = InstancedRenderData.LODModels[LOD].Sections.Num();
+		int32 SectionCount = InstancedRenderData.LODModels[LODIndex].Sections.Num();
 
 		for (int32 SectionIdx = 0; SectionIdx < SectionCount; ++SectionIdx)
 		{
 			//#dxr_todo: so far we use the parent static mesh path to get material data
-			FMeshBatch MeshBatch;
-			FMeshBatch DynamicMeshBatch;
-
-			if (!GetMeshElement(LOD, 0, SectionIdx, 0, false, false, DynamicMeshBatch))
 			{
-				DynamicMeshBatch.MaterialRenderProxy = UMaterial::GetDefaultMaterial(MD_Surface)->GetRenderProxy();
-				DynamicMeshBatch.SegmentIndex = SectionIdx;
-				DynamicMeshBatch.MeshIdInPrimitive = SectionIdx;
+				FMeshBatch MeshBatch;
+
+				if (!FStaticMeshSceneProxy::GetMeshElement(LODIndex, 0, SectionIdx, 0, false, false, MeshBatch))
+				{
+					MeshBatch.MaterialRenderProxy = UMaterial::GetDefaultMaterial(MD_Surface)->GetRenderProxy();
+					MeshBatch.VertexFactory = &RenderData->LODVertexFactories[LODIndex].VertexFactory;
+					MeshBatch.SegmentIndex = SectionIdx;
+					MeshBatch.MeshIdInPrimitive = SectionIdx;
+				};
+
+				RayTracingWPOInstanceTemplate.Materials.Add(MeshBatch);
 			}
 
-			if (!FStaticMeshSceneProxy::GetMeshElement(LOD, 0, SectionIdx, 0, false, false, MeshBatch))
-			{				
-				MeshBatch.MaterialRenderProxy = UMaterial::GetDefaultMaterial(MD_Surface)->GetRenderProxy();
-				MeshBatch.VertexFactory = &RenderData->LODVertexFactories[LOD].VertexFactory;
-				MeshBatch.SegmentIndex = SectionIdx;
-				MeshBatch.MeshIdInPrimitive = SectionIdx;
-			};
+			{
+				FMeshBatch DynamicMeshBatch;
 
-			DynamicMeshBatch.VertexFactory = &InstancedRenderData.VertexFactories[LOD];
+				if (!GetMeshElement(LODIndex, 0, SectionIdx, 0, false, false, DynamicMeshBatch))
+				{
+					DynamicMeshBatch.MaterialRenderProxy = UMaterial::GetDefaultMaterial(MD_Surface)->GetRenderProxy();
+					DynamicMeshBatch.SegmentIndex = SectionIdx;
+					DynamicMeshBatch.MeshIdInPrimitive = SectionIdx;
+				}
 
-			RayTracingWPOInstanceTemplate.Materials.Add(MeshBatch);
-			RayTracingWPODynamicTemplate.Materials.Add(DynamicMeshBatch);
-		}
-	
-		if (RayTracingDynamicData.Num() != SimulatedInstances || LOD != CachedRayTracingLOD)
-		{
-			SetupRayTracingDynamicInstances(SimulatedInstances, LOD);
+				DynamicMeshBatch.VertexFactory = &InstancedRenderData.VertexFactories[LODIndex];
+
+				RayTracingWPODynamicTemplate.Materials.Add(DynamicMeshBatch);
+			}
 		}
+
+		SetupRayTracingDynamicInstances(SimulatedInstances, LODIndex);
+
 		ActiveInstances.AddZeroed(SimulatedInstances);
 
 		for (auto &Instance : ActiveInstances)
@@ -1705,13 +1713,12 @@ void FInstancedStaticMeshSceneProxy::GetDynamicRayTracingInstances(struct FRayTr
 
 	VisibleInstances.Reserve(InstanceCount);
 
-
 	const FBox CurrentBounds = StaticMeshBounds.GetBox();
 
 	constexpr float LocalToWorldScale = 1.0f;
 	FVector ViewPosition = Context.ReferenceView->ViewLocation;
 
-	const FInstanceSceneDataBuffers *InstanceSceneDataBuffers = GetInstanceSceneDataBuffers();
+	const FInstanceSceneDataBuffers* InstanceSceneDataBuffers = GetInstanceSceneDataBuffers();
 	check(InstanceSceneDataBuffers && InstanceSceneDataBuffers->GetNumInstances() == InstanceCount);
 
 	auto GetDistanceToInstance = [&ViewPosition, InstanceSceneDataBuffers](int32 InstanceIndex, float& OutInstanceRadius, float& OutDistanceToInstanceCenter, float& OutDistanceToInstanceStart)
@@ -1849,14 +1856,14 @@ void FInstancedStaticMeshSceneProxy::GetDynamicRayTracingInstances(struct FRayTr
 		FMatrix InstanceToWorld = InstanceSceneDataBuffers->GetInstanceToWorld(InstanceIndex);
 		const uint32 DynamicInstanceIdx = InstanceIndex % SimulatedInstances;
 
-		if (bHasWorldPositionOffset && InstancedRenderData.VertexFactories[LOD].GetType()->SupportsRayTracingDynamicGeometry())
+		if (bHasWorldPositionOffset && InstancedRenderData.VertexFactories[LODIndex].GetType()->SupportsRayTracingDynamicGeometry())
 		{
 			FRayTracingInstance* DynamicInstance = nullptr;
 
 			if (ActiveInstances[DynamicInstanceIdx] == -1)
 			{
 				// first case of this dynamic instance, setup the material and add it
-				const FStaticMeshLODResources& LODModel = RenderData->LODResources[LOD];
+				const FStaticMeshLODResources& LODModel = RenderData->LODResources[LODIndex];
 
 				FRayTracingDynamicData& DynamicData = RayTracingDynamicData[DynamicInstanceIdx];
 
@@ -1899,19 +1906,19 @@ void FInstancedStaticMeshSceneProxy::GetDynamicRayTracingInstances(struct FRayTr
 
 	if (RayTracingInstanceTemplate.InstanceTransforms.Num() > 0)
 	{
-		int32 SectionCount = InstancedRenderData.LODModels[LOD].Sections.Num();
+		int32 SectionCount = InstancedRenderData.LODModels[LODIndex].Sections.Num();
 
 		for (int32 SectionIdx = 0; SectionIdx < SectionCount; ++SectionIdx)
 		{
 			//#dxr_todo: so far we use the parent static mesh path to get material data
 			FMeshBatch MeshBatch;
 
-			bool bResult = FStaticMeshSceneProxy::GetMeshElement(LOD, 0, SectionIdx, 0, false, false, MeshBatch);
+			bool bResult = FStaticMeshSceneProxy::GetMeshElement(LODIndex, 0, SectionIdx, 0, false, false, MeshBatch);
 			if (!bResult)
 			{
 				// Hidden material
 				MeshBatch.MaterialRenderProxy = UMaterial::GetDefaultMaterial(MD_Surface)->GetRenderProxy();
-				MeshBatch.VertexFactory = &RenderData->LODVertexFactories[LOD].VertexFactory;
+				MeshBatch.VertexFactory = &RenderData->LODVertexFactories[LODIndex].VertexFactory;
 				MeshBatch.SegmentIndex = SectionIdx;
 				MeshBatch.MeshIdInPrimitive = SectionIdx;
 			}
@@ -1924,13 +1931,13 @@ void FInstancedStaticMeshSceneProxy::GetDynamicRayTracingInstances(struct FRayTr
 }
 
 
-void FInstancedStaticMeshSceneProxy::SetupRayTracingDynamicInstances(int32 NumDynamicInstances, int32 LOD)
+void FInstancedStaticMeshSceneProxy::SetupRayTracingDynamicInstances(int32 NumDynamicInstances, int32 LODIndex)
 {
-	if (RayTracingDynamicData.Num() > NumDynamicInstances || CachedRayTracingLOD != LOD)
+	if (RayTracingDynamicData.Num() > NumDynamicInstances || CachedRayTracingLODIndex != LODIndex)
 	{
-		//free the unused/out of date entries
+		// free the unused/out of date entries
 
-		int32 FirstToFree = (CachedRayTracingLOD != LOD) ? 0 : NumDynamicInstances;
+		int32 FirstToFree = (CachedRayTracingLODIndex != LODIndex) ? 0 : NumDynamicInstances;
 		for (int32 Item = FirstToFree; Item < RayTracingDynamicData.Num(); Item++)
 		{
 			auto& DynamicRayTracingItem = RayTracingDynamicData[Item];
@@ -1944,7 +1951,7 @@ void FInstancedStaticMeshSceneProxy::SetupRayTracingDynamicInstances(int32 NumDy
 	{
 		RayTracingDynamicData.Reserve(NumDynamicInstances);
 		const int32 StartIndex = RayTracingDynamicData.Num();
-		const FStaticMeshLODResources& LODModel = RenderData->LODResources[LOD];
+		const FStaticMeshLODResources& LODModel = RenderData->LODResources[LODIndex];
 
 		for (int32 Item = StartIndex; Item < NumDynamicInstances; Item++)
 		{
@@ -1963,7 +1970,7 @@ void FInstancedStaticMeshSceneProxy::SetupRayTracingDynamicInstances(int32 NumDy
 		}
 	}
 
-	CachedRayTracingLOD = LOD;
+	CachedRayTracingLODIndex = LODIndex;
 }
 
 #endif
