Exposed a way to get the source string of a Text in BP

diff --git a/Engine/Source/Runtime/Engine/Classes/Kismet/KismetTextLibrary.h b/Engine/Source/Runtime/Engine/Classes/Kismet/KismetTextLibrary.h
--- a/Engine/Source/Runtime/Engine/Classes/Kismet/KismetTextLibrary.h
+++ b/Engine/Source/Runtime/Engine/Classes/Kismet/KismetTextLibrary.h
@@ -304,6 +304,13 @@ class UKismetTextLibrary : public UBlueprintFunctionLibrary
 	UFUNCTION(BlueprintPure, Category="Utilities|Text", meta=(DisplayName="Get ID from Text"))
 	static ENGINE_API bool GetTextId(FText Text, FString& OutNamespace, FString& OutKey);
 
+	/**
+	 * Get the (non-localized) source string of the given text.
+	 * @note For a generated text (eg, the result of a Format), this will deep build the source string as if the generation had run for the native language.
+	 */
+	UFUNCTION(BlueprintPure, Category="Utilities|Text", meta=(DisplayName="Get Source String from Text"))
+	static ENGINE_API FString GetTextSourceString(FText Text);
+
 	/**
 	 * Check whether the given polyglot data is valid.
 	 * @return True if the polyglot data is valid, false otherwise. ErrorMessage will be filled in if the the data is invalid.
diff --git a/Engine/Source/Runtime/Engine/Private/KismetTextLibrary.cpp b/Engine/Source/Runtime/Engine/Private/KismetTextLibrary.cpp
--- a/Engine/Source/Runtime/Engine/Private/KismetTextLibrary.cpp
+++ b/Engine/Source/Runtime/Engine/Private/KismetTextLibrary.cpp
@@ -338,6 +338,11 @@ bool UKismetTextLibrary::GetTextId(FText Text, FString& OutNamespace, FString& O
 	return false;
 }
 
+FString UKismetTextLibrary::GetTextSourceString(FText Text)
+{
+	return Text.BuildSourceString();
+}
+
 void UKismetTextLibrary::IsPolyglotDataValid(const FPolyglotTextData& PolyglotData, bool& IsValid, FText& ErrorMessage)
 {
 	IsValid = PolyglotData.IsValid(&ErrorMessage);
