Allow classes inheriting from USkeletalMesh to override UStreamableRenderAsset methods for streaming.
* Move SkeletalMeshUpdate and RenderAssetUpdate from private to internal to allow its use within the engine.
* Mark FRenderAssetUpdate, FSkelMeshUpdateContext, and FSkeletalMeshStreamIn as ENGINE_API types.

#rnx

diff --git a/Engine/Source/Runtime/Engine/Classes/Engine/SkeletalMesh.h b/Engine/Source/Runtime/Engine/Classes/Engine/SkeletalMesh.h
--- a/Engine/Source/Runtime/Engine/Classes/Engine/SkeletalMesh.h
+++ b/Engine/Source/Runtime/Engine/Classes/Engine/SkeletalMesh.h
@@ -2512,11 +2512,11 @@ public:
 	//~ End UObject Interface.
 
 	//~ Begin UStreamableRenderAsset Interface.
-	ENGINE_API virtual int32 CalcCumulativeLODSize(int32 NumLODs) const final override;
-	ENGINE_API virtual FIoFilenameHash GetMipIoFilenameHash(const int32 MipIndex) const final override;
-	ENGINE_API virtual bool DoesMipDataExist(const int32 MipIndex) const final override;
-	ENGINE_API virtual bool StreamOut(int32 NewMipCount) final override;
-	ENGINE_API virtual bool StreamIn(int32 NewMipCount, bool bHighPrio) final override;
+	ENGINE_API virtual int32 CalcCumulativeLODSize(int32 NumLODs) const override;
+	ENGINE_API virtual FIoFilenameHash GetMipIoFilenameHash(const int32 MipIndex) const override;
+	ENGINE_API virtual bool DoesMipDataExist(const int32 MipIndex) const override;
+	ENGINE_API virtual bool StreamOut(int32 NewMipCount) override;
+	ENGINE_API virtual bool StreamIn(int32 NewMipCount, bool bHighPrio) override;
 	ENGINE_API virtual bool HasPendingRenderResourceInitialization() const;
 	virtual EStreamableRenderAssetType GetRenderAssetType() const final override { return EStreamableRenderAssetType::SkeletalMesh; }
 	//~ End UStreamableRenderAsset Interface.
diff --git a/Engine/Source/Runtime/Engine/Internal/Streaming/RenderAssetUpdate.inl b/Engine/Source/Runtime/Engine/Internal/Streaming/RenderAssetUpdate.inl
new file mode 100644
--- /dev/null
+++ b/Engine/Source/Runtime/Engine/Internal/Streaming/RenderAssetUpdate.inl
@@ -0,0 +1,107 @@
+// Copyright Epic Games, Inc. All Rights Reserved.
+
+/*=============================================================================
+RenderAssetUpdate.inl: Base class of helpers to stream in and out texture/mesh LODs
+=============================================================================*/
+
+#pragma once
+
+#include "RenderAssetUpdate.h"
+
+template <typename TContext>
+TRenderAssetUpdate<TContext>::TRenderAssetUpdate(const UStreamableRenderAsset* InAsset)
+	: FRenderAssetUpdate(InAsset)
+	, TaskThread(TT_None)
+	, TaskCallback(nullptr)
+	, CancelationThread(TT_None)
+	, CancelationCallback(nullptr)
+{
+
+}
+
+template <typename TContext>
+void TRenderAssetUpdate<TContext>::PushTask(const FContext& Context, EThreadType InTaskThread, const FCallback& InTaskCallback, EThreadType InCancelationThread, const FCallback& InCancelationCallback)
+{
+	// PushTask can only be called by the one thread/callback that is doing the processing. 
+	// This means we don't need to check whether other threads could be trying to push tasks.
+	check(TaskState == TS_Locked || TaskState == TS_Init);
+	checkSlow((bool)InTaskCallback == (InTaskThread != TT_None));
+	checkSlow((bool)InCancelationCallback == (InCancelationThread != TT_None));
+
+	TaskThread = InTaskThread;
+	TaskCallback = InTaskCallback;
+	CancelationThread = InCancelationThread;
+	CancelationCallback = InCancelationCallback;
+}
+
+template <typename TContext>
+FRenderAssetUpdate::ETaskState TRenderAssetUpdate<TContext>::TickInternal(EThreadType InCurrentThread, bool bCheckForSuspension)
+{
+	check(TaskState == TS_Locked);
+
+	// Thread, callback and asset must be coherent because Abort() could be called while this is executing.
+	EThreadType RelevantThread = TaskThread;
+	FCallback RelevantCallback = TaskCallback;
+	const UStreamableRenderAsset* RelevantAsset = StreamableAsset;
+	if (bIsCancelled || !RelevantAsset)
+	{
+		RelevantThread = CancelationThread;
+		RelevantCallback = CancelationCallback;
+		// RelevantAsset = nullptr; // TODO once Cancel supports it!
+	}
+
+	if (RelevantThread == TT_None)
+	{
+		ClearCallbacks();
+		return TS_Done;
+	}
+	else if (TaskSynchronization.GetValue() > 0)
+	{
+		return TS_Suspended;
+	}
+	else if (bCheckForSuspension && IsAssetStreamingSuspended())
+	{
+		return TS_Suspended;
+	}
+	else if (bDeferExecution)
+	{
+		bDeferExecution = false;
+		return TS_Suspended;
+	}
+	else if (RelevantThread == InCurrentThread)
+	{
+		ClearCallbacks();
+		FContext Context(RelevantAsset, InCurrentThread);
+		RelevantCallback(Context);
+		return TS_Locked;
+	} 
+	else if (RelevantThread == FRenderAssetUpdate::TT_GameThread && !ScheduledGTTasks)
+	{
+		ScheduleGTTask();
+		return TS_Locked;
+	}
+	else if (RelevantThread == FRenderAssetUpdate::TT_Render && !ScheduledRenderTasks)
+	{
+		if (GIsThreadedRendering)
+		{
+			// Note that it could execute now if this is the renderthread.
+			ScheduleRenderTask();
+			return TS_Locked;
+		}
+		else
+		{
+			// If render commands do not have their own thread, avoid unrelated thread to enqueue commands.
+			// This is to avoid issues from using RHIAsyncCreateTexture2D which doesn't give any feedback of when it is ready.
+			return TS_InProgress;
+		}
+	}
+	else if (RelevantThread == FRenderAssetUpdate::TT_Async && !ScheduledAsyncTasks)
+	{
+		ScheduleAsyncTask();
+		return TS_Locked;
+	}
+	else
+	{
+		return TS_InProgress;
+	}
+}
diff --git a/Engine/Source/Runtime/Engine/Internal/Streaming/SkeletalMeshUpdate.cpp b/Engine/Source/Runtime/Engine/Internal/Streaming/SkeletalMeshUpdate.cpp
new file mode 100644
--- /dev/null
+++ b/Engine/Source/Runtime/Engine/Internal/Streaming/SkeletalMeshUpdate.cpp
@@ -0,0 +1,669 @@
+// Copyright Epic Games, Inc. All Rights Reserved.
+
+/*=============================================================================
+SkeletalMeshUpdate.cpp: Helpers to stream in and out skeletal mesh LODs.
+=============================================================================*/
+
+#include "Streaming/SkeletalMeshUpdate.h"
+#include "HAL/PlatformFile.h"
+#include "RenderUtils.h"
+#include "ContentStreaming.h"
+#include "Streaming/TextureStreamingHelpers.h"
+#include "Serialization/MemoryReader.h"
+#include "Rendering/SkeletalMeshRenderData.h"
+#include "Components/SkinnedMeshComponent.h"
+#include "Streaming/RenderAssetUpdate.inl"
+#include "RHIResourceReplace.h"
+
+extern int32 GStreamingMaxReferenceChecks;
+
+template class TRenderAssetUpdate<FSkelMeshUpdateContext>;
+
+static constexpr uint32 GSkelMeshMaxNumResourceUpdatesPerLOD = 16;
+static constexpr uint32 GSkelMeshMaxNumResourceUpdatesPerBatch = (MAX_MESH_LOD_COUNT - 1) * GSkelMeshMaxNumResourceUpdatesPerLOD;
+
+FSkelMeshUpdateContext::FSkelMeshUpdateContext(const USkeletalMesh* InMesh, EThreadType InCurrentThread)
+	: Mesh(InMesh)
+	, CurrentThread(InCurrentThread)
+{
+	check(InMesh);
+	checkSlow(InCurrentThread != FSkeletalMeshUpdate::TT_Render || IsInRenderingThread());
+	RenderData = Mesh->GetResourceForRendering();
+	AssetLODBias = Mesh->GetStreamableResourceState().AssetLODBias;
+	if (RenderData)
+	{
+		LODResourcesView = TArrayView<FSkeletalMeshLODRenderData*>(RenderData->LODRenderData.GetData() + AssetLODBias, Mesh->GetStreamableResourceState().MaxNumLODs);
+	}
+}
+
+FSkelMeshUpdateContext::FSkelMeshUpdateContext(const UStreamableRenderAsset* InMesh, EThreadType InCurrentThread)
+#if UE_BUILD_SHIPPING
+	: FSkelMeshUpdateContext(static_cast<const USkeletalMesh*>(InMesh), InCurrentThread)
+#else
+	: FSkelMeshUpdateContext(Cast<USkeletalMesh>(InMesh), InCurrentThread)
+#endif
+{}
+
+FSkeletalMeshUpdate::FSkeletalMeshUpdate(const USkeletalMesh* InMesh)
+	: TRenderAssetUpdate<FSkelMeshUpdateContext>(InMesh)
+{
+}
+
+void FSkeletalMeshStreamIn::FIntermediateBuffers::CreateFromCPUData(FRHICommandListBase& RHICmdList, FSkeletalMeshLODRenderData& LODResource)
+{
+	FStaticMeshVertexBuffers& VBs = LODResource.StaticVertexBuffers;
+	TangentsVertexBuffer = VBs.StaticMeshVertexBuffer.CreateTangentsRHIBuffer(RHICmdList);
+	TexCoordVertexBuffer = VBs.StaticMeshVertexBuffer.CreateTexCoordRHIBuffer(RHICmdList);
+	PositionVertexBuffer = VBs.PositionVertexBuffer.CreateRHIBuffer(RHICmdList);
+	ColorVertexBuffer = VBs.ColorVertexBuffer.CreateRHIBuffer(RHICmdList);
+	LODResource.SkinWeightProfilesData.CreateRHIBuffers(RHICmdList, AltSkinWeightVertexBuffers);
+	SkinWeightVertexBuffer = LODResource.SkinWeightVertexBuffer.CreateRHIBuffer(RHICmdList);
+	ClothVertexBuffer = LODResource.ClothVertexBuffer.CreateRHIBuffer(RHICmdList);
+	IndexBuffer = LODResource.MultiSizeIndexContainer.CreateRHIBuffer(RHICmdList);
+}
+
+void FSkeletalMeshStreamIn::FIntermediateBuffers::TransferBuffers(FSkeletalMeshLODRenderData& LODResource, FRHIResourceReplaceBatcher& Batcher)
+{
+	FStaticMeshVertexBuffers& VBs = LODResource.StaticVertexBuffers;
+	VBs.StaticMeshVertexBuffer.InitRHIForStreaming(TangentsVertexBuffer, TexCoordVertexBuffer, Batcher);
+	VBs.PositionVertexBuffer.InitRHIForStreaming(PositionVertexBuffer, Batcher);
+	VBs.ColorVertexBuffer.InitRHIForStreaming(ColorVertexBuffer, Batcher);
+	LODResource.SkinWeightVertexBuffer.InitRHIForStreaming(SkinWeightVertexBuffer, Batcher);
+	LODResource.ClothVertexBuffer.InitRHIForStreaming(ClothVertexBuffer, Batcher);
+	LODResource.MultiSizeIndexContainer.InitRHIForStreaming(IndexBuffer, Batcher);
+	LODResource.SkinWeightProfilesData.InitRHIForStreaming(AltSkinWeightVertexBuffers, Batcher);
+}
+
+#if RHI_RAYTRACING
+
+void FSkeletalMeshStreamIn::FIntermediateRayTracingGeometry::CreateFromCPUData(FRHICommandListBase& RHICmdList, FRayTracingGeometry& RayTracingGeometry)
+{
+	Initializer = RayTracingGeometry.Initializer;
+	Initializer.Type = ERayTracingGeometryInitializerType::StreamingSource;
+
+	if (RayTracingGeometry.RawData.Num())
+	{
+		check(Initializer.OfflineData == nullptr);
+		Initializer.OfflineData = &RayTracingGeometry.RawData;
+	}
+
+	static const auto CVarDebugForceRuntimeBLAS = IConsoleManager::Get().FindTConsoleVariableDataInt(TEXT("r.Raytracing.DebugForceRuntimeBLAS"));
+	const bool bDebugForceRuntimeBLAS = (!CVarDebugForceRuntimeBLAS) || (CVarDebugForceRuntimeBLAS->GetValueOnAnyThread() != 0);
+
+	if (bDebugForceRuntimeBLAS && Initializer.OfflineData != nullptr)
+	{
+		Initializer.OfflineData->Discard();
+		Initializer.OfflineData = nullptr;
+	}
+
+	RayTracingGeometryRHI = RHICmdList.CreateRayTracingGeometry(Initializer);
+	bRequiresBuild = Initializer.OfflineData == nullptr || RayTracingGeometryRHI->IsCompressed();
+}
+
+void FSkeletalMeshStreamIn::FIntermediateRayTracingGeometry::SafeRelease()
+{
+	Initializer = {};
+	RayTracingGeometryRHI.SafeRelease();
+}
+
+void FSkeletalMeshStreamIn::FIntermediateRayTracingGeometry::TransferRayTracingGeometry(FRayTracingGeometry& RayTracingGeometry, FRHIResourceReplaceBatcher& Batcher)
+{
+	RayTracingGeometry.InitRHIForStreaming(RayTracingGeometryRHI, Batcher);
+	RayTracingGeometry.SetRequiresBuild(bRequiresBuild);
+
+	SafeRelease();
+}
+
+#endif
+
+FSkeletalMeshStreamIn::FSkeletalMeshStreamIn(const USkeletalMesh* InMesh, EThreadType CreateResourcesThread)
+	: FSkeletalMeshUpdate(InMesh)
+	, CreateResourcesThread(CreateResourcesThread)
+{
+	if (!ensure(PendingFirstLODIdx < CurrentFirstLODIdx))
+	{
+		bIsCancelled = true;
+	}
+}
+
+FSkeletalMeshStreamIn::~FSkeletalMeshStreamIn()
+{
+	check(!StreamingRHICmdList);
+}
+
+void FSkeletalMeshStreamIn::CreateBuffers(const FContext& Context)
+{
+	LLM_SCOPE(ELLMTag::SkeletalMesh);
+
+	check(Context.Mesh && Context.RenderData);
+
+	StreamingRHICmdList = new FRHICommandList();
+	StreamingRHICmdList->SwitchPipeline(ERHIPipeline::Graphics);
+	{
+		SCOPED_DRAW_EVENTF(*StreamingRHICmdList, SkeletalMesh_StreamIn, TEXT("SkeletalMesh - StreamIn: %s"), Context.Mesh->GetFName());
+
+		for (int32 LODIndex = PendingFirstLODIdx; LODIndex < CurrentFirstLODIdx; ++LODIndex)
+		{
+			FSkeletalMeshLODRenderData& LODResource = *Context.LODResourcesView[LODIndex];
+			IntermediateBuffersArray[LODIndex].CreateFromCPUData(*StreamingRHICmdList, LODResource);
+
+#if RHI_RAYTRACING
+			// Skip LODs that have their render data stripped
+			if (IsRayTracingEnabled() && Context.Mesh->GetSupportRayTracing() && LODResource.GetNumVertices() > 0 && LODResource.bReferencedByStaticSkeletalMeshObjects_RenderThread)
+			{
+				IntermediateRayTracingGeometry[LODIndex].CreateFromCPUData(*StreamingRHICmdList, LODResource.StaticRayTracingGeometry);
+			}
+#endif
+		}
+
+		// Use a scope to flush the batcher before updating CurrentFirstLODIdx
+		{
+			FRHIResourceReplaceBatcher Batcher(*StreamingRHICmdList, GSkelMeshMaxNumResourceUpdatesPerBatch);
+
+			for (int32 LODIndex = PendingFirstLODIdx; LODIndex < CurrentFirstLODIdx; ++LODIndex)
+			{
+				FSkeletalMeshLODRenderData& LODResource = *Context.LODResourcesView[LODIndex];
+				LODResource.IncrementMemoryStats(Context.Mesh->GetHasVertexColors());
+				LODResource.InitMorphResources();
+				IntermediateBuffersArray[LODIndex].TransferBuffers(LODResource, Batcher);
+			}
+		}
+	}
+
+	StreamingRHICmdList->FinishRecording();
+}
+
+void FSkeletalMeshStreamIn::DiscardNewLODs(const FContext& Context)
+{
+	if (Context.RenderData)
+	{
+		for (int32 LODIndex = PendingFirstLODIdx; LODIndex < CurrentFirstLODIdx; ++LODIndex)
+		{
+			FSkeletalMeshLODRenderData& LODResource = *Context.LODResourcesView[LODIndex];
+			LODResource.ReleaseCPUResources(true);
+		}
+	}
+}
+
+void FSkeletalMeshStreamIn::DoFinishUpdate(const FContext& Context)
+{
+	check(Context.CurrentThread == TT_Render);
+	check(IsInRenderingThread());
+
+	if (StreamingRHICmdList)
+	{
+		FRHICommandListImmediate::Get().QueueAsyncCommandListSubmit(StreamingRHICmdList);
+		StreamingRHICmdList = nullptr;
+	}
+
+#if RHI_RAYTRACING
+	if (IsRayTracingAllowed() && Context.Mesh->GetSupportRayTracing())
+	{
+		// Use a scope to flush the batcher before updating CurrentFirstLODIdx
+		{
+			FRHIResourceReplaceBatcher Batcher(FRHICommandListImmediate::Get(), GSkelMeshMaxNumResourceUpdatesPerBatch);
+			for (int32 LODIdx = PendingFirstLODIdx; LODIdx < CurrentFirstLODIdx; ++LODIdx)
+			{
+				FSkeletalMeshLODRenderData& LODResource = *Context.LODResourcesView[LODIdx];
+
+				if (LODResource.GetNumVertices() > 0 && LODResource.bReferencedByStaticSkeletalMeshObjects_RenderThread)
+				{
+					IntermediateRayTracingGeometry[LODIdx].TransferRayTracingGeometry(LODResource.StaticRayTracingGeometry, Batcher);
+				}
+			}
+		}
+
+		// Must happen after the batched updates have been flushed
+		for (int32 LODIndex = PendingFirstLODIdx; LODIndex < CurrentFirstLODIdx; ++LODIndex)
+		{
+			FSkeletalMeshLODRenderData& LODResource = *Context.LODResourcesView[LODIndex];
+
+			// Skip LODs that have their render data stripped
+			if (LODResource.GetNumVertices() > 0 && LODResource.bReferencedByStaticSkeletalMeshObjects_RenderThread)
+			{
+				// Under very rare circumstances that we switch ray tracing on/off right in the middle of streaming RayTracingGeometryRHI might not be valid.
+				if (IsRayTracingEnabled() && ensure(LODResource.StaticRayTracingGeometry.IsValid()))
+				{
+					LODResource.StaticRayTracingGeometry.RequestBuildIfNeeded(ERTAccelerationStructureBuildPriority::Normal);
+				}
+			}
+		}
+
+	}
+#endif
+
+	Context.RenderData->PendingFirstLODIdx = Context.RenderData->CurrentFirstLODIdx = ResourceState.LODCountToAssetFirstLODIdx(ResourceState.NumRequestedLODs);
+	MarkAsSuccessfullyFinished();
+}
+
+void FSkeletalMeshStreamIn::DoCancel(const FContext& Context)
+{
+	// TODO: support streaming CPU data for editor builds
+	if (!GIsEditor)
+	{
+		DiscardNewLODs(Context);
+	}
+
+	check(!StreamingRHICmdList);
+	MarkAsSuccessfullyFinished();
+}
+
+FSkeletalMeshStreamOut::FSkeletalMeshStreamOut(const USkeletalMesh* InMesh)
+	: FSkeletalMeshUpdate(InMesh)
+{
+	PushTask(FContext(InMesh, TT_None), TT_GameThread, SRA_UPDATE_CALLBACK(ConditionalMarkComponentsDirty), TT_None, nullptr);
+}
+
+void FSkeletalMeshStreamOut::ConditionalMarkComponentsDirty(const FContext& Context)
+{
+	DECLARE_SCOPE_CYCLE_COUNTER(TEXT("FSkeletalMeshStreamOut::ConditionalMarkComponentsDirty"), STAT_SkeletalMeshStreamOut_ConditionalMarkComponentsDirty, STATGROUP_StreamingDetails);
+	CSV_SCOPED_TIMING_STAT_GLOBAL(SkStreamingMarkDirtyTime);
+	check(Context.CurrentThread == TT_GameThread);
+
+	const USkeletalMesh* Mesh = Context.Mesh;
+	FSkeletalMeshRenderData* RenderData = Context.RenderData;
+	if (!IsCancelled() && Mesh && RenderData)
+	{
+		RenderData->PendingFirstLODIdx = ResourceState.LODCountToAssetFirstLODIdx(ResourceState.NumRequestedLODs);
+
+		TArray<const UPrimitiveComponent*> Comps;
+		IStreamingManager::Get().GetTextureStreamingManager().GetAssetComponents(Mesh, Comps, [](const UPrimitiveComponent* Comp)
+		{
+			return !Comp->IsComponentTickEnabled();
+		});
+		for (int32 Idx = 0; Idx < Comps.Num(); ++Idx)
+		{
+			check(Comps[Idx]->IsA<USkinnedMeshComponent>());
+			USkinnedMeshComponent* Comp = (USkinnedMeshComponent*)Comps[Idx];
+			if (Comp->GetPredictedLODLevel() < RenderData->PendingFirstLODIdx)
+			{
+				Comp->SetPredictedLODLevel(RenderData->PendingFirstLODIdx);
+				Comp->bForceMeshObjectUpdate = true;
+				Comp->MarkRenderDynamicDataDirty();
+			}
+		}
+	}
+	else
+	{
+		Abort();
+	}
+	PushTask(Context, TT_Async, SRA_UPDATE_CALLBACK(WaitForReferences), (EThreadType)Context.CurrentThread, SRA_UPDATE_CALLBACK(Cancel));
+}
+
+void FSkeletalMeshStreamOut::WaitForReferences(const FContext& Context)
+{
+	DECLARE_SCOPE_CYCLE_COUNTER(TEXT("FSkeletalMeshStreamOut::WaitForReferences"), STAT_SkeletalMeshStreamOut_WaitForReferences, STATGROUP_StreamingDetails);
+	check(Context.CurrentThread == TT_Async);
+
+	const USkeletalMesh* Mesh = Context.Mesh;
+	FSkeletalMeshRenderData* RenderData = Context.RenderData;
+	uint32 NumExternalReferences = 0;
+
+	if (Mesh && RenderData)
+	{
+		for (int32 LODIndex = CurrentFirstLODIdx; LODIndex < PendingFirstLODIdx; ++LODIndex)
+		{
+			// Minus 1 since the LODResources reference is not considered external
+			NumExternalReferences += Context.LODResourcesView[LODIndex]->GetRefCount() - 1;
+		}
+
+		if (NumExternalReferences > PreviousNumberOfExternalReferences && NumReferenceChecks > 0)
+		{
+			PreviousNumberOfExternalReferences = NumExternalReferences;
+			UE_LOG(LogSkeletalMesh, Warning, TEXT("[%s] Streamed out LODResources got referenced while in pending stream out."), *Mesh->GetName());
+		}
+	}
+
+	if (!NumExternalReferences || NumReferenceChecks >= GStreamingMaxReferenceChecks)
+	{
+		PushTask(Context, TT_Render, SRA_UPDATE_CALLBACK(ReleaseBuffers), (EThreadType)Context.CurrentThread, SRA_UPDATE_CALLBACK(Cancel));
+		
+		// This is required to allow the engine to generate the bone buffers for the PendingFirstLODIdx. See logic in FSkeletalMeshSceneProxy::GetMeshElementsConditionallySelectable().
+		if (NumReferenceChecks == 0)
+		{
+			bDeferExecution = true;
+		}
+	}
+	else
+	{
+		++NumReferenceChecks;
+		if (NumReferenceChecks >= GStreamingMaxReferenceChecks)
+		{
+			UE_LOG(LogSkeletalMesh, Warning, TEXT("[%s] Streamed out LODResources references are not getting released."), *Mesh->GetName());
+		}
+
+		bDeferExecution = true;
+		PushTask(Context, TT_Async, SRA_UPDATE_CALLBACK(WaitForReferences),  (EThreadType)Context.CurrentThread, SRA_UPDATE_CALLBACK(Cancel));
+	}
+}
+
+void FSkeletalMeshStreamOut::ReleaseBuffers(const FContext& Context)
+{
+	DECLARE_SCOPE_CYCLE_COUNTER(TEXT("FSkeletalMeshStreamOut::ReleaseBuffers"), STAT_SkeletalMeshStreamOut_ReleaseBuffers, STATGROUP_StreamingDetails);
+	check(Context.CurrentThread == TT_Render);
+	check(IsInRenderingThread());
+
+	const USkeletalMesh* Mesh = Context.Mesh;
+	FSkeletalMeshRenderData* RenderData = Context.RenderData;
+	if (!IsCancelled() && Mesh && RenderData)
+	{
+		RenderData->CurrentFirstLODIdx = Context.RenderData->PendingFirstLODIdx;
+
+		{
+			FRHIResourceReplaceBatcher Batcher(FRHICommandListImmediate::Get(), GSkelMeshMaxNumResourceUpdatesPerBatch);
+
+			for (int32 LODIndex = CurrentFirstLODIdx; LODIndex < PendingFirstLODIdx; ++LODIndex)
+			{
+				FSkeletalMeshLODRenderData& LODResource = *Context.LODResourcesView[LODIndex];
+				FStaticMeshVertexBuffers& VBs = LODResource.StaticVertexBuffers;
+				LODResource.DecrementMemoryStats();
+				VBs.StaticMeshVertexBuffer.ReleaseRHIForStreaming(Batcher);
+				VBs.PositionVertexBuffer.ReleaseRHIForStreaming(Batcher);
+				VBs.ColorVertexBuffer.ReleaseRHIForStreaming(Batcher);
+				LODResource.SkinWeightVertexBuffer.ReleaseRHIForStreaming(Batcher);
+				LODResource.ClothVertexBuffer.ReleaseRHIForStreaming(Batcher);
+				LODResource.MultiSizeIndexContainer.ReleaseRHIForStreaming(Batcher);
+				LODResource.SkinWeightProfilesData.ReleaseRHIForStreaming(Batcher);
+
+				if (!FPlatformProperties::HasEditorOnlyData())
+				{
+					// TODO requires more testing : LODResource.ReleaseCPUResources(true);
+				}
+
+#if RHI_RAYTRACING
+				if (IsRayTracingAllowed())
+				{
+					if (RenderData->LODRenderData[LODIndex].bReferencedByStaticSkeletalMeshObjects_RenderThread)
+					{
+						LODResource.StaticRayTracingGeometry.ReleaseResource();
+					}
+				}
+#endif
+			}
+		}
+
+		MarkAsSuccessfullyFinished();
+	}
+}
+
+void FSkeletalMeshStreamOut::Cancel(const FContext& Context)
+{
+	DECLARE_SCOPE_CYCLE_COUNTER(TEXT("FSkeletalMeshStreamOut::Cancel"), STAT_SkeletalMeshStreamOut_Cancel, STATGROUP_StreamingDetails);
+
+	if (Context.RenderData)
+	{
+		Context.RenderData->PendingFirstLODIdx = Context.RenderData->CurrentFirstLODIdx;
+	}
+}
+
+void FSkeletalMeshStreamIn_IO::FCancelIORequestsTask::DoWork()
+{
+	check(PendingUpdate);
+	// Acquire the lock of this object in order to cancel any pending IO.
+	// If the object is currently being ticked, wait.
+	ETaskState OldState = PendingUpdate->DoLock();
+	PendingUpdate->CancelIORequest();
+	PendingUpdate->DoUnlock(OldState);
+}
+
+void FSkeletalMeshStreamIn_IO::Abort()
+{
+	if (!IsCancelled() && !IsCompleted())
+	{
+		FSkeletalMeshStreamIn::Abort();
+
+		if (BulkDataRequest.IsPending())
+		{
+			// Prevent the update from being considered done before this is finished.
+			// By checking that it was not already cancelled, we make sure this doesn't get called twice.
+			(new FAsyncCancelIORequestsTask(this))->StartBackgroundTask();
+		}
+	}
+}
+
+void FSkeletalMeshStreamIn_IO::SetIORequest(const FContext& Context)
+{
+	if (IsCancelled())
+	{
+		return;
+	}
+	check(BulkDataRequest.IsNone() && PendingFirstLODIdx < CurrentFirstLODIdx);
+
+	const USkeletalMesh* Mesh = Context.Mesh;
+	FSkeletalMeshRenderData* RenderData = Context.RenderData;
+	if (Mesh && RenderData)
+	{
+		const int32 BatchCount = CurrentFirstLODIdx - PendingFirstLODIdx;
+		FBulkDataBatchRequest::FScatterGatherBuilder Batch = FBulkDataBatchRequest::ScatterGather(BatchCount);
+		for (int32 Index = PendingFirstLODIdx; Index < CurrentFirstLODIdx; ++Index)
+		{
+			Batch.Read(Context.LODResourcesView[Index]->StreamingBulkData);
+		}
+
+		// Increment as we push the request. If a request complete immediately, then it will call the callback
+		// but that won't do anything because the tick would not try to acquire the lock since it is already locked.
+		TaskSynchronization.Increment();
+
+		EAsyncIOPriorityAndFlags Priority = AIOP_Low;
+		if (bHighPrioIORequest)
+		{
+			static IConsoleVariable* CVarAsyncLoadingPrecachePriority = IConsoleManager::Get().FindConsoleVariable(TEXT("s.AsyncLoadingPrecachePriority"));
+			const bool bLoadBeforeAsyncPrecache = CVarStreamingLowResHandlingMode.GetValueOnAnyThread() == (int32)FRenderAssetStreamingSettings::LRHM_LoadBeforeAsyncPrecache;
+
+			if (CVarAsyncLoadingPrecachePriority && bLoadBeforeAsyncPrecache)
+			{
+				const int32 AsyncIOPriority = CVarAsyncLoadingPrecachePriority->GetInt();
+				// Higher priority than regular requests but don't go over max
+				Priority = (EAsyncIOPriorityAndFlags)FMath::Clamp<int32>(AsyncIOPriority + 1, AIOP_BelowNormal, AIOP_MAX);
+			}
+			else
+			{
+				Priority = AIOP_BelowNormal;
+			}
+		}
+
+		Batch.Issue(BulkData, Priority, [this](FBulkDataRequest::EStatus Status)
+		{
+			// At this point task synchronization would hold the number of pending requests.
+			TaskSynchronization.Decrement();
+
+			if (FBulkDataRequest::EStatus::Ok != Status)
+			{
+				// If IO requests was cancelled but the streaming request wasn't, this is an IO error.
+				if (!bIsCancelled)
+				{
+					bFailedOnIOError = true;
+				}
+				MarkAsCancelled();
+			}
+
+#if !UE_BUILD_SHIPPING
+			// On some platforms the IO is too fast to test cancelation requests timing issues.
+			if (FRenderAssetStreamingSettings::ExtraIOLatency > 0 && TaskSynchronization.GetValue() == 0)
+			{
+				FPlatformProcess::Sleep(FRenderAssetStreamingSettings::ExtraIOLatency * .001f); // Slow down the streaming.
+			}
+#endif
+			// The tick here is intended to schedule the success or cancel callback.
+			// Using TT_None ensure gets which could create a dead lock.
+			Tick(FSkeletalMeshUpdate::TT_None);
+		},
+		BulkDataRequest);
+	}
+	else
+	{
+		MarkAsCancelled();
+	}
+}
+
+void FSkeletalMeshStreamIn_IO::ClearIORequest(const FContext& Context)
+{
+	if (BulkDataRequest.IsPending())
+	{
+		BulkDataRequest.Cancel();
+		BulkDataRequest.Wait();
+	}
+	
+	BulkDataRequest = FBulkDataBatchRequest();
+	BulkData = FIoBuffer();
+}
+
+void FSkeletalMeshStreamIn_IO::ReportIOError(const FContext& Context)
+{
+	// Invalidate the cache state of all initial mips (note that when using FIoChunkId each mip has a different value).
+	if (bFailedOnIOError && Context.Mesh)
+	{
+		IRenderAssetStreamingManager& StreamingManager = IStreamingManager::Get().GetRenderAssetStreamingManager();
+		for (int32 MipIndex = 0; MipIndex < CurrentFirstLODIdx; ++MipIndex)
+		{
+			StreamingManager.MarkMountedStateDirty(Context.Mesh->GetMipIoFilenameHash(MipIndex));
+		}
+
+		UE_LOG(LogContentStreaming, Warning, TEXT("[%s] SkeletalMesh stream in request failed due to IO error (LOD %d-%d)."), *Context.Mesh->GetName(), PendingFirstLODIdx, CurrentFirstLODIdx - 1);
+	}
+}
+
+void FSkeletalMeshStreamIn_IO::SerializeLODData(const FContext& Context)
+{
+	LLM_SCOPE_BYNAME(TEXT("SkeletalMesh/Serialize"));
+
+	check(!TaskSynchronization.GetValue());
+	const USkeletalMesh* Mesh = Context.Mesh;
+	FSkeletalMeshRenderData* RenderData = Context.RenderData;
+	if (!IsCancelled() && Mesh && RenderData)
+	{
+		check(BulkData.GetSize() >= 0 && BulkData.GetSize() <= TNumericLimits<uint32>::Max());
+
+		FMemoryReaderView Ar(BulkData.GetView(), true);
+		for (int32 LODIndex = PendingFirstLODIdx; LODIndex < CurrentFirstLODIdx; ++LODIndex)
+		{
+			FSkeletalMeshLODRenderData& LODResource = *Context.LODResourcesView[LODIndex];
+			const bool bForceKeepCPUResources = FSkeletalMeshLODRenderData::ShouldForceKeepCPUResources();
+			const bool bNeedsCPUAccess = FSkeletalMeshLODRenderData::ShouldKeepCPUResources(Mesh, LODIndex + Context.AssetLODBias, bForceKeepCPUResources);
+			constexpr uint8 DummyStripFlags = 0;
+			LODResource.SerializeStreamedData(Ar, const_cast<USkeletalMesh*>(Mesh), LODIndex + Context.AssetLODBias, DummyStripFlags, bNeedsCPUAccess, bForceKeepCPUResources);
+
+			// Attempt to recover from possibly corrupted data
+			if (Ar.IsError())
+			{
+				UE_LOG(LogContentStreaming, Error,
+					TEXT("[%s] SkeletalMesh stream in failed due to possibly corrupted data. LOD %d %d-%d. BulkData %#x offset %lld size %lld flags %#x. bForceKeepCPUResources %d. bNeedsCPUAccess %d."),
+					*Mesh->GetPathName(),
+					LODIndex,
+					PendingFirstLODIdx,
+					CurrentFirstLODIdx - 1,
+					LODResource.StreamingBulkData.GetIoFilenameHash(),
+					LODResource.StreamingBulkData.GetBulkDataOffsetInFile(),
+					LODResource.StreamingBulkData.GetBulkDataSize(),
+					LODResource.StreamingBulkData.GetBulkDataFlags(),
+					bForceKeepCPUResources,
+					bNeedsCPUAccess);
+
+#if STREAMING_RETRY_ON_DESERIALIZATION_ERROR
+				bFailedOnIOError = true;
+				MarkAsCancelled();
+				break;
+#else
+				GLog->FlushThreadedLogs();
+				GLog->Flush();
+				UE_LOG(LogContentStreaming, Fatal, TEXT("Possibly corrupted skeletal mesh LOD data detected."));
+#endif
+			}
+		}
+
+		BulkData = FIoBuffer();
+	}
+}
+
+void FSkeletalMeshStreamIn_IO::Cancel(const FContext& Context)
+{
+	DoCancel(Context);
+	ReportIOError(Context);
+}
+
+void FSkeletalMeshStreamIn_IO::CancelIORequest()
+{
+	if (BulkDataRequest.IsPending())
+	{
+		BulkDataRequest.Cancel();
+	}
+}
+
+FSkeletalMeshStreamIn_IO::FSkeletalMeshStreamIn_IO(const USkeletalMesh* InMesh, bool bHighPrio, EThreadType CreateResourcesThread)
+	: FSkeletalMeshStreamIn(InMesh, CreateResourcesThread)
+	, bHighPrioIORequest(bHighPrio)
+{
+	PushTask(FContext(InMesh, TT_None), TT_Async, SRA_UPDATE_CALLBACK(DoInitiateIO), TT_None, nullptr);
+}
+
+void FSkeletalMeshStreamIn_IO::DoInitiateIO(const FContext& Context)
+{
+	check(Context.CurrentThread == TT_Async);
+
+	SetIORequest(Context);
+
+	PushTask(Context, TT_Async, SRA_UPDATE_CALLBACK(DoSerializeLODData), TT_Async, SRA_UPDATE_CALLBACK(DoCancelIO));
+}
+
+void FSkeletalMeshStreamIn_IO::DoSerializeLODData(const FContext& Context)
+{
+	check(Context.CurrentThread == TT_Async);
+	SerializeLODData(Context);
+	ClearIORequest(Context);
+
+	PushTask(Context
+		, CreateResourcesThread, SRA_UPDATE_CALLBACK(DoCreateBuffers)
+		, (EThreadType)Context.CurrentThread, SRA_UPDATE_CALLBACK(Cancel));
+}
+
+void FSkeletalMeshStreamIn_IO::DoCreateBuffers(const FContext& Context)
+{
+	CreateBuffers(Context);
+
+	check(!TaskSynchronization.GetValue());
+
+	// We cannot cancel once DoCreateBuffers has started executing, as there's an RHICmdList that must be submitted.
+	// Pass the same callback for both task and cancel.
+	PushTask(Context
+		, TT_Render, SRA_UPDATE_CALLBACK(DoFinishUpdate)
+		, TT_Render, SRA_UPDATE_CALLBACK(DoFinishUpdate)
+	);
+}
+
+void FSkeletalMeshStreamIn_IO::DoCancelIO(const FContext& Context)
+{
+	ClearIORequest(Context);
+	PushTask(Context, TT_None, nullptr, (EThreadType)Context.CurrentThread, SRA_UPDATE_CALLBACK(Cancel));
+}
+
+#if WITH_EDITOR
+FSkeletalMeshStreamIn_DDC::FSkeletalMeshStreamIn_DDC(const USkeletalMesh* InMesh, EThreadType CreateResourcesThread)
+	: FSkeletalMeshStreamIn(InMesh, CreateResourcesThread)
+{
+	PushTask(FContext(InMesh, TT_None), TT_Async, SRA_UPDATE_CALLBACK(DoLoadNewLODsFromDDC), TT_None, nullptr);
+}
+
+void FSkeletalMeshStreamIn_DDC::LoadNewLODsFromDDC(const FContext& Context)
+{
+	check(Context.CurrentThread == TT_Async);
+	// TODO: support streaming CPU data for editor builds
+}
+
+void FSkeletalMeshStreamIn_DDC::DoLoadNewLODsFromDDC(const FContext& Context)
+{
+	LoadNewLODsFromDDC(Context);
+	check(!TaskSynchronization.GetValue());
+
+	PushTask(Context
+		, CreateResourcesThread, SRA_UPDATE_CALLBACK(DoCreateBuffers)
+		, (EThreadType)Context.CurrentThread, SRA_UPDATE_CALLBACK(DoCancel));
+}
+
+void FSkeletalMeshStreamIn_DDC::DoCreateBuffers(const FContext& Context)
+{
+	CreateBuffers(Context);
+	
+	check(!TaskSynchronization.GetValue());
+	PushTask(Context, TT_Render, SRA_UPDATE_CALLBACK(DoFinishUpdate), TT_None, nullptr);
+}
+#endif
diff --git a/Engine/Source/Runtime/Engine/Internal/Streaming/SkeletalMeshUpdate.h b/Engine/Source/Runtime/Engine/Internal/Streaming/SkeletalMeshUpdate.h
new file mode 100644
--- /dev/null
+++ b/Engine/Source/Runtime/Engine/Internal/Streaming/SkeletalMeshUpdate.h
@@ -0,0 +1,250 @@
+// Copyright Epic Games, Inc. All Rights Reserved.
+
+/*=============================================================================
+SkeletalMeshUpdate.h: Helpers to stream in and out skeletal mesh LODs.
+=============================================================================*/
+
+#pragma once
+
+#include "CoreMinimal.h"
+#include "Engine/SkeletalMesh.h"
+#include "IO/IoDispatcher.h"
+#include "RenderAssetUpdate.h"
+#include "Rendering/SkinWeightVertexBuffer.h"
+#include "Serialization/BulkData.h"
+#include "RayTracingGeometry.h"
+
+/**
+* A context used to update or proceed with the next update step.
+* The mesh and render data references could be stored in the update object
+* but are currently kept outside to avoid lifetime management within the object.
+*/
+struct ENGINE_API FSkelMeshUpdateContext
+{
+	typedef int32 EThreadType;
+
+	FSkelMeshUpdateContext(const USkeletalMesh* InMesh, EThreadType InCurrentThread);
+
+	FSkelMeshUpdateContext(const UStreamableRenderAsset* InMesh, EThreadType InCurrentThread);
+
+	EThreadType GetCurrentThread() const
+	{
+		return CurrentThread;
+	}
+
+	/** The mesh to update, this must be the same one as the one used when creating the FSkeletalMeshUpdate object. */
+	const USkeletalMesh* Mesh;
+	/** The current render data of this mesh. */
+	FSkeletalMeshRenderData* RenderData;
+	/** The array view of streamable LODs from the asset. Takes into account FStreamableRenderResourceState::AssetLODBias and FStreamableRenderResourceState::MaxNumLODs. */
+	TArrayView<FSkeletalMeshLODRenderData*> LODResourcesView;
+	/** Cached value of mesh its LOD bias (MinLOD for SkeletalMesh). */
+	int32 AssetLODBias = 0;
+
+	/** The thread on which the context was created. */
+	EThreadType CurrentThread;
+};
+
+extern template class TRenderAssetUpdate<FSkelMeshUpdateContext>;
+
+/**
+* This class provides a framework for loading and unloading the LODs of skeletal meshes.
+* Each thread essentially calls Tick() until the job is done.
+* The object can be safely deleted when IsCompleted() returns true.
+*/
+class FSkeletalMeshUpdate : public TRenderAssetUpdate<FSkelMeshUpdateContext>
+{
+public:
+	FSkeletalMeshUpdate(const USkeletalMesh* InMesh);
+
+	virtual ~FSkeletalMeshUpdate() {}
+
+	virtual void Abort()
+	{
+		TRenderAssetUpdate<FSkelMeshUpdateContext>::Abort();
+	}
+};
+
+class ENGINE_API FSkeletalMeshStreamIn : public FSkeletalMeshUpdate
+{
+public:
+	FSkeletalMeshStreamIn(const USkeletalMesh* InMesh, EThreadType CreateResourcesThread);
+
+	virtual ~FSkeletalMeshStreamIn();
+
+protected:
+	/** Correspond to the buffers in FSkeletalMeshLODRenderData */
+	struct FIntermediateBuffers
+	{
+		FBufferRHIRef TangentsVertexBuffer;
+		FBufferRHIRef TexCoordVertexBuffer;
+		FBufferRHIRef PositionVertexBuffer;
+		FBufferRHIRef ColorVertexBuffer;
+		FSkinWeightRHIInfo SkinWeightVertexBuffer;
+		FBufferRHIRef ClothVertexBuffer;
+		FBufferRHIRef IndexBuffer;
+		TArray<TPair<FName, FSkinWeightRHIInfo>> AltSkinWeightVertexBuffers;
+
+		void CreateFromCPUData(FRHICommandListBase& RHICmdList, FSkeletalMeshLODRenderData& LODResource);
+
+		/** Transfer ownership of buffers to a LOD resource */
+		void TransferBuffers(FSkeletalMeshLODRenderData& LODResource, FRHIResourceReplaceBatcher& Batcher);
+	};
+
+#if RHI_RAYTRACING
+	struct FIntermediateRayTracingGeometry
+	{
+	private:
+		FRayTracingGeometryInitializer Initializer;
+		FRayTracingGeometryRHIRef RayTracingGeometryRHI;
+		bool bRequiresBuild = false;
+
+	public:
+		void CreateFromCPUData(FRHICommandListBase& RHICmdList, FRayTracingGeometry& RayTracingGeometry);
+
+		void SafeRelease();
+
+		void TransferRayTracingGeometry(FRayTracingGeometry& RayTracingGeometry, FRHIResourceReplaceBatcher& Batcher);
+	};
+#endif
+
+	/** Create buffers with new LOD data */
+	void CreateBuffers(const FContext& Context);
+
+	/** Discard newly streamed-in CPU data */
+	void DiscardNewLODs(const FContext& Context);
+
+	/** Apply the new buffers (if not cancelled) and finish the update process. When cancelled, the intermediate buffers simply gets discarded. */
+	void DoFinishUpdate(const FContext& Context);
+
+	/** Discard streamed-in CPU data and intermediate RHI buffers */
+	void DoCancel(const FContext& Context);
+
+	/** The intermediate buffers created in the update process. */
+	FIntermediateBuffers IntermediateBuffersArray[MAX_MESH_LOD_COUNT];
+
+#if RHI_RAYTRACING
+	FIntermediateRayTracingGeometry IntermediateRayTracingGeometry[MAX_MESH_LOD_COUNT];
+#endif
+
+	/** RHI command list used for creating buffers and replacing the streaming placeholders. Submitted in DoFinishUpdate */
+	FRHICommandList* StreamingRHICmdList = nullptr;
+
+	/** The thread to use for recording the above command list */
+	const EThreadType CreateResourcesThread;
+};
+
+class FSkeletalMeshStreamOut : public FSkeletalMeshUpdate
+{
+public:
+	FSkeletalMeshStreamOut(const USkeletalMesh* InMesh);
+
+	virtual ~FSkeletalMeshStreamOut() {}
+
+private:
+
+	int32 NumReferenceChecks = 0;
+	uint32 PreviousNumberOfExternalReferences = 0;
+
+
+	/** Notify components that the LOD is being streamed out so that they can release references. */
+	void ConditionalMarkComponentsDirty(const FContext& Context);
+
+	/** Wait for all references to be released. */
+	void WaitForReferences(const FContext& Context);
+
+	/** Release RHI buffers and update SRVs. */
+	void ReleaseBuffers(const FContext& Context);
+
+	/** Cancel the pending mip change. */
+	void Cancel(const FContext& Context);
+};
+
+class FSkeletalMeshStreamIn_IO : public FSkeletalMeshStreamIn
+{
+public:
+	FSkeletalMeshStreamIn_IO(const USkeletalMesh* InMesh, bool bHighPrio, EThreadType CreateResourcesThread);
+
+	virtual ~FSkeletalMeshStreamIn_IO() {}
+
+	virtual void Abort() override;
+
+protected:
+	class FCancelIORequestsTask : public FNonAbandonableTask
+	{
+	public:
+		FCancelIORequestsTask(FSkeletalMeshStreamIn_IO* InPendingUpdate)
+			: PendingUpdate(InPendingUpdate)
+		{}
+
+		void DoWork();
+
+		FORCEINLINE TStatId GetStatId() const
+		{
+			RETURN_QUICK_DECLARE_CYCLE_STAT(FCancelIORequestsTask_SkeletalMesh, STATGROUP_ThreadPoolAsyncTasks);
+		}
+
+	private:
+		TRefCountPtr<FSkeletalMeshStreamIn_IO> PendingUpdate;
+	};
+
+	typedef FAutoDeleteAsyncTask<FCancelIORequestsTask> FAsyncCancelIORequestsTask;
+	friend class FCancelIORequestsTask;
+
+	/** Create a new async IO request to read in LOD data */
+	void SetIORequest(const FContext& Context);
+
+	/** Release IORequest and IOFileHandle. IORequest will be cancelled if still inflight */
+	void ClearIORequest(const FContext& Context);
+
+	/** Report IO errors if any. */
+	void ReportIOError(const FContext& Context);
+
+	/** Serialize data of new LODs to corresponding FStaticMeshLODResources */
+	void SerializeLODData(const FContext& Context);
+
+	/** Cancel and report IO error */
+	void Cancel(const FContext& Context);
+
+	/** Called by FAsyncCancelIORequestsTask to cancel inflight IO request if any */
+	void CancelIORequest();
+
+	/** Handle to bulk data I/O request */
+	FBulkDataBatchRequest BulkDataRequest;
+
+	/** Bulk data I/O buffer */
+	FIoBuffer BulkData;
+
+	bool bHighPrioIORequest;
+
+	// Whether an IO error was detected (when files do not exists).
+	bool bFailedOnIOError = false;
+
+	void DoInitiateIO(const FContext& Context);
+
+	void DoSerializeLODData(const FContext& Context);
+
+	void DoCreateBuffers(const FContext& Context);
+
+	void DoCancelIO(const FContext& Context);
+};
+
+#if WITH_EDITOR
+class FSkeletalMeshStreamIn_DDC : public FSkeletalMeshStreamIn
+{
+public:
+	FSkeletalMeshStreamIn_DDC(const USkeletalMesh* InMesh, EThreadType CreateResourcesThread);
+
+	virtual ~FSkeletalMeshStreamIn_DDC() {}
+
+protected:
+	void LoadNewLODsFromDDC(const FContext& Context);
+
+private:
+	/** Load new LOD buffers from DDC and queue a task to create RHI buffers on RT */
+	void DoLoadNewLODsFromDDC(const FContext& Context);
+
+	/** Create RHI buffers for newly streamed-in LODs and queue a task to rename references on RT */
+	void DoCreateBuffers(const FContext& Context);
+};
+#endif
diff --git a/Engine/Source/Runtime/Engine/Private/Streaming/RenderAssetUpdate.inl b/Engine/Source/Runtime/Engine/Private/Streaming/RenderAssetUpdate.inl
deleted file mode 100644
diff --git a/Engine/Source/Runtime/Engine/Private/Streaming/SkeletalMeshUpdate.cpp b/Engine/Source/Runtime/Engine/Private/Streaming/SkeletalMeshUpdate.cpp
deleted file mode 100644
diff --git a/Engine/Source/Runtime/Engine/Private/Streaming/SkeletalMeshUpdate.h b/Engine/Source/Runtime/Engine/Private/Streaming/SkeletalMeshUpdate.h
deleted file mode 100644
diff --git a/Engine/Source/Runtime/Engine/Public/RenderAssetUpdate.h b/Engine/Source/Runtime/Engine/Public/RenderAssetUpdate.h
--- a/Engine/Source/Runtime/Engine/Public/RenderAssetUpdate.h
+++ b/Engine/Source/Runtime/Engine/Public/RenderAssetUpdate.h
@@ -24,7 +24,7 @@ ENGINE_API bool IsAssetStreamingSuspended();
 * Each thread essentially calls Tick() until the job is done.
 * The object can be safely deleted when IsCompleted() returns true.
 */
-class FRenderAssetUpdate : public IRefCountedObject
+class ENGINE_API FRenderAssetUpdate : public IRefCountedObject
 {
 public:
 
@@ -56,7 +56,7 @@ public:
 	* @param InAsset - the texture/mesh being updated, this must be the same texture/mesh as the texture/mesh used to create this object.
 	* @param InCurrentThread - the thread from which the tick is being called. Using TT_None ensures that no work will be immediately performed.
 	*/
-	ENGINE_API void Tick(EThreadType InCurrentThread);
+	void Tick(EThreadType InCurrentThread);
 
 	/** Returns whether the task has finished executing and there is no other thread possibly accessing it. */
 	bool IsCompleted() const
@@ -106,7 +106,7 @@ public:
 		return (uint32)NumRefs.Increment(); 
 	}
 
-	ENGINE_API uint32 Release() const final override;
+	uint32 Release() const final override;
 
 	uint32 GetRefCount() const final override
 	{
