Revert erroneous change that made a warning only a displayed log

The warning was firing in places where it shouldn't, but the correct fix was to add a check before the warning (is a memory limit set?) and not to make the warning itself only be at Display.

diff --git a/Engine/Source/Runtime/Engine/Private/AssetCompilingManager.cpp b/Engine/Source/Runtime/Engine/Private/AssetCompilingManager.cpp
--- a/Engine/Source/Runtime/Engine/Private/AssetCompilingManager.cpp
+++ b/Engine/Source/Runtime/Engine/Private/AssetCompilingManager.cpp
@@ -301,7 +301,7 @@ public:
 					{
 						// Separately report the case if an asset was bigger than our manually set memory limit. Such assets are always too big,
 						// and don't just exceed the currently available memory.
-						UE_LOGFMT_NSLOC(LogAsyncCompilation, Display, "AsyncAssetCompilation", "HardMemoryLimitExceeded",
+						UE_LOGFMT_NSLOC(LogAsyncCompilation, Warning, "AsyncAssetCompilation", "HardMemoryLimitExceeded",
 							"BEWARE: AssetCompile memory estimate is greater than the hard memory limit, but we're running it [{TaskName}] anyway! "
 							"RequiredMemory = {TotalEstimatedMemory} MiB + {RequiredMemory} MiB, MemoryLimit = {MemoryLimit} MiB, HardMemoryLimit = {HardMemoryLimit} MiB",
 							("TaskName", DebugName),
