Fix crash when changing r.RayTracing.Geometry.InstancedStaticMeshes.SimulationCount
- TArray reallocations move FRayTracingGeometry around in memory without calling copy constructors, etc
- Prevent crash by manually releasing and recreating all geometries when reallocating the array.
- In the future we could work on alternative involving re-registering geometries with FRayTracingGeometryManager.

diff --git a/Engine/Source/Runtime/Engine/Private/InstancedStaticMesh.cpp b/Engine/Source/Runtime/Engine/Private/InstancedStaticMesh.cpp
--- a/Engine/Source/Runtime/Engine/Private/InstancedStaticMesh.cpp
+++ b/Engine/Source/Runtime/Engine/Private/InstancedStaticMesh.cpp
@@ -1453,6 +1453,8 @@ void FInstancedStaticMeshSceneProxy::DestroyRenderThreadResources()
 		DynamicRayTracingItem.DynamicGeometry.ReleaseResource();
 		DynamicRayTracingItem.DynamicGeometryVertexBuffer.Release();
 	}
+
+	RayTracingDynamicData.Empty();
 #endif
 }
 
@@ -1933,44 +1935,49 @@ void FInstancedStaticMeshSceneProxy::GetDynamicRayTracingInstances(struct FRayTr
 
 void FInstancedStaticMeshSceneProxy::SetupRayTracingDynamicInstances(int32 NumDynamicInstances, int32 LODIndex)
 {
-	if (RayTracingDynamicData.Num() > NumDynamicInstances || CachedRayTracingLODIndex != LODIndex)
+	if (NumDynamicInstances == RayTracingDynamicData.Num() && LODIndex == CachedRayTracingLODIndex)
 	{
-		// free the unused/out of date entries
-
-		int32 FirstToFree = (CachedRayTracingLODIndex != LODIndex) ? 0 : NumDynamicInstances;
-		for (int32 Item = FirstToFree; Item < RayTracingDynamicData.Num(); Item++)
-		{
-			auto& DynamicRayTracingItem = RayTracingDynamicData[Item];
-			DynamicRayTracingItem.DynamicGeometry.ReleaseResource();
-			DynamicRayTracingItem.DynamicGeometryVertexBuffer.Release();
-		}
-		RayTracingDynamicData.SetNum(FirstToFree);
+		return;
 	}
 
-	if (RayTracingDynamicData.Num() < NumDynamicInstances)
+	CachedRayTracingLODIndex = LODIndex;
+
+	// if either NumDynamicInstances or LOD changed
+	// need to recreate RayTracingDynamicData array
+	// FRayTracingGeometry is no relocate-able so can't grow or shrink array
+	// TODO: Investigate manually re-registering with FRayTracingGeometryManager, using sparse array or other alternative data structure instead
+
+	// release entries
+
+	for (int32 Item = 0; Item < RayTracingDynamicData.Num(); Item++)
 	{
-		RayTracingDynamicData.Reserve(NumDynamicInstances);
-		const int32 StartIndex = RayTracingDynamicData.Num();
-		const FStaticMeshLODResources& LODModel = RenderData->LODResources[LODIndex];
+		auto& DynamicRayTracingItem = RayTracingDynamicData[Item];
+		DynamicRayTracingItem.DynamicGeometry.ReleaseResource();
+		DynamicRayTracingItem.DynamicGeometryVertexBuffer.Release();
+	}
 
-		for (int32 Item = StartIndex; Item < NumDynamicInstances; Item++)
-		{
-			FRayTracingDynamicData &DynamicData = RayTracingDynamicData.AddDefaulted_GetRef();
+	// clear and resize
+	RayTracingDynamicData.Empty(NumDynamicInstances);
 
-			FRayTracingGeometryInitializer Initializer = LODModel.RayTracingGeometry.Initializer;
-			for (FRayTracingGeometrySegment& Segment : Initializer.Segments)
-			{
-				Segment.VertexBuffer = nullptr; 
-			}
-			Initializer.bAllowUpdate = true;
-			Initializer.bFastBuild = true;
+	// create new geometries
+
+	const FStaticMeshLODResources& LODModel = RenderData->LODResources[LODIndex];
 
-			DynamicData.DynamicGeometry.SetInitializer(MoveTemp(Initializer));
-			DynamicData.DynamicGeometry.InitResource(FRHICommandListImmediate::Get());
+	for (int32 Item = 0; Item < NumDynamicInstances; Item++)
+	{
+		FRayTracingDynamicData& DynamicData = RayTracingDynamicData.AddDefaulted_GetRef();
+
+		FRayTracingGeometryInitializer Initializer = LODModel.RayTracingGeometry.Initializer;
+		for (FRayTracingGeometrySegment& Segment : Initializer.Segments)
+		{
+			Segment.VertexBuffer = nullptr;
 		}
-	}
+		Initializer.bAllowUpdate = true;
+		Initializer.bFastBuild = true;
 
-	CachedRayTracingLODIndex = LODIndex;
+		DynamicData.DynamicGeometry.SetInitializer(MoveTemp(Initializer));
+		DynamicData.DynamicGeometry.InitResource(FRHICommandListImmediate::Get());
+	}
 }
 
 #endif
