Modified FInstancedStaticMeshSceneProxy::GetDynamicRayTracingInstances(...) to select first LOD with valid ray tracing geometry.

diff --git a/Engine/Source/Runtime/Engine/Private/InstancedStaticMesh.cpp b/Engine/Source/Runtime/Engine/Private/InstancedStaticMesh.cpp
--- a/Engine/Source/Runtime/Engine/Private/InstancedStaticMesh.cpp
+++ b/Engine/Source/Runtime/Engine/Private/InstancedStaticMesh.cpp
@@ -1625,17 +1625,27 @@ void FInstancedStaticMeshSceneProxy::GetDynamicRayTracingInstances(struct FRayTr
 
 	int32 LODIndex = FMath::Max<int32>(MinAllowedLODIndex, GetCurrentFirstLODIdx_RenderThread());
 
-	FRayTracingGeometry* RayTracingGeometry = &RenderData->LODResources[LODIndex].RayTracingGeometry;
+	FRayTracingGeometry* RayTracingGeometry = nullptr;
 
-	if (!RayTracingGeometry->IsInitialized())
+	// Select first LOD with valid ray tracing geometry
+	for (; LODIndex < RenderData->LODResources.Num(); ++LODIndex)
 	{
-		return;
+		FStaticMeshLODResources& CurrentLODResources = RenderData->LODResources[LODIndex];
+
+		if (CurrentLODResources.RayTracingGeometry.HasPendingBuildRequest())
+		{
+			ensure(CurrentLODResources.RayTracingGeometry.IsValid());
+			CurrentLODResources.RayTracingGeometry.BoostBuildPriority();
+		}
+		else if (CurrentLODResources.RayTracingGeometry.IsValid())
+		{
+			RayTracingGeometry = &CurrentLODResources.RayTracingGeometry;
+			break;
+		}
 	}
 
-	// TODO: Select different LOD when current LOD is still requested for build?
-	if (RayTracingGeometry->HasPendingBuildRequest())
+	if (RayTracingGeometry == nullptr)
 	{
-		RayTracingGeometry->BoostBuildPriority();
 		return;
 	}
 
