[Shaders] - Add LLM markers to track memory of shader compilation to help diagnose out of memory issues.

- Instruments: (FShaderCompilingManager::SubmitJobs, FShaderCompilingManager::ProcessAsyncResults, GlobalBeginCompileShader, FShaderJobCache::FindOrAdd)

diff --git a/Engine/Source/Runtime/Engine/Private/ShaderCompiler/ShaderCompiler.cpp b/Engine/Source/Runtime/Engine/Private/ShaderCompiler/ShaderCompiler.cpp
--- a/Engine/Source/Runtime/Engine/Private/ShaderCompiler/ShaderCompiler.cpp
+++ b/Engine/Source/Runtime/Engine/Private/ShaderCompiler/ShaderCompiler.cpp
@@ -6147,6 +6147,8 @@ void FShaderCompilingManager::ReleaseJob(FShaderCommonCompileJob* Job)
 
 void FShaderCompilingManager::SubmitJobs(TArray<FShaderCommonCompileJobPtr>& NewJobs, const FString MaterialBasePath, const FString PermutationString)
 {
+	LLM_SCOPE_BYTAG(ShaderCompiler);
+
 	// make sure no compiling can start if not allowed
 	if (!AllowShaderCompiling())
 	{
@@ -7397,6 +7399,8 @@ void FShaderCompilingManager::ProcessAsyncResults(bool bLimitExecutionTime, bool
 
 void FShaderCompilingManager::ProcessAsyncResults(float TimeSlice, bool bBlockOnGlobalShaderCompletion)
 {
+	LLM_SCOPE_BYTAG(ShaderCompiler);
+
 	TRACE_CPUPROFILER_EVENT_SCOPE(FShaderCompilingManager::ProcessAsyncResults)
 
 	COOK_STAT(FScopedDurationTimer Timer(ShaderCompilerCookStats::ProcessAsyncResultsTimeSec));
@@ -7774,6 +7778,8 @@ void GlobalBeginCompileShader(
 	const TCHAR* DebugExtension
 	)
 {
+	LLM_SCOPE_BYTAG(ShaderCompiler);
+
 	TRACE_CPUPROFILER_EVENT_SCOPE(GlobalBeginCompileShader);
 	COOK_STAT(ShaderCompilerCookStats::GlobalBeginCompileShaderCalls++);
 	COOK_STAT(FScopedDurationTimer DurationTimer(ShaderCompilerCookStats::GlobalBeginCompileShaderTimeSec));
@@ -10724,6 +10730,8 @@ namespace
 
 FShaderJobCacheRef FShaderJobCache::FindOrAdd(const FJobInputHash& Hash, EShaderCompileJobPriority JobPriority, const bool bCheckDDC, TPimplPtr<UE::DerivedData::FRequestOwner>& InoutRequestOwner, FJobCachedOutput*& OutCachedOutput)
 {
+	LLM_SCOPE_BYTAG(ShaderCompiler);
+
 	check(ShaderCompiler::IsJobCacheEnabled());
 
 	++TotalSearchAttempts;
