Fixed SoftObjectPath remapping done at cook time for ExternalStreamingObjects.

diff --git a/Engine/Source/Runtime/Engine/Classes/Engine/Level.h b/Engine/Source/Runtime/Engine/Classes/Engine/Level.h
--- a/Engine/Source/Runtime/Engine/Classes/Engine/Level.h
+++ b/Engine/Source/Runtime/Engine/Classes/Engine/Level.h
@@ -1121,10 +1121,17 @@ public:
 	 * @param InLevelPackage the package to get the external actors package name of
 	 * @param InActorPackagingScheme the packaging scheme to use
 	 * @param InActorPath the fully qualified actor path, in the format: 'Outermost.Outer.Name'
-	 * @param InMountPointContext an optional context object used to determine the mount point of the package
+	 * @param InLevelMountPointContext an optional context object used to determine the mount point of the package
 	 * @return the package name
 	 */
-	static ENGINE_API FString GetActorPackageName(UPackage* InLevelPackage, EActorPackagingScheme InActorPackagingScheme, const FString& InActorPath, const UObject* InMountPointContext = nullptr);
+	static ENGINE_API FString GetActorPackageName(UPackage* InLevelPackage, EActorPackagingScheme InActorPackagingScheme, const FString& InActorPath, const UObject* InLevelMountPointContext = nullptr);
+
+	/**
+	 * Returns a resolved level path using the level mount point context
+	 * @param InLevelPackage the level package name
+	 * @param InLevelMountPointContext an optional context object used to determine the mount point of the package
+	 */
+	static ENGINE_API FString ResolveRootPath(const FString& LevelPackageName, const UObject* InLevelMountPointContext = nullptr);
 
 	/**
 	 * Get the package name for this actor
diff --git a/Engine/Source/Runtime/Engine/Private/Level.cpp b/Engine/Source/Runtime/Engine/Private/Level.cpp
--- a/Engine/Source/Runtime/Engine/Private/Level.cpp
+++ b/Engine/Source/Runtime/Engine/Private/Level.cpp
@@ -3548,11 +3548,8 @@ bool ULevel::HasAnyActorsOfType(UClass *SearchType)
 }
 
 #if WITH_EDITOR
-FString ULevel::GetActorPackageName(UPackage* InLevelPackage, EActorPackagingScheme ActorPackagingScheme, const FString& InActorPath, const UObject* InLevelMountPointContext)
+FString ULevel::ResolveRootPath(const FString& LevelPackageName, const UObject* InLevelMountPointContext)
 {
-	check(InLevelPackage);
-	const FString LevelPackageName = InLevelPackage->GetName();
-
 	TOptional<FString> ResolvedLevelMountPoint;
 	if (InLevelMountPointContext)
 	{
@@ -3566,9 +3563,16 @@ FString ULevel::GetActorPackageName(UPackage* InLevelPackage, EActorPackagingSch
 			}
 		}
 	}
-	
+
 	// If ResolvedLevelMountPoint is not set, fallback on LevelPackageName
 	const FString LevelRootPath = ResolvedLevelMountPoint.Get(LevelPackageName);
+	return LevelRootPath;
+}
+
+FString ULevel::GetActorPackageName(UPackage* InLevelPackage, EActorPackagingScheme ActorPackagingScheme, const FString& InActorPath, const UObject* InLevelMountPointContext)
+{
+	check(InLevelPackage);
+	const FString LevelRootPath = ULevel::ResolveRootPath(InLevelPackage->GetName(), InLevelMountPointContext);
 	return ULevel::GetActorPackageName(ULevel::GetExternalActorsPath(LevelRootPath), ActorPackagingScheme, InActorPath);
 }
 
diff --git a/Engine/Source/Runtime/Engine/Private/WorldPartition/WorldPartitionLevelStreamingPolicy.cpp b/Engine/Source/Runtime/Engine/Private/WorldPartition/WorldPartitionLevelStreamingPolicy.cpp
--- a/Engine/Source/Runtime/Engine/Private/WorldPartition/WorldPartitionLevelStreamingPolicy.cpp
+++ b/Engine/Source/Runtime/Engine/Private/WorldPartition/WorldPartitionLevelStreamingPolicy.cpp
@@ -210,21 +210,6 @@ bool UWorldPartitionLevelStreamingPolicy::ConvertEditorPathToRuntimePath(const F
 		return false;
 	}
 
-	// In the editor, the _LevelInstance_ID is appended to the persistent level, while at runtime it is appended to each cell package, so we need to remap it there if present.
-	// Also handle prefixes like "/Temp"
-	FString LevelInstanceSuffix;
-	FString LevelInstancePrefix;
-	const FString WorldAssetPackageName = WorldAssetPath.GetPackageName().ToString();
-	const FString SourceWorldAssetPackageName = SourceWorldAssetPath.GetPackageName().ToString();
-	if (WorldAssetPackageName.Len() > SourceWorldAssetPackageName.Len())
-	{
-		if (const int32 Index = WorldAssetPackageName.Find(SourceWorldAssetPackageName); Index != INDEX_NONE)
-		{
-			LevelInstancePrefix = WorldAssetPackageName.Mid(0, Index);
-			LevelInstanceSuffix = WorldAssetPackageName.Mid(Index + SourceWorldAssetPackageName.Len());
-		}
-	}
-
 	FString SubObjectString;
 	FString SubObjectContext;
 	if (SrcObjectPath.GetSubPathString().Split(TEXT("."), &SubObjectContext, &SubObjectString))
@@ -232,7 +217,8 @@ bool UWorldPartitionLevelStreamingPolicy::ConvertEditorPathToRuntimePath(const F
 		if (SubObjectContext == TEXT("PersistentLevel"))
 		{
 			FString OutSubObjectString;
-			const FName* CellName = FindCellNameForSubObject(SubObjectString, true, OutSubObjectString);
+			const UObject* OutLevelMountPointContext = nullptr;
+			const FName* CellName = FindCellNameForSubObject(SubObjectString, true, OutSubObjectString, OutLevelMountPointContext);
 			FString SubPathString = SubObjectContext + TEXT(".") + OutSubObjectString;
 
 			if (!CellName)
@@ -248,6 +234,28 @@ bool UWorldPartitionLevelStreamingPolicy::ConvertEditorPathToRuntimePath(const F
 #endif
 			else
 			{
+				// In the editor, the _LevelInstance_ID is appended to the persistent level, while at runtime it is appended to each cell package, so we need to remap it there if present.
+				// Also handle prefixes like "/Temp"
+				FString LevelInstanceSuffix;
+				FString LevelInstancePrefix;
+				const FString WorldAssetPackageName = WorldAssetPath.GetPackageName().ToString();
+#if WITH_EDITOR
+				// Transform if necessary the world package name using the level mount point context object
+				const FString SourceWorldAssetPackageName = ULevel::ResolveRootPath(SourceWorldAssetPath.GetPackageName().ToString(), OutLevelMountPointContext);
+#else
+				//@todo_ow: Verify if we need to make ULevel::ResolveRootPath available at runtime
+				const FString SourceWorldAssetPackageName = SourceWorldAssetPath.GetPackageName().ToString();
+#endif
+
+				if (WorldAssetPackageName.Len() > SourceWorldAssetPackageName.Len())
+				{
+					if (const int32 Index = WorldAssetPackageName.Find(SourceWorldAssetPackageName); Index != INDEX_NONE)
+					{
+						LevelInstancePrefix = WorldAssetPackageName.Mid(0, Index);
+						LevelInstanceSuffix = WorldAssetPackageName.Mid(Index + SourceWorldAssetPackageName.Len());
+					}
+				}
+
 				OutPath = FString::Printf(TEXT("%s%s/_Generated_/%s%s.%s:%s"), *LevelInstancePrefix, *SourceWorldAssetPackageName, *(CellName->ToString()), *LevelInstanceSuffix, *WorldAssetPath.GetAssetName().ToString(), *SubPathString);
 			}
 
@@ -383,8 +391,10 @@ void UWorldPartitionLevelStreamingPolicy::DrawRuntimeCellsDetails(UCanvas* Canva
 	Offset.Y = MaxPosY;
 }
 
-const FName* UWorldPartitionLevelStreamingPolicy::FindCellNameForSubObject(const FString& InSubObjectString, bool bInResolveContainers, FString& OutSubPathString) const
+const FName* UWorldPartitionLevelStreamingPolicy::FindCellNameForSubObject(const FString& InSubObjectString, bool bInResolveContainers, FString& OutSubPathString, const UObject*& OutLevelMountPointContext) const
 {
+	OutLevelMountPointContext = nullptr;
+
 	FString SubObjectString;
 	FString SubObjectContext(InSubObjectString);
 	InSubObjectString.Split(TEXT("."), &SubObjectContext, &SubObjectString);
@@ -401,8 +411,10 @@ const FName* UWorldPartitionLevelStreamingPolicy::FindCellNameForSubObject(const
 	{
 		if (ExternalStreamingObject.IsValid())
 		{
-			if (const FName* CellName = ExternalStreamingObject.Get()->SubObjectsToCellRemapping.Find(*SubObjectContext))
+			const URuntimeHashExternalStreamingObjectBase* ExternalStreamingObjectPtr = ExternalStreamingObject.Get();
+			if (const FName* CellName = ExternalStreamingObjectPtr->SubObjectsToCellRemapping.Find(*SubObjectContext))
 			{
+				OutLevelMountPointContext = ExternalStreamingObjectPtr->GetLevelMountPointContextObject();
 				return CellName;
 			}
 		}
@@ -413,7 +425,7 @@ const FName* UWorldPartitionLevelStreamingPolicy::FindCellNameForSubObject(const
 	{
 		if (ContainerResolver.ResolveContainerPath(InSubObjectString, OutResolvedSubPathString))
 		{
-			return FindCellNameForSubObject(OutResolvedSubPathString, false, OutSubPathString);
+			return FindCellNameForSubObject(OutResolvedSubPathString, false, OutSubPathString, OutLevelMountPointContext);
 		}
 
 		for (const TWeakObjectPtr<URuntimeHashExternalStreamingObjectBase>& ExternalStreamingObject : ExternalStreamingObjects)
@@ -422,7 +434,7 @@ const FName* UWorldPartitionLevelStreamingPolicy::FindCellNameForSubObject(const
 			{
 				if (ExternalStreamingObject.Get()->ContainerResolver.ResolveContainerPath(InSubObjectString, OutResolvedSubPathString))
 				{
-					return FindCellNameForSubObject(OutResolvedSubPathString, false, OutSubPathString);
+					return FindCellNameForSubObject(OutResolvedSubPathString, false, OutSubPathString, OutLevelMountPointContext);
 				}
 			}
 		}
diff --git a/Engine/Source/Runtime/Engine/Private/WorldPartition/WorldPartitionRuntimeHash.cpp b/Engine/Source/Runtime/Engine/Private/WorldPartition/WorldPartitionRuntimeHash.cpp
--- a/Engine/Source/Runtime/Engine/Private/WorldPartition/WorldPartitionRuntimeHash.cpp
+++ b/Engine/Source/Runtime/Engine/Private/WorldPartition/WorldPartitionRuntimeHash.cpp
@@ -58,6 +58,11 @@ TSet<TObjectPtr<UDataLayerInstance>>& URuntimeHashExternalStreamingObjectBase::G
 	return DataLayerInstances;
 }
 
+const UObject* URuntimeHashExternalStreamingObjectBase::GetLevelMountPointContextObject() const
+{
+	return GetRootExternalDataLayerAsset();
+}
+
 UWorld* URuntimeHashExternalStreamingObjectBase::GetOwningWorld() const
 {
 	// Once OnStreamingObjectLoaded is called and OwningWorld is set, use this cached value
diff --git a/Engine/Source/Runtime/Engine/Private/WorldPartition/WorldPartitionRuntimeLevelStreamingCell.cpp b/Engine/Source/Runtime/Engine/Private/WorldPartition/WorldPartitionRuntimeLevelStreamingCell.cpp
--- a/Engine/Source/Runtime/Engine/Private/WorldPartition/WorldPartitionRuntimeLevelStreamingCell.cpp
+++ b/Engine/Source/Runtime/Engine/Private/WorldPartition/WorldPartitionRuntimeLevelStreamingCell.cpp
@@ -5,6 +5,8 @@
 #include "WorldPartition/WorldPartitionLevelStreamingPolicy.h"
 #include "WorldPartition/WorldPartitionDebugHelper.h"
 #include "WorldPartition/WorldPartitionStreamingGenerationContext.h"
+#include "WorldPartition/WorldPartitionRuntimeHash.h"
+#include "WorldPartition/WorldPartitionStreamingPolicy.h"
 #include "Engine/LevelStreaming.h"
 #include "Engine/Level.h"
 #include "Misc/HierarchicalLogArchive.h"
@@ -435,6 +437,42 @@ bool UWorldPartitionRuntimeLevelStreamingCell::OnPopulateGeneratorPackageForCook
 	return PrepareCellForCook(InPackage);
 }
 
+// Helper used by UWorldPartitionRuntimeLevelStreamingCell::OnPopulateGeneratedPackageForCook
+class FScopedCookingExternalStreamingObject
+{
+private:
+	FScopedCookingExternalStreamingObject(const URuntimeHashExternalStreamingObjectBase* InExternalStreamingObject)
+		: ExternalStreamingObject(const_cast<URuntimeHashExternalStreamingObjectBase*>(InExternalStreamingObject))
+	{
+		check(IsRunningCookCommandlet());
+		if (ExternalStreamingObject)
+		{
+			UWorld* World = ExternalStreamingObject->GetOuterWorld();
+			check(World);
+			UWorldPartition* WorldPartition = World->GetWorldPartition();
+			check(WorldPartition);
+			check(WorldPartition->StreamingPolicy);
+			verify(WorldPartition->StreamingPolicy->InjectExternalStreamingObject(const_cast<URuntimeHashExternalStreamingObjectBase*>(ExternalStreamingObject)));
+		}
+	}
+
+	~FScopedCookingExternalStreamingObject()
+	{
+		if (ExternalStreamingObject)
+		{
+			UWorld* World = ExternalStreamingObject->GetOuterWorld();
+			check(World);
+			UWorldPartition* WorldPartition = World->GetWorldPartition();
+			check(WorldPartition);
+			check(WorldPartition->StreamingPolicy);
+			verify(WorldPartition->StreamingPolicy->RemoveExternalStreamingObject(ExternalStreamingObject));
+		}
+	}
+
+	URuntimeHashExternalStreamingObjectBase* ExternalStreamingObject;
+	friend class UWorldPartitionRuntimeLevelStreamingCell;
+};
+
 bool UWorldPartitionRuntimeLevelStreamingCell::OnPopulateGeneratedPackageForCook(UPackage* InPackage, TArray<UPackage*>& OutModifiedPackages)
 {
 	check(!IsAlwaysLoaded());
@@ -472,6 +510,10 @@ bool UWorldPartitionRuntimeLevelStreamingCell::OnPopulateGeneratedPackageForCook
 		check(NewLevel->GetPackage() == InPackage);
 		FWorldPartitionLevelHelper::MoveExternalActorsToLevel(Packages, NewLevel, OutModifiedPackages);
 
+		// Push temporarily the cooking ExternalStreamingObject in the policy for RemapLevelSoftObjectPaths to use it to resolve softobjectpaths
+		const URuntimeHashExternalStreamingObjectBase* ExternalStreamingObject = GetTypedOuter<URuntimeHashExternalStreamingObjectBase>();
+		FScopedCookingExternalStreamingObject ScopeCookingExternalStreamingObject(ExternalStreamingObject);
+		
 		// Remap Level's SoftObjectPaths
 		FWorldPartitionLevelHelper::RemapLevelSoftObjectPaths(NewLevel, WorldPartition);
 	}
diff --git a/Engine/Source/Runtime/Engine/Public/WorldPartition/WorldPartition.h b/Engine/Source/Runtime/Engine/Public/WorldPartition/WorldPartition.h
--- a/Engine/Source/Runtime/Engine/Public/WorldPartition/WorldPartition.h
+++ b/Engine/Source/Runtime/Engine/Public/WorldPartition/WorldPartition.h
@@ -716,4 +716,7 @@ private:
 	friend class AWorldPartitionReplay;
 	friend class UWorldPartitionSubsystem;
 	friend class UExternalDataLayerManager;
+#if WITH_EDITOR
+	friend class FScopedCookingExternalStreamingObject;
+#endif
 };
\ No newline at end of file
diff --git a/Engine/Source/Runtime/Engine/Public/WorldPartition/WorldPartitionLevelStreamingPolicy.h b/Engine/Source/Runtime/Engine/Public/WorldPartition/WorldPartitionLevelStreamingPolicy.h
--- a/Engine/Source/Runtime/Engine/Public/WorldPartition/WorldPartitionLevelStreamingPolicy.h
+++ b/Engine/Source/Runtime/Engine/Public/WorldPartition/WorldPartitionLevelStreamingPolicy.h
@@ -50,7 +50,7 @@ protected:
 	void ForEachActiveRuntimeCell(TFunctionRef<void(const UWorldPartitionRuntimeCell*)> Func) const;
 
 private:
-	const FName* FindCellNameForSubObject(const FString& InSubObjectString, bool bInResolveContainers, FString& OutSubObjectString) const;
+	const FName* FindCellNameForSubObject(const FString& InSubObjectString, bool bInResolveContainers, FString& OutSubObjectString, const UObject*& OutLevelMountPointContext) const;
 	const UWorldPartitionRuntimeLevelStreamingCell* FindCellForSubObject(FName SubObjectName) const;
 
 	UPROPERTY()
diff --git a/Engine/Source/Runtime/Engine/Public/WorldPartition/WorldPartitionRuntimeHash.h b/Engine/Source/Runtime/Engine/Public/WorldPartition/WorldPartitionRuntimeHash.h
--- a/Engine/Source/Runtime/Engine/Public/WorldPartition/WorldPartitionRuntimeHash.h
+++ b/Engine/Source/Runtime/Engine/Public/WorldPartition/WorldPartitionRuntimeHash.h
@@ -77,6 +77,7 @@ public:
 	virtual const UExternalDataLayerInstance* GetRootExternalDataLayerInstance() const override { return RootExternalDataLayerInstance; }
 	// ~End IDataLayerInstanceProvider
 	UExternalDataLayerInstance* GetRootExternalDataLayerInstance() { return const_cast<UExternalDataLayerInstance*>(RootExternalDataLayerInstance.Get()); }
+	ENGINE_API const UObject* GetLevelMountPointContextObject() const;
 
 #if WITH_EDITOR
 	UE_DEPRECATED(5.4, "PopulateGeneratorPackageForCook is depreacted and was replaced by OnPopulateGeneratorPackageForCook")
@@ -270,4 +271,4 @@ public:
 	static ENGINE_API void RegisterWorldPartitionRuntimeHashConverter(const UClass* InSrcClass, const UClass* InDstClass, FRuntimeHashConvertFunc&& InConverter);
 	static ENGINE_API UWorldPartitionRuntimeHash* ConvertWorldPartitionHash(const UWorldPartitionRuntimeHash* InSrcHash, const UClass* InDstClass);
 #endif
-};
+};
\ No newline at end of file
