Fixed places where ActorName was used instead of ActorPath to generate the actor package name in External Data Layer Manager.

diff --git a/Engine/Plugins/Runtime/GameFeatures/Source/GameFeaturesEditor/Private/GameFeatureActionConvertContentBundleWorldPartitionBuilder.cpp b/Engine/Plugins/Runtime/GameFeatures/Source/GameFeaturesEditor/Private/GameFeatureActionConvertContentBundleWorldPartitionBuilder.cpp
--- a/Engine/Plugins/Runtime/GameFeatures/Source/GameFeaturesEditor/Private/GameFeatureActionConvertContentBundleWorldPartitionBuilder.cpp
+++ b/Engine/Plugins/Runtime/GameFeatures/Source/GameFeaturesEditor/Private/GameFeatureActionConvertContentBundleWorldPartitionBuilder.cpp
@@ -277,7 +277,7 @@ bool UGameFeatureActionConvertContentBundleWorldPartitionBuilder::RunInternal(UW
 
 			// Validation
 			check(NewActorPackage == Actor->GetExternalPackage());
-			check(NewActorPackage->GetName() == ExternalDataLayerManager->GetActorPackageName(ExternalDataLayerAsset, DestinationLevel, Actor->GetName()));
+			check(NewActorPackage->GetName() == ExternalDataLayerManager->GetActorPackageName(ExternalDataLayerAsset, DestinationLevel, Actor->GetPathName()));
 			check(NewActorPackage->GetName() != OldActorPackageName);
 
 			// Move dependant objects into the new actor package
diff --git a/Engine/Source/Runtime/Engine/Private/WorldPartition/DataLayer/ExternalDataLayerManager.cpp b/Engine/Source/Runtime/Engine/Private/WorldPartition/DataLayer/ExternalDataLayerManager.cpp
--- a/Engine/Source/Runtime/Engine/Private/WorldPartition/DataLayer/ExternalDataLayerManager.cpp
+++ b/Engine/Source/Runtime/Engine/Private/WorldPartition/DataLayer/ExternalDataLayerManager.cpp
@@ -825,10 +825,10 @@ void UExternalDataLayerManager::ForEachExternalStreamingObjects(TFunctionRef<boo
 	}
 }
 
-FString UExternalDataLayerManager::GetActorPackageName(const UExternalDataLayerAsset* InExternalDataLayerAsset, const ULevel* InDestinationLevel, const FString& ActorName) const
+FString UExternalDataLayerManager::GetActorPackageName(const UExternalDataLayerAsset* InExternalDataLayerAsset, const ULevel* InDestinationLevel, const FString& InActorPath) const
 {
 	const FString ContainerRootPath = GetExternalDataLayerLevelRootPath(InExternalDataLayerAsset);
-	const FString ActorPackageName = ULevel::GetActorPackageName(ULevel::GetExternalActorsPath(ContainerRootPath), InDestinationLevel->GetActorPackagingScheme(), ActorName);
+	const FString ActorPackageName = ULevel::GetActorPackageName(ULevel::GetExternalActorsPath(ContainerRootPath), InDestinationLevel->GetActorPackagingScheme(), InActorPath);
 	return ActorPackageName;
 }
 
@@ -840,7 +840,7 @@ bool UExternalDataLayerManager::SetupActorPackageForExternalDataLayerAsset(AActo
 	// First check if we really need to rename the package at all.
 	// For example, when reinstancing (after compiling a BP), we reuse the old actor package.
 	const FString OldActorPackageName = InActor->GetPackage()->GetName();
-	const FString NewActorPackageName = GetActorPackageName(InExternalDataLayerAsset, InActor->GetLevel(), InActor->GetName());
+	const FString NewActorPackageName = GetActorPackageName(InExternalDataLayerAsset, InActor->GetLevel(), InActor->GetPathName());
 	if (OldActorPackageName == NewActorPackageName)
 	{
 		return true;
@@ -896,8 +896,8 @@ AWorldDataLayers* UExternalDataLayerManager::GetWorldDataLayers(const UExternalD
 	// @todo_ow: Add LevelInstance EDL support (choose the right destination level)
 	ULevel* DestinationLevel = GetTypedOuter<UWorld>()->PersistentLevel;
 	const FString WorldDataLayersName = FString::Printf(TEXT("%s%s"), *AWorldDataLayers::StaticClass()->GetName(), *ObjectTools::SanitizeObjectName(InExternalDataLayerAsset->GetPathName()));
-	const FString PackageName = GetActorPackageName(InExternalDataLayerAsset, DestinationLevel, WorldDataLayersName);
-		
+	const FString WorldDataLayersPathName = DestinationLevel->GetPathName() + TEXT(".") + WorldDataLayersName;
+	const FString PackageName = GetActorPackageName(InExternalDataLayerAsset, DestinationLevel, WorldDataLayersPathName);
 	UPackage* ActorPackage = FindObject<UPackage>(nullptr, *PackageName);
 	AWorldDataLayers* EDLWorldDataLayers = FindObject<AWorldDataLayers>(ActorPackage, *WorldDataLayersName);
 	if (IsValid(EDLWorldDataLayers))
@@ -907,7 +907,7 @@ AWorldDataLayers* UExternalDataLayerManager::GetWorldDataLayers(const UExternalD
 
 	if (bAllowCreate)
 	{
-		ActorPackage = ULevel::CreateActorPackage(DestinationLevel->GetPackage(), DestinationLevel->GetActorPackagingScheme(), WorldDataLayersName, InExternalDataLayerAsset);
+		ActorPackage = ULevel::CreateActorPackage(DestinationLevel->GetPackage(), DestinationLevel->GetActorPackagingScheme(), WorldDataLayersPathName, InExternalDataLayerAsset);
 		FActorSpawnParameters SpawnParameters;
 		SpawnParameters.Name = *WorldDataLayersName;
 		SpawnParameters.OverrideLevel = DestinationLevel;
diff --git a/Engine/Source/Runtime/Engine/Public/WorldPartition/DataLayer/ExternalDataLayerManager.h b/Engine/Source/Runtime/Engine/Public/WorldPartition/DataLayer/ExternalDataLayerManager.h
--- a/Engine/Source/Runtime/Engine/Public/WorldPartition/DataLayer/ExternalDataLayerManager.h
+++ b/Engine/Source/Runtime/Engine/Public/WorldPartition/DataLayer/ExternalDataLayerManager.h
@@ -82,7 +82,7 @@ private:
 	const UExternalDataLayerAsset* GetActorEditorContextCurrentExternalDataLayer() const;
 	AWorldDataLayers* GetWorldDataLayers(const UExternalDataLayerAsset* InExternalDataLayerAsset, bool bInAllowCreate = false) const;
 	bool OnActorPreSpawnInitialization(AActor* InActor, const UExternalDataLayerAsset* InExternalDataLayerAsset);
-	FString GetActorPackageName(const UExternalDataLayerAsset* InExternalDataLayerAsset, const ULevel* InDestinationLevel, const FString& ActorName) const;
+	FString GetActorPackageName(const UExternalDataLayerAsset* InExternalDataLayerAsset, const ULevel* InDestinationLevel, const FString& InActorPath) const;
 	bool SetupActorPackageForExternalDataLayerAsset(AActor* InActor, const UExternalDataLayerAsset* InExternalDataLayerAsset);
 	URuntimeHashExternalStreamingObjectBase* CreateExternalStreamingObjectUsingStreamingGeneration(const UExternalDataLayerAsset* InExternalDataLayerAsset);
 
