[Material Editor] Added function name to a Material Editor warning log line.

diff --git a/Engine/Source/Runtime/Engine/Private/Materials/MaterialExpressions.cpp b/Engine/Source/Runtime/Engine/Private/Materials/MaterialExpressions.cpp
--- a/Engine/Source/Runtime/Engine/Private/Materials/MaterialExpressions.cpp
+++ b/Engine/Source/Runtime/Engine/Private/Materials/MaterialExpressions.cpp
@@ -15156,9 +15156,6 @@ void UMaterialFunction::PostLoad()
 		{
 			UE_LOG(LogMaterial, Warning, TEXT("Some expression in Material Function %s failed to load correctly. This will cause any material using this MF to fail translation. Please check open affected Material Function, make sure its expression graph is valid and resave it."), *GetFullName());
 			
-			// Mark this function as invalid. Translating a material containing a call to it will fail.
-			bAllExpressionsLoadedCorrectly = false;
-
 			// Dirty this function by deterministically changing its StateId.
 			static FGuid NotAllExpressionsLoadedCorrectlyToken(TEXT("6B9D300E-ED9D-4E4A-A141-05DE059B5704"));
 			StateId.A ^= NotAllExpressionsLoadedCorrectlyToken.A;
@@ -16928,7 +16925,9 @@ int32 UMaterialExpressionMaterialFunctionCall::Compile(class FMaterialCompiler*
 
 	if (!MaterialFunction->GetBaseFunction()->bAllExpressionsLoadedCorrectly)
 	{
-		return Compiler->Errorf(TEXT("Called function is in an invalid state because some expressions didn't load correctly. Please open the affected function, review and correct the expression graph then save again."));
+		return Compiler->Errorf(TEXT("Called function '%s' is in an invalid state because some expressions didn't load correctly. "
+			"Please open the affected function, review and correct the expression graph then save again."),
+			*MaterialFunction->GetBaseFunction()->GetFullName());
 	}
 
 	// Verify that all function inputs and outputs are in a valid state to be linked into this material for compiling
