Fixed invalid check when validating the actor package name of an actor created under an External Data Layer

diff --git a/Engine/Source/Runtime/Engine/Private/WorldPartition/DataLayer/ExternalDataLayerManager.cpp b/Engine/Source/Runtime/Engine/Private/WorldPartition/DataLayer/ExternalDataLayerManager.cpp
--- a/Engine/Source/Runtime/Engine/Private/WorldPartition/DataLayer/ExternalDataLayerManager.cpp
+++ b/Engine/Source/Runtime/Engine/Private/WorldPartition/DataLayer/ExternalDataLayerManager.cpp
@@ -886,7 +886,7 @@ bool UExternalDataLayerManager::OnActorExternalDataLayerAssetChanged(AActor* InA
 #endif
 	
 	const FString ActorPackageName = InActor->GetExternalPackage()->GetName();
-	const FString NewActorPackageName = GetActorPackageName(ExternalDataLayerAsset, InActor->GetLevel(), InActor->GetName());
+	const FString NewActorPackageName = GetActorPackageName(ExternalDataLayerAsset, InActor->GetLevel(), InActor->GetPathName());
 	check(NewActorPackageName == ActorPackageName);
 	return (NewActorPackageName == ActorPackageName);
 }
