World Partition: fix bad check when doing queries on world streaming.

diff --git a/Engine/Source/Runtime/Engine/Private/WorldPartition/RuntimeHashSet/WorldPartitionRuntimeHashSet.cpp b/Engine/Source/Runtime/Engine/Private/WorldPartition/RuntimeHashSet/WorldPartitionRuntimeHashSet.cpp
--- a/Engine/Source/Runtime/Engine/Private/WorldPartition/RuntimeHashSet/WorldPartitionRuntimeHashSet.cpp
+++ b/Engine/Source/Runtime/Engine/Private/WorldPartition/RuntimeHashSet/WorldPartitionRuntimeHashSet.cpp
@@ -365,25 +365,20 @@ void UWorldPartitionRuntimeHashSet::ForEachStreamingCellsQuery(const FWorldParti
 {
 	auto ShouldAddCell = [this](const UWorldPartitionRuntimeCell* Cell, const FWorldPartitionStreamingQuerySource& QuerySource)
 	{
-#if WITH_EDITOR
-		if (!IsCellRelevantFor(Cell->GetClientOnlyVisible()))
+		if (IsCellRelevantFor(Cell->GetClientOnlyVisible()))
 		{
-			return false;
-		}
-#else
-		check(IsCellRelevantFor(Cell->GetClientOnlyVisible()));
-#endif
-		if (Cell->HasDataLayers())
-		{
-			if (Cell->GetDataLayers().FindByPredicate([&](const FName& DataLayerName) { return QuerySource.DataLayers.Contains(DataLayerName); }))
+			if (Cell->HasDataLayers())
+			{
+				if (Cell->GetDataLayers().FindByPredicate([&](const FName& DataLayerName) { return QuerySource.DataLayers.Contains(DataLayerName); }))
+				{
+					return true;
+				}
+			}
+			else if (!QuerySource.bDataLayersOnly)
 			{
 				return true;
 			}
 		}
-		else if (!QuerySource.bDataLayersOnly)
-		{
-			return true;
-		}
 
 		return false;
 	};
