Make sure required bones are invalidated when the skeletal mesh component unregisters

diff --git a/Engine/Source/Runtime/Engine/Private/Components/SkeletalMeshComponent.cpp b/Engine/Source/Runtime/Engine/Private/Components/SkeletalMeshComponent.cpp
--- a/Engine/Source/Runtime/Engine/Private/Components/SkeletalMeshComponent.cpp
+++ b/Engine/Source/Runtime/Engine/Private/Components/SkeletalMeshComponent.cpp
@@ -875,7 +875,13 @@ void USkeletalMeshComponent::OnUnregister()
 		}
 	}
 
+	// Invalidate required bones and our cached data
 	RequiredBones.Reset();
+	bRequiredBonesUpToDate = false;
+	if (SharedRequiredBones)
+	{
+		SharedRequiredBones->Reset();
+	}
 
 	Super::OnUnregister();
 }
