Avoid crashing ODSC server when we detect a shader compilation error
Do not replace the global shader map when we know some shaders will be missing. We will crash later inside VerifyGlobalShaders anyway
Add an on-screen message to notify shader compilation is happening

#rnx

diff --git a/Engine/Source/Runtime/Engine/Private/ODSC/ODSCManager.cpp b/Engine/Source/Runtime/Engine/Private/ODSC/ODSCManager.cpp
--- a/Engine/Source/Runtime/Engine/Private/ODSC/ODSCManager.cpp
+++ b/Engine/Source/Runtime/Engine/Private/ODSC/ODSCManager.cpp
@@ -24,11 +24,24 @@ FODSCManager::FODSCManager()
 		FCoreDelegates::OnEnginePreExit.AddRaw(this, &FODSCManager::OnEnginePreExit);
 		Thread = new FODSCThread(Host);
 		Thread->StartThread();
+		OnScreenMessagesHandle = FCoreDelegates::OnGetOnScreenMessages.AddLambda([this](TMultiMap<FCoreDelegates::EOnScreenMessageSeverity, FText >& OutMessages)
+			{
+				if (Thread && Thread->HasPendingRequests())
+				{
+					OutMessages.Add(FCoreDelegates::EOnScreenMessageSeverity::Info, FText::FromString(FString::Printf(TEXT("Recompiling shaders"))));
+				}
+			}
+		);
 	}
 }
 
 FODSCManager::~FODSCManager()
 {
+	if (OnScreenMessagesHandle.IsValid())
+	{
+		FCoreDelegates::OnGetOnScreenMessages.Remove(OnScreenMessagesHandle);
+	}
+
 	FCoreDelegates::OnEnginePreExit.RemoveAll(this);
 	StopThread();
 }
diff --git a/Engine/Source/Runtime/Engine/Private/ODSC/ODSCThread.cpp b/Engine/Source/Runtime/Engine/Private/ODSC/ODSCThread.cpp
--- a/Engine/Source/Runtime/Engine/Private/ODSC/ODSCThread.cpp
+++ b/Engine/Source/Runtime/Engine/Private/ODSC/ODSCThread.cpp
@@ -253,6 +253,8 @@ void FODSCThread::Process()
 		RequestsToStart.Add(Request);
 	}
 
+	bHasPendingRequests = PayloadsToAggregate.Num() > 0 || RequestsToStart.Num() > 0;
+
 	// process any material or recompile change shader requests or global shader compile requests.
 	for (FODSCMessageHandler* NextRequest : RequestsToStart)
 	{
@@ -277,6 +279,8 @@ void FODSCThread::Process()
 		CompletedThreadedRequests.Enqueue(RequestHandler);
 	}
 
+	// SendMessageToServer is synchronous, so when we're here, we know we've processed all the requests
+	bHasPendingRequests = false;
 	WakeupEvent->Reset();
 }
 
diff --git a/Engine/Source/Runtime/Engine/Private/ODSC/ODSCThread.h b/Engine/Source/Runtime/Engine/Private/ODSC/ODSCThread.h
--- a/Engine/Source/Runtime/Engine/Private/ODSC/ODSCThread.h
+++ b/Engine/Source/Runtime/Engine/Private/ODSC/ODSCThread.h
@@ -148,6 +148,8 @@ public:
 	*/
 	void Wakeup();
 
+	bool HasPendingRequests() const { return bHasPendingRequests; }
+
 protected:
 
 	//~ Begin FRunnable Interface
@@ -202,4 +204,6 @@ private:
 
 	/** Special connection to the cooking server.  This is only used to send recompileshaders commands on. */
 	TUniquePtr<UE::Cook::ICookOnTheFlyServerConnection> CookOnTheFlyServerConnection;
+
+	bool bHasPendingRequests = false;
 };
diff --git a/Engine/Source/Runtime/Engine/Private/ShaderCompiler/ShaderCompiler.cpp b/Engine/Source/Runtime/Engine/Private/ShaderCompiler/ShaderCompiler.cpp
--- a/Engine/Source/Runtime/Engine/Private/ShaderCompiler/ShaderCompiler.cpp
+++ b/Engine/Source/Runtime/Engine/Private/ShaderCompiler/ShaderCompiler.cpp
@@ -9684,69 +9684,62 @@ static inline bool ShouldCacheGlobalShaderTypeName(const FGlobalShaderType* Glob
 };
 
 
-bool IsGlobalShaderMapComplete(const TCHAR* TypeNameSubstring)
+bool IsGlobalShaderMapComplete(const TCHAR* TypeNameSubstring, FGlobalShaderMap* GlobalShaderMap, EShaderPlatform Platform)
 {
-	for (int32 i = 0; i < SP_NumPlatforms; ++i)
+	// look at any shadermap in the GlobalShaderMap for the permutation flags, as they will all be the same
+	if (GlobalShaderMap)
 	{
-		EShaderPlatform Platform = (EShaderPlatform)i;
-
-		FGlobalShaderMap* GlobalShaderMap = GGlobalShaderMap[Platform];
-
-		// look at any shadermap in the GlobalShaderMap for the permutation flags, as they will all be the same
-		if (GlobalShaderMap)
+		const FGlobalShaderMapSection* FirstShaderMap = GlobalShaderMap->GetFirstSection();
+		if (FirstShaderMap == nullptr)
 		{
-			const FGlobalShaderMapSection* FirstShaderMap = GlobalShaderMap->GetFirstSection();
-			if (FirstShaderMap == nullptr)
-			{
-				// if we had no sections at all, we know we aren't complete
-				return false;
-			}
-			EShaderPermutationFlags GlobalShaderPermutation = FirstShaderMap->GetPermutationFlags();
+			// if we had no sections at all, we know we aren't complete
+			return false;
+		}
+		EShaderPermutationFlags GlobalShaderPermutation = FirstShaderMap->GetPermutationFlags();
 
-			// Check if the individual shaders are complete
-			for (TLinkedList<FShaderType*>::TIterator ShaderTypeIt(FShaderType::GetTypeList()); ShaderTypeIt; ShaderTypeIt.Next())
+		// Check if the individual shaders are complete
+		for (TLinkedList<FShaderType*>::TIterator ShaderTypeIt(FShaderType::GetTypeList()); ShaderTypeIt; ShaderTypeIt.Next())
+		{
+			FGlobalShaderType* GlobalShaderType = ShaderTypeIt->GetGlobalShaderType();
+			int32 PermutationCount = GlobalShaderType ? GlobalShaderType->GetPermutationCount() : 1;
+			for (int32 PermutationId = 0; PermutationId < PermutationCount; PermutationId++)
 			{
-				FGlobalShaderType* GlobalShaderType = ShaderTypeIt->GetGlobalShaderType();
-				int32 PermutationCount = GlobalShaderType ? GlobalShaderType->GetPermutationCount() : 1;
-				for (int32 PermutationId = 0; PermutationId < PermutationCount; PermutationId++)
+				if (ShouldCacheGlobalShaderTypeName(GlobalShaderType, PermutationId, TypeNameSubstring, Platform, GlobalShaderPermutation))
 				{
-					if (ShouldCacheGlobalShaderTypeName(GlobalShaderType, PermutationId, TypeNameSubstring, Platform, GlobalShaderPermutation))
+					if (!GlobalShaderMap->HasShader(GlobalShaderType, PermutationId))
 					{
-						if (!GlobalShaderMap->HasShader(GlobalShaderType, PermutationId))
-						{
-							return false;
-						}
+						return false;
 					}
 				}
 			}
+		}
 
-			// Then the pipelines as it may be sharing shaders
-			for (TLinkedList<FShaderPipelineType*>::TIterator ShaderPipelineIt(FShaderPipelineType::GetTypeList()); ShaderPipelineIt; ShaderPipelineIt.Next())
+		// Then the pipelines as it may be sharing shaders
+		for (TLinkedList<FShaderPipelineType*>::TIterator ShaderPipelineIt(FShaderPipelineType::GetTypeList()); ShaderPipelineIt; ShaderPipelineIt.Next())
+		{
+			const FShaderPipelineType* Pipeline = *ShaderPipelineIt;
+			if (Pipeline->IsGlobalTypePipeline())
 			{
-				const FShaderPipelineType* Pipeline = *ShaderPipelineIt;
-				if (Pipeline->IsGlobalTypePipeline())
+				auto& Stages = Pipeline->GetStages();
+				int32 NumStagesNeeded = 0;
+				for (const FShaderType* Shader : Stages)
 				{
-					auto& Stages = Pipeline->GetStages();
-					int32 NumStagesNeeded = 0;
-					for (const FShaderType* Shader : Stages)
+					const FGlobalShaderType* GlobalShaderType = Shader->GetGlobalShaderType();
+					if (ShouldCacheGlobalShaderTypeName(GlobalShaderType, kUniqueShaderPermutationId, TypeNameSubstring, Platform, GlobalShaderPermutation))
 					{
-						const FGlobalShaderType* GlobalShaderType = Shader->GetGlobalShaderType();
-						if (ShouldCacheGlobalShaderTypeName(GlobalShaderType, kUniqueShaderPermutationId, TypeNameSubstring, Platform, GlobalShaderPermutation))
-						{
-							++NumStagesNeeded;
-						}
-						else
-						{
-							break;
-						}
+						++NumStagesNeeded;
+					}
+					else
+					{
+						break;
 					}
+				}
 
-					if (NumStagesNeeded == Stages.Num())
+				if (NumStagesNeeded == Stages.Num())
+				{
+					if (!GlobalShaderMap->HasShaderPipeline(Pipeline))
 					{
-						if (!GlobalShaderMap->HasShaderPipeline(Pipeline))
-						{
-							return false;
-						}
+						return false;
 					}
 				}
 			}
@@ -9756,6 +9749,24 @@ bool IsGlobalShaderMapComplete(const TCHAR* TypeNameSubstring)
 	return true;
 }
 
+
+bool IsGlobalShaderMapComplete(const TCHAR* TypeNameSubstring)
+{
+	for (int32 i = 0; i < SP_NumPlatforms; ++i)
+	{
+		EShaderPlatform Platform = (EShaderPlatform)i;
+
+		FGlobalShaderMap* GlobalShaderMap = GGlobalShaderMap[Platform];
+
+		if (!IsGlobalShaderMapComplete(TypeNameSubstring, GlobalShaderMap, Platform))
+		{
+			return false;
+		}
+	}
+
+	return true;
+}
+
 static bool TryLoadCookedGlobalShaderMap(EShaderPlatform Platform, FScopedSlowTask& SlowTask)
 {
 	SlowTask.EnterProgressFrame(50);
@@ -10311,6 +10322,9 @@ void RecompileShadersForRemote(
 {
 	TRACE_CPUPROFILER_EVENT_SCOPE(RecompileShadersForRemote);
 
+	CVarAreShaderErrorsFatal->Set(0, ECVF_SetByCode);
+	ON_SCOPE_EXIT{ CVarAreShaderErrorsFatal->Unset(ECVF_SetByCode); };
+
 	// figure out what shader platforms to recompile
 	ITargetPlatformManagerModule* TPM = GetTargetPlatformManager();
 	ITargetPlatform* TargetPlatform = TPM->FindTargetPlatform(Args.PlatformName);
@@ -10685,7 +10699,9 @@ void LoadGlobalShadersForRemoteRecompile(FArchive& Ar, EShaderPlatform ShaderPla
 		{
 			NewGlobalShaderMap->LoadFromGlobalArchive(Ar);
 
-			if (GGlobalShaderMap[ShaderPlatform])
+			bool bIsNewGlobalShaderMapComplete = IsGlobalShaderMapComplete(nullptr, NewGlobalShaderMap, ShaderPlatform);
+
+			if (GGlobalShaderMap[ShaderPlatform] && bIsNewGlobalShaderMapComplete)
 			{
 				GGlobalShaderMap[ShaderPlatform]->ReleaseAllSections();
 
@@ -10705,6 +10721,10 @@ void LoadGlobalShadersForRemoteRecompile(FArchive& Ar, EShaderPlatform ShaderPla
 			}
 			else
 			{
+				if (!bIsNewGlobalShaderMapComplete)
+				{
+					UE_LOG(LogShaderCompilers, Error, TEXT("New shader map is incomplete. Look at the ODSC server log to see shader errors"));
+				}
 				delete NewGlobalShaderMap;
 			}
 		}
diff --git a/Engine/Source/Runtime/Engine/Public/ODSC/ODSCManager.h b/Engine/Source/Runtime/Engine/Public/ODSC/ODSCManager.h
--- a/Engine/Source/Runtime/Engine/Public/ODSC/ODSCManager.h
+++ b/Engine/Source/Runtime/Engine/Public/ODSC/ODSCManager.h
@@ -88,6 +88,9 @@ private:
 
 	/** Handles communicating directly with the cook on the fly server. */
 	FODSCThread* Thread = nullptr;
+
+	FDelegateHandle OnScreenMessagesHandle;
+
 };
 
 /** The global shader ODSC manager. */
