Substrate - Fixed material attribute conversion for water.

diff --git a/Engine/Source/Runtime/Engine/Private/Materials/Material.cpp b/Engine/Source/Runtime/Engine/Private/Materials/Material.cpp
--- a/Engine/Source/Runtime/Engine/Private/Materials/Material.cpp
+++ b/Engine/Source/Runtime/Engine/Private/Materials/Material.cpp
@@ -3330,24 +3330,12 @@ void UMaterial::ConvertMaterialToSubstrateMaterial()
 			check(ConvertAttributeNode);
 			GatherCustomNodes();
 
-			if (ThinTranslucentOutput)
-			{
-				ConnectionTo(*ThinTranslucentOutput->GetInput(0), ConvertAttributeNode, 1);	 // TransmittanceColor
-			}
 			if (SingleLayerWaterOutput)
 			{
-				ConnectionTo(*SingleLayerWaterOutput->GetInput(0), ConvertAttributeNode, 2); // WaterScatteringCoefficients
-				ConnectionTo(*SingleLayerWaterOutput->GetInput(1), ConvertAttributeNode, 3); // WaterAbsorptionCoefficients
-				ConnectionTo(*SingleLayerWaterOutput->GetInput(2), ConvertAttributeNode, 4); // WaterPhaseG
-				ConnectionTo(*SingleLayerWaterOutput->GetInput(3), ConvertAttributeNode, 5); // ColorScaleBehindWater
-			}
-			if (ClearCoatBottomNormalOutput)
-			{
-				ConnectionTo(*ClearCoatBottomNormalOutput->GetInput(0), ConvertAttributeNode, 6, SUBSTRATE_COPY_CONNECTION); // ClearCoatNormal
-			}
-			if (TangentOutput)
-			{
-				ConnectionTo(*TangentOutput->GetInput(0), ConvertAttributeNode, 7, SUBSTRATE_COPY_CONNECTION);	// TangentOutput
+				ConnectionTo(*SingleLayerWaterOutput->GetInput(0), ConvertAttributeNode, 1); // WaterScatteringCoefficients
+				ConnectionTo(*SingleLayerWaterOutput->GetInput(1), ConvertAttributeNode, 2); // WaterAbsorptionCoefficients
+				ConnectionTo(*SingleLayerWaterOutput->GetInput(2), ConvertAttributeNode, 3); // WaterPhaseG
+				ConnectionTo(*SingleLayerWaterOutput->GetInput(3), ConvertAttributeNode, 4); // ColorScaleBehindWater
 			}
 		}
 
