Use non-fast conversion for LocalToWorld and InstanceToWorld to fix precision issues seen in plant WPO.
Allow more direct transformation to and from translated worldspace, without absolute worldspace as an intermediate.

diff --git a/Engine/Shaders/Private/MaterialTemplate.ush b/Engine/Shaders/Private/MaterialTemplate.ush
--- a/Engine/Shaders/Private/MaterialTemplate.ush
+++ b/Engine/Shaders/Private/MaterialTemplate.ush
@@ -1668,13 +1668,13 @@ FMaterialLWCData MakeMaterialLWCData(FMaterialVertexParameters Parameters)
 
 	Result.AbsoluteWorldPosition = WSSubtract(Parameters.WorldPosition, Result.PreViewTranslation);
 
-	Result.InstanceToWorld = DFFastToWS(GetInstanceToWorldDF(Parameters));
+	Result.InstanceToWorld = DFToWS(GetInstanceToWorldDF(Parameters));
 	Result.WorldToInstance = DFFastToWS(GetWorldToInstanceDF(Parameters));
-	Result.PreviousInstanceToWorld = DFFastToWS(GetPrevInstanceToWorldDF(Parameters));
+	Result.PreviousInstanceToWorld = DFToWS(GetPrevInstanceToWorldDF(Parameters));
 
-	Result.LocalToWorld = DFFastToWS(GetLocalToWorldDF(Parameters));
+	Result.LocalToWorld = DFToWS(GetLocalToWorldDF(Parameters));
 	Result.WorldToLocal = DFFastToWS(GetWorldToLocalDF(Parameters));
-	Result.PreviousLocalToWorld = DFFastToWS(GetPrevLocalToWorldDF(Parameters));
+	Result.PreviousLocalToWorld = DFToWS(GetPrevLocalToWorldDF(Parameters));
 	Result.PreviousWorldToLocal = DFFastToWS(GetPrevWorldToLocalDF(Parameters));
 
 #if DECAL_PRIMITIVE
@@ -1740,13 +1740,13 @@ FMaterialLWCData MakeMaterialLWCData(FMaterialPixelParameters Parameters)
 	Result.AbsoluteWorldPosition = WSSubtract(Parameters.WorldPosition_CamRelative, Result.PreViewTranslation);
 	Result.WorldPosition_NoOffsets = WSSubtract(Parameters.WorldPosition_NoOffsets_CamRelative, Result.PreViewTranslation);
 
-	Result.InstanceToWorld = DFFastToWS(GetInstanceToWorldDF(Parameters));
+	Result.InstanceToWorld = DFToWS(GetInstanceToWorldDF(Parameters));
 	Result.WorldToInstance = DFFastToWS(GetWorldToInstanceDF(Parameters));
-	Result.PreviousInstanceToWorld = DFFastToWS(GetPrevInstanceToWorldDF(Parameters));
+	Result.PreviousInstanceToWorld = DFToWS(GetPrevInstanceToWorldDF(Parameters));
 
-	Result.LocalToWorld = DFFastToWS(GetLocalToWorldDF(Parameters));
+	Result.LocalToWorld = DFToWS(GetLocalToWorldDF(Parameters));
 	Result.WorldToLocal = DFFastToWS(GetWorldToLocalDF(Parameters));
-	Result.PreviousLocalToWorld = DFFastToWS(GetPrevLocalToWorldDF(Parameters));
+	Result.PreviousLocalToWorld = DFToWS(GetPrevLocalToWorldDF(Parameters));
 	Result.PreviousWorldToLocal = DFFastToWS(GetPrevWorldToLocalDF(Parameters));
 
 #if DECAL_PRIMITIVE
diff --git a/Engine/Source/Runtime/Core/Private/UObject/DevObjectVersion.cpp b/Engine/Source/Runtime/Core/Private/UObject/DevObjectVersion.cpp
--- a/Engine/Source/Runtime/Core/Private/UObject/DevObjectVersion.cpp
+++ b/Engine/Source/Runtime/Core/Private/UObject/DevObjectVersion.cpp
@@ -129,7 +129,7 @@ FDevSystemGuids::FDevSystemGuids()
 	, POSESEARCHDB_DERIVEDDATA_VER(0x389117E4, 0x807A4CC0, 0x9F37C2E6, 0xD808A78D)
 	, SkeletalMeshDerivedDataVersion(0x9B5F4544, 0x76D7481C, 0x9AD3F614, 0xA6C07904)
 	, STATICMESH_DERIVEDDATA_VER(0x5B7A05A7, 0x7371440C, 0xA314621A, 0x9B1E08BC)
-	, MaterialTranslationDDCVersion(0xBBD65080, 0x4D314CE7, 0xA44ADF9B, 0xB4F3F5F9)
+	, MaterialTranslationDDCVersion(0x59E1F296, 0x27064757, 0x8CA7EB9B, 0xFC366BD2)
 {
 }
 
diff --git a/Engine/Source/Runtime/Engine/Private/Materials/HLSLMaterialTranslator.cpp b/Engine/Source/Runtime/Engine/Private/Materials/HLSLMaterialTranslator.cpp
--- a/Engine/Source/Runtime/Engine/Private/Materials/HLSLMaterialTranslator.cpp
+++ b/Engine/Source/Runtime/Engine/Private/Materials/HLSLMaterialTranslator.cpp
@@ -10039,7 +10039,7 @@ int32 FHLSLMaterialTranslator::TransformBase(EMaterialCommonBasis SourceCoordBas
 	}
 
 	const EMaterialValueType SourceType = GetParameterType(A);
-	const bool bIsPositionTranform = AWComponent != 0;
+	const bool bIsPositionTransform = AWComponent != 0;
 		
 	{ // validation
 		if (ShaderFrequency != SF_Pixel && ShaderFrequency != SF_Compute && ShaderFrequency != SF_Vertex)
@@ -10055,7 +10055,7 @@ int32 FHLSLMaterialTranslator::TransformBase(EMaterialCommonBasis SourceCoordBas
 			}
 		}
 
-		if (AWComponent != 0 && (SourceCoordBasis == MCB_Tangent || DestCoordBasis == MCB_Tangent))
+		if (bIsPositionTransform && (SourceCoordBasis == MCB_Tangent || DestCoordBasis == MCB_Tangent))
 		{
 			return Errorf(TEXT("Tangent basis not available for position transformations"));
 		}
@@ -10087,7 +10087,7 @@ int32 FHLSLMaterialTranslator::TransformBase(EMaterialCommonBasis SourceCoordBas
 	{
 		case MCB_Tangent:
 		{
-			check(AWComponent == 0);
+			check(!bIsPositionTransform);
 			if (DestCoordBasis == MCB_World)
 			{
 				CodeStr = TEXT("mul(<A>, Parameters.TangentToWorld)");
@@ -10100,10 +10100,18 @@ int32 FHLSLMaterialTranslator::TransformBase(EMaterialCommonBasis SourceCoordBas
 		{
 			if (DestCoordBasis == MCB_World)
 			{
-				if (bIsPositionTranform) { AddLWCFuncUsage(ELWCFunctionKind::MultiplyVectorMatrix); }
+				if (bIsPositionTransform) { AddLWCFuncUsage(ELWCFunctionKind::MultiplyVectorMatrix); }
 				CodeStr = TEXT("TransformLocal<TO><PREV>World(Parameters, <A>)");
 				CodeDerivStr = TEXT("TransformLocal<TO><PREV>World(Parameters, <A>)");
 			}
+			else if (DestCoordBasis == MCB_TranslatedWorld)
+			{
+				if (bIsPositionTransform)
+				{
+					CodeStr = MultiplyMatrix(TEXT("<A>"), TEXT("DFFastToTranslatedWorld(Get<PREV>LocalToWorldDF(Parameters), ResolvedView.PreViewTranslation)"), AWComponent);
+					CodeDerivStr = MultiplyMatrix(TEXT("<A>"), TEXT("DFFastToTranslatedWorld(Get<PREV>LocalToWorldDF(Parameters), ResolvedView.PreViewTranslation)"), 0);
+				}
+			}
 			// else use MCB_World as intermediary basis
 			break;
 		}
@@ -10111,7 +10119,7 @@ int32 FHLSLMaterialTranslator::TransformBase(EMaterialCommonBasis SourceCoordBas
 		{
 			if (DestCoordBasis == MCB_World)
 			{
-				if (AWComponent)
+				if (bIsPositionTransform)
 				{
 					AddLWCFuncUsage(ELWCFunctionKind::Subtract);
 					CodeStr = TEXT("WSSubtract(<A>, Get<PREV>PreViewTranslation(Parameters))");
@@ -10132,6 +10140,39 @@ int32 FHLSLMaterialTranslator::TransformBase(EMaterialCommonBasis SourceCoordBas
 				CodeStr = MultiplyMatrix(TEXT("<A>"), TEXT("ResolvedView.<PREV>TranslatedWorldToView"), AWComponent);
 				CodeDerivStr = MultiplyMatrix(TEXT("<A>"), TEXT("ResolvedView.<PREV>TranslatedWorldToView"), 0);
 			}
+			else if (DestCoordBasis == MCB_Tangent)
+			{
+				CodeStr = MultiplyTransposeMatrix(TEXT("Parameters.TangentToWorld"), TEXT("<A>"), AWComponent);
+				CodeDerivStr = MultiplyTransposeMatrix(TEXT("Parameters.TangentToWorld"), TEXT("<A>"), 0);
+			}
+			else if (DestCoordBasis == MCB_Local)
+			{
+				const EMaterialDomain Domain = (const EMaterialDomain)Material->GetMaterialDomain();
+
+				if (Domain != MD_Surface && Domain != MD_Volume)
+				{
+					Errorf(TEXT("This transformation is only supported in the 'Surface' material domain."));
+					return INDEX_NONE;
+				}
+
+				// TODO: inconsistent with TransformLocal<TO>World with instancing
+				// We have explicit options for "local" and "instance" spaces, but then GetLocalToWorld returns instance space, while GetWorldToLocal always returns primitive space. 
+				// It's inconsistent, but replacing either will break existing materials.
+				CodeStr = MultiplyMatrix(TEXT("<A>"), TEXT("DFFastToTranslatedWorld(Get<PREV>WorldToLocalDF(Parameters), ResolvedView.PreViewTranslation)"), AWComponent);
+				CodeDerivStr = MultiplyMatrix(TEXT("<A>"), TEXT("DFFastToTranslatedWorld(Get<PREV>WorldToLocalDF(Parameters), ResolvedView.PreViewTranslation)"), 0);
+			}
+			else if (DestCoordBasis == MCB_MeshParticle)
+			{
+				CodeStr = MultiplyMatrix(TEXT("<A>"), TEXT("DFFastToTranslatedWorld(Parameters.Particle.WorldToParticle, ResolvedView.PreViewTranslation)"), AWComponent);
+				CodeDerivStr = MultiplyMatrix(TEXT("<A>"), TEXT("DFFastToTranslatedWorld(Parameters.Particle.WorldToParticle, ResolvedView.PreViewTranslation)"), 0);
+				bUsesParticleWorldToLocal = true;
+			}
+			else if (DestCoordBasis == MCB_Instance)
+			{
+				CodeStr = MultiplyMatrix(TEXT("<A>"), TEXT("DFFastToTranslatedWorld(GetWorldToInstanceDF(Parameters), ResolvedView.PreViewTranslation)"), AWComponent);
+				CodeDerivStr = MultiplyMatrix(TEXT("<A>"), TEXT("DFFastToTranslatedWorld(GetWorldToInstanceDF(Parameters), ResolvedView.PreViewTranslation)"), 0);
+				bUsesInstanceWorldToLocalPS |= ShaderFrequency == SF_Pixel;
+			}
 			// else use MCB_World as intermediary basis
 			break;
 		}
@@ -10160,7 +10201,7 @@ int32 FHLSLMaterialTranslator::TransformBase(EMaterialCommonBasis SourceCoordBas
 			}
 			else if (DestCoordBasis == MCB_TranslatedWorld)
 			{
-				if (AWComponent)
+				if (bIsPositionTransform)
 				{
 					AddLWCFuncUsage(ELWCFunctionKind::Add);
 					CodeStr = TEXT("WSAddDemote(<A>, Get<PREV>PreViewTranslation(Parameters))");
@@ -10221,6 +10262,12 @@ int32 FHLSLMaterialTranslator::TransformBase(EMaterialCommonBasis SourceCoordBas
 				CodeDerivStr = LWCMultiplyMatrix(TEXT("<A>"), TEXT("GetParticleToWorld(Parameters)"), 0, false);
 				bUsesParticleLocalToWorld = true;
 			}
+			else if (DestCoordBasis == MCB_TranslatedWorld)
+			{
+				CodeStr = MultiplyMatrix(TEXT("<A>"), TEXT("DFFastToTranslatedWorld(Parameters.Particle.ParticleToWorld, ResolvedView.PreViewTranslation)"), AWComponent);
+				CodeDerivStr = MultiplyMatrix(TEXT("<A>"), TEXT("DFFastToTranslatedWorld(Parameters.Particle.ParticleToWorld, ResolvedView.PreViewTranslation)"), 0);
+				bUsesParticleLocalToWorld = true;
+			}
 			// use World as an intermediary base
 			break;
 		}
@@ -10233,6 +10280,12 @@ int32 FHLSLMaterialTranslator::TransformBase(EMaterialCommonBasis SourceCoordBas
 				CodeDerivStr = LWCMultiplyMatrix(TEXT("<A>"), TEXT("Get<PREV>InstanceToWorld(Parameters)"), 0, false);
 				bUsesInstanceLocalToWorldPS |= ShaderFrequency == SF_Pixel;
 			}
+			else if (DestCoordBasis == MCB_TranslatedWorld)
+			{
+				CodeStr = MultiplyMatrix(TEXT("<A>"), TEXT("DFFastToTranslatedWorld(Get<PREV>InstanceToWorldDF(Parameters), ResolvedView.PreViewTranslation)"), AWComponent);
+				CodeDerivStr = MultiplyMatrix(TEXT("<A>"), TEXT("DFFastToTranslatedWorld(Get<PREV>InstanceToWorldDF(Parameters), ResolvedView.PreViewTranslation)"), 0);
+				bUsesInstanceLocalToWorldPS |= ShaderFrequency == SF_Pixel;
+			}
 			// use World as an intermediary base
 			break;
 		}
@@ -10254,7 +10307,7 @@ int32 FHLSLMaterialTranslator::TransformBase(EMaterialCommonBasis SourceCoordBas
 		return TransformBase(IntermediaryBasis, DestCoordBasis, IntermediaryA, AWComponent);
 	}
 		
-	if (bIsPositionTranform)
+	if (bIsPositionTransform)
 	{
 		CodeStr.ReplaceInline(TEXT("<TO>"),TEXT("PositionTo"));
 	}
