Fix montages not being marked dirty when (e.g.) adding curve comments or changing curve colors

Duplicated package dirtying logic into UAnimSequenceBase from UAnimSequence.

diff --git a/Engine/Source/Runtime/Engine/Private/Animation/AnimSequenceBase.cpp b/Engine/Source/Runtime/Engine/Private/Animation/AnimSequenceBase.cpp
--- a/Engine/Source/Runtime/Engine/Private/Animation/AnimSequenceBase.cpp
+++ b/Engine/Source/Runtime/Engine/Private/Animation/AnimSequenceBase.cpp
@@ -32,6 +32,7 @@
 
 #include "Animation/AnimationSettings.h"
 #include "UObject/UE5MainStreamObjectVersion.h"
+#include "UObject/UObjectThreadContext.h"
 
 DEFINE_LOG_CATEGORY(LogAnimMarkerSync);
 CSV_DECLARE_CATEGORY_MODULE_EXTERN(ENGINE_API, Animation);
@@ -1274,6 +1275,8 @@ void UAnimSequenceBase::OnModelModified(const EAnimDataModelNotifyType& NotifyTy
 		PRAGMA_ENABLE_DEPRECATION_WARNINGS
 	};
 
+	bool bShouldMarkPackageDirty = !FUObjectThreadContext::Get().IsRoutingPostLoad && NotifyType != EAnimDataModelNotifyType::BracketOpened;
+
 	switch (NotifyType)
 	{
 		case EAnimDataModelNotifyType::SequenceLengthChanged:
@@ -1369,6 +1372,8 @@ void UAnimSequenceBase::OnModelModified(const EAnimDataModelNotifyType& NotifyTy
 
 				const auto LengthChangingNotifies = { EAnimDataModelNotifyType::SequenceLengthChanged, EAnimDataModelNotifyType::FrameRateChanged, EAnimDataModelNotifyType::Reset, EAnimDataModelNotifyType::Populated };
 
+				bShouldMarkPackageDirty = NotifyCollector.WasDataModified();
+
 				if (NotifyCollector.Contains(CurveCopyNotifies) || NotifyCollector.Contains(LengthChangingNotifies))
 				{
 					CopyCurvesFromModel();
@@ -1383,6 +1388,18 @@ void UAnimSequenceBase::OnModelModified(const EAnimDataModelNotifyType& NotifyTy
 			break;
 		}
 	}
+
+	if (NotifyCollector.IsNotWithinBracket())
+	{
+		if (bShouldMarkPackageDirty)
+		{
+			MarkPackageDirty();
+		}
+	}
+	else if (bShouldMarkPackageDirty)
+	{
+		NotifyCollector.MarkDataModified();
+	}
 }
 
 IAnimationDataModel* UAnimSequenceBase::GetDataModel() const
