Switched missing bones ensure to be a logged error

Ensures are reported to crash reporter, but this one is about user content rather than program conditions.

diff --git a/Engine/Source/Runtime/Engine/Private/Animation/Skeleton.cpp b/Engine/Source/Runtime/Engine/Private/Animation/Skeleton.cpp
--- a/Engine/Source/Runtime/Engine/Private/Animation/Skeleton.cpp
+++ b/Engine/Source/Runtime/Engine/Private/Animation/Skeleton.cpp
@@ -875,11 +875,12 @@ void USkeleton::BuildLinkupData(const USkinnedAsset* InSkinnedAsset, FSkeletonTo
 		const FName MeshBoneName = MeshRefSkel.GetBoneName(MeshBoneIndex);
 		int32 SkeletonBoneIndex = SkeletonRefSkel.FindBoneIndex(MeshBoneName);
 
-#if WITH_EDITOR
-		// If we're in editor, and skeleton is missing a bone, fix it.
-		// not currently supported in-game.
 		if (SkeletonBoneIndex == INDEX_NONE)
 		{
+#if WITH_EDITOR
+			// If we're in editor, and skeleton is missing a bone, fix it.
+			// not currently supported in-game.
+
 			static FName NAME_LoadErrors("LoadErrors");
 			FMessageLog LoadErrors(NAME_LoadErrors);
 
@@ -896,12 +897,12 @@ void USkeleton::BuildLinkupData(const USkinnedAsset* InSkinnedAsset, FSkeletonTo
 
 			// Fix missing bone.
 			SkeletonBoneIndex = SkeletonRefSkel.FindBoneIndex(MeshBoneName);
-		}
 #else
-		// If we're not in editor, we still want to know which skeleton is missing a bone.
-		ensureMsgf(SkeletonBoneIndex != INDEX_NONE, TEXT("USkeleton::BuildLinkup: The Skeleton %s, is missing bones that SkeletalMesh %s needs. MeshBoneName %s"),
+			// If we're not in editor, we still want to know which skeleton is missing a bone.
+			UE_LOG(LogAnimation, Error, TEXT("USkeleton::BuildLinkup: The Skeleton %s, is missing bones that SkeletalMesh %s needs. MeshBoneName %s"),
 				*GetNameSafe(this), *GetNameSafe(InSkinnedAsset), *MeshBoneName.ToString());
 #endif
+		}
 
 		NewMeshLinkup.MeshToSkeletonTable[MeshBoneIndex] = SkeletonBoneIndex;
 	}
