- Fix for Cascade GPU particles not working
- We can't wrap AdvanceGPUParticleFrame inside a Pass as GDME will point at the wrong data

diff --git a/Engine/Source/Runtime/Engine/Private/Particles/FXSystem.cpp b/Engine/Source/Runtime/Engine/Private/Particles/FXSystem.cpp
--- a/Engine/Source/Runtime/Engine/Private/Particles/FXSystem.cpp
+++ b/Engine/Source/Runtime/Engine/Private/Particles/FXSystem.cpp
@@ -428,18 +428,9 @@ void FFXSystem::PreInitViews(FRDGBuilder& GraphBuilder, bool bAllowGPUParticleUp
 {
 	if (RHISupportsGPUParticles())
 	{
-		RDG_GPU_STAT_SCOPE(GraphBuilder, FXSystemPreInitViews);
-		RDG_CSV_STAT_EXCLUSIVE_SCOPE(GraphBuilder, FXSystem);
-
-		AddPass(
-			GraphBuilder,
-			RDG_EVENT_NAME("FFXSystem::PreInitViews"),
-			[this, bAllowGPUParticleUpdate](FRHICommandListImmediate& RHICmdList)
-			{
-
-				AdvanceGPUParticleFrame(RHICmdList, bAllowGPUParticleUpdate);
-			}
-		);
+		// Note: This can not be put into a GraphBuilder pass directly, the internals need to be refactored in order to that.
+		// This is because the data modified in here will be used in GDME pass which must have the most up to date information
+		AdvanceGPUParticleFrame(GraphBuilder.RHICmdList, bAllowGPUParticleUpdate);
 	}
 }
 
