- Remove ensure as it triggers a false positive during level transitions

diff --git a/Engine/Source/Runtime/Engine/Private/Particles/ParticleComponents.cpp b/Engine/Source/Runtime/Engine/Private/Particles/ParticleComponents.cpp
--- a/Engine/Source/Runtime/Engine/Private/Particles/ParticleComponents.cpp
+++ b/Engine/Source/Runtime/Engine/Private/Particles/ParticleComponents.cpp
@@ -5299,8 +5299,8 @@ void UParticleSystemComponent::OnAttachmentChanged()
 
 	if (IsTickManaged())
 	{
-		FParticleSystemWorldManager* PSCMan = GetWorldManager();
-		if (ensure(PSCMan))
+		// Note: the PSCMan can become invalid during GC / level change
+		if (FParticleSystemWorldManager* PSCMan = GetWorldManager())
 		{
 			//Reregister component to recalculate dependencies and re add to manager's lists.
 			PSCMan->UnregisterComponent(this);
