Small refactor of FStaticMeshSceneProxy::GetDynamicRayTracingInstances(...)
- also BoostBuildPriority of LODs pending build request when target LOD is not yet ready.

diff --git a/Engine/Source/Runtime/Engine/Private/StaticMeshRender.cpp b/Engine/Source/Runtime/Engine/Private/StaticMeshRender.cpp
--- a/Engine/Source/Runtime/Engine/Private/StaticMeshRender.cpp
+++ b/Engine/Source/Runtime/Engine/Private/StaticMeshRender.cpp
@@ -1925,9 +1925,17 @@ void FStaticMeshSceneProxy::GetDynamicRayTracingInstances(FRayTracingMaterialGat
 
 	if (!bEvaluateWPO)
 	{
+		// Select first LOD with valid ray tracing geometry
 		for (; LODIndex < NumLODs; LODIndex++)
 		{
-			if (RenderData->LODResources[LODIndex].RayTracingGeometry.IsValid())
+			FStaticMeshLODResources& CurrentLODResources = RenderData->LODResources[LODIndex];
+
+			if (CurrentLODResources.RayTracingGeometry.HasPendingBuildRequest())
+			{
+				ensure(CurrentLODResources.RayTracingGeometry.IsValid());
+				CurrentLODResources.RayTracingGeometry.BoostBuildPriority();
+			}
+			else if (CurrentLODResources.RayTracingGeometry.IsValid())
 			{
 				break;
 			}
@@ -1941,20 +1949,26 @@ void FStaticMeshSceneProxy::GetDynamicRayTracingInstances(FRayTracingMaterialGat
 
 	const FStaticMeshLODResources& LODModel = RenderData->LODResources[LODIndex];
 
-	// TODO: Need to validate that the DynamicRayTracingGeometries are still valid - they could contain streamed out IndexBuffers from the shared StaticMesh (UE-139474)
-	FRayTracingGeometry& Geometry = bEvaluateWPO ? DynamicRayTracingGeometries[LODIndex] : RenderData->LODResources[LODIndex].RayTracingGeometry;
-	
-	if (LODModel.GetNumVertices() <= 0 || Geometry.Initializer.TotalPrimitiveCount <= 0)
+	if (LODModel.GetNumVertices() <= 0)
 	{
 		return;
 	}
 
-	// Early out for now if no valid RHI RT geometry yet (still pending build request)
-	// TODO: select different LOD if available
-	if (Geometry.HasPendingBuildRequest())
+	// TODO: Need to validate that the DynamicRayTracingGeometries are still valid - they could contain streamed out IndexBuffers from the shared StaticMesh (UE-139474)
+	FRayTracingGeometry& Geometry = bEvaluateWPO ? DynamicRayTracingGeometries[LODIndex] : RenderData->LODResources[LODIndex].RayTracingGeometry;
+
+	if (bEvaluateWPO)
 	{
-		Geometry.BoostBuildPriority();
-		return;
+		if (Geometry.HasPendingBuildRequest())
+		{
+			// This should only happen if geometry was recently made resident and build hasn't happened yet.
+			// TODO: could cancel build request and let it go through the dynamic code path which will build it as necessary
+			return;
+		}
+	}
+	else
+	{
+		check(Geometry.IsValid() && !Geometry.HasPendingBuildRequest());
 	}
 
 	{
