Prevent crashes when setting a screen percentage outside the GetMinUpsampleResolutionFraction() to GetMaxUpsampleResolutionFraction() range specified by the active third-party temporal upscaler.

This crash could occur when using dynamic resolution and seeing poor performance. We now late update the target screen percentage in FSceneRenderer::PrepareViewRectsForRendering based on the the temporal upscaler's minimum and maximum limits. Forcing a percentage with r.DynamicRes.TestScreenPercentage will still bypass this.

diff --git a/Engine/Source/Runtime/Engine/Private/DynamicResolution.cpp b/Engine/Source/Runtime/Engine/Private/DynamicResolution.cpp
--- a/Engine/Source/Runtime/Engine/Private/DynamicResolution.cpp
+++ b/Engine/Source/Runtime/Engine/Private/DynamicResolution.cpp
@@ -202,6 +202,9 @@ void FDynamicResolutionHeuristicProxy::ResetInternal()
 	DynamicFrameTimeBudgetMs = 0.0f;
 	DynamicFrameTimeVSyncFactor = 0;
 
+	TemporalUpscalerMinResolutionFraction = ISceneViewFamilyScreenPercentage::kMinResolutionFraction;
+	TemporalUpscalerMaxResolutionFraction = ISceneViewFamilyScreenPercentage::kMaxResolutionFraction;
+
 	// Ignore previous frame timings.
 	IgnoreFrameRemainingCount = 1;
 }
@@ -544,6 +547,12 @@ void FDynamicResolutionHeuristicProxy::RefreshCurrentFrameResolutionFraction_Ren
 			BudgetSettings.MinResolutionFraction,
 			FinalMaxResolutionFraction);
 
+		// Also clamp with the temporal upscaler's minimum and maximum fractions (set to theoretical minimum and maximum if not in use)
+		NewFrameResolutionFraction = FMath::Clamp(
+			NewFrameResolutionFraction,
+			TemporalUpscalerMinResolutionFraction,
+			TemporalUpscalerMaxResolutionFraction);
+
 		NewFrameResolutionFractions[Budget] = NewFrameResolutionFraction;
 	}
 
@@ -574,6 +583,24 @@ void FDynamicResolutionHeuristicProxy::RefreshCurrentFrameResolutionFraction_Ren
 	RefreshHeuristicStats_RenderThread();
 }
 
+void FDynamicResolutionHeuristicProxy::SetTemporalUpscaler(const UE::Renderer::Private::ITemporalUpscaler* InTemporalUpscaler)
+{
+	check(IsInRenderingThread());
+
+	float NewMinResolutionFraction = InTemporalUpscaler ? InTemporalUpscaler->GetMinUpsampleResolutionFraction() : ISceneViewFamilyScreenPercentage::kMinResolutionFraction;
+	float NewMaxResolutionFraction = InTemporalUpscaler ? InTemporalUpscaler->GetMaxUpsampleResolutionFraction() : ISceneViewFamilyScreenPercentage::kMaxResolutionFraction;
+
+	if (NewMinResolutionFraction != TemporalUpscalerMinResolutionFraction ||
+		NewMaxResolutionFraction != TemporalUpscalerMaxResolutionFraction)
+	{
+		TemporalUpscalerMinResolutionFraction = NewMinResolutionFraction;
+		TemporalUpscalerMaxResolutionFraction = NewMaxResolutionFraction;
+
+		// If the temporal upscaler (or its supported range) have changed, refresh the fractions for this frame
+		RefreshCurrentFrameResolutionFraction_RenderThread();
+	}
+}
+
 void FDynamicResolutionHeuristicProxy::RefreshCurrentFrameResolutionFractionUpperBound_RenderThread()
 {
 	// Compute max resolution for each budget by quantizing the new resolution fraction (falls back to the MaxResolution setting if BudgetSetting.UpperBoundQuantization==0)
@@ -1192,6 +1219,11 @@ public:
 		}
 	}
 
+	virtual void SetTemporalUpscaler(const UE::Renderer::Private::ITemporalUpscaler* InTemporalUpscaler)
+	{
+		Proxy->Heuristic.SetTemporalUpscaler(InTemporalUpscaler);
+	}
+
 private:
 	// Owned render thread proxy.
 	FDefaultDynamicResolutionStateProxy* const Proxy;
diff --git a/Engine/Source/Runtime/Engine/Public/DynamicResolutionProxy.h b/Engine/Source/Runtime/Engine/Public/DynamicResolutionProxy.h
--- a/Engine/Source/Runtime/Engine/Public/DynamicResolutionProxy.h
+++ b/Engine/Source/Runtime/Engine/Public/DynamicResolutionProxy.h
@@ -5,6 +5,7 @@
 #include "CoreMinimal.h"
 #include "SceneView.h"
 #include "Engine/EngineTypes.h"
+#include "TemporalUpscaler.h"
 
 /** Render thread proxy that holds the heuristic for dynamic resolution. */
 class FDynamicResolutionHeuristicProxy
@@ -55,6 +56,9 @@ public:
 	/** Creates a default dynamic resolution state using this proxy that queries GPU timing from the RHI. */
 	static ENGINE_API TSharedPtr<class IDynamicResolutionState> CreateDefaultState();
 
+	/** Applies the minimum/maximum resolution fraction for a third-party temporal upscaler. */
+	ENGINE_API void SetTemporalUpscaler(const UE::Renderer::Private::ITemporalUpscaler* InTemporalUpscaler);
+
 private:
 	struct FrameHistoryEntry
 	{
@@ -115,6 +119,10 @@ private:
 	// Frame counter to allocate unique ID for CommitPreviousFrameGPUTimings_RenderThread().
 	uint64 FrameCounter;
 
+	// Minimum and maximum resolution fractions supported by the main view family's third-party temporal upscaler.
+	float TemporalUpscalerMinResolutionFraction;
+	float TemporalUpscalerMaxResolutionFraction;
+
 	inline const FrameHistoryEntry& GetPreviousFrameEntry(int32 BrowsingFrameId) const
 	{
 		if (BrowsingFrameId < 0 || BrowsingFrameId >= HistorySize)
diff --git a/Engine/Source/Runtime/Engine/Public/DynamicResolutionState.h b/Engine/Source/Runtime/Engine/Public/DynamicResolutionState.h
--- a/Engine/Source/Runtime/Engine/Public/DynamicResolutionState.h
+++ b/Engine/Source/Runtime/Engine/Public/DynamicResolutionState.h
@@ -51,6 +51,9 @@ public:
 	/** Setup a screen percentage driver for a given view family. */
 	virtual void SetupMainViewFamily(class FSceneViewFamily& ViewFamily) = 0;
 
+	/** Apply the minimum/maximum resolution fraction for a third-party temporal upscaler. */
+	virtual void SetTemporalUpscaler(const UE::Renderer::Private::ITemporalUpscaler* InTemporalUpscaler) = 0;
+
 protected:
 
 	/** Returns a non thread safe aproximation of the current resolution fraction applied on render thread, mostly used for stats and analytic. */
diff --git a/Engine/Source/Runtime/Renderer/Private/SceneRendering.cpp b/Engine/Source/Runtime/Renderer/Private/SceneRendering.cpp
--- a/Engine/Source/Runtime/Renderer/Private/SceneRendering.cpp
+++ b/Engine/Source/Runtime/Renderer/Private/SceneRendering.cpp
@@ -2986,6 +2986,12 @@ void FSceneRenderer::PrepareViewRectsForRendering(FRHICommandListImmediate& RHIC
 	check(IsInRenderingThread());
 	TRACE_CPUPROFILER_EVENT_SCOPE(PrepareViewRectsForRendering);
 
+	// If we support screen percentage, update the dynamic resolution state with our current temporal upscaler, which clamps the screen percentage to its supported range.
+	if (ViewFamily.SupportsScreenPercentage())
+	{
+		GEngine->GetDynamicResolutionState()->SetTemporalUpscaler(ViewFamily.GetTemporalUpscalerInterface());
+	}
+
 	// Read the resolution data.
 	{
 		check(ViewFamily.ScreenPercentageInterface);
