Avoid repeated DDC remote existence queries for texture mips in editor by only considering the mips that are loaded or prefetched up front.  If we attempt to stream in a mip that wasn't loaded or prefetched up front and that fails, a forced rebuild of the texture is enqueued.

diff --git a/Engine/Source/Runtime/Engine/Private/Streaming/Texture2DMipDataProvider_DDC.cpp b/Engine/Source/Runtime/Engine/Private/Streaming/Texture2DMipDataProvider_DDC.cpp
--- a/Engine/Source/Runtime/Engine/Private/Streaming/Texture2DMipDataProvider_DDC.cpp
+++ b/Engine/Source/Runtime/Engine/Private/Streaming/Texture2DMipDataProvider_DDC.cpp
@@ -8,14 +8,16 @@ Texture2DMipDataProvider_DDC.cpp : Implementation of FTextureMipDataProvider usi
 #include "EngineLogs.h"
 #include "Rendering/StreamableTextureResource.h"
 #include "Serialization/MemoryReader.h"
+#include "TextureCompiler.h"
 
 #if WITH_EDITORONLY_DATA
 
 #include "DerivedDataCache.h"
 
-FTexture2DMipDataProvider_DDC::FTexture2DMipDataProvider_DDC(const UTexture* Texture)
-	: FTextureMipDataProvider(Texture, ETickState::Init, ETickThread::Async)
+FTexture2DMipDataProvider_DDC::FTexture2DMipDataProvider_DDC(UTexture* InTexture)
+	: FTextureMipDataProvider(InTexture, ETickState::Init, ETickThread::Async)
 	, DDCRequestOwner(UE::DerivedData::EPriority::Normal)
+	, Texture(InTexture)
 {
 }
 
@@ -93,6 +95,10 @@ void FTexture2DMipDataProvider_DDC::Init(const FTextureUpdateContext& Context, c
 						check(!DDCBuffers[MipIndex]);
 						DDCBuffers[MipIndex] = MoveTemp(Response.RawData);
 					}
+					else if (Response.Status == EStatus::Error)
+					{
+						FTextureCompilingManager::Get().ForceDeferredTextureRebuildAnyThread({Texture});
+					}
 				});
 			}
 		}
diff --git a/Engine/Source/Runtime/Engine/Private/Streaming/Texture2DMipDataProvider_DDC.h b/Engine/Source/Runtime/Engine/Private/Streaming/Texture2DMipDataProvider_DDC.h
--- a/Engine/Source/Runtime/Engine/Private/Streaming/Texture2DMipDataProvider_DDC.h
+++ b/Engine/Source/Runtime/Engine/Private/Streaming/Texture2DMipDataProvider_DDC.h
@@ -9,6 +9,7 @@ Texture2DMipDataProvider_DDC.h : Implementation of FTextureMipDataProvider using
 #include "CoreMinimal.h"
 #include "Memory/SharedBuffer.h"
 #include "Streaming/TextureMipDataProvider.h"
+#include "UObject/WeakObjectPtr.h"
 
 #if WITH_EDITORONLY_DATA
 #include "DerivedDataRequestOwner.h"
@@ -20,7 +21,7 @@ class FTexture2DMipDataProvider_DDC : public FTextureMipDataProvider
 {
 public:
 
-	FTexture2DMipDataProvider_DDC(const UTexture* Texture);
+	FTexture2DMipDataProvider_DDC(UTexture* Texture);
 	~FTexture2DMipDataProvider_DDC();
 
 
@@ -45,6 +46,7 @@ private:
 
 	TArray<FSharedBuffer, TInlineAllocator<MAX_TEXTURE_MIP_COUNT> > DDCBuffers;
 	UE::DerivedData::FRequestOwner DDCRequestOwner;
+	TWeakObjectPtr<UTexture> Texture;
 };
 
 #endif //WITH_EDITORONLY_DATA
diff --git a/Engine/Source/Runtime/Engine/Private/Streaming/Texture2DStreamIn_DDC.cpp b/Engine/Source/Runtime/Engine/Private/Streaming/Texture2DStreamIn_DDC.cpp
--- a/Engine/Source/Runtime/Engine/Private/Streaming/Texture2DStreamIn_DDC.cpp
+++ b/Engine/Source/Runtime/Engine/Private/Streaming/Texture2DStreamIn_DDC.cpp
@@ -6,6 +6,7 @@ Texture2DStreamIn_DDC.cpp: Stream in helper for 2D textures loading DDC files.
 
 #include "Streaming/Texture2DStreamIn_DDC.h"
 #include "EngineLogs.h"
+#include "TextureCompiler.h"
 #include "Rendering/Texture2DResource.h"
 #include "RenderUtils.h"
 #include "Streaming/Texture2DStreamIn.h"
@@ -31,6 +32,7 @@ static FAutoConsoleVariableRef CVarStreamingDDCPendingSleep(
 FTexture2DStreamIn_DDC::FTexture2DStreamIn_DDC(UTexture2D* InTexture)
 	: FTexture2DStreamIn(InTexture)
 	, DDCRequestOwner(UE::DerivedData::EPriority::Normal)
+	, Texture(InTexture)
 {
 	DDCMipRequestStatus.AddZeroed(ResourceState.MaxNumLODs);
 }
@@ -130,6 +132,10 @@ void FTexture2DStreamIn_DDC::DoCreateAsyncDDCRequests(const FContext& Context)
 					check(Status.bRequestIssued);
 					Status.bRequestIssued = false;
 				}
+				else if (Response.Status == EStatus::Error)
+				{
+					FTextureCompilingManager::Get().ForceDeferredTextureRebuildAnyThread({Texture});
+				}
 			});
 		}
 	}
diff --git a/Engine/Source/Runtime/Engine/Private/Streaming/Texture2DStreamIn_DDC.h b/Engine/Source/Runtime/Engine/Private/Streaming/Texture2DStreamIn_DDC.h
--- a/Engine/Source/Runtime/Engine/Private/Streaming/Texture2DStreamIn_DDC.h
+++ b/Engine/Source/Runtime/Engine/Private/Streaming/Texture2DStreamIn_DDC.h
@@ -9,6 +9,7 @@ Texture2DStreamIn_DDC.h: Stream in helper for 2D textures loading DDC files.
 #include "CoreMinimal.h"
 #include "Memory/SharedBuffer.h"
 #include "Texture2DStreamIn.h"
+#include "UObject/WeakObjectPtr.h"
 
 #if WITH_EDITORONLY_DATA
 #include "DerivedDataRequestOwner.h"
@@ -32,6 +33,7 @@ protected:
 
 	TArray<FMipRequestStatus, TInlineAllocator<MAX_TEXTURE_MIP_COUNT> > DDCMipRequestStatus;
 	UE::DerivedData::FRequestOwner DDCRequestOwner;
+	TWeakObjectPtr<UTexture2D> Texture;
 
 	// ****************************
 	// ********* Helpers **********
diff --git a/Engine/Source/Runtime/Engine/Private/TextureCompiler.cpp b/Engine/Source/Runtime/Engine/Private/TextureCompiler.cpp
--- a/Engine/Source/Runtime/Engine/Private/TextureCompiler.cpp
+++ b/Engine/Source/Runtime/Engine/Private/TextureCompiler.cpp
@@ -289,6 +289,16 @@ void FTextureCompilingManager::AddTextures(TArrayView<UTexture* const> InTexture
 	TRACE_COUNTER_SET(QueuedTextureCompilation, GetNumRemainingTextures());
 }
 
+void FTextureCompilingManager::ForceDeferredTextureRebuildAnyThread(TArrayView<const TWeakObjectPtr<UTexture>> InTextures)
+{
+	TRACE_CPUPROFILER_EVENT_SCOPE(FTextureCompilingManager::AddTexturesDeferredAnyThread)
+
+	for (const TWeakObjectPtr<UTexture>& Texture : InTextures)
+	{
+		DeferredRebuildRequestQueue.ProduceItem(Texture);
+	}
+}
+
 void FTextureCompilingManager::FinishCompilationForObjects(TArrayView<UObject* const> InObjects)
 {
 	TRACE_CPUPROFILER_EVENT_SCOPE(FTextureCompilingManager::FinishCompilationForObjects);
@@ -722,6 +732,7 @@ void FTextureCompilingManager::ProcessAsyncTasks(bool bLimitExecutionTime)
 void FTextureCompilingManager::ProcessAsyncTasks(const AssetCompilation::FProcessAsyncTaskParams& Params)
 {
 	FObjectCacheContextScope ObjectCacheScope;
+	ProcessDeferredRequests();
 	FinishCompilationsForGame();
 
 	if (!Params.bPlayInEditorAssetsOnly)
@@ -732,6 +743,30 @@ void FTextureCompilingManager::ProcessAsyncTasks(const AssetCompilation::FProces
 	UpdateCompilationNotification();
 }
 
+void FTextureCompilingManager::ProcessDeferredRequests()
+{
+	TSet<UTexture*> DeferredTextures;
+	DeferredRebuildRequestQueue.ConsumeAllFifo([this, &DeferredTextures](TWeakObjectPtr<UTexture> WeakTexture)
+	{
+		if (UTexture* Texture = WeakTexture.Get())
+		{
+			if (Texture->IsAsyncCacheComplete() && IsAsyncCompilationAllowed(Texture) && !IsCompilingTexture(Texture))
+			{
+				DeferredTextures.Add(Texture);
+			}
+		}
+	});
+
+	if (!DeferredTextures.IsEmpty())
+	{
+		for (UTexture* DeferredTexture : DeferredTextures)
+		{
+			DeferredTexture->ForceRebuildPlatformData();
+		}
+	}
+}
+
+
 #undef LOCTEXT_NAMESPACE
 
 #endif // #if WITH_EDITOR
diff --git a/Engine/Source/Runtime/Engine/Private/TextureDerivedDataTask.cpp b/Engine/Source/Runtime/Engine/Private/TextureDerivedDataTask.cpp
--- a/Engine/Source/Runtime/Engine/Private/TextureDerivedDataTask.cpp
+++ b/Engine/Source/Runtime/Engine/Private/TextureDerivedDataTask.cpp
@@ -1402,11 +1402,11 @@ static bool TryCacheStreamingMips(const FString& TexturePathName, int32 FirstMip
 
 	TArray<FCacheGetValueRequest, TInlineAllocator<16>> MipRequests;
 
-	int32 MipIndex = -1;
+	int32 LowestMipIndexToPrefetchOrLoad = FMath::Min(FirstMipToPrefetch, FirstMipToLoad);
 	const FSharedString Name(TexturePathName);
-	for (const FTexture2DMipMap& Mip : DerivedData->Mips)
+	for (int32 MipIndex = 0; MipIndex < LowestMipIndexToPrefetchOrLoad; ++MipIndex)
 	{
-		++MipIndex;
+		const FTexture2DMipMap& Mip = DerivedData->Mips[MipIndex];
 		if (Mip.IsPagedToDerivedData())
 		{
 			const FCacheKey MipKey = ConvertLegacyCacheKey(DerivedData->GetDerivedDataMipKeyString(MipIndex, Mip));
diff --git a/Engine/Source/Runtime/Engine/Public/TextureCompiler.h b/Engine/Source/Runtime/Engine/Public/TextureCompiler.h
--- a/Engine/Source/Runtime/Engine/Public/TextureCompiler.h
+++ b/Engine/Source/Runtime/Engine/Public/TextureCompiler.h
@@ -2,6 +2,7 @@
 #pragma once
 
 #include "UObject/WeakObjectPtr.h"
+#include "Containers/ConsumeAllMpmcQueue.h"
 #include "Containers/Set.h"
 #include "IAssetCompilingManager.h"
 
@@ -34,6 +35,11 @@ public:
 	 */
 	ENGINE_API void AddTextures(TArrayView<UTexture* const> InTextures);
 
+	/** 
+	 * Forces textures to be recompiled asynchronously later. 
+	 */
+	ENGINE_API void ForceDeferredTextureRebuildAnyThread(TArrayView<const TWeakObjectPtr<UTexture>> InTextures);
+
 	/** 
 	 * Blocks until completion of the requested textures.
 	 */
@@ -100,9 +106,12 @@ private:
 	void PostCompilation(UTexture* Texture);
 	void PostCompilation(TArrayView<UTexture* const> InCompiledTextures);
 
+	void ProcessDeferredRequests();
+
 	double LastReschedule = 0.0f;
 	bool bHasShutdown = false;
 	bool bIsRoutingPostCompilation = false;
+	UE::TConsumeAllMpmcQueue<TWeakObjectPtr<UTexture>> DeferredRebuildRequestQueue;
 	TArray<TSet<TWeakObjectPtr<UTexture>>> RegisteredTextureBuckets;
 	TUniquePtr<FAsyncCompilationNotification> Notification;
 
