Followup change to CL 32039781 to fix compilation errors for plugins that include DynamicResolutionState.h.

diff --git a/Engine/Source/Runtime/Engine/Public/DynamicResolutionState.h b/Engine/Source/Runtime/Engine/Public/DynamicResolutionState.h
--- a/Engine/Source/Runtime/Engine/Public/DynamicResolutionState.h
+++ b/Engine/Source/Runtime/Engine/Public/DynamicResolutionState.h
@@ -9,6 +9,7 @@
 #include "CoreMinimal.h"
 #include "Engine/EngineTypes.h"
 #include "DynamicRenderScaling.h"
+#include "TemporalUpscaler.h"
 
 
 /** Dynamic resolution of the primary ScreenPercentage. */
