Shader Compiler:  Fix deadlock when running with -DDC-Verify enabled for shaders.  Verification can cause the DDC query to complete synchronously, and the callback is unable to acquire a write lock to the job cache since the current thread running SubmitJob already holds the lock.  Added a thread local variable to track when the current thread is running SubmitJob, and switch to the synchronous code paths (the same code used when async DDC queries are disabled).

#rnx

diff --git a/Engine/Source/Runtime/Engine/Private/ShaderCompiler/ShaderCompiler.cpp b/Engine/Source/Runtime/Engine/Private/ShaderCompiler/ShaderCompiler.cpp
--- a/Engine/Source/Runtime/Engine/Private/ShaderCompiler/ShaderCompiler.cpp
+++ b/Engine/Source/Runtime/Engine/Private/ShaderCompiler/ShaderCompiler.cpp
@@ -928,6 +928,9 @@ private:
 	/** Guards access to the structure */
 	mutable FRWLock JobLock;
 
+	/** Needed to detect if DDC query callback is completing in the same thread as SubmitJob */
+	static thread_local bool bInSubmitJobThread;
+
 	/** List of jobs waiting on SubmitJob task or DDC query (not yet added to a pending queue). */
 	FShaderCommonCompileJob* PendingSubmitJobTaskJobs = nullptr;
 
@@ -974,6 +977,8 @@ private:
 	uint64 CurrentlyAllocatedMemory = 0;
 };
 
+thread_local bool FShaderJobCache::bInSubmitJobThread;
+
 static FShaderJobData& GetShaderJobData(const FShaderJobCacheRef& CacheRef)
 {
 	check(CacheRef.Block);
@@ -1587,6 +1592,14 @@ int32 FShaderJobCache::RemoveAllPendingJobsWithId(uint32 InId)
 
 void FShaderJobCache::SubmitJob(FShaderCommonCompileJob* Job)
 {
+	// Set thread local so DDC query callback can detect if it's in the same thread, and we need to run through the non-async code path.
+	struct InSubmitJobScope
+	{
+		InSubmitJobScope() { bInSubmitJobThread = true; }
+		~InSubmitJobScope() { bInSubmitJobThread = false; }
+	};
+	InSubmitJobScope InSubmitJob;
+
 	check(Job->Priority != EShaderCompileJobPriority::None);
 	check(Job->PendingPriority == EShaderCompileJobPriority::None);
 
@@ -10831,19 +10844,31 @@ FShaderJobCacheRef FShaderJobCache::FindOrAdd(const FJobInputHash& Hash, EShader
 			Request.Key.Hash = Hash;
 			Request.Policy = IsShaderJobCacheDDCRemotePolicyEnabled() ? UE::DerivedData::ECachePolicy::Default : UE::DerivedData::ECachePolicy::Local;
 
-			// If blocking, we'll read the cached output back to the main thread
-			FJobCachedOutput** OutCachedOutputPtr = DerivedDataPriority == UE::DerivedData::EPriority::Blocking ? &OutCachedOutput : nullptr;
+			// Optionally read the cached output back to the main thread
+			FJobCachedOutput** OutCachedOutputPtr = &OutCachedOutput;
+			bool bCompletedSynchronously = false;
+			bool* bCompletedSynchronouslyPtr = &bCompletedSynchronously;
 
 			UE::DerivedData::GetCache().Get(
 				{ Request },
 				*RequestOwner,
-				[this, JobDataPtr = &JobData, OutCachedOutputPtr, DerivedDataPriority](UE::DerivedData::FCacheGetResponse&& Response)
+				[this, JobDataPtr = &JobData, OutCachedOutputPtr, DerivedDataPriority, bCompletedSynchronouslyPtr](UE::DerivedData::FCacheGetResponse&& Response)
 				{
 					if (GShaderCompilerDebugStallDDCQuery > 0)
 					{
 						FPlatformProcess::Sleep(GShaderCompilerDebugStallDDCQuery * 0.001f);
 					}
 
+					bool bIsAsync = (DerivedDataPriority != UE::DerivedData::EPriority::Blocking);
+
+					// Check thread local variable to see if we're in the submit job thread (DDC request completing synchronously), in which
+					// case we want to go through the synchronous code paths below, instead of async.
+					if (bInSubmitJobThread)
+					{
+						bIsAsync = false;
+						*bCompletedSynchronouslyPtr = true;
+					}
+
 					if (Response.Status == UE::DerivedData::EStatus::Ok)
 					{
 						// Retrieve the shared buffer containing the job output and compute the associated output hash for the result retrieved from DDC
@@ -10856,7 +10881,7 @@ FShaderJobCacheRef FShaderJobCache::FindOrAdd(const FJobInputHash& Hash, EShader
 
 						// If we are running the cache logic async (not blocking in the main thread), we need a lock before writing to the job cache.
 						// Otherwise, the lock will already be held by the main thread (and trying to lock here would just deadlock).
-						if (DerivedDataPriority != UE::DerivedData::EPriority::Blocking)
+						if (bIsAsync)
 						{
 							JobLock.WriteLock();
 							check(JobDataPtr->JobInFlight);
@@ -10904,14 +10929,8 @@ FShaderJobCacheRef FShaderJobCache::FindOrAdd(const FJobInputHash& Hash, EShader
 						JobDataPtr->OutputHash = OutputHash;
 						JobDataPtr->bOutputFromDDC = true;
 
-						// Optionally send results back to the main thread
-						if (OutCachedOutputPtr)
-						{
-							*OutCachedOutputPtr = &StoredOutput->JobOutput;
-						}
-
-						// If non-blocking, add processed results to output.  For the blocking case, this is handled back in the main thread.
-						if (DerivedDataPriority != UE::DerivedData::EPriority::Blocking)
+						// If async, add processed results to output.  For the synchronous case, this is handled back in the main thread.
+						if (bIsAsync)
 						{
 							check(JobDataPtr->JobInFlight);
 							FShaderCommonCompileJobPtr Job = JobDataPtr->JobInFlight;
@@ -10958,11 +10977,16 @@ FShaderJobCacheRef FShaderJobCache::FindOrAdd(const FJobInputHash& Hash, EShader
 								UE_LOG(LogShaderCompilers, UE_SHADERCACHE_LOG_LEVEL, TEXT("Processed %d outstanding jobs with the same ihash %s."), FinishedJobs.Num() - 1, *LexToString(Job->InputHash));
 							}
 						}
+						else
+						{
+							// Send results back to the main thread when running synchronous
+							*OutCachedOutputPtr = &StoredOutput->JobOutput;
+						}
 					}
 					else
 					{
-						// If non-blocking, add job to pending queue.  For the blocking case, this is handled back in the main thread.
-						if (DerivedDataPriority != UE::DerivedData::EPriority::Blocking)
+						// If async, add job to pending queue.  For the synchronous case, this is handled back in the main thread.
+						if (bIsAsync)
 						{
 							FWriteScopeLock Locker(JobLock);
 							FShaderCommonCompileJob* Job = JobDataPtr->JobInFlight;
@@ -11001,6 +11025,14 @@ FShaderJobCacheRef FShaderJobCache::FindOrAdd(const FJobInputHash& Hash, EShader
 				RequestOwner->Wait();
 				delete RequestOwner;
 			}
+			else if (bCompletedSynchronously)
+			{
+				// It's also possible (notably when DDC verification is enabled) for the request to have completed synchronously,
+				// in which case we can delete the TPimplPtr request owner by setting it to null.  This tells the main thread
+				// there is no async DDC request in flight, and it should handle adding the pending job to the queue, since the
+				// DDC request callback won't be handling that.
+				InoutRequestOwner = nullptr;
+			}
 		}
 	}
 #endif
