Fix bug in the logic of TryCacheStreamingMips introduced by my last change.

diff --git a/Engine/Source/Runtime/Engine/Private/TextureDerivedDataTask.cpp b/Engine/Source/Runtime/Engine/Private/TextureDerivedDataTask.cpp
--- a/Engine/Source/Runtime/Engine/Private/TextureDerivedDataTask.cpp
+++ b/Engine/Source/Runtime/Engine/Private/TextureDerivedDataTask.cpp
@@ -1402,9 +1402,10 @@ static bool TryCacheStreamingMips(const FString& TexturePathName, int32 FirstMip
 
 	TArray<FCacheGetValueRequest, TInlineAllocator<16>> MipRequests;
 
-	int32 LowestMipIndexToPrefetchOrLoad = FMath::Min(FirstMipToPrefetch, FirstMipToLoad);
+	const int32 LowestMipIndexToPrefetchOrLoad = FMath::Min(FirstMipToPrefetch, FirstMipToLoad);
+	const int32 NumMips = DerivedData->Mips.Num();
 	const FSharedString Name(TexturePathName);
-	for (int32 MipIndex = 0; MipIndex < LowestMipIndexToPrefetchOrLoad; ++MipIndex)
+	for (int32 MipIndex = LowestMipIndexToPrefetchOrLoad; MipIndex < NumMips; ++MipIndex)
 	{
 		const FTexture2DMipMap& Mip = DerivedData->Mips[MipIndex];
 		if (Mip.IsPagedToDerivedData())
