Reconnect client TO ODSC server if connection has been lost

diff --git a/Engine/Source/Runtime/Engine/Private/ODSC/ODSCThread.cpp b/Engine/Source/Runtime/Engine/Private/ODSC/ODSCThread.cpp
--- a/Engine/Source/Runtime/Engine/Private/ODSC/ODSCThread.cpp
+++ b/Engine/Source/Runtime/Engine/Private/ODSC/ODSCThread.cpp
@@ -101,23 +101,16 @@ bool FODSCMessageHandler::ReloadGlobalShaders() const
 }
 
 FODSCThread::FODSCThread(const FString& HostIP)
-	: Thread(nullptr),
-	  WakeupEvent(FPlatformProcess::GetSynchEventFromPool(true))
+	: Thread(nullptr)
+	, WakeupEvent(FPlatformProcess::GetSynchEventFromPool(true))
+	, ODSCHostIP(HostIP)
 {
 	UE_LOG(LogODSC, Log, TEXT("ODSC Thread active."));
 
-	// Attempt to get a default connection to the COTF server (which cooks assets).
-	UE::Cook::ICookOnTheFlyModule& CookOnTheFlyModule = FModuleManager::LoadModuleChecked<UE::Cook::ICookOnTheFlyModule>(TEXT("CookOnTheFly"));
-	if (!CookOnTheFlyModule.GetDefaultServerConnection())
+	bHasDefaultConnection = (FModuleManager::LoadModuleChecked<UE::Cook::ICookOnTheFlyModule>(TEXT("CookOnTheFly")).GetDefaultServerConnection() != nullptr);
+	if (!bHasDefaultConnection)
 	{
-		// If we don't have a default connection make a specific connection to the HostIP provided.
-		UE::Cook::FCookOnTheFlyHostOptions CookOnTheFlyHostOptions;
-		CookOnTheFlyHostOptions.Hosts.Add(HostIP);
-		CookOnTheFlyServerConnection = CookOnTheFlyModule.ConnectToServer(CookOnTheFlyHostOptions);
-		if (!CookOnTheFlyServerConnection)
-		{
-			UE_LOG(LogODSC, Warning, TEXT("Failed to connect to cook on the fly server."));
-		}
+		ConnectToODSCHost();
 	}
 }
 
@@ -129,6 +122,20 @@ FODSCThread::~FODSCThread()
 	WakeupEvent = nullptr;
 }
 
+bool FODSCThread::ConnectToODSCHost()
+{
+	// If we don't have a default connection make a specific connection to the HostIP provided.
+	UE::Cook::FCookOnTheFlyHostOptions CookOnTheFlyHostOptions;
+	CookOnTheFlyHostOptions.Hosts.Add(ODSCHostIP);
+	CookOnTheFlyServerConnection = FModuleManager::LoadModuleChecked<UE::Cook::ICookOnTheFlyModule>(TEXT("CookOnTheFly")).ConnectToServer(CookOnTheFlyHostOptions);
+	if (!CookOnTheFlyServerConnection)
+	{
+		UE_LOG(LogODSC, Warning, TEXT("Failed to connect to cook on the fly server."));
+		return false;
+	}
+	return CookOnTheFlyServerConnection != nullptr && CookOnTheFlyServerConnection->IsConnected();
+}
+
 void FODSCThread::StartThread()
 {
 	Thread = FRunnableThread::Create(this, TEXT("ODSCThread"), 128 * 1024, TPri_Normal);
@@ -289,8 +296,19 @@ void FODSCThread::SendMessageToServer(IPlatformFile::IFileServerMessageHandler*
 	// If we have a default connection that already exists, send directly to that.
 	if ((CookOnTheFlyServerConnection == nullptr) || (!CookOnTheFlyServerConnection->IsConnected()))
 	{
-		IFileManager::Get().SendMessageToServer(TEXT("RecompileShaders"), Handler);
-		return;
+		if (bHasDefaultConnection)
+		{
+			IFileManager::Get().SendMessageToServer(TEXT("RecompileShaders"), Handler);
+			return;
+		}
+		else
+		{
+			UE_LOG(LogODSC, Display, TEXT("Detected that CookOnTheFlyServerConnection has been lost, trying again"));
+			if (!ConnectToODSCHost())
+			{
+				return;
+			}
+		}
 	}
 
 	// We don't have a default COTF connection so use our specific connection to send our command.
diff --git a/Engine/Source/Runtime/Engine/Private/ODSC/ODSCThread.h b/Engine/Source/Runtime/Engine/Private/ODSC/ODSCThread.h
--- a/Engine/Source/Runtime/Engine/Private/ODSC/ODSCThread.h
+++ b/Engine/Source/Runtime/Engine/Private/ODSC/ODSCThread.h
@@ -170,6 +170,8 @@ private:
 	 */
 	void Process();
 
+	bool ConnectToODSCHost();
+
 	/**
 	 * Threaded requests that are waiting to be processed on the ODSC thread.
 	 * Added to on (any) non-ODSC thread, processed then cleared on ODSC thread.
@@ -205,5 +207,8 @@ private:
 	/** Special connection to the cooking server.  This is only used to send recompileshaders commands on. */
 	TUniquePtr<UE::Cook::ICookOnTheFlyServerConnection> CookOnTheFlyServerConnection;
 
+	FString ODSCHostIP;
+
 	bool bHasPendingRequests = false;
+	bool bHasDefaultConnection = false;
 };
