Added SurfaceCoverage for thin translucent shading model.
Works with shading model and with material attributes.

diff --git a/Engine/Shaders/Private/PathTracing/PathTracingMaterialHitShader.usf b/Engine/Shaders/Private/PathTracing/PathTracingMaterialHitShader.usf
--- a/Engine/Shaders/Private/PathTracing/PathTracingMaterialHitShader.usf
+++ b/Engine/Shaders/Private/PathTracing/PathTracingMaterialHitShader.usf
@@ -923,6 +923,10 @@ RAY_TRACING_ENTRY_CLOSEST_HIT(PathTracingMaterialCHS,
 	{
 		const float3 TransmittanceColor = GetThinTranslucentMaterialOutput0(MaterialParameters);
 		Payload.SetTransmittanceColor(TransmittanceColor);
+
+		const float SurfaceCoverage = GetThinTranslucentMaterialOutput1(MaterialParameters);
+		Payload.BSDFOpacity *= SurfaceCoverage;
+
 		Payload.Ior = Ior;
 		if (Ior > 0.0)
 		{
@@ -940,6 +944,9 @@ RAY_TRACING_ENTRY_CLOSEST_HIT(PathTracingMaterialCHS,
 			const float  Specular = Payload.Specular;
 			const float  F0 = F0RGBToF0(ComputeF0(Specular, BaseColor, Metallic));
 			Payload.TransparencyColor = (1 - Payload.BSDFOpacity) * ComputeThinSlabWeights(TransmittanceColor, VoN, 0.0, F0).Transmitted;
+
+			// Fade out the TransparencyColor as a function of the SurfaceCoverage.
+			Payload.TransparencyColor = lerp(1.0.xxx, Payload.TransparencyColor, SurfaceCoverage.xxx);
 		}
 	}
 	else
@@ -1345,6 +1352,10 @@ RAY_TRACING_ENTRY_ANY_HIT(PathTracingMaterialAHS,
 				}
 #endif
 			}
+
+			// Fade out the Transparency as a function of the SurfaceCoverage.
+			const float SurfaceCoverage = GetThinTranslucentMaterialOutput1(MaterialParameters);
+			Transparency = lerp(1.0.xxx, Transparency, SurfaceCoverage.xxx);
 		}
 		else
 #endif // MATERIAL_SHADINGMODEL_THIN_TRANSLUCENT
diff --git a/Engine/Shaders/Private/Substrate/SubstrateLegacyConversion.ush b/Engine/Shaders/Private/Substrate/SubstrateLegacyConversion.ush
--- a/Engine/Shaders/Private/Substrate/SubstrateLegacyConversion.ush
+++ b/Engine/Shaders/Private/Substrate/SubstrateLegacyConversion.ush
@@ -141,7 +141,8 @@ FSubstrateData SubstrateConvertLegacyMaterialDynamic(
 	float ClearCoat, float ClearCoatRoughness,
 	float3 Emissive,
 	float Opacity,
-	float3 TransmittanceColor,
+	float3 ThinTranslucentTransmittanceColor,
+	float ThinTranslucentSurfaceCoverage,
 	float3 WaterScatteringCoefficients, float3 WaterAbsorptionCoefficients, float WaterPhaseG, float3 ColorScaleBehindWater,
 	uint ShadingModel,
 	float3 RawNormal,
@@ -399,7 +400,7 @@ FSubstrateData SubstrateConvertLegacyMaterialDynamic(
 		// we need to do some approximation. In order to reduce avoid hue shift due to difference between 'solid' base-color/diffuse-albedo and 'transparent' transmission 
 		// color, we lerp BaseColor with transmittance color in LAB space so that the interpolated color minized its 'distance' 
 		const float3 BaseColor_LAB			= LinearRGB_2_LAB(BaseColor);
-		const float3 TransmittanceColor_LAB = LinearRGB_2_LAB(TransmittanceColor);
+		const float3 TransmittanceColor_LAB = LinearRGB_2_LAB(ThinTranslucentTransmittanceColor);
 		const float3 Color_LAB			    = lerp(TransmittanceColor_LAB, BaseColor_LAB, Opacity);
 		const float3 Color					= LAB_2_LinearRGB(Color_LAB);
 
@@ -420,17 +421,18 @@ FSubstrateData SubstrateConvertLegacyMaterialDynamic(
 		LegacySlab.SSSMFP			= lerp(BotMFP,    TopMFP,    Opacity);
 		LegacySlab.F0				= lerp(BotF0,     TopF0,     Opacity);
 		LegacySlab.F90				= Ones;
-		LegacySlab.UseMetalness	= false;
+		LegacySlab.UseMetalness		= false;
 		LegacySlab.Roughness 		= Roughness;
 		LegacySlab.Anisotropy 		= 0.f;
 		LegacySlab.SSSProfileID 	= 0.f;
-		LegacySlab.SSSMFPScale 	= 1.f;
+		LegacySlab.SSSMFPScale 		= 1.f;
 		LegacySlab.Emissive 		= Emissive;
 		LegacySlab.Thickness 		= Thickness;
 		LegacySlab.SharedLocalBasisIndex = SharedLocalBasisIndex;
 
-		// Thin translucency model always have a coverage of 1. It also canont be used in a multiple shading model context.
-		FinalWeight = 1.0f;
+		// Opacity on the root node is the coverage of the disney model on top of the translucent part.
+		// Thin translucency model have a special surface coverage input.
+		FinalWeight = ThinTranslucentSurfaceCoverage;
 	}
 
 	// Fixed layer structure for helping compiler to unroll and optimize shader
@@ -454,7 +456,8 @@ FSubstrateData SubstrateConvertLegacyMaterialStatic(
 	float ClearCoat, float ClearCoatRoughness,
 	float3 Emissive,
 	float Opacity,
-	float3 TransmittanceColor,
+	float3 ThinTranslucentTransmittanceColor,
+	float ThinTranslucentSurfaceCoverage,
 	float3 WaterScatteringCoefficients, float3 WaterAbsorptionCoefficients, float WaterPhaseG, float3 ColorScaleBehindWater,
 	uint ShadingModel,
 	float3 RawNormal,
@@ -569,7 +572,8 @@ FSubstrateData SubstrateConvertLegacyMaterialStatic(
 			ClearCoat, ClearCoatRoughness,
 			Emissive,
 			Opacity,
-			TransmittanceColor,
+			ThinTranslucentTransmittanceColor,
+			ThinTranslucentSurfaceCoverage,
 			WaterScatteringCoefficients, WaterAbsorptionCoefficients, WaterPhaseG, ColorScaleBehindWater,
 			ShadingModel,
 			RawNormal,
diff --git a/Engine/Shaders/Private/ThinTranslucentCommon.ush b/Engine/Shaders/Private/ThinTranslucentCommon.ush
--- a/Engine/Shaders/Private/ThinTranslucentCommon.ush
+++ b/Engine/Shaders/Private/ThinTranslucentCommon.ush
@@ -28,7 +28,8 @@ void AccumulateThinTranslucentModel(inout float3 DualBlendSurfaceLuminancePostCo
 	// how much to add for foreground color
 	float3 SurfaceColor = float3(0.0f,0.0f,0.0f);
 
-	float3 TransmittanceColor = GetThinTranslucentMaterialOutput0(MaterialParams);
+	const float3 TransmittanceColor = GetThinTranslucentMaterialOutput0(MaterialParams);
+	const float  SurfaceCoverage    = GetThinTranslucentMaterialOutput1(MaterialParams);
 	
 	// color is for normalized thickness
 #if THIN_TRANSLUCENT_USE_DOTNV_THICKNESS
@@ -39,7 +40,7 @@ void AccumulateThinTranslucentModel(inout float3 DualBlendSurfaceLuminancePostCo
 	float PathLength = 1.0f;
 #endif	
 	float3 NegativeAbsorptionCoefficient = log(TransmittanceColor);
-	float3 ColorTopMaterialCoverage = exp(NegativeAbsorptionCoefficient * PathLength);
+	float3 BottomMaterialTransmittance = exp(NegativeAbsorptionCoefficient * PathLength);
 
 	// Light goes from background -> solid surface -> camera, and we need fresnel at both interactions.
 	const float3 FresnelRatio = F_Schlick(GBuffer.SpecularColor, NoV);
@@ -48,7 +49,7 @@ void AccumulateThinTranslucentModel(inout float3 DualBlendSurfaceLuminancePostCo
 	Transmittance = Transmittance * (1-FresnelRatio);
 
 	// Light gets lost from absorption through the surface
-	Transmittance = Transmittance * ColorTopMaterialCoverage;
+	Transmittance = Transmittance * BottomMaterialTransmittance;
 
 	// Exiting the surface, the Fresnel ratio is the same! This is technically only true for F_Fresnel, but since we are using
 	// F_Schlick as a stand-in, applying the same logic is a valid approximation
@@ -68,7 +69,7 @@ void AccumulateThinTranslucentModel(inout float3 DualBlendSurfaceLuminancePostCo
 	SurfaceColor += SpecularColor;
 
 	// Luminance and transmitance assumin a full coverage of 1.
-	DualBlendSurfaceCoverage = 1.0f;
+	DualBlendSurfaceCoverage = SurfaceCoverage;
 	DualBlendSurfaceLuminancePostCoverage = SurfaceColor * DualBlendSurfaceCoverage;
 	DualBlendSurfaceTransmittancePreCoverage = Transmittance;
 }
diff --git a/Engine/Source/Runtime/Engine/Classes/Materials/MaterialExpressionSubstrate.h b/Engine/Source/Runtime/Engine/Classes/Materials/MaterialExpressionSubstrate.h
--- a/Engine/Source/Runtime/Engine/Classes/Materials/MaterialExpressionSubstrate.h
+++ b/Engine/Source/Runtime/Engine/Classes/Materials/MaterialExpressionSubstrate.h
@@ -166,14 +166,21 @@ class UMaterialExpressionSubstrateShadingModels : public UMaterialExpressionSubs
 	 */
 	UPROPERTY()
 	FShadingModelMaterialInput ShadingModel;
+
+	/**
+	 * The coverage of the surface using a thin translucent shading model. This will reduce the visibility of the thin translucent surface & plastic/metal BRDF overall. (type = float1, unit = unitless, dafaults to 1, range is [0,1])
+	 */
+	UPROPERTY()
+	FExpressionInput ThinTranslucentSurfaceCoverage;
+
+	// Always show at the bottom of the pin list
+	UPROPERTY(EditAnywhere, Category = ShadingModel, meta = (ShowAsInputPin = "Primary", DisplayName = "Single Shading Model"))
+	TEnumAsByte<enum EMaterialShadingModel> ShadingModelOverride = MSM_DefaultLit;
 	
 	/** SubsurfaceProfile, for Screen Space Subsurface Scattering. The profile needs to be set up on both the Substrate diffuse node, and the material node at the moment. */
 	UPROPERTY(EditAnywhere, BlueprintReadOnly, Category = Material, meta = (DisplayName = "Subsurface Profile"))
 	TObjectPtr<class USubsurfaceProfile> SubsurfaceProfile;
 
-	UPROPERTY(EditAnywhere, Category = ShadingModel, meta = (ShowAsInputPin = "Primary", DisplayName = "Single Shading Model"))
-	TEnumAsByte<enum EMaterialShadingModel> ShadingModelOverride = MSM_DefaultLit;
-
 	//~ Begin UMaterialExpression Interface
 #if WITH_EDITOR
 	virtual int32 Compile(class FMaterialCompiler* Compiler, int32 OutputIndex) override;
diff --git a/Engine/Source/Runtime/Engine/Classes/Materials/MaterialExpressionThinTranslucentMaterialOutput.h b/Engine/Source/Runtime/Engine/Classes/Materials/MaterialExpressionThinTranslucentMaterialOutput.h
--- a/Engine/Source/Runtime/Engine/Classes/Materials/MaterialExpressionThinTranslucentMaterialOutput.h
+++ b/Engine/Source/Runtime/Engine/Classes/Materials/MaterialExpressionThinTranslucentMaterialOutput.h
@@ -17,6 +17,10 @@ class UMaterialExpressionThinTranslucentMaterialOutput : public UMaterialExpress
 	UPROPERTY()
 	FExpressionInput TransmittanceColor;
 
+	/** Input for the surface coverage of both the thin surface part and the material on top (controled using Opacity input of the graph root node). Valid range is [0,1]. */
+	UPROPERTY()
+	FExpressionInput SurfaceCoverage;
+
 public:
 #if WITH_EDITOR
 	//~ Begin UMaterialExpression Interface
diff --git a/Engine/Source/Runtime/Engine/Private/Materials/HLSLMaterialTranslator.cpp b/Engine/Source/Runtime/Engine/Private/Materials/HLSLMaterialTranslator.cpp
--- a/Engine/Source/Runtime/Engine/Private/Materials/HLSLMaterialTranslator.cpp
+++ b/Engine/Source/Runtime/Engine/Private/Materials/HLSLMaterialTranslator.cpp
@@ -12844,7 +12844,8 @@ int32 FHLSLMaterialTranslator::SubstrateConversionFromLegacy(
 	int32 ClearCoat, int32 ClearCoatRoughness,
 	int32 EmissiveColor,
 	int32 Opacity,
-	int32 TransmittanceColor,
+	int32 ThinTranslucentTransmittanceColor,
+	int32 ThinTranslucentSurfaceCoverage,
 	int32 WaterScatteringCoefficients, int32 WaterAbsorptionCoefficients, int32 WaterPhaseG, int32 ColorScaleBehindWater,
 	int32 InShadingModel,
 	int32 Normal, int32 Tangent, const FString& SharedLocalBasisIndexMacro,
@@ -12886,31 +12887,32 @@ int32 FHLSLMaterialTranslator::SubstrateConversionFromLegacy(
 		}
 
 		return AddCodeChunk(
-			MCT_Substrate, TEXT("Parameters.%s.PromoteParameterBlendedBSDFToOperator(SubstrateConvertLegacyMaterial%s(Parameters.SubstratePixelFootprint, %s, %s, %s, %s, %s, %s, %s, %s, %s, %s, %s, %s, %s, %s, %s, %s, %s, %s, %s, %s, %s, %s, %s, Parameters.%s.Types) /* Normal = %s ; Tangent = %s ; ClearCoat_Normal = %s ; ClearCoat_Tangent = %s */, %u, %u, %u, %u)"),
+			MCT_Substrate, TEXT("Parameters.%s.PromoteParameterBlendedBSDFToOperator(SubstrateConvertLegacyMaterial%s(Parameters.SubstratePixelFootprint, %s, %s, %s, %s, %s, %s, %s, %s, %s, %s, %s, %s, %s, %s, %s, %s, %s, %s, %s, %s, %s, %s, %s, %s, Parameters.%s.Types) /* Normal = %s ; Tangent = %s ; ClearCoat_Normal = %s ; ClearCoat_Tangent = %s */, %u, %u, %u, %u)"),
 			*GetParametersSubstrateTreeName(CurrentSubstrateCompilationContext),
 			bHasDynamicShadingModels ? TEXT("Dynamic") : TEXT("Static"),
-			*SubstrateGetCastParameterCode(BaseColor,						MCT_Float3),
-			*SubstrateGetCastParameterCode(Specular,						MCT_Float),
-			*SubstrateGetCastParameterCode(Metallic,						MCT_Float),
-			*SubstrateGetCastParameterCode(Roughness,						MCT_Float),
-			*SubstrateGetCastParameterCode(Anisotropy,						MCT_Float),
-			*SubstrateGetCastParameterCode(SubSurfaceColor,				MCT_Float3),
-			*SubstrateGetCastParameterCode(SubSurfaceProfileId,			MCT_Float),
-			*SubstrateGetCastParameterCode(ClearCoat,						MCT_Float),
-			*SubstrateGetCastParameterCode(ClearCoatRoughness,				MCT_Float),
-			*SubstrateGetCastParameterCode(EmissiveColor,					MCT_Float3),
-			*SubstrateGetCastParameterCode(Opacity,						MCT_Float),
-			*SubstrateGetCastParameterCode(TransmittanceColor,				MCT_Float3),
-			*SubstrateGetCastParameterCode(WaterScatteringCoefficients,	MCT_Float3),
-			*SubstrateGetCastParameterCode(WaterAbsorptionCoefficients,	MCT_Float3),
-			*SubstrateGetCastParameterCode(WaterPhaseG,					MCT_Float),
-			*SubstrateGetCastParameterCode(ColorScaleBehindWater,			MCT_Float3),
+			*SubstrateGetCastParameterCode(BaseColor,							MCT_Float3),
+			*SubstrateGetCastParameterCode(Specular,							MCT_Float),
+			*SubstrateGetCastParameterCode(Metallic,							MCT_Float),
+			*SubstrateGetCastParameterCode(Roughness,							MCT_Float),
+			*SubstrateGetCastParameterCode(Anisotropy,							MCT_Float),
+			*SubstrateGetCastParameterCode(SubSurfaceColor,						MCT_Float3),
+			*SubstrateGetCastParameterCode(SubSurfaceProfileId,					MCT_Float),
+			*SubstrateGetCastParameterCode(ClearCoat,							MCT_Float),
+			*SubstrateGetCastParameterCode(ClearCoatRoughness,					MCT_Float),
+			*SubstrateGetCastParameterCode(EmissiveColor,						MCT_Float3),
+			*SubstrateGetCastParameterCode(Opacity,								MCT_Float),
+			*SubstrateGetCastParameterCode(ThinTranslucentTransmittanceColor,	MCT_Float3),
+			*SubstrateGetCastParameterCode(ThinTranslucentSurfaceCoverage,		MCT_Float1),
+			*SubstrateGetCastParameterCode(WaterScatteringCoefficients,			MCT_Float3),
+			*SubstrateGetCastParameterCode(WaterAbsorptionCoefficients,			MCT_Float3),
+			*SubstrateGetCastParameterCode(WaterPhaseG,							MCT_Float),
+			*SubstrateGetCastParameterCode(ColorScaleBehindWater,				MCT_Float3),
 			*GetParameterCode(InShadingModel),
 			// Raw access to Normal/Tangent/ClearCoatNormal/CustomTangent for conversion purpose
-			*SubstrateGetCastParameterCode(Normal,							MCT_Float3),
-			*SubstrateGetCastParameterCode(RawTangent,						MCT_Float3),
-			*SubstrateGetCastParameterCode(ClearCoat_Normal,				MCT_Float3),
-			*SubstrateGetCastParameterCode(CustomTangent_Tangent,			MCT_Float3),
+			*SubstrateGetCastParameterCode(Normal,								MCT_Float3),
+			*SubstrateGetCastParameterCode(RawTangent,							MCT_Float3),
+			*SubstrateGetCastParameterCode(ClearCoat_Normal,					MCT_Float3),
+			*SubstrateGetCastParameterCode(CustomTangent_Tangent,				MCT_Float3),
 			*SharedLocalBasisIndexMacro,
 			*ClearCoat_SharedLocalBasisIndexMacro,
 			*GetParametersSharedLocalBasesName(CurrentSubstrateCompilationContext),
@@ -12928,30 +12930,31 @@ int32 FHLSLMaterialTranslator::SubstrateConversionFromLegacy(
 	}
 
 	return AddCodeChunk(
-		MCT_Substrate, TEXT("SubstrateConvertLegacyMaterial%s(Parameters.SubstratePixelFootprint, %s, %s, %s, %s, %s, %s, %s, %s, %s, %s, %s, %s, %s, %s, %s, %s, %s, %s, %s, %s, %s, %s, %s, Parameters.%s.Types) /* Normal = %s ; Tangent = %s ; ClearCoat_Normal = %s ; ClearCoat_Tangent = %s */"),
+		MCT_Substrate, TEXT("SubstrateConvertLegacyMaterial%s(Parameters.SubstratePixelFootprint, %s, %s, %s, %s, %s, %s, %s, %s, %s, %s, %s, %s, %s, %s, %s, %s, %s, %s, %s, %s, %s, %s, %s, %s, Parameters.%s.Types) /* Normal = %s ; Tangent = %s ; ClearCoat_Normal = %s ; ClearCoat_Tangent = %s */"),
 		bHasDynamicShadingModels ? TEXT("Dynamic") : TEXT("Static"),
-		*SubstrateGetCastParameterCode(BaseColor,						MCT_Float3),
-		*SubstrateGetCastParameterCode(Specular,						MCT_Float),
-		*SubstrateGetCastParameterCode(Metallic,						MCT_Float),
-		*SubstrateGetCastParameterCode(Roughness,						MCT_Float),
-		*SubstrateGetCastParameterCode(Anisotropy,						MCT_Float),
-		*SubstrateGetCastParameterCode(SubSurfaceColor,				MCT_Float3),
-		*SubstrateGetCastParameterCode(SubSurfaceProfileId,			MCT_Float),
-		*SubstrateGetCastParameterCode(ClearCoat,						MCT_Float),
-		*SubstrateGetCastParameterCode(ClearCoatRoughness,				MCT_Float),
-		*SubstrateGetCastParameterCode(EmissiveColor,					MCT_Float3),
-		*SubstrateGetCastParameterCode(Opacity,						MCT_Float),
-		*SubstrateGetCastParameterCode(TransmittanceColor,				MCT_Float3),
-		*SubstrateGetCastParameterCode(WaterScatteringCoefficients,	MCT_Float3),
-		*SubstrateGetCastParameterCode(WaterAbsorptionCoefficients,	MCT_Float3),
-		*SubstrateGetCastParameterCode(WaterPhaseG,					MCT_Float),
-		*SubstrateGetCastParameterCode(ColorScaleBehindWater,			MCT_Float3),
+		*SubstrateGetCastParameterCode(BaseColor,							MCT_Float3),
+		*SubstrateGetCastParameterCode(Specular,							MCT_Float),
+		*SubstrateGetCastParameterCode(Metallic,							MCT_Float),
+		*SubstrateGetCastParameterCode(Roughness,							MCT_Float),
+		*SubstrateGetCastParameterCode(Anisotropy,							MCT_Float),
+		*SubstrateGetCastParameterCode(SubSurfaceColor,						MCT_Float3),
+		*SubstrateGetCastParameterCode(SubSurfaceProfileId,					MCT_Float),
+		*SubstrateGetCastParameterCode(ClearCoat,							MCT_Float),
+		*SubstrateGetCastParameterCode(ClearCoatRoughness,					MCT_Float),
+		*SubstrateGetCastParameterCode(EmissiveColor,						MCT_Float3),
+		*SubstrateGetCastParameterCode(Opacity,								MCT_Float),
+		*SubstrateGetCastParameterCode(ThinTranslucentTransmittanceColor,	MCT_Float3),
+		*SubstrateGetCastParameterCode(ThinTranslucentSurfaceCoverage,		MCT_Float1),
+		*SubstrateGetCastParameterCode(WaterScatteringCoefficients,			MCT_Float3),
+		*SubstrateGetCastParameterCode(WaterAbsorptionCoefficients,			MCT_Float3),
+		*SubstrateGetCastParameterCode(WaterPhaseG,							MCT_Float),
+		*SubstrateGetCastParameterCode(ColorScaleBehindWater,				MCT_Float3),
 		*GetParameterCode(InShadingModel),
 		// Raw access to Normal/Tangent/ClearCoatNormal/CustomTangent for conversion purpose
-		*SubstrateGetCastParameterCode(Normal,							MCT_Float3),
-		*SubstrateGetCastParameterCode(RawTangent,						MCT_Float3),
-		*SubstrateGetCastParameterCode(ClearCoat_Normal,				MCT_Float3),
-		*SubstrateGetCastParameterCode(CustomTangent_Tangent,			MCT_Float3),
+		*SubstrateGetCastParameterCode(Normal,								MCT_Float3),
+		*SubstrateGetCastParameterCode(RawTangent,							MCT_Float3),
+		*SubstrateGetCastParameterCode(ClearCoat_Normal,					MCT_Float3),
+		*SubstrateGetCastParameterCode(CustomTangent_Tangent,				MCT_Float3),
 		*SharedLocalBasisIndexMacro,
 		*ClearCoat_SharedLocalBasisIndexMacro,
 		*GetParametersSharedLocalBasesName(CurrentSubstrateCompilationContext),
diff --git a/Engine/Source/Runtime/Engine/Private/Materials/HLSLMaterialTranslator.h b/Engine/Source/Runtime/Engine/Private/Materials/HLSLMaterialTranslator.h
--- a/Engine/Source/Runtime/Engine/Private/Materials/HLSLMaterialTranslator.h
+++ b/Engine/Source/Runtime/Engine/Private/Materials/HLSLMaterialTranslator.h
@@ -1251,7 +1251,8 @@ protected:
 		int32 ClearCoat, int32 ClearCoatRoughness,
 		int32 EmissiveColor,
 		int32 Opacity,
-		int32 TransmittanceColor,
+		int32 ThinTranslucentTransmittanceColor,
+		int32 ThinTranslucentSurfaceCoverage,
 		int32 WaterScatteringCoefficients, int32 WaterAbsorptionCoefficients, int32 WaterPhaseG, int32 ColorScaleBehindWater,
 		int32 ShadingModel,
 		int32 Normal, int32 Tangent, const FString& SharedLocalBasisIndexMacro,
diff --git a/Engine/Source/Runtime/Engine/Private/Materials/Material.cpp b/Engine/Source/Runtime/Engine/Private/Materials/Material.cpp
--- a/Engine/Source/Runtime/Engine/Private/Materials/Material.cpp
+++ b/Engine/Source/Runtime/Engine/Private/Materials/Material.cpp
@@ -3673,7 +3673,16 @@ void UMaterial::ConvertMaterialToSubstrateMaterial()
 
 		if (ThinTranslucentOutput)
 		{
-			ConnectionTo(*ThinTranslucentOutput->GetInput(0), ConvertNode, 12);	 // TransmittanceColor
+			FExpressionInput* TransmittanceColorInput = ThinTranslucentOutput->GetInput(0);
+			FExpressionInput* ThinTranslucentSurfaceCoverageInput = ThinTranslucentOutput->GetInput(1);
+			if (TransmittanceColorInput && TransmittanceColorInput->IsConnected())
+			{
+				ConvertNode->TransmittanceColor.Connect(TransmittanceColorInput->OutputIndex, TransmittanceColorInput->Expression);
+			}
+			if (ThinTranslucentSurfaceCoverageInput && ThinTranslucentSurfaceCoverageInput->IsConnected())
+			{
+				ConvertNode->ThinTranslucentSurfaceCoverage.Connect(ThinTranslucentSurfaceCoverageInput->OutputIndex, ThinTranslucentSurfaceCoverageInput->Expression);
+			}
 		}
 		if (SingleLayerWaterOutput)
 		{
diff --git a/Engine/Source/Runtime/Engine/Private/Materials/MaterialAttributeDefinitionMap.cpp b/Engine/Source/Runtime/Engine/Private/Materials/MaterialAttributeDefinitionMap.cpp
--- a/Engine/Source/Runtime/Engine/Private/Materials/MaterialAttributeDefinitionMap.cpp
+++ b/Engine/Source/Runtime/Engine/Private/Materials/MaterialAttributeDefinitionMap.cpp
@@ -300,7 +300,8 @@ void FMaterialAttributeDefinitionMap::InitializeAttributeMap()
 	AddCustomAttribute(FGuid(0xfbd7b46e, 0xb1234824, 0xbde76b23, 0x609f984c), "BentNormal", "GetBentNormal", MCT_Float3, FVector4(0, 0, 1, 0));
 	AddCustomAttribute(FGuid(0xAA3D5C04, 0x16294716, 0xBBDEC869, 0x6A27DD72), "ClearCoatBottomNormal", "ClearCoatBottomNormal", MCT_Float3, FVector4(0, 0, 1, 0));
 	AddCustomAttribute(FGuid(0x8EAB2CB2, 0x73634A24, 0x8CD14F47, 0x3F9C8E55), "CustomEyeTangent", "GetTangentOutput", MCT_Float3, FVector4(0, 0, 0, 0));
-	AddCustomAttribute(FGuid(0xF2D8C70E, 0x42ECA0D1, 0x4652D0AD, 0xB785A065), "TransmittanceColor", "GetThinTranslucentMaterialOutput", MCT_Float3, FVector4(0.5, 0.5, 0.5, 0));
+	AddCustomAttribute(FGuid(0xF2D8C70E, 0x42ECA0D1, 0x4652D0AD, 0xB785A065), "TransmittanceColor", "GetThinTranslucentMaterialOutput0", MCT_Float3, FVector4(0.5, 0.5, 0.5, 0));
+	AddCustomAttribute(FGuid(0xAAA1D0E0, 0xE312458E, 0xF3556BAD, 0xA01277CD), "ThinTranslucentSurfaceCoverage", "GetThinTranslucentMaterialOutput1", MCT_Float, FVector4(1, 0, 0, 0));
 }
 
 void FMaterialAttributeDefinitionMap::Add(const FGuid& AttributeID, const FString& AttributeName, EMaterialProperty Property,
diff --git a/Engine/Source/Runtime/Engine/Private/Materials/MaterialExpressions.cpp b/Engine/Source/Runtime/Engine/Private/Materials/MaterialExpressions.cpp
--- a/Engine/Source/Runtime/Engine/Private/Materials/MaterialExpressions.cpp
+++ b/Engine/Source/Runtime/Engine/Private/Materials/MaterialExpressions.cpp
@@ -23719,6 +23719,10 @@ int32 UMaterialExpressionThinTranslucentMaterialOutput::Compile(class FMaterialC
 	{
 		CodeInput = TransmittanceColor.IsConnected() ? TransmittanceColor.Compile(Compiler) : Compiler->Constant3(0.5f, 0.5f, 0.5f);
 	}
+	if (OutputIndex == 1)
+	{
+		CodeInput = SurfaceCoverage.IsConnected() ? SurfaceCoverage.Compile(Compiler) : Compiler->Constant(1.0f);
+	}
 
 	return Compiler->CustomOutput(this, OutputIndex, CodeInput);
 }
@@ -23735,6 +23739,9 @@ uint32 UMaterialExpressionThinTranslucentMaterialOutput::GetInputType(int32 Inpu
 	case 0:
 		return MCT_Float3;
 		break;
+	case 1:
+		return MCT_Float1;
+		break;
 	}
 
 	check(false);
@@ -23745,7 +23752,7 @@ uint32 UMaterialExpressionThinTranslucentMaterialOutput::GetInputType(int32 Inpu
 
 int32 UMaterialExpressionThinTranslucentMaterialOutput::GetNumOutputs() const
 {
-	return 1;
+	return 2;
 }
 
 FString UMaterialExpressionThinTranslucentMaterialOutput::GetFunctionName() const
@@ -24135,7 +24142,7 @@ int32 UMaterialExpressionSubstrateShadingModels::Compile(class FMaterialCompiler
 		AnisotropyCodeChunk,
 		// SSS
 		CompileWithDefaultFloat3(Compiler, SubSurfaceColor, 1.0f, 1.0f, 1.0f),
-		SSSProfileCodeChunk != INDEX_NONE ? SSSProfileCodeChunk : Compiler->Constant(0.0f),	
+		SSSProfileCodeChunk != INDEX_NONE ? SSSProfileCodeChunk : Compiler->Constant(0.0f),
 		// Clear Coat / Custom
 		CompileWithDefaultFloat1(Compiler, ClearCoat, 1.0f),
 		CompileWithDefaultFloat1(Compiler, ClearCoatRoughness, 0.1f),
@@ -24143,6 +24150,7 @@ int32 UMaterialExpressionSubstrateShadingModels::Compile(class FMaterialCompiler
 		CompileWithDefaultFloat3(Compiler, EmissiveColor, 0.0f, 0.0f, 0.0f),
 		OpacityCodeChunk,
 		CompileWithDefaultFloat3(Compiler, TransmittanceColor, 0.5f, 0.5f, 0.5f),
+		CompileWithDefaultFloat1(Compiler, ThinTranslucentSurfaceCoverage, 1.0f),
 		// Water
 		CompileWithDefaultFloat3(Compiler, WaterScatteringCoefficients, 0.0f, 0.0f, 0.0f),
 		CompileWithDefaultFloat3(Compiler, WaterAbsorptionCoefficients, 0.0f, 0.0f, 0.0f),
@@ -24203,8 +24211,9 @@ uint32 UMaterialExpressionSubstrateShadingModels::GetInputType(int32 InputIndex)
 	else if (InputIndex == 16) return MCT_Float3; // ColorScaleBehindWater
 	else if (InputIndex == 17) return MCT_Float3; // ClearCoatNormal
 	else if (InputIndex == 18) return MCT_Float3; // CustomTangent
-	else if (InputIndex == 19) return MCT_ShadingModel; // ShadingModel
-	else if (InputIndex == 20) return MCT_ShadingModel; // ShadingModelOverride (as it uses 'ShowAsInputPin' metadata)
+	else if (InputIndex == 19) return MCT_Float1; // ThinTranslucentSurfaceCoverage
+	else if (InputIndex == 20) return MCT_ShadingModel; // ShadingModel
+	else if (InputIndex == 21) return MCT_ShadingModel; // EMaterialShadingModel with ShowAsInputPin seems to always show at the bottom
 
 	check(false);
 	return MCT_Float1;
@@ -24267,14 +24276,16 @@ FName UMaterialExpressionSubstrateShadingModels::GetInputName(int32 InputIndex)
 		return TEXT("Custom1");
 	}
 	else if (InputIndex == 11)	return TEXT("Opacity");
-	else if (InputIndex == 12)	return TEXT("TransmittanceColor");
+	else if (InputIndex == 12)	return TEXT("Thin Translucent Transmittance Color");
 	else if (InputIndex == 13)	return TEXT("Water Scattering Coefficients");
 	else if (InputIndex == 14)	return TEXT("Water Absorption Coefficients");
 	else if (InputIndex == 15)	return TEXT("Water Phase G");
 	else if (InputIndex == 16)	return TEXT("Color Scale BehindWater");
 	else if (InputIndex == 17)	return TEXT("Clear Coat Normal");
 	else if (InputIndex == 18)	return TEXT("Custom Tangent");
-	else if (InputIndex == 19)	return TEXT("Shading Model From Expression");
+	else if (InputIndex == 19)	return TEXT("Thin Translucent Surface Coverage");
+	else if (InputIndex == 20)	return TEXT("Single Shading Model");
+	else if (InputIndex == 21)	return TEXT("Shading Model From Expression");
 	return TEXT("Unknown");
 }
 
@@ -27217,9 +27228,10 @@ UMaterialExpressionSubstrateConvertMaterialAttributes::UMaterialExpressionSubstr
 #if WITH_EDITOR
 int32 UMaterialExpressionSubstrateConvertMaterialAttributes::Compile(class FMaterialCompiler* Compiler, int32 OutputIndex)
 {
-	static const FGuid ClearCoatBottomNormalGuid 	= FMaterialAttributeDefinitionMap::GetCustomAttributeID(TEXT("ClearCoatBottomNormal"));;
-	static const FGuid CustomEyeTangentGuid 		= FMaterialAttributeDefinitionMap::GetCustomAttributeID(TEXT("CustomEyeTangent"));
-	static const FGuid TransmittanceColorGuid 		= FMaterialAttributeDefinitionMap::GetCustomAttributeID(TEXT("TransmittanceColor"));
+	static const FGuid ClearCoatBottomNormalGuid 			= FMaterialAttributeDefinitionMap::GetCustomAttributeID(TEXT("ClearCoatBottomNormal"));
+	static const FGuid CustomEyeTangentGuid 				= FMaterialAttributeDefinitionMap::GetCustomAttributeID(TEXT("CustomEyeTangent"));
+	static const FGuid TransmittanceColorGuid 				= FMaterialAttributeDefinitionMap::GetCustomAttributeID(TEXT("TransmittanceColor"));
+	static const FGuid ThinTranslucentSurfaceCoverageGuid	= FMaterialAttributeDefinitionMap::GetCustomAttributeID(TEXT("ThinTranslucentSurfaceCoverage"));
 
 	if (OutputIndex != 0)
 	{
@@ -27321,6 +27333,17 @@ int32 UMaterialExpressionSubstrateConvertMaterialAttributes::Compile(class FMate
 	{
 		TransmittanceColorChunk = Compiler->Constant3(0.5f, 0.5f, 0.5f);
 	}
+	// Thin Translucent Surface Coverage
+	const bool bHasThinTranslucentSurfaceCoverage = IsCustomMaterialAttributeInputConnected(Cached, ThinTranslucentSurfaceCoverageGuid);
+	int32 ThinTranslucentSurfaceCoverageChunk = INDEX_NONE;
+	if (bHasThinTranslucentSurfaceCoverage)
+	{
+		ThinTranslucentSurfaceCoverageChunk = MaterialAttributes.CompileWithDefault(Compiler, ThinTranslucentSurfaceCoverageGuid);
+	}
+	else
+	{
+		ThinTranslucentSurfaceCoverageChunk = Compiler->Constant(1.0f);
+	}
 
 	int32 ShadingModelCodeChunk = MaterialAttributes.CompileWithDefault(Compiler, FMaterialAttributeDefinitionMap::GetID(MP_ShadingModel));
 	const bool bHasShadingModelExpression = IsMaterialAttributeInputConnected(Cached, MP_ShadingModel);
@@ -27349,6 +27372,7 @@ int32 UMaterialExpressionSubstrateConvertMaterialAttributes::Compile(class FMate
 		MaterialAttributes.CompileWithDefault(Compiler, FMaterialAttributeDefinitionMap::GetID(MP_EmissiveColor)),
 		OpacityCodeChunk,
 		TransmittanceColorChunk,
+		ThinTranslucentSurfaceCoverageChunk,
 		// Water
 		CompileWithDefaultFloat3(Compiler, WaterScatteringCoefficients, 0.0f, 0.0f, 0.0f),
 		CompileWithDefaultFloat3(Compiler, WaterAbsorptionCoefficients, 0.0f, 0.0f, 0.0f),
diff --git a/Engine/Source/Runtime/Engine/Public/MaterialCompiler.h b/Engine/Source/Runtime/Engine/Public/MaterialCompiler.h
--- a/Engine/Source/Runtime/Engine/Public/MaterialCompiler.h
+++ b/Engine/Source/Runtime/Engine/Public/MaterialCompiler.h
@@ -634,7 +634,8 @@ public:
 		int32 ClearCoat, int32 ClearCoatRoughness,
 		int32 EmissiveColor,
 		int32 Opacity,
-		int32 TransmittanceColor,
+		int32 ThinTranslucentTransmittanceColor,
+		int32 ThinTranslucentSurfaceCoverage,
 		int32 WaterScatteringCoefficients, int32 WaterAbsorptionCoefficients, int32 WaterPhaseG, int32 ColorScaleBehindWater,
 		int32 ShadingModel,
 		int32 Normal, int32 Tangent, const FString& SharedLocalBasisIndexMacro,
@@ -1289,7 +1290,8 @@ public:
 		int32 ClearCoat, int32 ClearCoatRoughness,
 		int32 EmissiveColor,
 		int32 Opacity,
-		int32 TransmittanceColor,
+		int32 ThinTranslucentTransmittanceColor,
+		int32 ThinTranslucentSurfaceCoverage,
 		int32 WaterScatteringCoefficients, int32 WaterAbsorptionCoefficients, int32 WaterPhaseG, int32 ColorScaleBehindWater,
 		int32 ShadingModel,
 		int32 Normal, int32 Tangent, const FString& SharedLocalBasisIndexMacro,
@@ -1305,7 +1307,8 @@ public:
 			ClearCoat, ClearCoatRoughness,
 			EmissiveColor,
 			Opacity,
-			TransmittanceColor,
+			ThinTranslucentTransmittanceColor,
+			ThinTranslucentSurfaceCoverage,
 			WaterScatteringCoefficients, WaterAbsorptionCoefficients, WaterPhaseG, ColorScaleBehindWater,
 			ShadingModel,
 			Normal, Tangent, SharedLocalBasisIndexMacro,
