AltSkinWeights - Get isUsingSkinWeightProfile is not Returning
#feature extend IsUsingSkinWeightProfile to also take static/default override buffers into consideration rather than just the manual SMC level override

diff --git a/Engine/Source/Runtime/Engine/Classes/Components/SkinnedMeshComponent.h b/Engine/Source/Runtime/Engine/Classes/Components/SkinnedMeshComponent.h
--- a/Engine/Source/Runtime/Engine/Classes/Components/SkinnedMeshComponent.h
+++ b/Engine/Source/Runtime/Engine/Classes/Components/SkinnedMeshComponent.h
@@ -1336,7 +1336,7 @@ public:
 
 	/** Check whether or not a Skin Weight Profile is currently set */
 	UFUNCTION(BlueprintCallable, Category = "Components|SkinnedMesh")
-	bool IsUsingSkinWeightProfile() const { return bSkinWeightProfileSet == 1;  }
+	ENGINE_API bool IsUsingSkinWeightProfile() const;
 
 	UE_DEPRECATED(4.26, "GetVertexOffsetUsage() has been deprecated. Support will be dropped in the future.")
 	UFUNCTION(BlueprintCallable, Category = "Components|SkinnedMesh")
diff --git a/Engine/Source/Runtime/Engine/Classes/Engine/SkeletalMesh.h b/Engine/Source/Runtime/Engine/Classes/Engine/SkeletalMesh.h
--- a/Engine/Source/Runtime/Engine/Classes/Engine/SkeletalMesh.h
+++ b/Engine/Source/Runtime/Engine/Classes/Engine/SkeletalMesh.h
@@ -996,7 +996,8 @@ public:
 
 	/** USkinnedAsset interface */
 	ENGINE_API virtual int32 GetPlatformMinLODIdx(const ITargetPlatform* TargetPlatform) const override;
-	ENGINE_API virtual void SetSkinWeightProfilesData(int32 LODIndex, FSkinWeightProfilesData& SkinWeightProfilesData) override;
+	ENGINE_API virtual void SetSkinWeightProfilesData(int32 LODIndex, FSkinWeightProfilesData& SkinWeightProfilesData) override;	
+	ENGINE_API virtual FSkinWeightProfilesData* GetSkinWeightProfilesData(int32 LODIndex) override;
 
 	static ENGINE_API void OnLodStrippingQualityLevelChanged(IConsoleVariable* Variable);
 
diff --git a/Engine/Source/Runtime/Engine/Classes/Engine/SkinnedAsset.h b/Engine/Source/Runtime/Engine/Classes/Engine/SkinnedAsset.h
--- a/Engine/Source/Runtime/Engine/Classes/Engine/SkinnedAsset.h
+++ b/Engine/Source/Runtime/Engine/Classes/Engine/SkinnedAsset.h
@@ -213,6 +213,7 @@ public:
 	PURE_VIRTUAL(USkinnedAsset::GetDisableBelowMinLodStripping, static const FPerPlatformBool Dummy; return Dummy;);
 
 	virtual void SetSkinWeightProfilesData(int32 LODIndex, struct FSkinWeightProfilesData& SkinWeightProfilesData) {}
+	virtual FSkinWeightProfilesData* GetSkinWeightProfilesData(int32 LODIndex) { return nullptr; }
 
 	/** Computes flags for building vertex buffers. */
 	ENGINE_API virtual ESkeletalMeshVertexFlags GetVertexBufferFlags() const;
diff --git a/Engine/Source/Runtime/Engine/Private/Components/SkinnedMeshComponent.cpp b/Engine/Source/Runtime/Engine/Private/Components/SkinnedMeshComponent.cpp
--- a/Engine/Source/Runtime/Engine/Private/Components/SkinnedMeshComponent.cpp
+++ b/Engine/Source/Runtime/Engine/Private/Components/SkinnedMeshComponent.cpp
@@ -4871,6 +4871,26 @@ void USkinnedMeshComponent::UnloadSkinWeightProfile(FName InProfileName)
 	}
 }
 
+bool USkinnedMeshComponent::IsUsingSkinWeightProfile() const
+{
+	if (bSkinWeightProfileSet)
+	{		
+		return true;
+	}
+	else
+	{
+		if (USkinnedAsset* Asset = GetSkinnedAsset())
+		{
+			if (const FSkinWeightProfilesData* ProfileData = Asset->GetSkinWeightProfilesData(GetPredictedLODLevel()))
+			{
+				return ProfileData->IsDefaultOverridden() || ProfileData->IsStaticOverridden();
+			}
+		}		
+	}
+
+	return false;
+}
+
 void USkinnedMeshComponent::UpdateSkinWeightOverrideBuffer()
 {
 	// Force a mesh update to ensure bone buffers are up to date
diff --git a/Engine/Source/Runtime/Engine/Private/SkeletalMesh.cpp b/Engine/Source/Runtime/Engine/Private/SkeletalMesh.cpp
--- a/Engine/Source/Runtime/Engine/Private/SkeletalMesh.cpp
+++ b/Engine/Source/Runtime/Engine/Private/SkeletalMesh.cpp
@@ -5827,6 +5827,18 @@ void USkeletalMesh::SetSkinWeightProfilesData(int32 LODIndex, FSkinWeightProfile
 	}
 }
 
+FSkinWeightProfilesData* USkeletalMesh::GetSkinWeightProfilesData(int32 LODIndex)
+{
+	FSkeletalMeshRenderData* RenderData = GetSkeletalMeshRenderData();
+	if (RenderData && RenderData->LODRenderData.IsValidIndex(LODIndex))
+	{
+		FSkeletalMeshLODRenderData& LODRenderData = RenderData->LODRenderData[LODIndex];
+		return &LODRenderData.SkinWeightProfilesData;
+	}
+	
+	return nullptr;
+}
+
 void USkeletalMesh::OnLodStrippingQualityLevelChanged(IConsoleVariable* Variable) {
 #if WITH_EDITOR || PLATFORM_DESKTOP
 	if (GEngine && GEngine->UseSkeletalMeshMinLODPerQualityLevels)
diff --git a/Engine/Source/Runtime/Engine/Public/Animation/SkinWeightProfile.h b/Engine/Source/Runtime/Engine/Public/Animation/SkinWeightProfile.h
--- a/Engine/Source/Runtime/Engine/Public/Animation/SkinWeightProfile.h
+++ b/Engine/Source/Runtime/Engine/Public/Animation/SkinWeightProfile.h
@@ -127,7 +127,7 @@ struct FSkinweightReadbackData
 /** Runtime structure for keeping track of skin weight profile(s) and the associated buffer */
 struct FSkinWeightProfilesData
 {
-	FSkinWeightProfilesData() : BaseBuffer(nullptr), DefaultOverrideSkinWeightBuffer(nullptr), bDefaultOverriden(false), bStaticOverriden(false), DefaultProfileName(NAME_None) {}
+	FSkinWeightProfilesData() : BaseBuffer(nullptr), DefaultOverrideSkinWeightBuffer(nullptr), bDefaultOverridden(false), bStaticOverridden(false), DefaultProfileName(NAME_None) {}
 	ENGINE_API void Init(FSkinWeightVertexBuffer* InBaseBuffer);
 
 	ENGINE_API ~FSkinWeightProfilesData();
@@ -177,6 +177,8 @@ struct FSkinWeightProfilesData
 	ENGINE_API void ResetGPUReadback();
 	ENGINE_API void InitialiseProfileBuffer(const FName& ProfileName);
 
+	ENGINE_API bool IsDefaultOverridden() const { return bDefaultOverridden; }
+	ENGINE_API bool IsStaticOverridden() const { return bStaticOverridden; }
 protected:
 	ENGINE_API void ApplyOverrideProfile(FSkinWeightVertexBuffer* OverrideBuffer, const FName& ProfileName);
 
@@ -186,8 +188,8 @@ protected:
 	TMap<FName, FSkinWeightVertexBuffer*> ProfileNameToBuffer;
 	TMap<FName, FRuntimeSkinWeightProfileData> OverrideData;
 
-	bool bDefaultOverriden;
-	bool bStaticOverriden;
+	bool bDefaultOverridden;
+	bool bStaticOverridden;
 	FName DefaultProfileName;
 
 protected:
