Adding missing file from CL 32050495:

AltSkinWeights - Get isUsingSkinWeightProfile is not Returning
#feature extend IsUsingSkinWeightProfile to also take static/default override buffers into consideration rather than just the manual SMC level override

diff --git a/Engine/Source/Runtime/Engine/Private/Animation/SkinWeightProfile.cpp b/Engine/Source/Runtime/Engine/Private/Animation/SkinWeightProfile.cpp
--- a/Engine/Source/Runtime/Engine/Private/Animation/SkinWeightProfile.cpp
+++ b/Engine/Source/Runtime/Engine/Private/Animation/SkinWeightProfile.cpp
@@ -253,8 +253,8 @@ FSkinWeightProfilesData::~FSkinWeightProfilesData()
 	BaseBuffer = nullptr;
 	DefaultOverrideSkinWeightBuffer = nullptr;
 
-	bDefaultOverriden = false;
-	bStaticOverriden = false;
+	bDefaultOverridden = false;
+	bStaticOverridden = false;
 	DefaultProfileName = NAME_None;
 }
 
@@ -294,8 +294,8 @@ void FSkinWeightProfilesData::OverrideBaseBufferSkinWeightData(USkeletalMesh* Me
 				ProfilePtr->ApplyDefaultOverride(BaseBuffer);
 			}
 
-			bDefaultOverriden = true;
-			bStaticOverriden = true;
+			bDefaultOverridden = true;
+			bStaticOverridden = true;
 			DefaultProfileName = Profiles[DefaultProfileIndex].Name;
 		}
 	}
@@ -304,7 +304,7 @@ void FSkinWeightProfilesData::OverrideBaseBufferSkinWeightData(USkeletalMesh* Me
 
 void FSkinWeightProfilesData::SetDynamicDefaultSkinWeightProfile(USkeletalMesh* Mesh, int32 LODIndex, bool bSerialization /*= false*/)
 {
-	if (bStaticOverriden)
+	if (bStaticOverridden)
 	{
 		UE_LOG(LogSkeletalMesh, Error, TEXT("[%s] Skeletal Mesh has overridden the default Skin Weights buffer during serialization, cannot set any other skin weight profile."), *Mesh->GetName());
 		return;
@@ -337,7 +337,7 @@ void FSkinWeightProfilesData::SetDynamicDefaultSkinWeightProfile(USkeletalMesh*
 		if (DefaultProfileIndex != INDEX_NONE)
 		{
 			const bool bNoDefaultProfile = DefaultOverrideSkinWeightBuffer == nullptr;
-			const bool bDifferentDefaultProfile = bNoDefaultProfile && (!bDefaultOverriden || DefaultProfileName != Profiles[DefaultProfileIndex].Name);
+			const bool bDifferentDefaultProfile = bNoDefaultProfile && (!bDefaultOverridden || DefaultProfileName != Profiles[DefaultProfileIndex].Name);
 			if (bNoDefaultProfile || bDifferentDefaultProfile)
 			{
 				if (GetOverrideBuffer(Profiles[DefaultProfileIndex].Name) == nullptr)
@@ -364,7 +364,7 @@ void FSkinWeightProfilesData::SetDynamicDefaultSkinWeightProfile(USkeletalMesh*
 								ProfilePtr->ApplyOverrides(OverrideBuffer, BaseBufferData, BaseBuffer->GetNumVertices());
 
 								DefaultOverrideSkinWeightBuffer = OverrideBuffer;
-								bDefaultOverriden = true;
+								bDefaultOverridden = true;
 								DefaultProfileName = Profiles[DefaultProfileIndex].Name;
 								
 								const FName OwnerName(USkinnedAsset::GetLODPathName(Mesh, LODIndex));
@@ -381,7 +381,7 @@ void FSkinWeightProfilesData::SetDynamicDefaultSkinWeightProfile(USkeletalMesh*
 							if (WeakMesh.IsValid())
 							{
 								FSkinnedMeshComponentRecreateRenderStateContext RecreateState(WeakMesh.Get());
-								DataPtr->bDefaultOverriden = true;
+								DataPtr->bDefaultOverridden = true;
 								DataPtr->DefaultProfileName = ProfileName;
 								DataPtr->SetupDynamicDefaultSkinweightProfile();								
 							}
@@ -409,8 +409,8 @@ void FSkinWeightProfilesData::SetDynamicDefaultSkinWeightProfile(USkeletalMesh*
 				}
 				else
 				{
-				bDefaultOverriden = true;
-				DefaultProfileName = Profiles[DefaultProfileIndex].Name;
+					bDefaultOverridden = true;
+					DefaultProfileName = Profiles[DefaultProfileIndex].Name;
 
 					SetupDynamicDefaultSkinweightProfile();
 				}
@@ -421,13 +421,13 @@ void FSkinWeightProfilesData::SetDynamicDefaultSkinWeightProfile(USkeletalMesh*
 
 void FSkinWeightProfilesData::ClearDynamicDefaultSkinWeightProfile(USkeletalMesh* Mesh, int32 LODIndex)
 {
-	if (bStaticOverriden)
+	if (bStaticOverridden)
 	{
 		UE_LOG(LogSkeletalMesh, Error, TEXT("[%s] Skeletal Mesh has overridden the default Skin Weights buffer during serialization, cannot clear the skin weight profile."), *Mesh->GetName());		
 		return;
 	}
 
-	if (bDefaultOverriden)
+	if (bDefaultOverridden)
 	{
 		if (DefaultOverrideSkinWeightBuffer != nullptr)
 		{
@@ -438,14 +438,14 @@ void FSkinWeightProfilesData::ClearDynamicDefaultSkinWeightProfile(USkeletalMesh
 			DefaultOverrideSkinWeightBuffer = nullptr;
 		}
 
-		bDefaultOverriden = false;
+		bDefaultOverridden = false;
 		DefaultProfileName = NAME_None;		
 	}
 }
 
 void FSkinWeightProfilesData::SetupDynamicDefaultSkinweightProfile()
 {
-	if (ProfileNameToBuffer.Contains(DefaultProfileName) && bDefaultOverriden && !bStaticOverriden)
+	if (ProfileNameToBuffer.Contains(DefaultProfileName) && bDefaultOverridden && !bStaticOverridden)
 	{
 		DefaultOverrideSkinWeightBuffer = ProfileNameToBuffer.FindChecked(DefaultProfileName);
 	}
@@ -461,9 +461,9 @@ FSkinWeightVertexBuffer* FSkinWeightProfilesData::GetOverrideBuffer(const FName&
 	SCOPED_NAMED_EVENT(FSkinWeightProfilesData_GetOverrideBuffer, FColor::Red);
 
 	// In case we have overridden the default skin weight buffer we do not need to create an override buffer, if it was statically overridden we cannot load any other profile
-	if (bDefaultOverriden && (ProfileName == DefaultProfileName || bStaticOverriden))
+	if (bDefaultOverridden && (ProfileName == DefaultProfileName || bStaticOverridden))
 	{	
-		if (bStaticOverriden && ProfileName != DefaultProfileName)
+		if (bStaticOverridden && ProfileName != DefaultProfileName)
 		{
 			UE_LOG(LogSkeletalMesh, Error, TEXT("Skeletal Mesh has overridden the default Skin Weights buffer during serialization, cannot set any other skin weight profile."));
 		}	
@@ -504,7 +504,7 @@ FRuntimeSkinWeightProfileData& FSkinWeightProfilesData::AddOverrideData(const FN
 
 void FSkinWeightProfilesData::ReleaseBuffer(const FName& ProfileName, bool bForceRelease /*= false*/)
 {
-	if (ProfileNameToBuffer.Contains(ProfileName) && (!bDefaultOverriden || ProfileName != DefaultProfileName || bForceRelease))
+	if (ProfileNameToBuffer.Contains(ProfileName) && (!bDefaultOverridden || ProfileName != DefaultProfileName || bForceRelease))
 	{
 		FSkinWeightVertexBuffer* Buffer = nullptr;
 		ProfileNameToBuffer.RemoveAndCopyValue(ProfileName, Buffer);
@@ -529,7 +529,7 @@ void FSkinWeightProfilesData::ReleaseResources()
 	ProfileNameToBuffer.Empty();
 
 	// Never release a default _dynamic_ buffer
-	if (bDefaultOverriden && !bStaticOverriden)
+	if (bDefaultOverridden && !bStaticOverridden)
 	{
 		ensure(DefaultOverrideSkinWeightBuffer != nullptr);
 		Buffers.Remove(DefaultOverrideSkinWeightBuffer);
