Add logging to show the memory usage estimate when compressing an animation

diff --git a/Engine/Source/Runtime/Engine/Private/Animation/AnimationCompressionDerivedData.cpp b/Engine/Source/Runtime/Engine/Private/Animation/AnimationCompressionDerivedData.cpp
--- a/Engine/Source/Runtime/Engine/Private/Animation/AnimationCompressionDerivedData.cpp
+++ b/Engine/Source/Runtime/Engine/Private/Animation/AnimationCompressionDerivedData.cpp
@@ -95,43 +95,8 @@ void FAnimationSequenceAsyncCacheTask::BeginCache(const FIoHash& KeyHash)
 	{
 		FQueuedThreadPool* ThreadPool = Anim::FAnimSequenceCompilingManager::Get().GetThreadPool();
 		const EQueuedWorkPriority BasePriority = Anim::FAnimSequenceCompilingManager::Get().GetBasePriority(AnimSequence);
-		
-		const int64 AdditiveAnimSize = !AnimSequence->IsValidAdditive() ? 0 : [AnimSequence]()
-		{
-			if (const UAnimSequence* RefPoseSeq = AnimSequence->RefPoseSeq)
-			{
-				return RefPoseSeq->GetApproxRawSize();
-			}
-
-			return AnimSequence->GetApproxRawSize();
-		}();
-		
-		int64 RequiredMemory = AnimSequence->GetApproxRawSize() + AdditiveAnimSize;
-		if (Compression::FAnimationCompressionMemorySummaryScope::ShouldStoreCompressionResults())
-		{
-			Compression::FAnimationCompressionMemorySummaryScope::CompressionResultSummary().GatherPreCompressionStats(AnimSequence->GetApproxRawSize(), AnimSequence->GetApproxCompressedSize());
-		}
+		const int64 RequiredMemory = GetRequiredMemoryEstimate();
 
-		if (const UAnimBoneCompressionSettings* BoneCompressionSettings = AnimSequence->BoneCompressionSettings.Get())
-		{
-			// We try out all compression codecs in parallel to find the best one, so we need to sum up the cost for every codec.
-			for (TObjectPtr<UAnimBoneCompressionCodec> Codec : BoneCompressionSettings->Codecs)
-			{
-				if (Codec)
-				{
-					const int64 PeakMemoryEstimate = Codec->EstimateCompressionMemoryUsage(*AnimSequence);
-					if (PeakMemoryEstimate < 0)
-					{
-						// Assume the worst and default to the default behavior when no estimate is given.
-						UE_LOG(LogAnimationCompression, Warning, TEXT("Got invalid memory usage estimate from codec %s for %s. This can negatively affect the time compression takes."), *Codec->GetFullName(), *AnimSequence->GetFullName());
-						RequiredMemory = -1;
-						break;
-					}
-					RequiredMemory += PeakMemoryEstimate;
-				}
-			}
-		}
-		
 		CompressionStartTime = FPlatformTime::Seconds();
 
 		check(BuildTask == nullptr);
@@ -209,7 +174,8 @@ bool FAnimationSequenceAsyncCacheTask::BuildData() const
 {	
 	// This is where we should do the compression parts
 	TRACE_CPUPROFILER_EVENT_SCOPE_TEXT(*(FString(TEXT("FAnimationSequenceAsyncCacheTask::BuildData ") + CompressibleAnimPtr->Name)));
-	UE_LOG(LogAnimationCompression, Display, TEXT("Building compressed animation data for %s"), *CompressibleAnimPtr->FullName);
+	UE_LOG(LogAnimationCompression, Display, TEXT("Building compressed animation data for %s (Required Memory Estimate: %.2f MB)"),
+		*CompressibleAnimPtr->FullName, double(GetRequiredMemoryEstimate()) / (1024.0 * 1024.0));
 
 	check(CompressibleAnimPtr.IsValid());
 	FCompressibleAnimData& DataToCompress = *CompressibleAnimPtr.Get();
@@ -298,6 +264,52 @@ void FAnimationSequenceAsyncCacheTask::LaunchCompressionTask(const UE::DerivedDa
 		});
 }
 
+int64 FAnimationSequenceAsyncCacheTask::GetRequiredMemoryEstimate() const
+{
+	int64 RequiredMemory = -1;
+
+	if (const UAnimSequence* AnimSequence = WeakAnimSequence.Get())
+	{
+		const int64 AdditiveAnimSize = !AnimSequence->IsValidAdditive() ? 0 : [AnimSequence]()
+			{
+				if (const UAnimSequence* RefPoseSeq = AnimSequence->RefPoseSeq)
+				{
+					return RefPoseSeq->GetApproxRawSize();
+				}
+
+				return AnimSequence->GetApproxRawSize();
+			}();
+
+		RequiredMemory = AnimSequence->GetApproxRawSize() + AdditiveAnimSize;
+		if (Compression::FAnimationCompressionMemorySummaryScope::ShouldStoreCompressionResults())
+		{
+			Compression::FAnimationCompressionMemorySummaryScope::CompressionResultSummary().GatherPreCompressionStats(AnimSequence->GetApproxRawSize(), AnimSequence->GetApproxCompressedSize());
+		}
+
+		if (const UAnimBoneCompressionSettings* BoneCompressionSettings = AnimSequence->BoneCompressionSettings.Get())
+		{
+			// We try out all compression codecs in parallel to find the best one, so we need to sum up the cost for every codec.
+			for (TObjectPtr<UAnimBoneCompressionCodec> Codec : BoneCompressionSettings->Codecs)
+			{
+				if (Codec)
+				{
+					const int64 PeakMemoryEstimate = Codec->EstimateCompressionMemoryUsage(*AnimSequence);
+					if (PeakMemoryEstimate < 0)
+					{
+						// Assume the worst and default to the default behavior when no estimate is given.
+						UE_LOG(LogAnimationCompression, Warning, TEXT("Got invalid memory usage estimate from codec %s for %s. This can negatively affect the time compression takes."), *Codec->GetFullName(), *AnimSequence->GetFullName());
+						RequiredMemory = -1;
+						break;
+					}
+					RequiredMemory += PeakMemoryEstimate;
+				}
+			}
+		}
+	}
+
+	return RequiredMemory;
+}
+
 int32 GSkipDDC = 0;
 static FAutoConsoleVariableRef CVarSkipDDC(
 	TEXT("a.SkipDDC"),
diff --git a/Engine/Source/Runtime/Engine/Private/Animation/AnimationCompressionDerivedData.h b/Engine/Source/Runtime/Engine/Private/Animation/AnimationCompressionDerivedData.h
--- a/Engine/Source/Runtime/Engine/Private/Animation/AnimationCompressionDerivedData.h
+++ b/Engine/Source/Runtime/Engine/Private/Animation/AnimationCompressionDerivedData.h
@@ -89,6 +89,7 @@ namespace UE::Anim
 		void EndCache(UE::DerivedData::FCacheGetValueResponse&& Response);
 		bool BuildData() const;
 		void LaunchCompressionTask(const UE::DerivedData::FSharedString& Name, const UE::DerivedData::FCacheKey& Key);
+		int64 GetRequiredMemoryEstimate() const;
 
 	private:
 		friend class FAnimationSequenceAsyncBuildWorker;
