"type IntegerAnimationAttribute was not registered previously" crash in packaged dev builds
#fix move build-in-type registration to earlier engine init point
#tests package/run TP_ThirdPerson

diff --git a/Engine/Source/Runtime/Engine/Private/Animation/AttributeTypes.cpp b/Engine/Source/Runtime/Engine/Private/Animation/AttributeTypes.cpp
--- a/Engine/Source/Runtime/Engine/Private/Animation/AttributeTypes.cpp
+++ b/Engine/Source/Runtime/Engine/Private/Animation/AttributeTypes.cpp
@@ -38,7 +38,7 @@ namespace UE
 			}
 		};
 
-		static FDelayedAutoRegisterHelper DelayedBuiltInTypesInitializationHelper(EDelayedRegisterRunPhase::ObjectSystemReady, []()
+		static FDelayedAutoRegisterHelper DelayedBuiltInTypesInitializationHelper(EDelayedRegisterRunPhase::PreObjectSystemReady, []()
 		{
 			FAttributeTypeRegistrar::RegisterBuiltInTypes();
 		});
