[InstancedStaticMesh] Fixed partial navigation update for Editor manipulations
- added partial update on component registration/unregistration
- skip updates when adding/removing instances when component is not registered
- added list of dirty areas to pending dirty element so they could be applied when udpate info is successfully processed
- added some missing updates when not using partial updates

diff --git a/Engine/Source/Runtime/Engine/Classes/AI/Navigation/NavigationDirtyElement.h b/Engine/Source/Runtime/Engine/Classes/AI/Navigation/NavigationDirtyElement.h
--- a/Engine/Source/Runtime/Engine/Classes/AI/Navigation/NavigationDirtyElement.h
+++ b/Engine/Source/Runtime/Engine/Classes/AI/Navigation/NavigationDirtyElement.h
@@ -9,6 +9,13 @@ class INavRelevantInterface;
 
 struct FNavigationDirtyElement
 {
+	/**
+	 * If not empty and the associated navigation relevant object controls the dirty areas explicitly (i.e. ShouldSkipDirtyAreaOnAddOrRemove returns true),
+	 * the list will be used to indicate the areas that need rebuilding.
+	 * Otherwise, the default behavior, NavRelevant object's bounds will be used.
+	 */
+	TArray<FBox> ExplicitAreasToDirty;
+
 	/** object owning this element */
 	FWeakObjectPtr Owner;
 	
diff --git a/Engine/Source/Runtime/Engine/Classes/Components/InstancedStaticMeshComponent.h b/Engine/Source/Runtime/Engine/Classes/Components/InstancedStaticMeshComponent.h
--- a/Engine/Source/Runtime/Engine/Classes/Components/InstancedStaticMeshComponent.h
+++ b/Engine/Source/Runtime/Engine/Classes/Components/InstancedStaticMeshComponent.h
@@ -388,6 +388,7 @@ class UInstancedStaticMeshComponent : public UStaticMeshComponent, public ISMIns
 
 	ENGINE_API virtual void PostLoad() override;
 	ENGINE_API virtual void OnRegister() override;
+	ENGINE_API virtual void OnUnregister() override;
 	
 	/** Sets to use RemoveAtSwap on instance removal. This is an optimization, but will change the resultant instance reordering. */
 	void SetRemoveSwap() { bSupportRemoveAtSwap = true; }
@@ -622,6 +623,8 @@ private:
 
 	void CalcAndCacheNavigationBounds();
 
+	void PartialNavigateUpdateForCurrentInstances();
+
 	/** Sets up new instance data to sensible defaults, creates physics counterparts if possible. */
 	ENGINE_API void SetupNewInstanceData(FInstancedStaticMeshInstanceData& InOutNewInstanceData, int32 InInstanceIndex, const FTransform& InInstanceTransform);
 
diff --git a/Engine/Source/Runtime/Engine/Private/InstancedStaticMesh.cpp b/Engine/Source/Runtime/Engine/Private/InstancedStaticMesh.cpp
--- a/Engine/Source/Runtime/Engine/Private/InstancedStaticMesh.cpp
+++ b/Engine/Source/Runtime/Engine/Private/InstancedStaticMesh.cpp
@@ -3149,18 +3149,29 @@ int32 UInstancedStaticMeshComponent::AddInstanceInternal(int32 InstanceIndex, FI
 	}
 #endif
 
-	// If it's the first instance, register the component. 
-	// If there was no instance on component register, component registration was skipped because of UInstancedStaticMeshComponent::IsNavigationRelevant().
-	if (GetInstanceCount() == 1)
+	// Update navigation relevancy
+	bNavigationRelevant = IsNavigationRelevant();
+
+	// Perform navigation update if the component is relevant to navigation and registered.
+	// Otherwise this will be handled in OnRegister (e.g. Editor manipulations).
+	if (bNavigationRelevant && IsRegistered())
 	{
-		if (const bool bNewNavigationRelevant = IsNavigationRelevant())
+		// If it's the first instance, register the component to the navigation system
+		// since it was skipped because IsNavigationRelevant() requires at least one instance.
+		if (GetInstanceCount() == 1)
 		{
-			bNavigationRelevant = bNewNavigationRelevant;
 			FNavigationSystem::RegisterComponent(*this);
 		}
+
+		if (SupportsPartialNavigationUpdate())
+		{
+			PartialNavigationUpdate(InstanceIndex);
+		}
+		else
+		{
+			FullNavigationUpdate();
+		}
 	}
-	
-	PartialNavigationUpdate(InstanceIndex);
 
 	if (FInstancedStaticMeshDelegates::OnInstanceIndexUpdated.IsBound())
 	{
@@ -3198,14 +3209,21 @@ TArray<int32> UInstancedStaticMeshComponent::AddInstancesInternal(TConstArrayVie
 	SelectedInstances.Add(false, Count);
 #endif
 
-	PerInstanceSMData.Reserve(InstanceIndex + Count);
+	const int32 NumInstances = InstanceIndex + Count;
+	PerInstanceSMData.Reserve(NumInstances);
+
+	// Update navigation relevancy
+	// Note that we update navigation relevancy using the base class since the new instances are not added yet
+	// and IsNavigationRelevant() takes the number of instance into account which can be 0 here.
+	bNavigationRelevant = NumInstances > 0 && Super::IsNavigationRelevant();
+
+	// Navigation update required if explicitly requested and if the component is relevant to navigation and registered.
+	// Otherwise this will be handled in OnRegister (e.g. Editor manipulations).
+	const bool bNavigationUpdateRequired = bUpdateNavigation && bNavigationRelevant && IsRegistered();
+	const bool bPartialNavigationUpdateRequired = bNavigationUpdateRequired && SupportsPartialNavigationUpdate();
 
-	// Perform partial navigation update if supported and explicitly requested.
-	// Note that we also test navigation relevancy but on the base class since this class
-	// implementation also take the number of instance into account which can be 0 here.
-	const bool bDoPartialNavigationUpdate = bUpdateNavigation && Super::IsNavigationRelevant() && SupportsPartialNavigationUpdate();
 	TArray<FTransform> NavigationUpdateTransforms;
-	if (bDoPartialNavigationUpdate)
+	if (bPartialNavigationUpdateRequired)
 	{
 		NavigationUpdateTransforms.Reserve(InstanceTransforms.Num());
 	}
@@ -3223,20 +3241,19 @@ TArray<int32> UInstancedStaticMeshComponent::AddInstancesInternal(TConstArrayVie
 			NewInstanceIndices.Add(InstanceIndex);
 		}
 
-		if (bDoPartialNavigationUpdate)
+		if (bNavigationUpdateRequired)
 		{
-			// If it's the first instance, register the component. 
-			// If there was no instance on component register, component registration was skipped because of UInstancedStaticMeshComponent::IsNavigationRelevant().
+			// If it's the first instance, register the component to the navigation system
+			// since it was skipped because IsNavigationRelevant() requires at least one instance.
 			if (GetInstanceCount() == 1)
 			{
-				if (const bool bNewNavigationRelevant = IsNavigationRelevant())
-				{
-					bNavigationRelevant = bNewNavigationRelevant;
-					FNavigationSystem::RegisterComponent(*this);
-				}
+				FNavigationSystem::RegisterComponent(*this);
 			}
 
-			NavigationUpdateTransforms.Emplace(InstanceTransform);
+			if (bPartialNavigationUpdateRequired)
+			{
+				NavigationUpdateTransforms.Emplace(InstanceTransform);
+			}
 		}
 
 		if (FInstancedStaticMeshDelegates::OnInstanceIndexUpdated.IsBound())
@@ -3248,9 +3265,9 @@ TArray<int32> UInstancedStaticMeshComponent::AddInstancesInternal(TConstArrayVie
 		++InstanceIndex;
 	}
 
-	if (bUpdateNavigation)
+	if (bNavigationUpdateRequired)
 	{
-		if (bDoPartialNavigationUpdate)
+		if (bPartialNavigationUpdateRequired)
 		{
 			PartialNavigationUpdates(NavigationUpdateTransforms);
 		}
@@ -3440,9 +3457,12 @@ bool UInstancedStaticMeshComponent::RemoveInstanceInternal(int32 InstanceIndex,
 	// remove instance
 	if (!InstanceAlreadyRemoved && PerInstanceSMData.IsValidIndex(InstanceIndex))
 	{
-		const bool bWasNavRelevant = bNavigationRelevant;
+		// Navigation update required explicitly requested and if the component is relevant to navigation and registered.
+		const bool bNavigationUpdateRequired = bUpdateNavigation && bNavigationRelevant && IsRegistered();
 
-		if (bUpdateNavigation)
+		// Note that it is done before removing the instance since partial update needs
+		// the instance's transform to dirty the area and full update will dirty the whole covered by the navigation bounds.
+		if (bNavigationUpdateRequired && SupportsPartialNavigationUpdate())
 		{
 			PartialNavigationUpdate(InstanceIndex);
 		}
@@ -3458,12 +3478,19 @@ bool UInstancedStaticMeshComponent::RemoveInstanceInternal(int32 InstanceIndex,
 			PerInstanceSMCustomData.RemoveAt(InstanceIndex * NumCustomDataFloats, NumCustomDataFloats);
 		}
 
-		// If it's the last instance, unregister the component since component with no instances are not registered. 
-		// (because of GetInstanceCount() > 0 in UInstancedStaticMeshComponent::IsNavigationRelevant())
-		if (bWasNavRelevant && GetInstanceCount() == 0)
+		if (bNavigationUpdateRequired)
 		{
-			bNavigationRelevant = false;
-			FNavigationSystem::UnregisterComponent(*this);
+			// If it's the last instance, unregister the component since component with no instances are not registered.
+			// (because of GetInstanceCount() > 0 in UInstancedStaticMeshComponent::IsNavigationRelevant())
+			if (GetInstanceCount() == 0)
+			{
+				bNavigationRelevant = false;
+				FNavigationSystem::UnregisterComponent(*this);
+			}
+			else if (!SupportsPartialNavigationUpdate())
+			{
+				FullNavigationUpdate();
+			}
 		}
 	}
 	if (bHasPreviousTransforms)
@@ -3696,7 +3723,8 @@ void UInstancedStaticMeshComponent::UpdateComponentTransform(EUpdateTransformFla
 				UpdateInstanceBodyTransform(i, NewInstanceTransform, bTeleport);
 			}
 		}
-		
+
+		// Only handle partial updates since modifying the transform is handled by base class (i.e. SceneComponent).
 		if (bDoPartialNavigationUpdate)
 		{
 			PartialNavigationUpdates(NavigationUpdateTransforms);
@@ -3774,15 +3802,23 @@ bool UInstancedStaticMeshComponent::UpdateInstanceTransform(const int32 Instance
 		UpdateInstanceBodyTransform(InstanceIndex, WorldTransform, bTeleport);
 	}
 
-	if (IsNavigationRelevant() && SupportsPartialNavigationUpdate())
+	// Navigation update required if the component is relevant to navigation and registered.
+	if (bNavigationRelevant && IsRegistered())
 	{
-		// Perform partial update after instance gets updated since we need NavigationBounds using up to date instances.
-		// Append instance's previous and new transforms to dirty both areas
-		PartialNavigationUpdates(
-			{
-				PreviousTransform * GetComponentTransform(),
-				(bWorldSpace ? NewInstanceTransform : NewInstanceTransform * GetComponentTransform())
-			});
+		if (SupportsPartialNavigationUpdate())
+		{
+			// Perform partial update after instance gets updated since we need NavigationBounds using up to date instances.
+			// Append instance's previous and new transforms to dirty both areas
+			PartialNavigationUpdates(
+				{
+					PreviousTransform * GetComponentTransform(),
+					(bWorldSpace ? NewInstanceTransform : NewInstanceTransform * GetComponentTransform())
+				});
+		}
+		else
+		{
+			FullNavigationUpdate();
+		}
 	}
 
 	if (bMarkRenderStateDirty)
@@ -3809,9 +3845,12 @@ bool UInstancedStaticMeshComponent::BatchUpdateInstancesTransforms(int32 StartIn
 
 	Modify();
 
-	const bool bDoPartialNavigationUpdate = IsNavigationRelevant() && SupportsPartialNavigationUpdate();
+	// Navigation update required if the component is relevant to navigation and registered.
+	const bool bNavigationUpdateRequired = bNavigationRelevant && IsRegistered();
+	const bool bPartialNavigationUpdateRequired = bNavigationUpdateRequired && SupportsPartialNavigationUpdate();
+
 	TArray<FTransform> NavigationUpdateTransforms;
-	if (bDoPartialNavigationUpdate)
+	if (bPartialNavigationUpdateRequired)
 	{
 		NavigationUpdateTransforms.Reserve(NewInstancesTransforms.Num());
 	}
@@ -3827,7 +3866,7 @@ bool UInstancedStaticMeshComponent::BatchUpdateInstancesTransforms(int32 StartIn
 		FMatrix& PrevInstanceData = PerInstancePrevTransform[InstanceIndex];
 
 		// Append instance's previous and new transforms to dirty both areas
-		if (bDoPartialNavigationUpdate)
+		if (bPartialNavigationUpdateRequired)
 		{
 			NavigationUpdateTransforms.Append(
 			{
@@ -3855,9 +3894,16 @@ bool UInstancedStaticMeshComponent::BatchUpdateInstancesTransforms(int32 StartIn
 		}
 	}
 
-	if (bDoPartialNavigationUpdate)
+	if (bNavigationUpdateRequired)
 	{
-		PartialNavigationUpdates(NavigationUpdateTransforms);
+		if (bPartialNavigationUpdateRequired)
+		{
+			PartialNavigationUpdates(NavigationUpdateTransforms);
+		}
+		else
+		{
+			FullNavigationUpdate();
+		}
 	}
 
 	if (bMarkRenderStateDirty)
@@ -3929,6 +3975,17 @@ PRAGMA_DISABLE_DEPRECATION_WARNINGS
 	TArray<int32> OldInstanceIds(PerInstanceIds);
 	TArray<int32> AddedInstances;
 
+	// Navigation update required if the component is relevant to navigation and registered.
+	const bool bNavigationUpdateRequired = bNavigationRelevant && IsRegistered();
+	const bool bPartialNavigationUpdateRequired = bNavigationUpdateRequired && SupportsPartialNavigationUpdate();
+
+	TArray<FTransform> NavigationUpdateTransforms;
+	if (bPartialNavigationUpdateRequired)
+	{
+		// Reserve enough space for previous and new transforms
+		NavigationUpdateTransforms.Reserve(UpdateInstanceIds.Num() * 2);
+	}
+
 	// Apply updates
 	for (int32 i = 0; i < UpdateInstanceIds.Num(); ++i)
 	{
@@ -3954,6 +4011,16 @@ PRAGMA_DISABLE_DEPRECATION_WARNINGS
 
 				PrimitiveInstanceDataManager.TransformChanged(InstanceIndex);
 
+				// Append instance's previous and new transforms to dirty both areas
+				if (bPartialNavigationUpdateRequired)
+				{
+					NavigationUpdateTransforms.Append(
+						{
+							UpdateInstancePreviousTransforms[i] * GetComponentTransform(),
+							UpdateInstanceTransforms[i] * GetComponentTransform()
+						});
+				}
+
 #if (CSV_PROFILER)
 				// We are updating a current and previous transform.
 				TotalSizeUpdateBytes += (sizeof(FTransform) + sizeof(FTransform));
@@ -3990,11 +4057,23 @@ PRAGMA_DISABLE_DEPRECATION_WARNINGS
 		}
 	}
 
-	// Remove instances 
+	// Remove instances
+	if (bPartialNavigationUpdateRequired)
+	{
+		// Reserve more space if need for instances to remove
+		NavigationUpdateTransforms.Reserve(NavigationUpdateTransforms.Num() + OldInstanceIds.Num());
+	}
+
 	for (int32 InstanceIndex = 0; InstanceIndex < OldInstanceIds.Num();)
 	{
 		if (OldInstanceIds[InstanceIndex] != INDEX_NONE)
 		{
+			// Append old instances transform
+			if (bPartialNavigationUpdateRequired)
+			{
+				NavigationUpdateTransforms.Add(FTransform(PerInstanceSMData[InstanceIndex].Transform) * GetComponentTransform());
+			}
+
 			// TODO: Move this to common helper function such that all data remove goes through one place in the code.
 			PrimitiveInstanceDataManager.RemoveAtSwap(InstanceIndex);
 			PerInstanceSMData.RemoveAtSwap(InstanceIndex, EAllowShrinking::No);
@@ -4011,8 +4090,14 @@ PRAGMA_DISABLE_DEPRECATION_WARNINGS
 		}
 		else
 		{
-			 ++InstanceIndex;
+			++InstanceIndex;
+		}
 	}
+
+	if (bPartialNavigationUpdateRequired)
+	{
+		// Reserve enough space for instances to add
+		NavigationUpdateTransforms.Reserve(NavigationUpdateTransforms.Num() + AddedInstances.Num());
 	}
 
 	// Add new instances to the component's data.
@@ -4024,14 +4109,14 @@ PRAGMA_DISABLE_DEPRECATION_WARNINGS
 		PerInstancePrevTransform.Add(UpdateInstancePreviousTransforms[Index].ToMatrixWithScale());
 		PerInstanceIds.Add(UpdateInstanceIds[Index]);
 
-			if (bHasCustomFloatData)
-			{
-				const int32 SrcCustomDataOffset = Index * NumCustomDataFloats;
-				const int32 CustomDataDestIndex = PerInstanceSMCustomData.AddUninitialized(NumCustomDataFloats);
+		if (bHasCustomFloatData)
+		{
+			const int32 SrcCustomDataOffset = Index * NumCustomDataFloats;
+			const int32 CustomDataDestIndex = PerInstanceSMCustomData.AddUninitialized(NumCustomDataFloats);
 
-				FMemory::Memcpy(&PerInstanceSMCustomData[CustomDataDestIndex], &CustomFloatData[SrcCustomDataOffset], NumCustomDataFloats * sizeof(float));
-			}
+			FMemory::Memcpy(&PerInstanceSMCustomData[CustomDataDestIndex], &CustomFloatData[SrcCustomDataOffset], NumCustomDataFloats * sizeof(float));
 		}
+	}
 
 	// Rebuild the mapping from ID to InstanceIndex.
 	InstanceIdToInstanceIndexMap.Reset();
@@ -4045,7 +4130,17 @@ PRAGMA_DISABLE_DEPRECATION_WARNINGS
 	check(PerInstanceSMData.Num() == PerInstanceIds.Num());
 	check(PerInstanceSMCustomData.Num() == (NumCustomDataFloats * PerInstanceSMData.Num()));
 
-	FullNavigationUpdate();
+	if (bNavigationUpdateRequired)
+	{
+		if (bPartialNavigationUpdateRequired)
+		{
+			PartialNavigationUpdates(NavigationUpdateTransforms);
+		}
+		else
+		{
+			FullNavigationUpdate();
+		}
+	}
 
 #if (CSV_PROFILER)
 	const int32 TotalSizeBytes = (UpdateInstanceTransforms.Num() * UpdateInstanceTransforms.GetTypeSize()) +
@@ -4087,11 +4182,15 @@ bool UInstancedStaticMeshComponent::BatchUpdateInstancesTransformsInternal(int32
 
 	Modify();
 
-	const bool bDoPartialNavigationUpdate = IsNavigationRelevant() && SupportsPartialNavigationUpdate();
+	// Navigation update required if the component is relevant to navigation and registered.
+	const bool bNavigationUpdateRequired = bNavigationRelevant && IsRegistered();
+	const bool bPartialNavigationUpdateRequired = bNavigationUpdateRequired && SupportsPartialNavigationUpdate();
+
 	TArray<FTransform> NavigationUpdateTransforms;
-	if (bDoPartialNavigationUpdate)
+	if (bPartialNavigationUpdateRequired)
 	{
-		NavigationUpdateTransforms.Reserve(NewInstancesTransforms.Num());
+		// Reserve enough space for previous and new transforms
+		NavigationUpdateTransforms.Reserve(NewInstancesTransforms.Num() * 2);
 	}
 
 	int32 InstanceIndex = StartInstanceIndex;
@@ -4100,7 +4199,7 @@ bool UInstancedStaticMeshComponent::BatchUpdateInstancesTransformsInternal(int32
 		FInstancedStaticMeshInstanceData& InstanceData = PerInstanceSMData[InstanceIndex];
 
 		// Append instance's previous and new transforms to dirty both areas
-		if (bDoPartialNavigationUpdate)
+		if (bPartialNavigationUpdateRequired)
 		{
 			NavigationUpdateTransforms.Append(
 			{
@@ -4127,9 +4226,16 @@ bool UInstancedStaticMeshComponent::BatchUpdateInstancesTransformsInternal(int32
 		InstanceIndex++;
 	}
 
-	if (bDoPartialNavigationUpdate)
+	if (bNavigationUpdateRequired)
 	{
-		PartialNavigationUpdates(NavigationUpdateTransforms);
+		if (bPartialNavigationUpdateRequired)
+		{
+			PartialNavigationUpdates(NavigationUpdateTransforms);
+		}
+		else
+		{
+			FullNavigationUpdate();
+		}
 	}
 
 	if (bMarkRenderStateDirty)
@@ -4149,9 +4255,12 @@ bool UInstancedStaticMeshComponent::BatchUpdateInstancesTransform(int32 StartIns
 
 	Modify();
 
-	const bool bDoPartialNavigationUpdate = IsNavigationRelevant() && SupportsPartialNavigationUpdate();
+	// Navigation update required if the component is relevant to navigation and registered.
+	const bool bNavigationUpdateRequired = bNavigationRelevant && IsRegistered();
+	const bool bPartialNavigationUpdateRequired = bNavigationUpdateRequired && SupportsPartialNavigationUpdate();
+
 	TArray<FTransform> NavigationUpdateTransforms;
-	if (bDoPartialNavigationUpdate)
+	if (bPartialNavigationUpdateRequired)
 	{
 		NavigationUpdateTransforms.Reserve(NumInstances);
 	}
@@ -4162,7 +4271,7 @@ bool UInstancedStaticMeshComponent::BatchUpdateInstancesTransform(int32 StartIns
 		FInstancedStaticMeshInstanceData& InstanceData = PerInstanceSMData[InstanceIndex];
 
 		// Append instance's previous and new transforms to dirty both areas
-		if (bDoPartialNavigationUpdate)
+		if (bPartialNavigationUpdateRequired)
 		{
 			NavigationUpdateTransforms.Append(
 			{
@@ -4187,9 +4296,16 @@ bool UInstancedStaticMeshComponent::BatchUpdateInstancesTransform(int32 StartIns
 		}
 	}
 
-	if (bDoPartialNavigationUpdate)
+	if (bNavigationUpdateRequired)
 	{
-		PartialNavigationUpdates(NavigationUpdateTransforms);
+		if (bPartialNavigationUpdateRequired)
+		{
+			PartialNavigationUpdates(NavigationUpdateTransforms);
+		}
+		else
+		{
+			FullNavigationUpdate();
+		}
 	}
 
 	if(bMarkRenderStateDirty)
@@ -4209,9 +4325,12 @@ bool UInstancedStaticMeshComponent::BatchUpdateInstancesData(int32 StartInstance
 
 	Modify();
 
-	const bool bDoPartialNavigationUpdate = IsNavigationRelevant() && SupportsPartialNavigationUpdate();
+	// Navigation update required if the component is relevant to navigation and registered.
+	const bool bNavigationUpdateRequired = bNavigationRelevant && IsRegistered();
+	const bool bPartialNavigationUpdateRequired = bNavigationUpdateRequired && SupportsPartialNavigationUpdate();
+
 	TArray<FTransform> NavigationUpdateTransforms;
-	if (bDoPartialNavigationUpdate)
+	if (bPartialNavigationUpdateRequired)
 	{
 		NavigationUpdateTransforms.Reserve(NumInstances);
 	}
@@ -4222,7 +4341,7 @@ bool UInstancedStaticMeshComponent::BatchUpdateInstancesData(int32 StartInstance
 		FInstancedStaticMeshInstanceData& InstanceData = PerInstanceSMData[InstanceIndex];
 
 		// Append instance's previous and new transforms to dirty both areas
-		if (bDoPartialNavigationUpdate)
+		if (bPartialNavigationUpdateRequired)
 		{
 			NavigationUpdateTransforms.Append(
 			{
@@ -4242,9 +4361,16 @@ bool UInstancedStaticMeshComponent::BatchUpdateInstancesData(int32 StartInstance
 		}
 	}
 
-	if (bDoPartialNavigationUpdate)
+	if (bNavigationUpdateRequired)
 	{
-		PartialNavigationUpdates(NavigationUpdateTransforms);
+		if (bPartialNavigationUpdateRequired)
+		{
+			PartialNavigationUpdates(NavigationUpdateTransforms);
+		}
+		else
+		{
+			FullNavigationUpdate();
+		}
 	}
 
 	if (bMarkRenderStateDirty)
@@ -4380,6 +4506,15 @@ void UInstancedStaticMeshComponent::ClearInstances()
 
 	const int32 PrevNumInstances = GetInstanceCount();
 
+	// Navigation update required if the component is relevant to navigation and registered.
+	// Note that it is done before removing the instance since partial update needs
+	// the instances transform to dirty all the areas.
+	const bool bNavigationUpdateRequired = bNavigationRelevant && IsRegistered();
+	if (bNavigationUpdateRequired && SupportsPartialNavigationUpdate())
+	{
+		PartialNavigateUpdateForCurrentInstances();
+	}
+
 	// Clear all the per-instance data
 	PerInstanceSMData.Empty();
 	PerInstanceSMCustomData.Empty();
@@ -4400,7 +4535,12 @@ void UInstancedStaticMeshComponent::ClearInstances()
 		FInstancedStaticMeshDelegates::OnInstanceIndexUpdated.Broadcast(this, MakeArrayView(&IndexUpdate, 1));
 	}
 
-	FNavigationSystem::UpdateComponentData(*this);
+	if (bNavigationUpdateRequired)
+	{
+		// Unregister since components with no instances are not registered.
+		FNavigationSystem::UnregisterComponent(*this);
+		bNavigationRelevant = false;
+	}
 
 #if WITH_EDITOR
 	DeletionState = EInstanceDeletionReason::NotDeleting;
@@ -4571,8 +4711,27 @@ void UInstancedStaticMeshComponent::OnRegister()
 		// At this point we need to let the manager reset if needed, also at this point we may assume that we have no idea of the state of individual members.
 		PrimitiveInstanceDataManager.OnRegister(PerInstanceSMData.Num());
 	}
+
+	// A component using partial updates will not dirty the whole area covered by the navigation bound (base class default behavior)
+	// so it needs to dirty areas around its current instances that are getting added to the scene.
+	if (bNavigationRelevant && SupportsPartialNavigationUpdate())
+	{
+		PartialNavigateUpdateForCurrentInstances();
+	}
 }
-		
+
+void UInstancedStaticMeshComponent::OnUnregister()
+{
+	// A component using partial updates will not dirty the whole area covered by the navigation bound (base class default behavior)
+	// so it needs to dirty areas around its current instances that are getting removed from the scene.
+	if (bNavigationRelevant && SupportsPartialNavigationUpdate())
+	{
+		PartialNavigateUpdateForCurrentInstances();
+	}
+
+	Super::OnUnregister();
+}
+
 #if WITH_EDITOR
 
 bool UInstancedStaticMeshComponent::CanEditChange(const FProperty* InProperty) const
@@ -4938,6 +5097,23 @@ void UInstancedStaticMeshComponent::FullNavigationUpdate()
 	FNavigationSystem::UpdateComponentData(*this); // just update everything
 }
 
+void UInstancedStaticMeshComponent::PartialNavigateUpdateForCurrentInstances()
+{
+	if (!IsNavigationRelevant() || PerInstanceSMData.IsEmpty())
+	{
+		return;
+	}
+
+	TArray<FTransform> NavigationUpdateTransforms;
+	NavigationUpdateTransforms.Reserve(PerInstanceSMData.Num());
+	for (int32 i = 0; i < PerInstanceSMData.Num(); i++)
+	{
+		NavigationUpdateTransforms.Add(FTransform(PerInstanceSMData[i].Transform) * GetComponentTransform());
+	}
+
+	PartialNavigationUpdates(NavigationUpdateTransforms);
+}
+
 void UInstancedStaticMeshComponent::PartialNavigationUpdate(int32 InstanceIdx)
 {
 	if (!IsNavigationRelevant())
@@ -5187,7 +5363,7 @@ bool UInstancedStaticMeshComponent::IsNavigationRelevant() const
 bool UInstancedStaticMeshComponent::ShouldSkipDirtyAreaOnAddOrRemove() const
 {
 	// If partial navigation updates are supported then we don't want to dirty the
-	// whole area covered by the navigation bounds when added added to the navigation octree,
+	// whole area covered by the navigation bounds when added to the navigation octree,
 	// instead we use the partial update to push the list of dirty areas.
 	return SupportsPartialNavigationUpdate();
 }
@@ -5304,7 +5480,11 @@ void UInstancedStaticMeshComponent::PostEditChangeChainProperty(FPropertyChanged
 		}
 		else if (PropertyChangedEvent.Property->GetFName() == GET_MEMBER_NAME_CHECKED(FInstancedStaticMeshInstanceData, Transform))
 		{
-			FullNavigationUpdate();
+			// No need to update when using partial update since OnRegister will take care of it.
+			if (!SupportsPartialNavigationUpdate())
+			{
+				FullNavigationUpdate();
+			}
 
 			// Mark all instances as changed because we don't know which one actually did.
 			PrimitiveInstanceDataManager.TransformsChangedAll();
diff --git a/Engine/Source/Runtime/NavigationSystem/Private/NavigationDataHandler.cpp b/Engine/Source/Runtime/NavigationSystem/Private/NavigationDataHandler.cpp
--- a/Engine/Source/Runtime/NavigationSystem/Private/NavigationDataHandler.cpp
+++ b/Engine/Source/Runtime/NavigationSystem/Private/NavigationDataHandler.cpp
@@ -100,7 +100,8 @@ FSetElementId FNavigationDataHandler::RegisterNavOctreeElement(UObject& ElementO
 			SetId = OctreeController.PendingOctreeUpdates.FindId(UpdateInfo);
 			if (SetId.IsValidId())
 			{
-				// make sure this request stays, in case it has been invalidated already
+				// make sure this request stays, in case it has been invalidated already and keep any dirty areas
+				UpdateInfo.ExplicitAreasToDirty = OctreeController.PendingOctreeUpdates[SetId].ExplicitAreasToDirty;
 				OctreeController.PendingOctreeUpdates[SetId] = UpdateInfo;
 			}
 			else
@@ -178,9 +179,16 @@ void FNavigationDataHandler::AddElementToNavOctree(const FNavigationDirtyElement
 	}
 
 	// mark area occupied by given element as dirty except if explicitly set to skip this default behavior
-	if (!GeneratedData.IsEmpty() && !GeneratedData.Data->bShouldSkipDirtyAreaOnAddOrRemove)
+	const int32 DirtyFlag = DirtyElement.FlagsOverride ? DirtyElement.FlagsOverride : GeneratedData.Data->GetDirtyFlag();
+	if (GeneratedData.Data->bShouldSkipDirtyAreaOnAddOrRemove)
+	{
+		if (DirtyElement.ExplicitAreasToDirty.Num() > 0)
+		{
+			DirtyAreasController.AddAreas(DirtyElement.ExplicitAreasToDirty, DirtyFlag, [&ElementOwner] { return ElementOwner; }, &DirtyElement, "Addition to navoctree");
+		}
+	}
+	else if (!GeneratedData.IsEmpty())
 	{
-		const int32 DirtyFlag = DirtyElement.FlagsOverride ? DirtyElement.FlagsOverride : GeneratedData.Data->GetDirtyFlag();
 		DirtyAreasController.AddArea(GeneratedData.Bounds.GetBox(), DirtyFlag, [&ElementOwner] { return ElementOwner; }, &DirtyElement, "Addition to navoctree");
 	}
 }
@@ -384,6 +392,22 @@ bool FNavigationDataHandler::UpdateNavOctreeElementBounds(UObject& Object, const
 		return true;
 	}
 
+	// If dirty areas are provided we need to append them to a pending update since the object is not added yet.
+	// Not necessary for the bounds since they are not stored in the update but fetched when the update is processed.
+	if (DirtyAreas.Num() > 0)
+	{
+		const FSetElementId PendingElementId = OctreeController.PendingOctreeUpdates.FindId(FNavigationDirtyElement(&Object));
+		if (PendingElementId.IsValidId())
+		{
+			FNavigationDirtyElement& DirtyElement = OctreeController.PendingOctreeUpdates[PendingElementId];
+			if (!DirtyElement.bInvalidRequest)
+			{
+				DirtyElement.ExplicitAreasToDirty.Append(DirtyAreas);
+				return true;
+			}
+		}
+	}
+
 	return false;
 }
 
