we didn't update the navigation bouinds for HISMComp, which causes an invalid bounds issue when adding the element to the octree.

diff --git a/Engine/Source/Runtime/Engine/Classes/Components/InstancedStaticMeshComponent.h b/Engine/Source/Runtime/Engine/Classes/Components/InstancedStaticMeshComponent.h
--- a/Engine/Source/Runtime/Engine/Classes/Components/InstancedStaticMeshComponent.h
+++ b/Engine/Source/Runtime/Engine/Classes/Components/InstancedStaticMeshComponent.h
@@ -621,8 +621,6 @@ private:
 	ENGINE_API bool ShouldInheritPerInstanceData(const UInstancedStaticMeshComponent* InArchetype) const;
 	ENGINE_API bool ShouldInheritPerInstanceData() const;
 
-	void CalcAndCacheNavigationBounds();
-
 	void PartialNavigateUpdateForCurrentInstances();
 
 	/** Sets up new instance data to sensible defaults, creates physics counterparts if possible. */
@@ -639,6 +637,8 @@ protected:
 
 	FPrimitiveInstanceDataManager PrimitiveInstanceDataManager;
 
+	ENGINE_API void CalcAndCacheNavigationBounds();
+
 	/** Creates body instances for all instances owned by this component. */
 	ENGINE_API void CreateAllInstanceBodies();
 
diff --git a/Engine/Source/Runtime/Engine/Private/HierarchicalInstancedStaticMesh.cpp b/Engine/Source/Runtime/Engine/Private/HierarchicalInstancedStaticMesh.cpp
--- a/Engine/Source/Runtime/Engine/Private/HierarchicalInstancedStaticMesh.cpp
+++ b/Engine/Source/Runtime/Engine/Private/HierarchicalInstancedStaticMesh.cpp
@@ -3265,6 +3265,8 @@ void UHierarchicalInstancedStaticMeshComponent::PartialNavigationUpdate(const in
 			}
 		}
 	}
+
+	CalcAndCacheNavigationBounds();
 }
 
 void UHierarchicalInstancedStaticMeshComponent::FlushAccumulatedNavigationUpdates()
