[MaterialEditor] Added calls ConditionalPostLoad() to texture assets referenced by parameters in Material Function when it is loaded, and made sure that ConditionalPostLoad() is called on layers and blends from Material Instances when loaded even outside of the editor.

diff --git a/Engine/Source/Runtime/Engine/Private/Materials/MaterialExpressions.cpp b/Engine/Source/Runtime/Engine/Private/Materials/MaterialExpressions.cpp
--- a/Engine/Source/Runtime/Engine/Private/Materials/MaterialExpressions.cpp
+++ b/Engine/Source/Runtime/Engine/Private/Materials/MaterialExpressions.cpp
@@ -16076,6 +16076,53 @@ void UMaterialFunctionInstance::PostLoad()
 	{
 		Parent->ConditionalPostLoad();
 	}
+
+#if WITH_EDITORONLY_DATA
+	for (const FScalarParameterValue& Param : ScalarParameterValues)
+	{
+		if (UCurveLinearColor* Curve = Param.AtlasData.Curve.Get())
+		{
+			Curve->ConditionalPostLoad();
+		}
+
+		if (UCurveLinearColorAtlas* Atlas = Param.AtlasData.Atlas.Get())
+		{
+			Atlas->ConditionalPostLoad();
+		}
+	}
+#endif // WITH_EDITORONLY_DATA
+
+	for (const FTextureParameterValue& Param : TextureParameterValues)
+	{
+		if (UTexture* Texture = Param.ParameterValue)
+		{
+			Texture->ConditionalPostLoad();
+		}
+	}
+
+	for (const FFontParameterValue& Param : FontParameterValues)
+	{
+		if (UFont* Font = Param.FontValue)
+		{
+			Font->ConditionalPostLoad();
+		}
+	}
+
+	for (const FRuntimeVirtualTextureParameterValue& Param : RuntimeVirtualTextureParameterValues)
+	{
+		if (URuntimeVirtualTexture* Texture = Param.ParameterValue)
+		{
+			Texture->ConditionalPostLoad();
+		}
+	}
+
+	for (const FSparseVolumeTextureParameterValue& Param : SparseVolumeTextureParameterValues)
+	{
+		if (USparseVolumeTexture* Texture = Param.ParameterValue)
+		{
+			Texture->ConditionalPostLoad();
+		}
+	}
 }
 
 #if WITH_EDITOR
diff --git a/Engine/Source/Runtime/Engine/Private/Materials/MaterialInstance.cpp b/Engine/Source/Runtime/Engine/Private/Materials/MaterialInstance.cpp
--- a/Engine/Source/Runtime/Engine/Private/Materials/MaterialInstance.cpp
+++ b/Engine/Source/Runtime/Engine/Private/Materials/MaterialInstance.cpp
@@ -3106,7 +3106,6 @@ void UMaterialInstance::PostLoad()
 		}
 	}
 
-#if WITH_EDITORONLY_DATA
 	// And any material layers parameter's functions
 	if (StaticParametersRuntime.bHasMaterialLayers)
 	{
@@ -3117,6 +3116,7 @@ void UMaterialInstance::PostLoad()
 				Dependency->ConditionalPostLoad();
 			}
 		}
+
 		for (UMaterialFunctionInterface* Dependency : StaticParametersRuntime.MaterialLayers.Blends)
 		{
 			if (Dependency)
@@ -3125,7 +3125,6 @@ void UMaterialInstance::PostLoad()
 			}
 		}
 	}
-#endif // WITH_EDITORONLY_DATA
 
 	if (!GIsEditor)
 	{
