Adding missing include.

diff --git a/Engine/Source/Runtime/Engine/Public/WorldPartition/WorldPartitionActorDescType.h b/Engine/Source/Runtime/Engine/Public/WorldPartition/WorldPartitionActorDescType.h
--- a/Engine/Source/Runtime/Engine/Public/WorldPartition/WorldPartitionActorDescType.h
+++ b/Engine/Source/Runtime/Engine/Public/WorldPartition/WorldPartitionActorDescType.h
@@ -1,6 +1,7 @@
 // Copyright Epic Games, Inc. All Rights Reserved.
 #pragma once
 
+#include "Templates/Requires.h"
 #include "Templates/UnrealTypeTraits.h"
 
 class AActor;
@@ -19,7 +20,3 @@ struct FWorldPartitionActorDescType<ActorType>			\
 #else
 #define DEFINE_ACTORDESC_TYPE(ActorType, ActorDescType)
 #endif
-
-#if UE_ENABLE_INCLUDE_ORDER_DEPRECATED_IN_5_2
-#include "CoreMinimal.h"
-#endif
