World Partition:
- Always time stamp all streaming generation logs.
- Fix streaming generation unit test to handle time stamped log.
- Don't output streaming generation logs when running cooking simulation.

diff --git a/Engine/Source/Runtime/Engine/Private/Tests/WorldPartition/WorldPartitionStreamingGenerationTests.cpp b/Engine/Source/Runtime/Engine/Private/Tests/WorldPartition/WorldPartitionStreamingGenerationTests.cpp
--- a/Engine/Source/Runtime/Engine/Private/Tests/WorldPartition/WorldPartitionStreamingGenerationTests.cpp
+++ b/Engine/Source/Runtime/Engine/Private/Tests/WorldPartition/WorldPartitionStreamingGenerationTests.cpp
@@ -53,9 +53,7 @@ namespace WorldPartitionTests
 			return false;
 		}
 
-		UWorldPartition::FGenerateStreamingParams Params = UWorldPartition::FGenerateStreamingParams()
-			.SetOutputLogPath(TEXT("StreamingGenerationTest"));
-
+		UWorldPartition::FGenerateStreamingParams Params;
 		UWorldPartition::FGenerateStreamingContext Context;
 		WorldPartition->GenerateStreaming(Params, Context);
 
@@ -66,7 +64,7 @@ namespace WorldPartitionTests
 
 		// Read reference output log
 		FString ReferenceOuputLog;
-		const FString ReferenceOutputLogPath(FPaths::GetPath(World->GetPackage()->GetLoadedPath().GetLocalFullPath()) / FPaths::GetCleanFilename(*Context.OutputLogFilename));
+		const FString ReferenceOutputLogPath(FPaths::GetPath(World->GetPackage()->GetLoadedPath().GetLocalFullPath()) / TEXT("WPUnitTests.log"));
 		if (!TestTrue(TEXT("Error Reading Reference Output Log File"), FFileHelper::LoadFileToString(ReferenceOuputLog, *ReferenceOutputLogPath)))
 		{
 			return false;
diff --git a/Engine/Source/Runtime/Engine/Private/WorldPartition/ContentBundle/ContentBundleEditor.cpp b/Engine/Source/Runtime/Engine/Private/WorldPartition/ContentBundle/ContentBundleEditor.cpp
--- a/Engine/Source/Runtime/Engine/Private/WorldPartition/ContentBundle/ContentBundleEditor.cpp
+++ b/Engine/Source/Runtime/Engine/Private/WorldPartition/ContentBundle/ContentBundleEditor.cpp
@@ -349,6 +349,13 @@ void FContentBundleEditor::GenerateStreaming(TArray<FString>* OutPackageToGenera
 	FActorDescContainerInstanceCollection Collection({ TObjectPtr<UActorDescContainerInstance>(ActorDescContainerInstance.Get()) });
 	UWorldPartition::FGenerateStreamingParams Params = UWorldPartition::FGenerateStreamingParams()
 		.SetContainerInstanceCollection(Collection, FStreamingGenerationContainerInstanceCollection::ECollectionType::BaseAsContentBundle);
+
+	// Only dump streaming generation logs when generating for PIE or cook
+	if (!bIsPIE && !IsRunningCookCommandlet())
+	{
+		Params.SetOutputLogType(TEXT(""));
+	}
+
 	UWorldPartition::FGenerateStreamingContext Context = UWorldPartition::FGenerateStreamingContext()
 		.SetLevelPackagesToGenerate(OutPackageToGenerate);
 
diff --git a/Engine/Source/Runtime/Engine/Private/WorldPartition/Cook/WorldPartitionCook.cpp b/Engine/Source/Runtime/Engine/Private/WorldPartition/Cook/WorldPartitionCook.cpp
--- a/Engine/Source/Runtime/Engine/Private/WorldPartition/Cook/WorldPartitionCook.cpp
+++ b/Engine/Source/Runtime/Engine/Private/WorldPartition/Cook/WorldPartitionCook.cpp
@@ -37,6 +37,12 @@ bool UWorldPartition::GatherPackagesToCook(IWorldPartitionCookPackageContext& Co
 	FGenerateStreamingParams Params = FGenerateStreamingParams()
 		.SetContainerInstanceCollection(*this, FStreamingGenerationContainerInstanceCollection::ECollectionType::BaseAndEDLs);
 
+	// Only dump streaming generation logs when cooking
+	if (!IsRunningCookCommandlet())
+	{
+		Params.SetOutputLogType(TEXT(""));
+	}
+
 	if (!GenerateContainerStreaming(Params, Context))
 	{
 		return false;
diff --git a/Engine/Source/Runtime/Engine/Private/WorldPartition/WorldPartitionStreamingGeneration.cpp b/Engine/Source/Runtime/Engine/Private/WorldPartition/WorldPartitionStreamingGeneration.cpp
--- a/Engine/Source/Runtime/Engine/Private/WorldPartition/WorldPartitionStreamingGeneration.cpp
+++ b/Engine/Source/Runtime/Engine/Private/WorldPartition/WorldPartitionStreamingGeneration.cpp
@@ -49,7 +49,8 @@ static FAutoConsoleCommand DumpStreamingGenerationLog(
 		{
 			if (UWorldPartition* WorldPartition = World->GetWorldPartition())
 			{
-				UWorldPartition::FGenerateStreamingParams Params;
+				UWorldPartition::FGenerateStreamingParams Params = UWorldPartition::FGenerateStreamingParams()
+					.SetOutputLogType(TEXT("DumpStreamingGeneration"));
 				UWorldPartition::FGenerateStreamingContext Context;
 				WorldPartition->GenerateStreaming(Params, Context);
 				WorldPartition->FlushStreaming();
@@ -1671,16 +1672,10 @@ public:
 		}
 	}
 
-	static TUniquePtr<FArchive> CreateDumpStateLogArchive(const TCHAR* Suffix, bool bTimeStamped = true)
+	static TUniquePtr<FArchive> CreateDumpStateLogArchive(const TCHAR* Suffix)
 	{
 		FString StateLogOutputFilename = FPaths::ProjectLogDir() / TEXT("WorldPartition") / FString::Printf(TEXT("StreamingGeneration-%s"), Suffix);
-
-		if (bTimeStamped)
-		{
-			StateLogOutputFilename += FString::Printf(TEXT("-%08x-%s"), FPlatformProcess::GetCurrentProcessId(), *FDateTime::Now().ToIso8601().Replace(TEXT(":"), TEXT(".")));
-		}
-
-		StateLogOutputFilename += TEXT(".log");
+		StateLogOutputFilename += FString::Printf(TEXT("-%08x-%s.log"), FPlatformProcess::GetCurrentProcessId(), *FDateTime::Now().ToIso8601().Replace(TEXT(":"), TEXT(".")));
 		return TUniquePtr<FArchive>(IFileManager::Get().CreateFileWriter(*StateLogOutputFilename));
 	}
 
@@ -1869,16 +1864,16 @@ bool UWorldPartition::GenerateContainerStreaming(const FGenerateStreamingParams&
 	TUniquePtr<FArchive> LogFileAr;
 	TUniquePtr<FHierarchicalLogArchive> HierarchicalLogAr;
 
-	const bool bIsStreamingGenerationLogAllowed = !bIsPIE || FModuleManager::LoadModuleChecked<IWorldPartitionEditorModule>("WorldPartitionEditor").GetEnableStreamingGenerationLogOnPIE();
+	const bool bIsStreamingGenerationLogAllowed = (!bIsPIE || FModuleManager::LoadModuleChecked<IWorldPartitionEditorModule>("WorldPartitionEditor").GetEnableStreamingGenerationLogOnPIE()) && (!InParams.OutputLogType.IsSet() || !InParams.OutputLogType.GetValue().IsEmpty());
 	const bool bIsStreamingGenerationLogRelevant = IsMainWorldPartition() && (!GIsBuildMachine || GIsAutomationTesting || IsRunningCookCommandlet());
 
 	if (bIsStreamingGenerationLogAllowed && bIsStreamingGenerationLogRelevant)
 	{
 		TStringBuilder<256> StateLogSuffix;
-		StateLogSuffix += bIsPIE ? TEXT("PIE") : (IsRunningGame() ? TEXT("Game") : (IsRunningCookCommandlet() ? TEXT("Cook") : (GIsAutomationTesting ? TEXT("UnitTest") : TEXT("Manual"))));
+		StateLogSuffix += bIsPIE ? TEXT("PIE") : (IsRunningGame() ? TEXT("Game") : (IsRunningCookCommandlet() ? TEXT("Cook") : (GIsAutomationTesting ? TEXT("UnitTest") : (InParams.OutputLogType.IsSet() ? *InParams.OutputLogType.GetValue() : TEXT("Manual")))));
 		StateLogSuffix += TEXT("_");
 		StateLogSuffix += ContainerShortName;
-		LogFileAr = FWorldPartitionStreamingGenerator::CreateDumpStateLogArchive(*StateLogSuffix, !InParams.OutputLogPath);
+		LogFileAr = FWorldPartitionStreamingGenerator::CreateDumpStateLogArchive(*StateLogSuffix);
 
 		if (LogFileAr.IsValid())
 		{
diff --git a/Engine/Source/Runtime/Engine/Public/WorldPartition/WorldPartition.h b/Engine/Source/Runtime/Engine/Public/WorldPartition/WorldPartition.h
--- a/Engine/Source/Runtime/Engine/Public/WorldPartition/WorldPartition.h
+++ b/Engine/Source/Runtime/Engine/Public/WorldPartition/WorldPartition.h
@@ -225,7 +225,7 @@ public:
 			return *this;
 		}
 		FGenerateStreamingParams& SetErrorHandler(IStreamingGenerationErrorHandler* InErrorHandler) { ErrorHandler = InErrorHandler; return *this; }
-		FGenerateStreamingParams& SetOutputLogPath(const FString& InOutputLogPath) { OutputLogPath = InOutputLogPath; return *this; }
+		FGenerateStreamingParams& SetOutputLogType(const FString& InOutputLogType) { OutputLogType = InOutputLogType; return *this; }
 
 		UE_DEPRECATED(5.4, "Use constructor receiving a ContainerInstanceCollection instead")
 		FGenerateStreamingParams& SetActorDescContainer(const UActorDescContainer* InActorDescContainer) { return *this; }
@@ -236,7 +236,7 @@ public:
 	private:
 
 		FStreamingGenerationContainerInstanceCollection ContainerInstanceCollection;
-		TOptional<const FString> OutputLogPath;
+		TOptional<const FString> OutputLogType;
 		IStreamingGenerationErrorHandler* ErrorHandler;
 
 		friend class UWorldPartition;
