Stopgap fix for HISM issue on server where instances are inadvertently cleared as the instance manager reports that it has changes.

Original CL: https://p4-swarm.epicgames.net/reviews/32049078
Created this CL for 29.01+ since 29.00 shelf is being held for a server build

#rnx

diff --git a/Engine/Source/Runtime/Engine/Public/InstancedStaticMesh/ISMInstanceDataManager.h b/Engine/Source/Runtime/Engine/Public/InstancedStaticMesh/ISMInstanceDataManager.h
--- a/Engine/Source/Runtime/Engine/Public/InstancedStaticMesh/ISMInstanceDataManager.h
+++ b/Engine/Source/Runtime/Engine/Public/InstancedStaticMesh/ISMInstanceDataManager.h
@@ -99,11 +99,20 @@ public:
 
 	ENGINE_API bool HasAnyInstanceChanges() const;
 
+
+#if USE_NULL_RHI
+	/**
+	 * Stopgap fix for server, where the manager reporting that it has changes causes (inadvertent) clearing of the instances in the HISM.
+	 * TODO: Remove when proper fix (not clearing instances) is implemented.
+	 */
+	inline bool HasAnyChanges() const { return false;}
+#else
 	/** 
 	 * Returns true if there are explicitly tracked instance changes, or the state is not tracked (because no proxy has been created yet),
 	 * and the tracking state is not Disabled.
 	 */
 	inline bool HasAnyChanges() const { return GetState() != ETrackingState::Disabled && (GetState() != ETrackingState::Tracked || HasAnyInstanceChanges());}
+#endif
 	bool FlushChanges(FInstanceUpdateComponentDesc &&ComponentData, bool bNewPrimitiveProxy);
 
 	// 
