World Partition: when destroying an actor descriptor container instance, unregister all actors that still have hard references to them.
Also updated World Partition Handles unit tests to include this case.

diff --git a/Engine/Source/Runtime/Engine/Private/Tests/WorldPartition/WorldPartitionHandleTests.cpp b/Engine/Source/Runtime/Engine/Private/Tests/WorldPartition/WorldPartitionHandleTests.cpp
--- a/Engine/Source/Runtime/Engine/Private/Tests/WorldPartition/WorldPartitionHandleTests.cpp
+++ b/Engine/Source/Runtime/Engine/Private/Tests/WorldPartition/WorldPartitionHandleTests.cpp
@@ -63,7 +63,7 @@ namespace WorldPartitionTests
 
 		TestTrue(TEXT("World type"), World->WorldType == EWorldType::None);
 
-		UActorDescContainerInstance* ActorDescContainerInstance = NewObject<UActorDescContainerInstance>(GetTransientPackage());
+		UActorDescContainerInstance* ActorDescContainerInstance = NewObject<UActorDescContainerInstance>(World);
 		ActorDescContainerInstance->Initialize({ TEXT("/Engine/WorldPartition/WorldPartitionUnitTest") });
 
 		FWorldPartitionHandle Handle = FWorldPartitionHandle(ActorDescContainerInstance, FGuid(TEXT("5D9F93BA407A811AFDDDAAB4F1CECC6A")));
@@ -294,8 +294,19 @@ namespace WorldPartitionTests
 			TestTrue(TEXT("Invalid container test"), Handle.IsValid());
 			TestTrue(TEXT("Invalid container test"), Reference.IsValid());
 
+			TestTrue(TEXT("Handle soft refcount"), FWorldPartitionActorDescUnitTestAcccessor::GetSoftRefCount(*Handle) == 1);
+			TestTrue(TEXT("Handle hard refcount"), FWorldPartitionActorDescUnitTestAcccessor::GetHardRefCount(*Handle) == 0);
+			TestTrue(TEXT("Reference soft refcount"), FWorldPartitionActorDescUnitTestAcccessor::GetSoftRefCount(*Reference) == 0);
+			TestTrue(TEXT("Reference hard refcount"), FWorldPartitionActorDescUnitTestAcccessor::GetHardRefCount(*Reference) == 1);
+
+			TWeakObjectPtr<AActor> LoadedActor = Reference->GetActor();
+			TestTrue(TEXT("Loaded actor"), LoadedActor.IsValid());
+			TestTrue(TEXT("Registered actor"), World->PersistentLevel->Actors.Find(LoadedActor.Get()) != INDEX_NONE);
+
 			ActorDescContainerInstance->Uninitialize();
 
+			TestTrue(TEXT("Registered actor"), World->PersistentLevel->Actors.Find(LoadedActor.Get()) == INDEX_NONE);
+
 			// Make sure to cleanup world before collecting garbage so it gets uninitialized
 			ScopedEditorWorld.Reset();
 			CollectGarbage(GARBAGE_COLLECTION_KEEPFLAGS);
diff --git a/Engine/Source/Runtime/Engine/Private/WorldPartition/ActorDescContainerInstance.cpp b/Engine/Source/Runtime/Engine/Private/WorldPartition/ActorDescContainerInstance.cpp
--- a/Engine/Source/Runtime/Engine/Private/WorldPartition/ActorDescContainerInstance.cpp
+++ b/Engine/Source/Runtime/Engine/Private/WorldPartition/ActorDescContainerInstance.cpp
@@ -189,6 +189,8 @@ void UActorDescContainerInstance::OnUnregisterChildContainerInstance(const FGuid
 
 void UActorDescContainerInstance::Uninitialize()
 {
+	const UWorld* OuterWorld = GetTypedOuter<UWorld>();
+
 	bIsInitialized = false;
 
 	UnregisterDelegates();
@@ -197,6 +199,15 @@ void UActorDescContainerInstance::Uninitialize()
 	{
 		if (ActorDescInstancePtr.IsValid())
 		{
+			if (ActorDescInstancePtr->GetHardRefCount())
+			{
+				if (OuterWorld && !OuterWorld->IsBeingCleanedUp())
+				{
+					// Force actor unregistration when destroying the container instance without the world being destroyed to avoid dangling actors
+					FWorldPartitionReference LastReference(ActorDescInstancePtr.Get());
+				}
+			}			
+
 			RemoveActorDescInstance(&ActorDescInstancePtr);
 		}
 	}
diff --git a/Engine/Source/Runtime/Engine/Private/WorldPartition/WorldPartitionHandle.cpp b/Engine/Source/Runtime/Engine/Private/WorldPartition/WorldPartitionHandle.cpp
--- a/Engine/Source/Runtime/Engine/Private/WorldPartition/WorldPartitionHandle.cpp
+++ b/Engine/Source/Runtime/Engine/Private/WorldPartition/WorldPartitionHandle.cpp
@@ -285,7 +285,11 @@ void FWorldPartitionReferenceImpl::IncRefCount(FWorldPartitionActorDescInstance*
 
 void FWorldPartitionReferenceImpl::DecRefCount(FWorldPartitionActorDescInstance* InActorDescInstance)
 {
-	if (InActorDescInstance->DecHardRefCount() == 0)
+	const bool bIsContainerInitialized = InActorDescInstance->GetContainerInstance()->IsInitialized();
+	const bool bIsActorDescInstanceUnregistering = InActorDescInstance->bIsRegisteringOrUnregistering;
+	const bool bReleasingLastReference = !InActorDescInstance->DecHardRefCount();
+
+	if ((bReleasingLastReference || !bIsContainerInitialized) && !bIsActorDescInstanceUnregistering)
 	{
 		FWorldPartitionLoadingContext::UnloadAndUnregisterActor(InActorDescInstance);
 	}
