Exposed sun disc color scale to sequencer/blueprint.

diff --git a/Engine/Source/Runtime/Engine/Classes/Components/DirectionalLightComponent.h b/Engine/Source/Runtime/Engine/Classes/Components/DirectionalLightComponent.h
--- a/Engine/Source/Runtime/Engine/Classes/Components/DirectionalLightComponent.h
+++ b/Engine/Source/Runtime/Engine/Classes/Components/DirectionalLightComponent.h
@@ -312,6 +312,9 @@ class UDirectionalLightComponent : public ULightComponent
 	UFUNCTION(BlueprintCallable, Category = "Rendering|Lighting")
 	ENGINE_API void SetAtmosphereSunLightIndex(int32 NewValue);
 
+	UFUNCTION(BlueprintCallable, Category = "Rendering|Lighting")
+	ENGINE_API void SetAtmosphereSunDiskColorScale(FLinearColor NewValue);
+
 	UFUNCTION(BlueprintCallable, Category = "Rendering|Lighting")
 	ENGINE_API void SetForwardShadingPriority(int32 NewValue);
 
diff --git a/Engine/Source/Runtime/Engine/Private/Components/DirectionalLightComponent.cpp b/Engine/Source/Runtime/Engine/Private/Components/DirectionalLightComponent.cpp
--- a/Engine/Source/Runtime/Engine/Private/Components/DirectionalLightComponent.cpp
+++ b/Engine/Source/Runtime/Engine/Private/Components/DirectionalLightComponent.cpp
@@ -1369,6 +1369,16 @@ void UDirectionalLightComponent::SetAtmosphereSunLightIndex(int32 NewValue)
 	}
 }
 
+void UDirectionalLightComponent::SetAtmosphereSunDiskColorScale(FLinearColor NewValue)
+{
+	if (AreDynamicDataChangesAllowed()
+		&& AtmosphereSunDiskColorScale != NewValue)
+	{
+		AtmosphereSunDiskColorScale = NewValue;
+		MarkRenderStateDirty();
+	}
+}
+
 void UDirectionalLightComponent::SetForwardShadingPriority(int32 NewValue)
 {
 	if (AreDynamicDataChangesAllowed()
