Fix crash when creating circular reference in Level Instance

#rnx

diff --git a/Engine/Source/Runtime/Engine/Private/WorldPartition/ActorDescContainerInstance.cpp b/Engine/Source/Runtime/Engine/Private/WorldPartition/ActorDescContainerInstance.cpp
--- a/Engine/Source/Runtime/Engine/Private/WorldPartition/ActorDescContainerInstance.cpp
+++ b/Engine/Source/Runtime/Engine/Private/WorldPartition/ActorDescContainerInstance.cpp
@@ -444,7 +444,6 @@ void UActorDescContainerInstance::RemoveActorDescInstance(TUniquePtr<FWorldParti
 
 	if (bCreateChildContainerHierarchy && ActorDescInstance->IsChildContainerInstance())
 	{
-		check(ChildContainerInstances.Contains(ActorDescInstance->GetGuid()));
 		ActorDescInstance->UnregisterChildContainerInstance();
 		check(!ChildContainerInstances.Contains(ActorDescInstance->GetGuid()));
 	}
diff --git a/Engine/Source/Runtime/Engine/Private/WorldPartition/LevelInstance/LevelInstanceActorDesc.cpp b/Engine/Source/Runtime/Engine/Private/WorldPartition/LevelInstance/LevelInstanceActorDesc.cpp
--- a/Engine/Source/Runtime/Engine/Private/WorldPartition/LevelInstance/LevelInstanceActorDesc.cpp
+++ b/Engine/Source/Runtime/Engine/Private/WorldPartition/LevelInstance/LevelInstanceActorDesc.cpp
@@ -254,6 +254,10 @@ void FLevelInstanceActorDesc::CheckForErrors(const IWorldPartitionActorDescInsta
 			ErrorHandler->OnLevelInstanceInvalidWorldAsset(*InActorDescView, ChildContainerPackage, IStreamingGenerationErrorHandler::ELevelInstanceInvalidReason::WorldAssetNotUsingExternalActors);
 		}
 	}
+	else if (!ValidateCircularReference(InActorDescView->GetContainerInstance(), ChildContainerPackage))
+	{
+		ErrorHandler->OnLevelInstanceInvalidWorldAsset(*InActorDescView, ChildContainerPackage, IStreamingGenerationErrorHandler::ELevelInstanceInvalidReason::CirculalReference);
+	}
 }
 
 void FLevelInstanceActorDesc::TransferFrom(const FWorldPartitionActorDesc* From)
@@ -284,6 +288,22 @@ UWorldPartition* FLevelInstanceActorDesc::GetLoadedChildWorldPartition(const FWo
 	return nullptr;
 }
 
+bool FLevelInstanceActorDesc::ValidateCircularReference(const UActorDescContainerInstance* InParentContainer, FName InChildContainerPackage)
+{
+	const UActorDescContainerInstance* CurrentParentContainerInstance = InParentContainer;
+	while (CurrentParentContainerInstance)
+	{
+		if (CurrentParentContainerInstance->GetContainerPackage() == InChildContainerPackage)
+		{
+			// found a circular reference
+			return false; 
+		}
+		CurrentParentContainerInstance = Cast<UActorDescContainerInstance>(CurrentParentContainerInstance->GetOuter());
+	}
+
+	return true;
+}
+
 UActorDescContainerInstance* FLevelInstanceActorDesc::CreateChildContainerInstance(const FWorldPartitionActorDescInstance* InActorDescInstance) const
 {
 	// Update Actor Desc Container in case of rename
@@ -293,9 +313,14 @@ UActorDescContainerInstance* FLevelInstanceActorDesc::CreateChildContainerInstan
 		NonConstThis->UnregisterChildContainer();
 		NonConstThis->RegisterChildContainer();
 	}
+
+	UActorDescContainerInstance* ContainerInstance = InActorDescInstance->GetContainerInstance();
+	if (!ValidateCircularReference(ContainerInstance, InActorDescInstance->GetChildContainerPackage()))
+	{
+		return nullptr;
+	}
 	
 	// Create ChildContainerInstance
-	UActorDescContainerInstance* ContainerInstance = InActorDescInstance->GetContainerInstance();
 	ULevelInstanceContainerInstance* ChildContainerInstance = NewObject<ULevelInstanceContainerInstance>(ContainerInstance, NAME_None, RF_Transient);
 
 	// Set Override Container which might be a regular UActorDescContainer if no overrides exist on this Container Instance
diff --git a/Engine/Source/Runtime/Engine/Private/WorldPartition/WorldPartitionActorDescInstance.cpp b/Engine/Source/Runtime/Engine/Private/WorldPartition/WorldPartitionActorDescInstance.cpp
--- a/Engine/Source/Runtime/Engine/Private/WorldPartition/WorldPartitionActorDescInstance.cpp
+++ b/Engine/Source/Runtime/Engine/Private/WorldPartition/WorldPartitionActorDescInstance.cpp
@@ -230,16 +230,20 @@ void FWorldPartitionActorDescInstance::RegisterChildContainerInstance()
 	check(!ChildContainerInstance);
 
 	ChildContainerInstance = ActorDesc->CreateChildContainerInstance(this);
-	check(ChildContainerInstance);
-	ContainerInstance->OnRegisterChildContainerInstance(GetGuid(), ChildContainerInstance);
+	if (ChildContainerInstance)
+	{
+		ContainerInstance->OnRegisterChildContainerInstance(GetGuid(), ChildContainerInstance);
+	}
 }
 
 void FWorldPartitionActorDescInstance::UnregisterChildContainerInstance()
 {
-	check(ChildContainerInstance);
-	ContainerInstance->OnUnregisterChildContainerInstance(GetGuid());
-	ChildContainerInstance->Uninitialize();
-	ChildContainerInstance = nullptr;
+	if (ChildContainerInstance)
+	{
+		ContainerInstance->OnUnregisterChildContainerInstance(GetGuid());
+		ChildContainerInstance->Uninitialize();
+		ChildContainerInstance = nullptr;
+	}
 }
 
 void FWorldPartitionActorDescInstance::UpdateChildContainerInstance()
diff --git a/Engine/Source/Runtime/Engine/Private/WorldPartition/WorldPartitionStreamingGeneration.cpp b/Engine/Source/Runtime/Engine/Private/WorldPartition/WorldPartitionStreamingGeneration.cpp
--- a/Engine/Source/Runtime/Engine/Private/WorldPartition/WorldPartitionStreamingGeneration.cpp
+++ b/Engine/Source/Runtime/Engine/Private/WorldPartition/WorldPartitionStreamingGeneration.cpp
@@ -946,16 +946,6 @@ class FWorldPartitionStreamingGenerator
 				FWorldPartitionActorDesc::FContainerInstance SubContainerInstance;
 				if (!ContainerCollectionInstanceView.GetChildContainerInstance(SubContainerInstance) || !SubContainerInstance.ContainerInstance)
 				{
-					ErrorHandler->OnLevelInstanceInvalidWorldAsset(ContainerCollectionInstanceView, ContainerCollectionInstanceView.GetChildContainerPackage(), IStreamingGenerationErrorHandler::ELevelInstanceInvalidReason::WorldAssetHasInvalidContainer);
-					continue;
-				}
-
-				bool bContainerWasAlreadyInSet;
-				ContainerInstancesStack.Add(SubContainerInstance.ContainerInstance->GetContainerPackage(), &bContainerWasAlreadyInSet);
-
-				if (bContainerWasAlreadyInSet)
-				{
-					ErrorHandler->OnLevelInstanceInvalidWorldAsset(ContainerCollectionInstanceView, ContainerCollectionInstanceView.GetChildContainerPackage(), IStreamingGenerationErrorHandler::ELevelInstanceInvalidReason::CirculalReference);
 					continue;
 				}
 
@@ -992,8 +982,6 @@ class FWorldPartitionStreamingGenerator
 				}
 
 				CreateActorDescriptorViewsRecursive(MoveTemp(SubContainerInstanceDescriptor));
-
-				verify(ContainerInstancesStack.Remove(SubContainerInstance.ContainerInstance->GetContainerPackage()));
 			}
 		}
 
@@ -1826,9 +1814,6 @@ private:
 	/** List of container instances participating in this streaming generation step */
 	TMap<FGuid, const UActorDescContainerInstance*> ActorGuidsToContainerInstanceMap;
 
-	/** List of current container instances on the stack to detect circular references */
-	TSet<FName> ContainerInstancesStack;
-
 	/** Maps containers IDs to their filtered actors use while creating FContainerCollectionInstanceDescriptor */
 	TMap<FActorContainerID, TSet<FGuid>> ContainerFilteredActors;
 };
diff --git a/Engine/Source/Runtime/Engine/Private/WorldPartition/WorldPartitionSubsystem.cpp b/Engine/Source/Runtime/Engine/Private/WorldPartition/WorldPartitionSubsystem.cpp
--- a/Engine/Source/Runtime/Engine/Private/WorldPartition/WorldPartitionSubsystem.cpp
+++ b/Engine/Source/Runtime/Engine/Private/WorldPartition/WorldPartitionSubsystem.cpp
@@ -204,19 +204,19 @@ const UWorldPartition* UWorldPartitionSubsystem::GetWorldPartition() const
 #if WITH_EDITOR
 FWorldPartitionActorFilter UWorldPartitionSubsystem::GetWorldPartitionActorFilter(const FString& InWorldPackage, EWorldPartitionActorFilterType InFilterTypes) const
 {
-	TSet<FString> VisitedPackages;
-	return GetWorldPartitionActorFilterInternal(InWorldPackage, InFilterTypes, VisitedPackages);
+	TSet<FString> VisitedPackageStack;
+	return GetWorldPartitionActorFilterInternal(InWorldPackage, InFilterTypes, VisitedPackageStack);
 }
 
 
-FWorldPartitionActorFilter UWorldPartitionSubsystem::GetWorldPartitionActorFilterInternal(const FString& InWorldPackage, EWorldPartitionActorFilterType InFilterTypes, TSet<FString>& InOutVisitedPackages) const
+FWorldPartitionActorFilter UWorldPartitionSubsystem::GetWorldPartitionActorFilterInternal(const FString& InWorldPackage, EWorldPartitionActorFilterType InFilterTypes, TSet<FString>& InOutVisitedPackageStack) const
 {
-	if (InOutVisitedPackages.Contains(InWorldPackage))
+	if (InOutVisitedPackageStack.Contains(InWorldPackage))
 	{
 		return FWorldPartitionActorFilter(InWorldPackage);
 	}
 
-	InOutVisitedPackages.Add(InWorldPackage);
+	InOutVisitedPackageStack.Add(InWorldPackage);
 	
 	// Most of the time if this will return an existing Container but when loading a new LevelInstance (Content Browser Drag&Drop, Create LI) 
 	// This will make sure Container exists.
@@ -257,10 +257,8 @@ FWorldPartitionActorFilter UWorldPartitionSubsystem::GetWorldPartitionActorFilte
 
 	for (const FWorldPartitionActorDesc* ContainerActorDesc : ContainerActorDescs)
 	{
-		TSet<FString> VisitedPackagesCopy(InOutVisitedPackages);
-
 		// Get World Default Filter
-		FWorldPartitionActorFilter* ChildFilter = new FWorldPartitionActorFilter(GetWorldPartitionActorFilterInternal(ContainerActorDesc->GetChildContainerPackage().ToString(), InFilterTypes, VisitedPackagesCopy));
+		FWorldPartitionActorFilter* ChildFilter = new FWorldPartitionActorFilter(GetWorldPartitionActorFilterInternal(ContainerActorDesc->GetChildContainerPackage().ToString(), InFilterTypes, InOutVisitedPackageStack));
 		ChildFilter->DisplayName = ContainerActorDesc->GetActorLabelOrName().ToString();
 
 		// Apply Filter to Default
@@ -272,6 +270,8 @@ FWorldPartitionActorFilter UWorldPartitionSubsystem::GetWorldPartitionActorFilte
 		Filter.AddChildFilter(ContainerActorDesc->GetGuid(), ChildFilter);
 	}
 
+	verify(InOutVisitedPackageStack.Remove(InWorldPackage));
+
 	return Filter;
 }
 
@@ -317,8 +317,14 @@ TMap<FActorContainerID, TSet<FGuid>> UWorldPartitionSubsystem::GetFilteredActors
 		return RegisteredContainer;
 	};
 
-	TFunction<void(const FActorContainerID&, const UActorDescContainer*)> ProcessContainers = [&FindOrRegisterContainer, &DataLayerFiltersPerContainer, &FilteredActors, &ProcessContainers, &InFilterTypes](const FActorContainerID& InContainerID, const UActorDescContainer* InContainer)
+	TFunction<void(const FActorContainerID&, const UActorDescContainer*, TSet<FName>&)> ProcessContainers = [&FindOrRegisterContainer, &DataLayerFiltersPerContainer, &FilteredActors, &ProcessContainers, &InFilterTypes](const FActorContainerID& InContainerID, const UActorDescContainer* InContainer, TSet<FName>& InOutVisitedPackageStack)
 	{
+		if(InOutVisitedPackageStack.Contains(InContainer->GetContainerPackage()))
+		{
+			return;
+		}
+		InOutVisitedPackageStack.Add(InContainer->GetContainerPackage());
+
 		const TMap<FSoftObjectPath, FWorldPartitionActorFilter::FDataLayerFilter>& DataLayerFilters = DataLayerFiltersPerContainer.FindChecked(InContainerID);
 		for (FActorDescList::TConstIterator<> ActorDescIt(InContainer); ActorDescIt; ++ActorDescIt)
 		{
@@ -354,14 +360,19 @@ TMap<FActorContainerID, TSet<FGuid>> UWorldPartitionSubsystem::GetFilteredActors
 				{
 					UActorDescContainer* ChildContainer = FindOrRegisterContainer(ActorDescIt->GetChildContainerPackage());
 					check(ChildContainer);
-					ProcessContainers(FActorContainerID(InContainerID, ActorDescIt->GetGuid()), ChildContainer);
+					ProcessContainers(FActorContainerID(InContainerID, ActorDescIt->GetGuid()), ChildContainer, InOutVisitedPackageStack);
 				}
 			}
 		}
+
+		verify(InOutVisitedPackageStack.Remove(InContainer->GetContainerPackage()));
 	};
 
 	UActorDescContainer* Container = FindOrRegisterContainer(*InWorldPackage);
-	ProcessContainers(InContainerID, Container);
+
+	TSet<FName> VisitedPackageStack;
+	ProcessContainers(InContainerID, Container, VisitedPackageStack);
+	verify(VisitedPackageStack.IsEmpty());
 	
 	// Unregister Containers
 	UActorDescContainerSubsystem& ContainerSubsystem = UActorDescContainerSubsystem::GetChecked();
diff --git a/Engine/Source/Runtime/Engine/Public/WorldPartition/LevelInstance/LevelInstanceActorDesc.h b/Engine/Source/Runtime/Engine/Public/WorldPartition/LevelInstance/LevelInstanceActorDesc.h
--- a/Engine/Source/Runtime/Engine/Public/WorldPartition/LevelInstance/LevelInstanceActorDesc.h
+++ b/Engine/Source/Runtime/Engine/Public/WorldPartition/LevelInstance/LevelInstanceActorDesc.h
@@ -51,6 +51,7 @@ protected:
 	ENGINE_API virtual UWorldPartition* GetLoadedChildWorldPartition(const FWorldPartitionActorDescInstance* InActorDescInstance) const override;
 	ENGINE_API FTransform GetChildContainerTransform() const;
 	FLevelInstancePropertyOverrideDesc* GetOverrideDesc() const { return PropertyOverrideDesc.Get(); }
+	ENGINE_API static bool ValidateCircularReference(const UActorDescContainerInstance* InParentContainer, FName InChildContainerPackage);
 
 	FSoftObjectPath WorldAsset;
 	ELevelInstanceRuntimeBehavior DesiredRuntimeBehavior;
diff --git a/Engine/Source/Runtime/Engine/Public/WorldPartition/WorldPartitionSubsystem.h b/Engine/Source/Runtime/Engine/Public/WorldPartition/WorldPartitionSubsystem.h
--- a/Engine/Source/Runtime/Engine/Public/WorldPartition/WorldPartitionSubsystem.h
+++ b/Engine/Source/Runtime/Engine/Public/WorldPartition/WorldPartitionSubsystem.h
@@ -118,7 +118,7 @@ public:
 
 	TSet<FWorldPartitionHandle> SelectedActorHandles;
 private:
-	ENGINE_API FWorldPartitionActorFilter GetWorldPartitionActorFilterInternal(const FString& InWorldPackage, EWorldPartitionActorFilterType InFilterTypes, TSet<FString>& InOutVisitedPackages) const;
+	ENGINE_API FWorldPartitionActorFilter GetWorldPartitionActorFilterInternal(const FString& InWorldPackage, EWorldPartitionActorFilterType InFilterTypes, TSet<FString>& InOutVisitedPackageStack) const;
 #endif
 
 protected:
