World Partition Converter: only initialize the newly created world partition when converting a world if the world is initialized, this was temporarily fixed by 30647533.
This is needed because initializing a world partition object will broadcast an event required for world subsystems, and if these subsystems are create afterward they will miss it. Defer the initialization step when initializing the world manually or when saving the package via InitializePhysicsSceneForSaveIfNecessary

diff --git a/Engine/Source/Runtime/Engine/Private/Level.cpp b/Engine/Source/Runtime/Engine/Private/Level.cpp
--- a/Engine/Source/Runtime/Engine/Private/Level.cpp
+++ b/Engine/Source/Runtime/Engine/Private/Level.cpp
@@ -3080,11 +3080,9 @@ void ULevel::OnLevelLoaded()
 
 			if (bIsMainWorldLevel || bInitializeForEditor)
 			{
-				if (!WorldPartition->IsInitialized())
-				{
-					FTransform Transform = LevelStreaming ? LevelStreaming->LevelTransform : FTransform::Identity;
-					WorldPartition->Initialize(OwningWorld, Transform);
-				}
+				check(!WorldPartition->IsInitialized());
+				FTransform Transform = LevelStreaming ? LevelStreaming->LevelTransform : FTransform::Identity;
+				WorldPartition->Initialize(OwningWorld, Transform);
 			}
 		}
 	}
diff --git a/Engine/Source/Runtime/Engine/Private/WorldPartition/ContentBundle/ContentBundleWorldSubsystem.cpp b/Engine/Source/Runtime/Engine/Private/WorldPartition/ContentBundle/ContentBundleWorldSubsystem.cpp
--- a/Engine/Source/Runtime/Engine/Private/WorldPartition/ContentBundle/ContentBundleWorldSubsystem.cpp
+++ b/Engine/Source/Runtime/Engine/Private/WorldPartition/ContentBundle/ContentBundleWorldSubsystem.cpp
@@ -33,14 +33,7 @@ void UContentBundleManager::Initialize()
 		PIEDuplicateHelper->Initialize();
 	}
 #endif
-	// Any world can be converted to a world partition world before the content bundle manager has been initialized.
-	if (UWorldPartition* WorldPartition = GetWorld()->GetWorldPartition())
-	{
-		if (WorldPartition->IsInitialized())
-		{
-			OnWorldPartitionInitialized(WorldPartition);
-		}
-	}
+
 	GetWorld()->OnWorldPartitionInitialized().AddUObject(this, &UContentBundleManager::OnWorldPartitionInitialized);
 	GetWorld()->OnWorldPartitionUninitialized().AddUObject(this, &UContentBundleManager::OnWorldPartitionUninitialized);
 }
diff --git a/Engine/Source/Runtime/Engine/Private/WorldPartition/WorldPartitionConverter.cpp b/Engine/Source/Runtime/Engine/Private/WorldPartition/WorldPartitionConverter.cpp
--- a/Engine/Source/Runtime/Engine/Private/WorldPartition/WorldPartitionConverter.cpp
+++ b/Engine/Source/Runtime/Engine/Private/WorldPartition/WorldPartitionConverter.cpp
@@ -169,7 +169,10 @@ bool FWorldPartitionConverter::Convert()
 
 	if (bCreatedWorldPartition)
 	{
-		WorldPartition->Initialize(World, FTransform::Identity);
+		if (World->IsInitialized())
+		{
+			WorldPartition->Initialize(World, FTransform::Identity);
+		}
 		UWorldPartition::WorldPartitionChangedEvent.Broadcast(World);
 	}
 
diff --git a/Engine/Source/Runtime/Engine/Private/WorldPartition/WorldPartitionSubsystem.cpp b/Engine/Source/Runtime/Engine/Private/WorldPartition/WorldPartitionSubsystem.cpp
--- a/Engine/Source/Runtime/Engine/Private/WorldPartition/WorldPartitionSubsystem.cpp
+++ b/Engine/Source/Runtime/Engine/Private/WorldPartition/WorldPartitionSubsystem.cpp
@@ -406,13 +406,6 @@ void UWorldPartitionSubsystem::Initialize(FSubsystemCollectionBase& Collection)
 	}
 #endif
 
-	if (UWorldPartition* WorldPartition = GetWorld()->GetWorldPartition())
-	{
-		if (WorldPartition->IsInitialized())
-		{
-			OnWorldPartitionInitialized(WorldPartition);
-		}
-	}
 	GetWorld()->OnWorldPartitionInitialized().AddUObject(this, &UWorldPartitionSubsystem::OnWorldPartitionInitialized);
 	GetWorld()->OnWorldPartitionUninitialized().AddUObject(this, &UWorldPartitionSubsystem::OnWorldPartitionUninitialized);
 	if (GetWorld()->IsGameWorld())
