[Material Editor] Added null check to FHLSLMaterialTranslator::GetMaterialEnvironment() to avoid crashing when a referenced Parameter Collectino fails to load (e.g. no longer exists).

diff --git a/Engine/Source/Runtime/Engine/Private/Materials/HLSLMaterialTranslator.cpp b/Engine/Source/Runtime/Engine/Private/Materials/HLSLMaterialTranslator.cpp
--- a/Engine/Source/Runtime/Engine/Private/Materials/HLSLMaterialTranslator.cpp
+++ b/Engine/Source/Runtime/Engine/Private/Materials/HLSLMaterialTranslator.cpp
@@ -2549,10 +2549,19 @@ void FHLSLMaterialTranslator::GetMaterialEnvironment(EShaderPlatform InPlatform,
 	{
 		// Add uniform buffer declarations for any parameter collections referenced
 		const FString CollectionName = FString::Printf(TEXT("MaterialCollection%u"), CollectionIndex);
+
+		// Check that the parameter collection loaded succesfully.
+		UMaterialParameterCollection* ParameterCollection = EnvironmentDefines->ParameterCollections[CollectionIndex];
+		if (!ParameterCollection)
+		{
+			UE_LOG(LogMaterial, Warning, TEXT("Null parameter collection found in environment defines while translating material."));
+			continue;
+		}
+
 		// This can potentially become an issue for MaterialCollection Uniform Buffers if they ever get non-numeric resources (eg Textures), as
 		// OutEnvironment.ResourceTableMap has a map by name, and the N ParameterCollection Uniform Buffers ALL are names "MaterialCollection"
 		// (and the hlsl cbuffers are named MaterialCollection0, etc, so the names don't match the layout)
-		FShaderUniformBufferParameter::ModifyCompilationEnvironment(*CollectionName, EnvironmentDefines->ParameterCollections[CollectionIndex]->GetUniformBufferStruct(), InPlatform, OutEnvironment);
+		FShaderUniformBufferParameter::ModifyCompilationEnvironment(*CollectionName, ParameterCollection->GetUniformBufferStruct(), InPlatform, OutEnvironment);
 	}
 
 	OutEnvironment.SetDefine(TEXT("IS_MATERIAL_SHADER"), TEXT("1"));
