Mark the Heterogeneous Volume Component as not containing deformable meshes.

Also declare ShadowCacheInvalidationBehavior to be static.

diff --git a/Engine/Source/Runtime/Engine/Private/Components/HeterogeneousVolumeComponent.cpp b/Engine/Source/Runtime/Engine/Private/Components/HeterogeneousVolumeComponent.cpp
--- a/Engine/Source/Runtime/Engine/Private/Components/HeterogeneousVolumeComponent.cpp
+++ b/Engine/Source/Runtime/Engine/Private/Components/HeterogeneousVolumeComponent.cpp
@@ -74,6 +74,10 @@ FHeterogeneousVolumeSceneProxy::FHeterogeneousVolumeSceneProxy(UHeterogeneousVol
 	bIsHeterogeneousVolume = true;
 	bCastDynamicShadow = InComponent->CastShadow;
 
+	// Heterogeneous volumes do not deform internally
+	bHasDeformableMesh = false;
+	ShadowCacheInvalidationBehavior = EShadowCacheInvalidationBehavior::Static;
+
 	HeterogeneousVolumeData.VoxelResolution = InComponent->VolumeResolution;
 	HeterogeneousVolumeData.InstanceToLocal = InComponent->FrameTransform.ToMatrixWithScale();
 
