Fixed crashes caused by an integer overflow and a race condition when streaming SparseVolumeTexture data from DDC. DDC requests finish in a callback on another thread. This thread reads and writes certain members of FPendingRequest. The streaming manager itself (running on the RT or a RT paralllel thread) also reads and writes some of the same members. The issue was fixed by introducing a critical section that is used when any FPendingRequest is accessed and we're streaming from DDC. This critical section should have very low contention, so performance impact should be minimal. The cooked path does not require the critical section, so this is compiled out when shipping. The integer overflow was a trivial fix: Requests are versioned and these versions must be compared using operator!=, not operator<.

#rnx

diff --git a/Engine/Source/Runtime/Engine/Private/SparseVolumeTexture/SparseVolumeTextureStreamingManager.cpp b/Engine/Source/Runtime/Engine/Private/SparseVolumeTexture/SparseVolumeTextureStreamingManager.cpp
--- a/Engine/Source/Runtime/Engine/Private/SparseVolumeTexture/SparseVolumeTextureStreamingManager.cpp
+++ b/Engine/Source/Runtime/Engine/Private/SparseVolumeTexture/SparseVolumeTextureStreamingManager.cpp
@@ -94,6 +94,16 @@ static FAutoConsoleVariableRef CVarSVTStreamingRequestSize(
 	ECVF_RenderThreadSafe
 );
 
+// When in editor, data is streamed from DDC. The request finishes in a callback on some other thread. This callback currently reads FPendingRequest::SVTHandle and FPendingRequest::RequestVersion
+// and writes to FPendingRequest::SharedBuffer and FPendingRequest::State. Other than this, FPendingRequest is never accessed from multiple threads at the same time. However, in order to avoid race conditions, 
+// all accesses to FPendingRequest should be guarded by using the following macro. Contention only happens if the streaming manager is currently accessing a request during an update and a DDC request 
+// finishes at the exact same time this access happens. While this is possible, it is very unlikely, so the performance impact of always using the lock should be minimal.
+#if WITH_EDITORONLY_DATA
+#define LOCK_PENDING_REQUEST(PendingRequest) FScopeLock Lock(&PendingRequest.DDCAsyncGuard)
+#else
+#define LOCK_PENDING_REQUEST(PendingRequest)
+#endif
+
 namespace UE
 {
 namespace SVT
@@ -402,6 +412,7 @@ void FStreamingManager::BeginAsyncUpdate(FRDGBuilder& GraphBuilder, bool bBlocki
 		// Cancel all IO requests
 		for (FPendingRequest& PendingRequest : PendingRequests)
 		{
+			LOCK_PENDING_REQUEST(PendingRequest);
 			PendingRequest.Reset();
 		}
 
@@ -455,6 +466,7 @@ void FStreamingManager::BeginAsyncUpdate(FRDGBuilder& GraphBuilder, bool bBlocki
 		{
 			const int32 PendingRequestIndex = (StartPendingRequestIndex + i) % MaxPendingRequests;
 			FPendingRequest& PendingRequest = PendingRequests[PendingRequestIndex];
+			LOCK_PENDING_REQUEST(PendingRequest);
 
 			// Skip if request was cancelled
 			if (!PendingRequest.IsValid())
@@ -762,6 +774,7 @@ void FStreamingManager::RemoveInternal(UStreamableSparseVolumeTexture* SparseVol
 		// Cancel any pending requests
 		for (FPendingRequest& PendingRequest : PendingRequests)
 		{
+			LOCK_PENDING_REQUEST(PendingRequest);
 			if (PendingRequest.SVTHandle == SVTInfo->SVTHandle)
 			{
 				PendingRequest.Reset();
@@ -1027,6 +1040,7 @@ void FStreamingManager::IssueRequests()
 				if (uint32* PendingRequestIndexPtr = FrameInfoTmp.TileIndexToPendingRequestIndex.Find((uint32)PreviousAllocation.TileIndexInFrame))
 				{
 					FPendingRequest& PendingRequest = PendingRequests[*PendingRequestIndexPtr];
+					LOCK_PENDING_REQUEST(PendingRequest);
 					check(PendingRequest.IsValid());
 					if (FrameInfoTmp.StreamingTiles.CountSetBits(PendingRequest.TileOffset, PendingRequest.TileOffset + PendingRequest.TileCount) == 0)
 					{
@@ -1053,6 +1067,7 @@ void FStreamingManager::IssueRequests()
 		{
 			const int32 PendingRequestIndex = NextPendingRequestIndex;
 			FPendingRequest& PendingRequest = PendingRequests[PendingRequestIndex];
+			LOCK_PENDING_REQUEST(PendingRequest);
 			PendingRequest.Reset();
 			PendingRequest.SVTHandle = TileRange.SVTHandle;
 			PendingRequest.FrameIndex = TileRange.FrameIndex;
@@ -1155,6 +1170,7 @@ int32 FStreamingManager::DetermineReadyRequests()
 	{
 		const int32 PendingRequestIndex = (StartPendingRequestIndex + i) % MaxPendingRequests;
 		FPendingRequest& PendingRequest = PendingRequests[PendingRequestIndex];
+		LOCK_PENDING_REQUEST(PendingRequest);
 
 		// Check if request was cancelled
 		if (!PendingRequest.IsValid())
@@ -1270,6 +1286,7 @@ void FStreamingManager::InstallReadyRequests()
 	{
 		const int32 PendingRequestIndex = (StartPendingRequestIndex + i) % MaxPendingRequests;
 		FPendingRequest& PendingRequest = PendingRequests[PendingRequestIndex];
+		LOCK_PENDING_REQUEST(PendingRequest);
 
 		// Skip if request was cancelled
 		if (!PendingRequest.IsValid())
@@ -1396,6 +1413,7 @@ void FStreamingManager::InstallReadyRequests()
 	for (int32 PendingRequestIndex : RequestsToCleanUp)
 	{
 		FPendingRequest& PendingRequest = PendingRequests[PendingRequestIndex];
+		LOCK_PENDING_REQUEST(PendingRequest);
 #if WITH_EDITORONLY_DATA
 		PendingRequest.SharedBuffer.Reset();
 #endif
@@ -1410,7 +1428,9 @@ void FStreamingManager::InstallReadyRequests()
 	for (int32 i = 0; i < AsyncState.NumReadyRequests; ++i)
 	{
 		const int32 PendingRequestIndex = (StartPendingRequestIndex + i) % MaxPendingRequests;
-		PendingRequests[PendingRequestIndex].Reset();
+		FPendingRequest& PendingRequest = PendingRequests[PendingRequestIndex];
+		LOCK_PENDING_REQUEST(PendingRequest);
+		PendingRequest.Reset();
 	}
 #endif
 
@@ -1628,13 +1648,15 @@ void FStreamingManager::RequestDDCData(TConstArrayView<UE::DerivedData::FCacheGe
 				const uint32 PendingRequestIndex = (uint32)(Response.UserData >> uint64(32));
 				const uint32 RequestVersion = (uint32)Response.UserData;
 
-				// In case the request returned after the mip level was already streamed out again we need to abort so that we do not overwrite data in the FPendingRequest slot.
-				if (RequestVersion < PendingRequests[PendingRequestIndex].RequestVersion)
+				FPendingRequest& PendingRequest = PendingRequests[PendingRequestIndex];
+				LOCK_PENDING_REQUEST(PendingRequest);
+
+				// In case the request returned after the data was already streamed out again we need to abort so that we do not overwrite data in the FPendingRequest slot.
+				if (RequestVersion != PendingRequest.RequestVersion)
 				{
 					return;
 				}
 
-				FPendingRequest& PendingRequest = PendingRequests[PendingRequestIndex];
 				check(PendingRequest.IsValid());
 
 				if (Response.Status == EStatus::Ok)
@@ -1783,3 +1805,5 @@ void FTileAllocator::Free(uint32 TileCoord)
 
 }
 }
+
+#undef LOCK_PENDING_REQUEST
diff --git a/Engine/Source/Runtime/Engine/Private/SparseVolumeTexture/SparseVolumeTextureStreamingManager.h b/Engine/Source/Runtime/Engine/Private/SparseVolumeTexture/SparseVolumeTextureStreamingManager.h
--- a/Engine/Source/Runtime/Engine/Private/SparseVolumeTexture/SparseVolumeTextureStreamingManager.h
+++ b/Engine/Source/Runtime/Engine/Private/SparseVolumeTexture/SparseVolumeTextureStreamingManager.h
@@ -163,6 +163,8 @@ private:
 		uint32 TileOffset = INDEX_NONE;
 		uint32 TileCount = 0;
 #if WITH_EDITORONLY_DATA
+		// When in editor, DDC requests finish on another thread, so we need to guard against race conditions.
+		FCriticalSection DDCAsyncGuard;
 		FSharedBuffer SharedBuffer;
 		enum class EState
 		{
@@ -275,7 +277,7 @@ private:
 	TMap<UStreamableSparseVolumeTexture*, uint16> SparseVolumeTextureToHandle; // Do not dereference the key! We just read the pointer itself.
 	TSparseArray<TUniquePtr<FStreamingInfo>> StreamingInfo; 
 	TMap<FFrameKey, FRequestPayload> RequestsHashTable;
-	TArray<FPendingRequest> PendingRequests;
+	TArray<FPendingRequest> PendingRequests; // Do not access any members of an element before first calling LOCK_PENDING_REQUEST() on the element!
 #if WITH_EDITORONLY_DATA
 	TUniquePtr<UE::DerivedData::FRequestOwner> RequestOwner;
 	TUniquePtr<UE::DerivedData::FRequestOwner> RequestOwnerBlocking;
