Metasounds: WaveProxy stops decoding on decoder error
AudioStreaming: Stop trying to decode on decoder error
Adding logging to help out in finding out what's breaking our decoders.

diff --git a/Engine/Source/Runtime/Engine/Private/AudioDecompress.cpp b/Engine/Source/Runtime/Engine/Private/AudioDecompress.cpp
--- a/Engine/Source/Runtime/Engine/Private/AudioDecompress.cpp
+++ b/Engine/Source/Runtime/Engine/Private/AudioDecompress.cpp
@@ -4,6 +4,7 @@
 #include "AudioDecompress.h"
 #include "AudioDevice.h"
 #include "Engine/Engine.h"
+#include "Hash/xxhash.h"
 #include "Interfaces/IAudioFormat.h"
 #include "AudioStreamingCache.h"
 #include "Misc/CoreStats.h"
@@ -458,13 +459,26 @@ bool IStreamedCompressedInfo::StreamCompressedData(uint8* Destination, bool bLoo
 	
 	while (RawPCMOffset < BufferSize)
 	{
+		if (HasError())
+		{
+			// If we've encountered an error, just always write zeroes and don't log anything because we already have
+			// and we don't want to spam every tick.
+			LastPCMByteSize = 0;
+			ZeroBuffer(Destination + RawPCMOffset, BufferSize - RawPCMOffset);
+
+			// We return true here to tell callers that we've completed decoding and should
+			// be looped or otherwise terminated.
+			return true;
+		}
+
 		// Decompress the next compression frame of audio (many samples) into the PCM buffer
 		const int32 DecodedFrames = DecompressToPCMBuffer(/*Unused*/ 0);
 
 		if (DecodedFrames < 0)
 		{
-			UE_LOG(LogAudioStreamCaching, Warning, TEXT("Zero pad buffer Chunk=%d, Wave=%s, Reason=Decoder returned negative samples."),
-				CurrentChunkIndex, *StreamingSoundWave->GetFName().ToString());
+			FXxHash64 Hash = FXxHash64::HashBuffer(SrcBufferData, SrcBufferDataSize);
+			UE_LOG(LogAudioStreamCaching, Warning, TEXT("Decoder error! Zero padding and terminating... Chunk=%d, Wave=%s DecodedFrames=%d SrcBufferOffset=%d SrcBufferDataSize=%d ChunkXxHash64=0x%llx"),
+				CurrentChunkIndex, *StreamingSoundWave->GetFName().ToString(), DecodedFrames, SrcBufferOffset, SrcBufferDataSize, Hash.Hash);
 
 			// Flag that the decoder has an unrecoverable error, so we don't try again.
 			bHasError = true;
diff --git a/Engine/Source/Runtime/Engine/Private/SoundWaveProxyReader.cpp b/Engine/Source/Runtime/Engine/Private/SoundWaveProxyReader.cpp
--- a/Engine/Source/Runtime/Engine/Private/SoundWaveProxyReader.cpp
+++ b/Engine/Source/Runtime/Engine/Private/SoundWaveProxyReader.cpp
@@ -545,6 +545,11 @@ FSoundWaveProxyReader::EDecodeResult FSoundWaveProxyReader::Decode()
 			bFinished = CompressedAudioInfo->ReadCompressedData(Buff, false, BuffSizeInBytes);
 		}
 
+		if (CompressedAudioInfo->HasError())
+		{
+			return EDecodeResult::Fail;
+		}
+
 		if (NumBytesStreamed == 0)
 		{
 			break;
