Fixed the two biggest causes for flickering when playing back SparseVolumeTextures: Exceeding the available IO bandwidth and reaching the streaming pool size limit.

IO Bandwidth:
The SVT streamer now features a cvar (r.SparseVolumeTexture.Streaming.BandwidthLimit) which is used to scale down the bandwidth which is used to stream in each streaming instance of a SVT. When the sum of the requested bandwidth of all streaming SVTs exceeds this value, a fractional mip level is computed which is the lowest (highest resolution) mip level which the streamer can continuously stream without exceeding the limit. In order to keep track of the individual streaming instances, there is now a new FStreamingInstance class, which replaces the old FStreamingWindow and is used to keep track for playback direction and to compute the requested and in-budget IO bandwidth. Users calling into the SVT streaming system can pass a uint32 "key" which is used to uniquely identify streaming instances across system ticks. This key makes it possible to tell requests apart that request the same frames but belong to different HV components or SVT viewers and might play at different frame rates.

Streaming Pool Size:
Exposed properties on the SVT asset to allow users to configure the streaming pool size. We now also expose per-SVT settings for prefetching in the same way. When these settings are changed at editor time, the SVT is simply reregistered with the system, which makes the new parameters take effect. Ideally we would automatically grow/resize the streaming pool, but for now this seems like a good intermediate solution.

In order to debug the bandwidth throttling, it turned out to be very helpful to have a way of visualizing what the streamer is doing, so this CL also introduces a debug visualization mode that can be enabled with r.SparseVolumeTexture.Streaming.ShowDebugInfo.

#rnx

diff --git a/Engine/Plugins/FX/Niagara/Source/Niagara/Private/NiagaraDataInterfaceRenderTargetVolume.cpp b/Engine/Plugins/FX/Niagara/Source/Niagara/Private/NiagaraDataInterfaceRenderTargetVolume.cpp
--- a/Engine/Plugins/FX/Niagara/Source/Niagara/Private/NiagaraDataInterfaceRenderTargetVolume.cpp
+++ b/Engine/Plugins/FX/Niagara/Source/Niagara/Private/NiagaraDataInterfaceRenderTargetVolume.cpp
@@ -769,8 +769,10 @@ bool UNiagaraDataInterfaceRenderTargetVolume::SimCacheReadFrame(UObject* Storage
 		UAnimatedSparseVolumeTexture* SVT = Cast<UAnimatedSparseVolumeTexture>(StorageObject);
 		
 		const int32 MipLevel = 0;
+		const float FrameRate = 0.0f;
 		const bool bBlocking = true;
-		USparseVolumeTextureFrame *SVTFrame = USparseVolumeTextureFrame::GetFrameAndIssueStreamingRequest(SVT, FMath::RoundToInt(FrameA + Interp), MipLevel, bBlocking);
+		const bool bHasFrameRate = false;
+		USparseVolumeTextureFrame *SVTFrame = USparseVolumeTextureFrame::GetFrameAndIssueStreamingRequest(SVT, GetTypeHash(this), MipLevel, FMath::RoundToInt(FrameA + Interp), MipLevel, bBlocking, bHasFrameRate);
 		
 		// The streaming manager normally ticks in FDeferredShadingSceneRenderer::Render(), but the SVT->DenseTexture conversion compute shader happens in a render command before that.
 		// At execution time of that command, the streamer hasn't had the chance to do any streaming yet, so we force another tick here.
diff --git a/Engine/Source/Runtime/Engine/Classes/SparseVolumeTexture/SparseVolumeTexture.h b/Engine/Source/Runtime/Engine/Classes/SparseVolumeTexture/SparseVolumeTexture.h
--- a/Engine/Source/Runtime/Engine/Classes/SparseVolumeTexture/SparseVolumeTexture.h
+++ b/Engine/Source/Runtime/Engine/Classes/SparseVolumeTexture/SparseVolumeTexture.h
@@ -364,11 +364,15 @@ public:
 	virtual ~USparseVolumeTextureFrame() = default;
 
 	// Retrieves a frame from the given SparseVolumeTexture and also issues a streaming request for it. 
+	// StreamingInstanceKey can be any arbitrary value that is suitable to keep track of the source of requests for a given SVT. This key is used internally to associate
+	// incoming requests with prior requests issued for the same SVT. A good value to pass here might be the pointer of the component the SVT is used within.
+	// FrameRate is an optional argument which helps to more accurately predict the required bandwidth when using non-blocking requests.
 	// FrameIndex is of float type so that the streaming system can use the fractional part to more easily keep track of playback speed and direction (forward/reverse playback).
 	// MipLevel is the lowest mip level that the caller intends to use but does not guarantee that the mip is actually resident.
 	// If bBlocking is true, DDC streaming requests will block on completion, guaranteeing that the requested frame will have been streamed in after the next streaming system update.
 	// If streaming cooked data from disk, the highest priority will be used, but no guarantee is given.
-	static ENGINE_API USparseVolumeTextureFrame* GetFrameAndIssueStreamingRequest(USparseVolumeTexture* SparseVolumeTexture, float FrameIndex, int32 MipLevel, bool bBlocking);
+	// if bHasValidFrameRate is true, the FrameRate argument will be use to predict the required streaming IO bandwidth.
+	static ENGINE_API USparseVolumeTextureFrame* GetFrameAndIssueStreamingRequest(USparseVolumeTexture* SparseVolumeTexture, uint32 StreamingInstanceKey, float FrameRate, float FrameIndex, int32 MipLevel, bool bBlocking, bool bHasValidFrameRate);
 
 	ENGINE_API bool Initialize(USparseVolumeTexture* InOwner, int32 InFrameIndex, const FTransform& InFrameTransform, UE::SVT::FTextureData& UncookedFrame);
 	int32 GetFrameIndex() const { return FrameIndex; }
@@ -476,6 +480,24 @@ public:
 	UPROPERTY(EditAnywhere, BlueprintReadWrite, Category = "Texture", AdvancedDisplay)
 	bool bLocalDDCOnly = true;
 
+	// The SVT streaming pool is sized such that it can hold the largest frame multiplied by this value. There should be some slack to allow for prefetching frames.
+	UPROPERTY(EditAnywhere, Category = "Texture", AdvancedDisplay)
+	float StreamingPoolSizeFactor = 3.0f;
+
+	// When using non-blocking streaming requests, upcoming frames are loaded into memory in advance. This property controls how many frames to prefetch.
+	UPROPERTY(EditAnywhere, Category = "Texture", AdvancedDisplay)
+	int32 NumberOfPrefetchFrames = 3;
+
+	// When using non-blocking streaming requests, upcoming frames are loaded into memory in advance. This property controls the size reduction in percent of each additional prefetched frames.
+	// A value of 20.0 would prefetch frame N+1 at 80%, N+2 at 60%, N+3 at 40% etc.
+	UPROPERTY(EditAnywhere, Category = "Texture", AdvancedDisplay)
+	float PrefetchPercentageStepSize = 20.0f;
+
+	// When using non-blocking streaming requests, upcoming frames are loaded into memory in advance. This property applies a bias in percent to how much data is prefetched for every frame.
+	// A value of 20.0 adds 20% to all prefetch percentages. So if PrefetchPercentageStepSize is set to 20.0, frame N+1 is prefetched at 80% + 20% = 100%, frame N+2 at 60% + 20% = 80%, N+3 at 40% + 20% = 60% etc.
+	UPROPERTY(EditAnywhere, Category = "Texture", AdvancedDisplay)
+	float PrefetchPercentageBias = 20.0f;
+
 #if WITH_EDITORONLY_DATA
 	UPROPERTY(VisibleAnywhere, Instanced, Category = ImportSettings)
 	TObjectPtr<class UAssetImportData> AssetImportData;
diff --git a/Engine/Source/Runtime/Engine/Private/Components/HeterogeneousVolumeComponent.cpp b/Engine/Source/Runtime/Engine/Private/Components/HeterogeneousVolumeComponent.cpp
--- a/Engine/Source/Runtime/Engine/Private/Components/HeterogeneousVolumeComponent.cpp
+++ b/Engine/Source/Runtime/Engine/Private/Components/HeterogeneousVolumeComponent.cpp
@@ -580,8 +580,9 @@ void UHeterogeneousVolumeComponent::TickComponent(float DeltaTime, ELevelTick Ti
 				Frame = FMath::Clamp(Frame, StartFrame, EndFrame);
 			}
 
-			bool bIsBlocking = bIssueBlockingRequests != 0;
-			USparseVolumeTextureFrame* SparseVolumeTextureFrame = USparseVolumeTextureFrame::GetFrameAndIssueStreamingRequest(SparseVolumeTexture, Frame, MipLevel, bIsBlocking);
+			const bool bIsBlocking = bIssueBlockingRequests != 0;
+			const bool bHasValidFrameRate = bPlaying != 0;
+			USparseVolumeTextureFrame* SparseVolumeTextureFrame = USparseVolumeTextureFrame::GetFrameAndIssueStreamingRequest(SparseVolumeTexture, GetTypeHash(this), FrameRate, Frame, MipLevel, bIsBlocking, bHasValidFrameRate);
 			if (SparseVolumeTextureFrame)
 			{
 				FIntVector PerFrameVolumeResolution = SparseVolumeTextureFrame->GetVolumeResolution();
diff --git a/Engine/Source/Runtime/Engine/Private/SparseVolumeTexture/SparseVolumeTexture.cpp b/Engine/Source/Runtime/Engine/Private/SparseVolumeTexture/SparseVolumeTexture.cpp
--- a/Engine/Source/Runtime/Engine/Private/SparseVolumeTexture/SparseVolumeTexture.cpp
+++ b/Engine/Source/Runtime/Engine/Private/SparseVolumeTexture/SparseVolumeTexture.cpp
@@ -939,11 +939,14 @@ USparseVolumeTextureFrame::USparseVolumeTextureFrame(const FObjectInitializer& O
 {
 }
 
-USparseVolumeTextureFrame* USparseVolumeTextureFrame::GetFrameAndIssueStreamingRequest(USparseVolumeTexture* SparseVolumeTexture, float FrameIndex, int32 MipLevel, bool bBlocking)
+USparseVolumeTextureFrame* USparseVolumeTextureFrame::GetFrameAndIssueStreamingRequest(USparseVolumeTexture* SparseVolumeTexture, uint32 StreamingInstanceKey, float FrameRate, float FrameIndex, int32 MipLevel, bool bBlocking, bool bHasValidFrameRate)
 {
 	if (UStreamableSparseVolumeTexture* StreamableSVT = Cast<UStreamableSparseVolumeTexture>(SparseVolumeTexture))
 	{
-		UE::SVT::GetStreamingManager().Request_GameThread(StreamableSVT, FrameIndex, MipLevel, bBlocking);
+		UE::SVT::EStreamingRequestFlags RequestFlags = UE::SVT::EStreamingRequestFlags::None;
+		RequestFlags |= bBlocking ? UE::SVT::EStreamingRequestFlags::Blocking : UE::SVT::EStreamingRequestFlags::None;
+		RequestFlags |= bHasValidFrameRate ? UE::SVT::EStreamingRequestFlags::HasFrameRate : UE::SVT::EStreamingRequestFlags::None;
+		UE::SVT::GetStreamingManager().Request_GameThread(StreamableSVT, StreamingInstanceKey, FrameRate, FrameIndex, MipLevel, RequestFlags);
 		return StreamableSVT->GetFrame(static_cast<int32>(FrameIndex));
 	}
 	return nullptr;
@@ -1409,7 +1412,21 @@ void UStreamableSparseVolumeTexture::PostEditChangeProperty(FPropertyChangedEven
 	Super::PostEditChangeProperty(PropertyChangedEvent);
 
 	// Don't bother trying to recache the frame data if the Frames array is invalid. This very likely means that this object is about to be deleted.
-	if (!bInvalidFramesArray)
+	bool bRecacheFrames = !bInvalidFramesArray;
+
+	if (PropertyChangedEvent.GetPropertyName() == GET_MEMBER_NAME_CHECKED(UStreamableSparseVolumeTexture, StreamingPoolSizeFactor)
+		|| PropertyChangedEvent.GetPropertyName() == GET_MEMBER_NAME_CHECKED(UStreamableSparseVolumeTexture, NumberOfPrefetchFrames)
+		|| PropertyChangedEvent.GetPropertyName() == GET_MEMBER_NAME_CHECKED(UStreamableSparseVolumeTexture, PrefetchPercentageStepSize)
+		|| PropertyChangedEvent.GetPropertyName() == GET_MEMBER_NAME_CHECKED(UStreamableSparseVolumeTexture, PrefetchPercentageBias))
+	{
+		// Re-register the SVT with the streamer so it picks up on the changed streaming parameters.
+		UE::SVT::GetStreamingManager().Remove_GameThread(this);
+		UE::SVT::GetStreamingManager().Add_GameThread(this);
+		
+		bRecacheFrames = false;
+	}
+	
+	if (bRecacheFrames)
 	{
 		RecacheFrames();
 	}
@@ -1659,7 +1676,7 @@ USparseVolumeTextureFrame* UAnimatedSparseVolumeTextureController::GetFrameByInd
 		return nullptr;
 	}
 
-	return USparseVolumeTextureFrame::GetFrameAndIssueStreamingRequest(SparseVolumeTexture, FrameIndex, MipLevel, bBlockingStreamingRequests);
+	return USparseVolumeTextureFrame::GetFrameAndIssueStreamingRequest(SparseVolumeTexture, GetTypeHash(this), 0.0f /*FrameRate*/, FrameIndex, MipLevel, bBlockingStreamingRequests, false /*bHasValidFrameRate*/);
 }
 
 USparseVolumeTextureFrame* UAnimatedSparseVolumeTextureController::GetCurrentFrame()
@@ -1672,7 +1689,7 @@ USparseVolumeTextureFrame* UAnimatedSparseVolumeTextureController::GetCurrentFra
 	// Compute (fractional) index of frame to sample
 	const float FrameIndexF = GetFractionalFrameIndex();
 
-	return USparseVolumeTextureFrame::GetFrameAndIssueStreamingRequest(SparseVolumeTexture, FrameIndexF, MipLevel, bBlockingStreamingRequests);
+	return USparseVolumeTextureFrame::GetFrameAndIssueStreamingRequest(SparseVolumeTexture, GetTypeHash(this), FrameRate, FrameIndexF, MipLevel, bBlockingStreamingRequests, true /*bHasValidFrameRate*/);
 }
 
 void UAnimatedSparseVolumeTextureController::GetCurrentFramesForInterpolation(USparseVolumeTextureFrame*& Frame0, USparseVolumeTextureFrame*& Frame1, float& LerpAlpha)
@@ -1687,8 +1704,9 @@ void UAnimatedSparseVolumeTextureController::GetCurrentFramesForInterpolation(US
 	const int32 FrameIndex = (int32)FrameIndexF;
 	LerpAlpha = FMath::Frac(FrameIndexF);
 
-	Frame0 = USparseVolumeTextureFrame::GetFrameAndIssueStreamingRequest(SparseVolumeTexture, FrameIndexF, MipLevel, bBlockingStreamingRequests);
-	Frame1 = USparseVolumeTextureFrame::GetFrameAndIssueStreamingRequest(SparseVolumeTexture, (FrameIndex + 1) % SparseVolumeTexture->GetNumFrames(), MipLevel, bBlockingStreamingRequests);
+	const uint32 StreamingInstanceKey = GetTypeHash(this);
+	Frame0 = USparseVolumeTextureFrame::GetFrameAndIssueStreamingRequest(SparseVolumeTexture, StreamingInstanceKey, FrameRate, FrameIndexF, MipLevel, bBlockingStreamingRequests, true /*bHasValidFrameRate*/);
+	Frame1 = USparseVolumeTextureFrame::GetFrameAndIssueStreamingRequest(SparseVolumeTexture, StreamingInstanceKey, FrameRate, (FrameIndex + 1) % SparseVolumeTexture->GetNumFrames(), MipLevel, bBlockingStreamingRequests, true /*bHasValidFrameRate*/);
 }
 
 float UAnimatedSparseVolumeTextureController::GetDuration()
diff --git a/Engine/Source/Runtime/Engine/Private/SparseVolumeTexture/SparseVolumeTextureStreamingInstance.cpp b/Engine/Source/Runtime/Engine/Private/SparseVolumeTexture/SparseVolumeTextureStreamingInstance.cpp
new file mode 100644
--- /dev/null
+++ b/Engine/Source/Runtime/Engine/Private/SparseVolumeTexture/SparseVolumeTextureStreamingInstance.cpp
@@ -0,0 +1,231 @@
+// Copyright Epic Games, Inc. All Rights Reserved.
+
+#include "SparseVolumeTextureStreamingInstance.h"
+#include "SparseVolumeTexture/ISparseVolumeTextureStreamingManager.h"
+#include "HAL/IConsoleManager.h"
+
+static float GSVTStreamingInstanceWindowSize = 5.0f;
+static FAutoConsoleVariableRef CVarSVTStreamingInstanceWindowSize(
+	TEXT("r.SparseVolumeTexture.Streaming.InstanceWindowSize"),
+	GSVTStreamingInstanceWindowSize,
+	TEXT("Window size in SVT frames around which requests to a SVT are considered to belong to the same instance. A streaming instance is an internal book keeping object to track playback of a SVT asset in a given context. Default: 5.0"),
+	ECVF_RenderThreadSafe
+);
+
+static int32 GSVTStreamingForceEstimateFrameRate = 0;
+static FAutoConsoleVariableRef CVarSVTStreamingForceEstimateFrameRate(
+	TEXT("r.SparseVolumeTexture.Streaming.ForceEstimateFrameRate"),
+	GSVTStreamingForceEstimateFrameRate,
+	TEXT("Forces the streaming system to always estimate the SVT playback frame rate instead of using the explicit frame rate that is passed in with the streaming requests. Intended for debugging."),
+	ECVF_RenderThreadSafe
+);
+
+namespace UE
+{
+namespace SVT
+{
+
+bool FStreamingInstanceRequest::IsBlocking() const
+{
+	return EnumHasAnyFlags(Flags, EStreamingRequestFlags::Blocking);
+}
+
+bool FStreamingInstanceRequest::HasValidFrameRate() const
+{
+	return EnumHasAnyFlags(Flags, EStreamingRequestFlags::HasFrameRate);
+}
+
+FStreamingInstance::FStreamingInstance(uint32 InKey, int32 InNumFrames, const TArrayView<uint32>& InMipLevelStreamingSizes, const FStreamingInstanceRequest& Request)
+	: MipLevelStreamingSizes(InMipLevelStreamingSizes)
+	, Key(InKey)
+	, NumFrames(InNumFrames)
+{
+	PreviousAverageRequestIssueTime = Request.Time;
+	AverageRequestIssueTime = Request.Time;
+	DeltaTime = 0.0;
+	UpdateIndex = Request.UpdateIndex;
+	PreviousAverageFrame = Request.FrameIndex;
+	AverageFrame = Request.FrameIndex;
+	NumRequestsThisUpdate = 1;
+	LowestRequestedMip = Request.IsBlocking() ? FLT_MAX : Request.MipLevel;
+	LowestRequestedBlockingMip = Request.IsBlocking() ? Request.MipLevel : FLT_MAX;
+	EstimatedFrameRate = Request.HasValidFrameRate() ? Request.FrameRate : 0.0f;
+	bPlayForwards = true; // No prior data, so just take a guess that playback is forwards
+	bPlayBackwards = false;
+	bIsBlocking = Request.IsBlocking();
+}
+
+void FStreamingInstance::AddRequest(const FStreamingInstanceRequest& Request)
+{
+	if (Request.UpdateIndex > UpdateIndex)
+	{
+		PreviousAverageRequestIssueTime = AverageRequestIssueTime;
+		UpdateIndex = Request.UpdateIndex;
+		PreviousAverageFrame = AverageFrame;
+		NumRequestsThisUpdate = 0;
+		LowestRequestedMip = FLT_MAX;
+		LowestRequestedBlockingMip = FLT_MAX;
+		bPlayForwards = false;
+		bPlayBackwards = false;
+		bIsBlocking = false;
+	}
+
+	// Calculate the average requested frame this update, taking into account potential wrapping when using looping playback. If this is the first request this update, just use its FrameIndex directly.
+	AverageFrame = NumRequestsThisUpdate == 0 ? Request.FrameIndex : GetWrappedWeightedAverage(Request.FrameIndex, 1.0f, AverageFrame, NumRequestsThisUpdate, NumFrames);
+	// Compute the new average request issue time based on the newly added request, the current average time and the number of already averaged requests.
+	AverageRequestIssueTime = (AverageRequestIssueTime * NumRequestsThisUpdate + Request.Time) / (NumRequestsThisUpdate + 1.0f);
+	// Estimate the the DeltaTime (and later down the frame rate) based on the average times requests were issued in the previous frame/update and the current one.
+	DeltaTime = FMath::Max(0.0, AverageRequestIssueTime - PreviousAverageRequestIssueTime);
+	EstimatedFrameRate = 0.0f;
+	if (Request.HasValidFrameRate() && GSVTStreamingForceEstimateFrameRate == 0)
+	{
+		// Use the explicitly passed in frame rate if possible.
+		EstimatedFrameRate = Request.FrameRate;
+	}
+	else if (DeltaTime > 0.0)
+	{
+		// Otherwise we estimate it based on the difference in requested frames and delta time.
+		EstimatedFrameRate = FMath::Abs(GetShortestWrappedDistance(AverageFrame, PreviousAverageFrame, NumFrames)) / DeltaTime;
+	}
+	if (Request.IsBlocking())
+	{
+		LowestRequestedBlockingMip = FMath::Min(LowestRequestedBlockingMip, Request.MipLevel);
+	}
+	else
+	{
+		LowestRequestedMip = FMath::Min(LowestRequestedMip, Request.MipLevel);
+	}
+	const float DistanceFromPreviousFrame = GetShortestWrappedDistance(Request.FrameIndex, PreviousAverageFrame, NumFrames);
+	bPlayForwards = bPlayForwards || (DistanceFromPreviousFrame > 0.0f);
+	bPlayBackwards = bPlayBackwards || (DistanceFromPreviousFrame < 0.0f);
+	bIsBlocking = bIsBlocking || Request.IsBlocking();
+
+	++NumRequestsThisUpdate;
+}
+
+void FStreamingInstance::ComputeLowestMipLevelInBandwidthBudget(int64 BandwidthBudget)
+{
+	if (BandwidthBudget < 0)
+	{
+		// No bandwidth limit
+		LowestMipInBandwidthBudget = 0.0f;
+	}
+	else
+	{
+		LowestMipInBandwidthBudget = MipLevelStreamingSizes.Num() - 1;
+
+		// Walk from highest to lowest mip and find the fractional mip that can be streamed continuously without exceeding the budget. Could probably use binary search here in the future. 
+		int64 PrevMipStreamingBandwidth = 0;
+		for (int32 MipLevelIndex = MipLevelStreamingSizes.Num() - 1; MipLevelIndex >= 0; --MipLevelIndex)
+		{
+			const int64 MipStreamingSize = MipLevelStreamingSizes[MipLevelIndex];
+			// Technically this assumes 1FPS for instances that aren't playing back at all (FrameRate==0.0), but this should be fine because we may still
+			// need to stream in data for the requested frame and it is probably better to err on the safe side of things.
+			const int64 MipStreamingBandwidth = FMath::CeilToInt64(MipStreamingSize * FMath::Max(1.0f, GetEstimatedFrameRate()));
+			if (MipStreamingBandwidth > BandwidthBudget)
+			{
+				check(MipStreamingBandwidth >= PrevMipStreamingBandwidth);
+				check(BandwidthBudget >= PrevMipStreamingBandwidth);
+				const int64 NumBytesFromPrevMip = MipStreamingBandwidth - PrevMipStreamingBandwidth;
+				const int64 NumBytesToLimit = BandwidthBudget - PrevMipStreamingBandwidth;
+				const float MipFraction = static_cast<float>(NumBytesToLimit / static_cast<double>(NumBytesFromPrevMip));
+				LowestMipInBandwidthBudget = (MipLevelIndex + 1) - MipFraction;
+				break;
+			}
+			LowestMipInBandwidthBudget = MipLevelIndex;
+			PrevMipStreamingBandwidth = MipStreamingBandwidth;
+		}
+	}
+}
+
+int64 FStreamingInstance::GetRequestedBandwidth(bool bZeroIfBlocking) const
+{
+	if (bZeroIfBlocking && bIsBlocking)
+	{
+		return 0;
+	}
+	const int32 MipLevelIndex = FMath::Clamp(FMath::FloorToInt32(LowestRequestedMip), 0, MipLevelStreamingSizes.Num() - 1);
+	check(MipLevelStreamingSizes.IsValidIndex(MipLevelIndex));
+	const int64 RequestedBandwidth = FMath::CeilToInt64(MipLevelStreamingSizes[MipLevelIndex] * FMath::Max(1.0f, GetEstimatedFrameRate()));
+	return RequestedBandwidth;
+}
+
+bool FStreamingInstance::IsFrameInWindow(float FrameIndex) const
+{
+	if (FrameIndex < 0.0f || FrameIndex >= static_cast<float>(NumFrames))
+	{
+		return false;
+	}
+	const float Distance = GetShortestWrappedDistance(AverageFrame, FrameIndex, NumFrames);
+	const float WindowSize = FMath::Max(UE_SMALL_NUMBER, GSVTStreamingInstanceWindowSize);
+	return FMath::Abs(Distance) <= WindowSize;
+}
+
+float FStreamingInstance::GetPrefetchMipLevel(float RequestedMipLevel, float Percentage) const
+{
+	const int32 RequestedMipLevelIndex = FMath::Clamp(FMath::FloorToInt32(RequestedMipLevel), 0, MipLevelStreamingSizes.Num() - 1);
+	check(MipLevelStreamingSizes.IsValidIndex(RequestedMipLevelIndex));
+	const uint32 RequestedMipLevelStreamingSize = MipLevelStreamingSizes[RequestedMipLevelIndex];
+	const int64 PrefetchStreamingSize = FMath::CeilToInt64(RequestedMipLevelStreamingSize * FMath::Clamp(Percentage, 0.0f, 1.0f));
+
+	// Walk from the highest to lowest mip level and find the fractional mip level which corresponds to StreamingMemorySizeOf(RequestedMipLevel) * Percentage.
+	int64 PrevMipStreamingSize = 0;
+	for (int32 MipLevelIndex = MipLevelStreamingSizes.Num() - 1; MipLevelIndex >= 0; --MipLevelIndex)
+	{
+		const int64 MipStreamingSize = MipLevelStreamingSizes[MipLevelIndex];
+		if (MipStreamingSize > PrefetchStreamingSize)
+		{
+			check(MipStreamingSize >= PrevMipStreamingSize);
+			check(PrefetchStreamingSize >= PrevMipStreamingSize);
+			const int64 NumBytesFromPrevMip = MipStreamingSize - PrevMipStreamingSize;
+			const int64 NumBytesToLimit = PrefetchStreamingSize - PrevMipStreamingSize;
+			const float MipFraction = static_cast<float>(NumBytesToLimit / static_cast<double>(NumBytesFromPrevMip));
+			return (MipLevelIndex + 1) - MipFraction;
+		}
+		PrevMipStreamingSize = MipStreamingSize;
+	}
+	return 0.0f;
+}
+
+float FStreamingInstance::GetShortestWrappedDistance(float To, float From, int32 NumFrames)
+{
+	const float NumFramesF = static_cast<float>(NumFrames);
+	check(NumFramesF >= 0.0f);
+	check(From >= 0.0f && From < NumFramesF);
+	check(To >= 0.0f && To < NumFramesF);
+	const float Distance = To - From;
+	// The wrapped distance must be the other way around, so we simply do NumFrames - NonWrappedDistance. In order to preserve
+	// the sign, we need to operate on absolute values and add it back later.
+	const float WrappedDistance = FMath::Max(0.0f, (NumFramesF - FMath::Abs(Distance))) * FMath::Sign(-Distance);
+	const float Result = FMath::Abs(Distance) < FMath::Abs(WrappedDistance) ? Distance : WrappedDistance;
+	return Result;
+}
+
+float FStreamingInstance::GetWrappedWeightedAverage(float ValueA, float WeightA, float ValueB, float WeightB, int32 NumFrames)
+{
+	const float NumFramesF = static_cast<float>(NumFrames);
+	const float ShortestWrappedDistance = GetShortestWrappedDistance(ValueA, ValueB, NumFrames);
+	const bool bIsWrapped = FMath::Abs(ShortestWrappedDistance) < FMath::Abs(ValueA - ValueB);
+
+	float ValueAUnwrapped = ValueA;
+	float ValueBUnwrapped = ValueB;
+	if (bIsWrapped && ShortestWrappedDistance > 0)
+	{
+		ValueAUnwrapped += NumFramesF; // Forward wrapping: From > To (ValueB > ValueA)
+	}
+	else if (bIsWrapped && ShortestWrappedDistance < 0)
+	{
+		ValueBUnwrapped += NumFramesF; // Reverse wrapping: From < To (ValueB < ValueA)
+	}
+
+	float Result = (ValueAUnwrapped * WeightA + ValueBUnwrapped * WeightB) / FMath::Max(WeightA + WeightB, UE_SMALL_NUMBER);
+	if (Result >= NumFramesF)
+	{
+		Result -= NumFramesF;
+	}
+
+	return Result;
+}
+
+}
+}
diff --git a/Engine/Source/Runtime/Engine/Private/SparseVolumeTexture/SparseVolumeTextureStreamingInstance.h b/Engine/Source/Runtime/Engine/Private/SparseVolumeTexture/SparseVolumeTextureStreamingInstance.h
new file mode 100644
--- /dev/null
+++ b/Engine/Source/Runtime/Engine/Private/SparseVolumeTexture/SparseVolumeTextureStreamingInstance.h
@@ -0,0 +1,81 @@
+// Copyright Epic Games, Inc. All Rights Reserved.
+
+#pragma once
+#include "CoreMinimal.h"
+
+namespace UE
+{
+namespace SVT
+{
+
+enum class EStreamingRequestFlags : uint8;
+
+// A streaming request that can be passed to FStreamingInstance.
+struct FStreamingInstanceRequest
+{
+	uint32 UpdateIndex = 0;
+	double Time = 0.0;
+	float FrameRate = 0.0f;
+	float FrameIndex = 0.0f;
+	float MipLevel = 0.0f;
+	EStreamingRequestFlags Flags = {};
+
+	explicit FStreamingInstanceRequest() = default;
+	explicit FStreamingInstanceRequest(uint32 InUpdateIndex, double InTime, float InFrameRate, float InFrameIndex, float InMipLevel, EStreamingRequestFlags InFlags)
+		: UpdateIndex(InUpdateIndex), Time(InTime), FrameRate(InFrameRate), FrameIndex(InFrameIndex), MipLevel(InMipLevel), Flags(InFlags) {}
+
+	bool IsBlocking() const;
+	bool HasValidFrameRate() const;
+};
+
+// Represents a streaming instance of a SVT. It is a "window" into the frame sequence that moves along the playback direction and is also identified by a "key", which is a user assigned value. 
+// The FStreamingInstance is used to cache the prefetch direction, to compute the total bandwidth requirements of all streaming SVTs and to determine the allocated bandwidth for a given SVT instance.
+class FStreamingInstance
+{
+public:
+	explicit FStreamingInstance(uint32 Key, int32 NumFrames, const TArrayView<uint32>& MipLevelStreamingSizes, const FStreamingInstanceRequest& Request);
+	void AddRequest(const FStreamingInstanceRequest& Request);
+	// Returns the lowest mip level which can be continuously streamed while staying within the given bandwidth budget. 
+	void ComputeLowestMipLevelInBandwidthBudget(int64 BandwidthBudget);
+	// Returns the bandwidth required to stream the SVT at the requested mip level. This is relevant for non-blocking streaming, so it can optionally return zero if the instance is using blocking requests.
+	int64 GetRequestedBandwidth(bool bZeroIfBlocking) const;
+	// Returns true if the given frame is within the "window" of this instance.
+	bool IsFrameInWindow(float FrameIndex) const;
+	// Computes a fractional mip level which corresponds to StreamingMemorySizeOf(RequestedMipLevel) * Percentage.
+	float GetPrefetchMipLevel(float RequestedMipLevel, float Percentage) const;
+	uint32 GetKey() const { return Key; }
+	uint32 GetUpdateIndex() const { return UpdateIndex; }
+	float GetEstimatedFrameRate() const { return EstimatedFrameRate; }
+	float GetAverageFrame() const { return AverageFrame; }
+	float GetLowestRequestedMipLevel() const { return LowestRequestedMip; }
+	float GetLowestMipLevelInBandwidthBudget() const { return LowestMipInBandwidthBudget; }
+	bool IsPlayingForwards() const { return bPlayForwards; }
+	bool IsPlayingBackwards() const { return bPlayBackwards; }
+
+private:
+	const TArrayView<uint32> MipLevelStreamingSizes;
+	const uint32 Key;
+	const int32 NumFrames;
+	double PreviousAverageRequestIssueTime = 0.0;
+	double AverageRequestIssueTime = 0.0;
+	double DeltaTime = 0.0;
+	uint32 UpdateIndex = 0;
+	float PreviousAverageFrame = -1.0f;
+	float AverageFrame = -1.0f; // Frame index this window is centered around
+	float EstimatedFrameRate = 0.0f;
+	int32 NumRequestsThisUpdate = 0;
+	float LowestRequestedMip = FLT_MAX;
+	float LowestRequestedBlockingMip = FLT_MAX;
+	float LowestMipInBandwidthBudget = 0.0f;
+	bool bPlayForwards = false;
+	bool bPlayBackwards = false;
+	bool bIsBlocking = false;
+
+	// Returns the shortest signed distance between From and To, potentially wrapping around the 0 - NumFrames range.
+	static float GetShortestWrappedDistance(float To, float From, int32 NumFrames);
+	// Computes a weighted average between ValueA and ValueB, applying wrapping logic to both the calculation and the result.
+	static float GetWrappedWeightedAverage(float ValueA, float WeightA, float ValueB, float WeightB, int32 NumFrames);
+};
+
+}
+}
diff --git a/Engine/Source/Runtime/Engine/Private/SparseVolumeTexture/SparseVolumeTextureStreamingManager.cpp b/Engine/Source/Runtime/Engine/Private/SparseVolumeTexture/SparseVolumeTextureStreamingManager.cpp
--- a/Engine/Source/Runtime/Engine/Private/SparseVolumeTexture/SparseVolumeTextureStreamingManager.cpp
+++ b/Engine/Source/Runtime/Engine/Private/SparseVolumeTexture/SparseVolumeTextureStreamingManager.cpp
@@ -11,6 +11,7 @@
 #include "Async/ParallelFor.h"
 #include "SparseVolumeTextureTileDataTexture.h"
 #include "SparseVolumeTextureUpload.h"
+#include "SparseVolumeTextureStreamingInstance.h"
 #include "GlobalRenderResources.h"
 
 #if WITH_EDITORONLY_DATA
@@ -22,22 +23,6 @@ DEFINE_LOG_CATEGORY(LogSparseVolumeTextureStreamingManager);
 
 DECLARE_GPU_STAT(SVTStreaming);
 
-static int32 GSVTStreamingNumPrefetchFrames = 3;
-static FAutoConsoleVariableRef CVarSVTStreamingNumPrefetchFrames(
-	TEXT("r.SparseVolumeTexture.Streaming.NumPrefetchFrames"),
-	GSVTStreamingNumPrefetchFrames,
-	TEXT("Number of frames to prefetch when a frame is requested."),
-	ECVF_RenderThreadSafe
-);
-
-static int32 GSVTStreamingPrefetchMipLevelBias = -1;
-static FAutoConsoleVariableRef CVarSVTStreamingPrefetchMipLevelBias(
-	TEXT("r.SparseVolumeTexture.Streaming.PrefetchMipLevelBias"),
-	GSVTStreamingPrefetchMipLevelBias,
-	TEXT("Bias to apply to the mip level of prefetched frames. Prefetching is done at increasingly higher mip levels (lower resolution), so setting a negative value here will increase the requested mip level resolution."),
-	ECVF_RenderThreadSafe
-);
-
 static int32 GSVTStreamingForceBlockingRequests = 0;
 static FAutoConsoleVariableRef CVarSVTStreamingForceBlockingRequests(
 	TEXT("r.SparseVolumeTexture.Streaming.ForceBlockingRequests"),
@@ -78,7 +63,7 @@ static FAutoConsoleVariableRef CVarSVTStreamingLogVerbosity(
 	ECVF_RenderThreadSafe
 );
 
-static int32 GSVTStreamingMaxPendingRequests = 1024;
+static int32 GSVTStreamingMaxPendingRequests = 256;
 static FAutoConsoleVariableRef CVarSVTStreamingMaxPendingRequests(
 	TEXT("r.SparseVolumeTexture.Streaming.MaxPendingRequests"),
 	GSVTStreamingMaxPendingRequests,
@@ -86,11 +71,27 @@ static FAutoConsoleVariableRef CVarSVTStreamingMaxPendingRequests(
 	ECVF_RenderThreadSafe | ECVF_ReadOnly
 );
 
-static int32 GSVTStreamingRequestSize = 1024;
+static int32 GSVTStreamingRequestSize = -1;
 static FAutoConsoleVariableRef CVarSVTStreamingRequestSize(
 	TEXT("r.SparseVolumeTexture.Streaming.RequestSize"),
 	GSVTStreamingRequestSize,
-	TEXT("IO request size in KiB. The SVT streaming manager will attempt to create IO requests of roughly this size. Default: 1024 KiB"),
+	TEXT("IO request size in KiB. The SVT streaming manager will attempt to create IO requests of roughly this size. Default: -1 (unlimited)"),
+	ECVF_RenderThreadSafe
+);
+
+static int32 GSVTStreamingBandwidthLimit = 512;
+static FAutoConsoleVariableRef CVarSVTStreamingBandwidthLimit(
+	TEXT("r.SparseVolumeTexture.Streaming.BandwidthLimit"),
+	GSVTStreamingBandwidthLimit,
+	TEXT("Bandwidth limit for SVT streaming in MiB/s. When requests exceed this limit, the system will stream at lower mip levels instead."),
+	ECVF_RenderThreadSafe
+);
+
+static int32 GSVTStreamingInstanceCleanupThreshold = 5;
+static FAutoConsoleVariableRef CVarSVTStreamingInstanceCleanupThreshold(
+	TEXT("r.SparseVolumeTexture.Streaming.InstanceCleanupThreshold"),
+	GSVTStreamingInstanceCleanupThreshold,
+	TEXT("Number of SVT streaming system updates to wait until unused streaming instances are cleaned up. A streaming instance is an internal book keeping object to track playback of a SVT asset in a given context. Default: 5"),
 	ECVF_RenderThreadSafe
 );
 
@@ -164,7 +165,7 @@ void FStreamingManager::InitRHI(FRHICommandListBase& RHICmdList)
 	PageTableUpdater = MakeUnique<FPageTableUpdater>();
 
 #if WITH_EDITORONLY_DATA
-	RequestOwner = MakeUnique<FRequestOwner>(EPriority::Normal);
+	RequestOwner = MakeUnique<FRequestOwner>(EPriority::Highest);
 	RequestOwnerBlocking = MakeUnique<FRequestOwner>(EPriority::Blocking);
 #endif
 }
@@ -192,6 +193,11 @@ void FStreamingManager::Add_GameThread(UStreamableSparseVolumeTexture* SparseVol
 	NewSVTInfo.FormatB = SparseVolumeTexture->GetFormat(1);
 	NewSVTInfo.FallbackValueA = SparseVolumeTexture->GetFallbackValue(0);
 	NewSVTInfo.FallbackValueB = SparseVolumeTexture->GetFallbackValue(1);
+	NewSVTInfo.NumMipLevelsGlobal = SparseVolumeTexture->GetNumMipLevels();
+	NewSVTInfo.StreamingPoolSizeFactor = SparseVolumeTexture->StreamingPoolSizeFactor;
+	NewSVTInfo.NumPrefetchFrames = SparseVolumeTexture->NumberOfPrefetchFrames;
+	NewSVTInfo.PrefetchPercentageStepSize = SparseVolumeTexture->PrefetchPercentageStepSize;
+	NewSVTInfo.PrefetchPercentageBias = SparseVolumeTexture->PrefetchPercentageBias;
 	NewSVTInfo.FrameInfo.SetNum(NumFrames);
 
 	for (int32 FrameIdx = 0; FrameIdx < NumFrames; ++FrameIdx)
@@ -228,16 +234,16 @@ void FStreamingManager::Remove_GameThread(UStreamableSparseVolumeTexture* Sparse
 		});
 }
 
-void FStreamingManager::Request_GameThread(UStreamableSparseVolumeTexture* SparseVolumeTexture, float FrameIndex, int32 MipLevel, bool bBlocking)
+void FStreamingManager::Request_GameThread(UStreamableSparseVolumeTexture* SparseVolumeTexture, uint32 StreamingInstanceKey, float FrameRate, float FrameIndex, int32 MipLevel, EStreamingRequestFlags Flags)
 {
 	if (!DoesPlatformSupportSparseVolumeTexture(GMaxRHIShaderPlatform) || !SparseVolumeTexture)
 	{
 		return;
 	}
 	ENQUEUE_RENDER_COMMAND(SVTRequest)(
-		[this, SparseVolumeTexture, FrameIndex, MipLevel, bBlocking](FRHICommandListImmediate& RHICmdList)
+		[this, SparseVolumeTexture, StreamingInstanceKey, FrameRate, FrameIndex, MipLevel, Flags](FRHICommandListImmediate& RHICmdList)
 		{
-			Request(SparseVolumeTexture, FrameIndex, MipLevel, bBlocking);
+			Request(SparseVolumeTexture, StreamingInstanceKey, FrameRate, FrameIndex, MipLevel, Flags);
 		});
 }
 
@@ -251,14 +257,14 @@ void FStreamingManager::Update_GameThread()
 		[](FRHICommandListImmediate& RHICmdList)
 		{
 			FRDGBuilder GraphBuilder(RHICmdList);
-			const bool bForceNonAsync = true; // No need to spin up a thread if we immediately wait on it anyways.
-			GStreamingManager.BeginAsyncUpdate(GraphBuilder, bForceNonAsync);
+			const bool bUseAsyncThread = false; // No need to spin up a thread if we immediately wait on it anyways.
+			GStreamingManager.BeginAsyncUpdate(GraphBuilder, bUseAsyncThread);
 			GStreamingManager.EndAsyncUpdate(GraphBuilder);
 			GraphBuilder.Execute();
 		});
 }
 
-void FStreamingManager::Request(UStreamableSparseVolumeTexture* SparseVolumeTexture, float FrameIndex, int32 MipLevel, bool bBlocking)
+void FStreamingManager::Request(UStreamableSparseVolumeTexture* SparseVolumeTexture, uint32 StreamingInstanceKey, float FrameRate, float FrameIndex, int32 MipLevel, EStreamingRequestFlags Flags)
 {
 	check(IsInRenderingThread());
 	if (!DoesPlatformSupportSparseVolumeTexture(GMaxRHIShaderPlatform) || !SparseVolumeTexture)
@@ -270,61 +276,23 @@ void FStreamingManager::Request(UStreamableSparseVolumeTexture* SparseVolumeText
 	if (SVTInfo)
 	{
 		const int32 NumFrames = SVTInfo->PerFrameInfo.Num();
-		const int32 FrameIndexI32 = static_cast<int32>(FrameIndex);
+		const int32 FrameIndexI32 = FMath::FloorToInt32(FrameIndex);
 		if (FrameIndexI32 < 0 || FrameIndexI32 >= NumFrames)
 		{
 			return;
 		}
 
-		// Try to find a FStreamingWindow around the requested frame index. This will inform us about which direction we need to prefetch into.
-		FStreamingWindow* StreamingWindow = nullptr;
-		for (FStreamingWindow& Window : SVTInfo->StreamingWindows)
-		{
-			if (FMath::Abs(FrameIndex - Window.CenterFrame) <= FStreamingWindow::WindowSize)
-			{
-				StreamingWindow = &Window;
-				break;
-			}
-		}
-		// Found an existing window!
-		if (StreamingWindow)
-		{
-			const bool bForward = StreamingWindow->LastCenterFrame <= FrameIndex;
-			if (StreamingWindow->LastRequested < NextUpdateIndex)
-			{
-				StreamingWindow->LastCenterFrame = StreamingWindow->CenterFrame;
-				StreamingWindow->CenterFrame = FrameIndex;
-				StreamingWindow->NumRequestsThisUpdate = 1;
-				StreamingWindow->LastRequested = NextUpdateIndex;
-				StreamingWindow->bPlayForward = bForward;
-				StreamingWindow->bPlayBackward = !bForward;
-			}
-			else
-			{
-				// Update the average center frame
-				StreamingWindow->CenterFrame = (StreamingWindow->CenterFrame * StreamingWindow->NumRequestsThisUpdate + FrameIndex) / (StreamingWindow->NumRequestsThisUpdate + 1.0f);
-				++StreamingWindow->NumRequestsThisUpdate;
-				StreamingWindow->bPlayForward |= bForward;
-				StreamingWindow->bPlayBackward |= !bForward;
-			}
-		}
-		// No existing window. Create a new one.
-		else
-		{
-			StreamingWindow = &SVTInfo->StreamingWindows.AddDefaulted_GetRef();
-			StreamingWindow->CenterFrame = FrameIndex;
-			StreamingWindow->LastCenterFrame = FrameIndex;
-			StreamingWindow->NumRequestsThisUpdate = 1;
-			StreamingWindow->LastRequested = NextUpdateIndex;
-			StreamingWindow->bPlayForward = true; // No prior data, so just take a guess that playback is forwards
-			StreamingWindow->bPlayBackward = false;
-		}
+		const bool bBlocking = !!(Flags & EStreamingRequestFlags::Blocking);
+		const double CurrentTime = (FTimespan(FDateTime::Now().GetTicks()) - FTimespan(InitTime.GetTicks())).GetTotalSeconds();
 
-		check(StreamingWindow);
+		// Try to find a FStreamingInstance around the requested frame index and add the current request to it. This will inform us about which direction we need to prefetch into.
+		const FStreamingInstanceRequest StreamingInstanceRequest(UpdateIndex, CurrentTime, FrameRate, FrameIndex, MipLevel, Flags);
+		FStreamingInstance* StreamingInstance = SVTInfo->GetAndUpdateStreamingInstance(StreamingInstanceKey, StreamingInstanceRequest);
+		check(StreamingInstance);
 
 		// Make sure the number of prefetched frames doesn't exceed the total number of frames.
 		// Not only does this make no sense, it also breaks the wrap around logic in the loop below if there is reverse playback.
-		const int32 NumPrefetchFrames = FMath::Clamp(GSVTStreamingNumPrefetchFrames, 0, NumFrames);
+		const int32 NumPrefetchFrames = FMath::Clamp(SVTInfo->NumPrefetchFrames, 0, NumFrames);
 
 		// No prefetching for blocking requests. Making the prefetches blocking would increase latency even more and making them non-blocking could lead
 		// to situations where lower mips are already streamed in in subsequent frames but can't be used because dependent higher mips haven't finished streaming
@@ -332,34 +300,38 @@ void FStreamingManager::Request(UStreamableSparseVolumeTexture* SparseVolumeText
 		// SVT_TODO: This can still break if blocking and non-blocking requests of the same frames/mips are made to the same SVT. We would need to cancel already scheduled non-blocking requests and reissue them as blocking.
 		// Or alternatively we could just block on all requests if we detect this case. If we had a single DDC request owner per request, we could just selectively wait on already scheduled non-blocking requests.
 		const int32 OffsetMagnitude = !bBlocking ? NumPrefetchFrames : 0;
-		const int32 LowerFrameOffset = StreamingWindow->bPlayBackward ? -OffsetMagnitude : 0;
-		const int32 UpperFrameOffset = StreamingWindow->bPlayForward ? OffsetMagnitude : 0;
+		const int32 LowerFrameOffset = StreamingInstance->IsPlayingBackwards() ? -OffsetMagnitude : 0;
+		const int32 UpperFrameOffset = StreamingInstance->IsPlayingForwards() ? OffsetMagnitude : 0;
+		const float PrefetchPercentageStepSize = FMath::Clamp(SVTInfo->PrefetchPercentageStepSize / 100.0f, 0.0f, 1.0f);
+		const float PrefetchPercentageBias = FMath::Clamp(SVTInfo->PrefetchPercentageBias / 100.0f, -1.0f, 1.0f);
 
 		for (int32 i = LowerFrameOffset; i <= UpperFrameOffset; ++i)
 		{
 			// Wrap around on both positive and negative numbers, assuming (i + NumFrames) >= 0. See the comment on NumPrefetchFrames.
-			const int32 RequestFrameIndex = (static_cast<int32>(FrameIndex) + i + NumFrames) % NumFrames;
-			const int32 RequestMipLevelOffset = FMath::Abs(i) + GSVTStreamingPrefetchMipLevelBias;
-			const uint32 MipLevelIndex = FMath::Clamp(MipLevel + RequestMipLevelOffset, 0, SVTInfo->PerFrameInfo[RequestFrameIndex].NumMipLevels);
+			const int32 RequestFrameIndex = (FrameIndexI32 + i + NumFrames) % NumFrames;
+			const float MipPercentage = FMath::Clamp((1.0f - FMath::Abs(i) * PrefetchPercentageStepSize) + (i != 0 ? PrefetchPercentageBias : 0.0f), 0.0f, 1.0f);
+			const float RequestMipLevel = StreamingInstance->GetPrefetchMipLevel(MipLevel, MipPercentage);
+			const int32 RequestMipLevelI = FMath::FloorToInt32(RequestMipLevel);
+			const float LowestMipLevelFraction = 1.0f - FMath::Frac(RequestMipLevel);
 			const uint8 Priority = bBlocking ? UINT8_MAX : FMath::Max(0, OffsetMagnitude - FMath::Abs(i));
 			FStreamingRequest Request;
 			Request.Key.SVTHandle = SVTInfo->SVTHandle;
 			Request.Key.FrameIndex = RequestFrameIndex;
-			Request.Payload.MipLevelMask = ~0u << MipLevelIndex;
-			Request.Payload.LowestMipFraction = 1.0f;
-			for (int32 j = MipLevelIndex; j < Request.Payload.Priorities.Num(); ++j)
+			Request.Payload.StreamingInstance = StreamingInstance;
+			Request.Payload.MipLevelMask = ~0u << RequestMipLevelI;
+			Request.Payload.LowestMipFraction = LowestMipLevelFraction;
+			for (int32 j = RequestMipLevelI; j < Request.Payload.Priorities.Num(); ++j)
 			{
 				Request.Payload.Priorities[j] = Priority;
 			}
 			AddRequest(Request);
 		}
 
-		// Clean up unused streaming windows
-		SVTInfo->StreamingWindows.RemoveAll([&](const FStreamingWindow& Window) { return (NextUpdateIndex - Window.LastRequested) > 5; });
+		ActiveStreamingInstances.Add(StreamingInstance);
 	}
 }
 
-void FStreamingManager::BeginAsyncUpdate(FRDGBuilder& GraphBuilder, bool bBlocking)
+void FStreamingManager::BeginAsyncUpdate(FRDGBuilder& GraphBuilder, bool bUseAsyncThread)
 {
 	check(IsInRenderingThread());
 	check(!AsyncState.bUpdateActive);
@@ -375,12 +347,21 @@ void FStreamingManager::BeginAsyncUpdate(FRDGBuilder& GraphBuilder, bool bBlocki
 
 	if (GSVTStreamingLogVerbosity > 2)
 	{
-		UE_LOG(LogSparseVolumeTextureStreamingManager, Display, TEXT("SVT Streaming Update %i"), NextUpdateIndex);
+		UE_LOG(LogSparseVolumeTextureStreamingManager, Display, TEXT("SVT Streaming Update %i"), UpdateIndex);
 	}
 
 	AsyncState = {};
 	AsyncState.bUpdateActive = true;
 
+	// Clean up unused streaming instances
+	for (TUniquePtr<FStreamingInfo>& SVTInfo : StreamingInfo)
+	{
+		if (SVTInfo)
+		{
+			SVTInfo->StreamingInstances.RemoveAll([&](const TUniquePtr<FStreamingInstance>& Instance) { return (UpdateIndex - Instance->GetUpdateIndex()) > (uint32)FMath::Max(1, GSVTStreamingInstanceCleanupThreshold); });
+		}
+	}
+
 	// For debugging, we can stream out ALL tiles
 	if (GSVTStreamingEmptyPhysicalTileTextures != 0)
 	{
@@ -415,6 +396,7 @@ void FStreamingManager::BeginAsyncUpdate(FRDGBuilder& GraphBuilder, bool bBlocki
 			LOCK_PENDING_REQUEST(PendingRequest);
 			PendingRequest.Reset();
 		}
+		NumPendingRequests = 0;
 
 		GSVTStreamingEmptyPhysicalTileTextures = 0;
 	}
@@ -454,6 +436,7 @@ void FStreamingManager::BeginAsyncUpdate(FRDGBuilder& GraphBuilder, bool bBlocki
 		GSVTStreamingPrintMemoryStats = 0;
 	}
 
+	ComputeBandwidthLimit();
 	FilterRequests();
 	IssueRequests();
 	AsyncState.NumReadyRequests = DetermineReadyRequests();
@@ -510,7 +493,7 @@ void FStreamingManager::BeginAsyncUpdate(FRDGBuilder& GraphBuilder, bool bBlocki
 	Parameters.StreamingManager = this;
 	
 	check(AsyncTaskEvents.IsEmpty());
-	if (GSVTStreamingAsyncThread && !bBlocking)
+	if (GSVTStreamingAsyncThread && bUseAsyncThread)
 	{
 		AsyncState.bUpdateIsAsync = true;
 		AsyncTaskEvents.Add(TGraphTask<FStreamingUpdateTask>::CreateTask().ConstructAndDispatchWhenReady(Parameters));
@@ -554,11 +537,80 @@ void FStreamingManager::EndAsyncUpdate(FRDGBuilder& GraphBuilder)
 
 	check(AsyncState.NumReadyRequests <= NumPendingRequests);
 	NumPendingRequests -= AsyncState.NumReadyRequests;
-	++NextUpdateIndex;
+	++UpdateIndex;
 	AsyncState.bUpdateActive = false;
 	AsyncState.bUpdateIsAsync = false;
 }
 
+const FStreamingDebugInfo* FStreamingManager::GetStreamingDebugInfo(FRDGBuilder& GraphBuilder) const
+{
+	const int64 BandwidthLimit = FMath::Max(GSVTStreamingBandwidthLimit, 1) * 1024LL * 1024LL;
+	const bool bLimitApplies = (GSVTStreamingBandwidthLimit > 0 && (TotalRequestedBandwidth > BandwidthLimit));
+	const double DownscaleFactor = bLimitApplies ? ((double)BandwidthLimit / (double)TotalRequestedBandwidth) : 1.0;
+	const int32 NumSVTs = StreamingInfo.Num();
+
+	FStreamingDebugInfo::FSVT* DebugInfoSVTs = GraphBuilder.AllocPODArray<FStreamingDebugInfo::FSVT>(NumSVTs);
+	int32 SVTArrayWriteSlot = 0;
+	for (const TUniquePtr<FStreamingInfo>& SVTInfo : StreamingInfo)
+	{
+		if (!SVTInfo)
+		{
+			continue;
+		}
+		const int32 NumFrames = SVTInfo->PerFrameInfo.Num();
+		const int32 NumInstances = SVTInfo->StreamingInstances.Num();
+		const int32 NameStrLength = SVTInfo->SVTName.GetStringLength() + 1;
+
+		TCHAR* AssetName = GraphBuilder.AllocPODArray<TCHAR>(NameStrLength);
+		SVTInfo->SVTName.ToString(AssetName, NameStrLength);
+		
+		float* ResidencyPercentages = GraphBuilder.AllocPODArray<float>(NumFrames);
+		float* StreamingPercentages = GraphBuilder.AllocPODArray<float>(NumFrames);
+		for (int32 FrameIndex = 0; FrameIndex < NumFrames; ++FrameIndex)
+		{
+			const FFrameInfo& FrameInfo = SVTInfo->PerFrameInfo[FrameIndex];
+			ResidencyPercentages[FrameIndex] = FrameInfo.ResidentTiles.CountSetBits() / (float)FMath::Max(FrameInfo.ResidentTiles.Num(), 1);
+			StreamingPercentages[FrameIndex] = FrameInfo.StreamingTiles.CountSetBits() / (float)FMath::Max(FrameInfo.StreamingTiles.Num(), 1);
+		}
+		
+		FStreamingDebugInfo::FSVT::FInstance* Instances = GraphBuilder.AllocPODArray<FStreamingDebugInfo::FSVT::FInstance>(NumInstances);
+		for (int32 InstanceIndex = 0; InstanceIndex < NumInstances; ++InstanceIndex)
+		{
+			FStreamingInstance* StreamingInstance = SVTInfo->StreamingInstances[InstanceIndex].Get();
+			if (StreamingInstance)
+			{
+				FStreamingDebugInfo::FSVT::FInstance Instance;
+				Instance.Key = StreamingInstance->GetKey();
+				Instance.Frame = StreamingInstance->GetAverageFrame();
+				Instance.FrameRate = StreamingInstance->GetEstimatedFrameRate();
+				Instance.RequestedBandwidth = float(StreamingInstance->GetRequestedBandwidth(true) / 1024.0 / 1024.0);
+				Instance.AllocatedBandwidth = Instance.RequestedBandwidth * DownscaleFactor;
+				Instance.RequestedMip = StreamingInstance->GetLowestRequestedMipLevel();
+				Instance.InBudgetMip = StreamingInstance->GetLowestMipLevelInBandwidthBudget();
+
+				Instances[InstanceIndex] = Instance;
+			}
+		}
+
+		FStreamingDebugInfo::FSVT& SVT = DebugInfoSVTs[SVTArrayWriteSlot++];
+		SVT.AssetName = AssetName;
+		SVT.FrameResidencyPercentages = ResidencyPercentages;
+		SVT.FrameStreamingPercentages = StreamingPercentages;
+		SVT.Instances = Instances;
+		SVT.NumFrames = NumFrames;
+		SVT.NumInstances = NumInstances;
+	}
+
+	FStreamingDebugInfo* DebugInfo = GraphBuilder.AllocPOD<FStreamingDebugInfo>();
+	DebugInfo->SVTs = DebugInfoSVTs;
+	DebugInfo->NumSVTs = SVTArrayWriteSlot;
+	DebugInfo->RequestedBandwidth = float(TotalRequestedBandwidth / 1024.0 / 1024.0);
+	DebugInfo->BandwidthLimit = float(BandwidthLimit / 1024.0 / 1024.0);
+	DebugInfo->BandwidthScale = float(DownscaleFactor);
+	
+	return DebugInfo;
+}
+
 void FStreamingManager::AddInternal(FRDGBuilder& GraphBuilder, FNewSparseVolumeTextureInfo&& NewSVTInfo)
 {
 	check(IsInRenderingThread());
@@ -587,7 +639,11 @@ void FStreamingManager::AddInternal(FRDGBuilder& GraphBuilder, FNewSparseVolumeT
 	SVTInfo.FormatB = NewSVTInfo.FormatB;
 	SVTInfo.FallbackValueA = NewSVTInfo.FallbackValueA;
 	SVTInfo.FallbackValueB = NewSVTInfo.FallbackValueB;
+	SVTInfo.NumPrefetchFrames = NewSVTInfo.NumPrefetchFrames;
+	SVTInfo.PrefetchPercentageStepSize = NewSVTInfo.PrefetchPercentageStepSize;
+	SVTInfo.PrefetchPercentageBias = NewSVTInfo.PrefetchPercentageBias;
 	SVTInfo.PerFrameInfo = MoveTemp(NewSVTInfo.FrameInfo);
+	SVTInfo.MipLevelStreamingSize.SetNumZeroed(NewSVTInfo.NumMipLevelsGlobal);
 
 	const int32 FormatSizes[] = { GPixelFormats[SVTInfo.FormatA].BlockBytes, GPixelFormats[SVTInfo.FormatB].BlockBytes };
 
@@ -623,12 +679,29 @@ void FStreamingManager::AddInternal(FRDGBuilder& GraphBuilder, FNewSparseVolumeT
 			NumRootVoxelsA += RootTileInfo.NumVoxels[0];
 			NumRootVoxelsB += RootTileInfo.NumVoxels[1];
 		}
+
+		// Determine the high water mark of per frame streaming size when streaming at any given mip level
+		uint32 MaxReferencedTileIndex = 0;
+		for (int32 MipLevel = FrameInfo.NumMipLevels - 2; MipLevel >= 0; --MipLevel)
+		{
+			if (Topology.MipInfo.IsValidIndex(MipLevel))
+			{
+				const FPageTopology::FMip& MipInfo = Topology.MipInfo[MipLevel];
+				for (uint32 PageIndex = MipInfo.PageOffset; PageIndex < (MipInfo.PageOffset + MipInfo.PageCount); ++PageIndex)
+				{
+					const uint32 TileIndex = Topology.TileIndices[PageIndex];
+					MaxReferencedTileIndex = FMath::Max(MaxReferencedTileIndex, TileIndex);
+				}
+				const uint32 MipStreamingSize = Resources->StreamingMetaData.TileDataOffsets[MaxReferencedTileIndex + 1] - Resources->StreamingMetaData.GetRootTileSize();
+				SVTInfo.MipLevelStreamingSize[MipLevel] = FMath::Max(SVTInfo.MipLevelStreamingSize[MipLevel], MipStreamingSize);
+			}
+		}
 	}
 
 	// Create RHI resources and upload root tile data
 	{
-		const int32 TileFactor = NumFrames <= 1 ? 1 : 3;
-		const int32 NumPhysicalTilesCapacity = FMath::Max(1, NumRootPhysicalTiles + (TileFactor * MaxNumPhysicalTiles)); // Ensure a minimum size of 1
+		const float TileFactor = NumFrames <= 1 ? 1.0f : FMath::Clamp(NewSVTInfo.StreamingPoolSizeFactor, 1.0f, static_cast<float>(NumFrames));
+		const int32 NumPhysicalTilesCapacity = FMath::Max(1, NumRootPhysicalTiles + FMath::CeilToInt32(TileFactor * MaxNumPhysicalTiles)); // Ensure a minimum size of 1
 		const FIntVector3 TileDataVolumeResolutionInTiles = FTileDataTexture::GetVolumeResolutionInTiles(NumPhysicalTilesCapacity);
 
 		SVTInfo.TileDataTexture = MakeUnique<FTileDataTexture>(TileDataVolumeResolutionInTiles, SVTInfo.FormatA, SVTInfo.FormatB, SVTInfo.FallbackValueA, SVTInfo.FallbackValueB);
@@ -698,7 +771,7 @@ void FStreamingManager::AddInternal(FRDGBuilder& GraphBuilder, FNewSparseVolumeT
 			{
 				check(!Resources->RootData.IsEmpty());
 				FTileAllocator::FAllocation PreviousAllocation;
-				const uint32 TileCoord = SVTInfo.TileAllocator.Allocate(NextUpdateIndex, 0 /*FreeThreshold*/, FrameIdx, 0 /*TileIndexInFrame*/, -1 /*TilePriority*/, true, PreviousAllocation);
+				const uint32 TileCoord = SVTInfo.TileAllocator.Allocate(UpdateIndex, 0 /*FreeThreshold*/, FrameIdx, 0 /*TileIndexInFrame*/, -1 /*TilePriority*/, true, PreviousAllocation);
 				check(TileCoord != INDEX_NONE);
 				check(TileCoord != 0);
 				check(!PreviousAllocation.bIsAllocated); // We should never have to evict another tile to upload a root tile
@@ -821,6 +894,49 @@ void FStreamingManager::AddRequest(const FStreamingRequest& Request)
 	}
 }
 
+void FStreamingManager::ComputeBandwidthLimit()
+{
+	TRACE_CPUPROFILER_EVENT_SCOPE(SVT::StreamingComputeBandwidthLimit);
+
+	// Get the total requested bandwidth and initialize all streaming instances to not apply any limit.
+	TotalRequestedBandwidth = 0;
+	for (FStreamingInstance* StreamingInstance : ActiveStreamingInstances)
+	{
+		check(StreamingInstance->GetUpdateIndex() == UpdateIndex);
+		TotalRequestedBandwidth += StreamingInstance->GetRequestedBandwidth(true /*bZeroIfBlocking*/);
+		StreamingInstance->ComputeLowestMipLevelInBandwidthBudget(-1); // -1 sets the LowestMipLevelInBandwidthBuget to 0.0f
+	}
+
+	// Apply the limit if the cvar is set and the requested bandwidth exceeds the limit.
+	const int64 AvailableBandwidth = FMath::Max(GSVTStreamingBandwidthLimit, 1) * 1024LL * 1024LL;
+	if (GSVTStreamingBandwidthLimit > 0 && (TotalRequestedBandwidth > AvailableBandwidth))
+	{
+		// This is the factor by which we need to scale down the bandwidth of all currently active streaming instances.
+		const double DownscaleFactor = ((double)AvailableBandwidth / (double)TotalRequestedBandwidth);
+		if (GSVTStreamingLogVerbosity > 2)
+		{
+			UE_LOG(LogSparseVolumeTextureStreamingManager, Display, TEXT("Total Requested Bandwidth: %f MiB/s, Total Available Bandwidth: %f MiB/s, Percentage: %f %%"), 
+				float(TotalRequestedBandwidth / 1024.0 / 1024.0), float(AvailableBandwidth / 1024.0 / 1024.0), float(DownscaleFactor * 100.0));
+		}
+
+		for (FStreamingInstance* StreamingInstance : ActiveStreamingInstances)
+		{
+			const int64 RequestedBandWidth = StreamingInstance->GetRequestedBandwidth(true /*bZeroIfBlocking*/);
+			const int64 AllocatedBandwidth = FMath::CeilToInt64(RequestedBandWidth * DownscaleFactor);
+			StreamingInstance->ComputeLowestMipLevelInBandwidthBudget(AllocatedBandwidth);
+
+			if (GSVTStreamingLogVerbosity > 2)
+			{
+				UE_LOG(LogSparseVolumeTextureStreamingManager, Display, TEXT("Key: %u, Requested: %f MiB, Allocated: %f MiB, FPS: %f, InBudgetMipLevel: %f"),
+					StreamingInstance->GetKey(), float(RequestedBandWidth / 1024.0 / 1024.0), float(AllocatedBandwidth / 1024.0 / 1024.0), StreamingInstance->GetEstimatedFrameRate(), StreamingInstance->GetLowestMipLevelInBandwidthBudget());
+			}
+		}
+	}
+
+	// We don't need this set after this point, so clear it for the next frames requests.
+	ActiveStreamingInstances.Reset();
+}
+
 void FStreamingManager::FilterRequests()
 {
 	TRACE_CPUPROFILER_EVENT_SCOPE(SVT::StreamingFilterRequests);
@@ -842,7 +958,12 @@ void FStreamingManager::FilterRequests()
 		const FPageTopology& Topology = FrameInfo.Resources->Topology;
 		const FTileStreamingMetaData& StreamingMetaData = FrameInfo.Resources->StreamingMetaData;
 		const int32 FirstStreamingMipLevel = FrameInfo.NumMipLevels - 2;
-		const int32 LowestRequestedMip = FMath::CountTrailingZeros((uint32)RequestValue.MipLevelMask);
+		const int32 LowestRequestedMipI = static_cast<int32>(FMath::CountTrailingZeros((uint32)RequestValue.MipLevelMask));
+		const float LowestRequestedMipF = LowestRequestedMipI + (1.0f - FMath::Clamp(RequestValue.LowestMipFraction, 0.0f, 1.0f));
+		const float LowestAllowedMipF = FMath::Clamp(RequestValue.StreamingInstance ? RequestValue.StreamingInstance->GetLowestMipLevelInBandwidthBudget() : 0.0f, LowestRequestedMipF, FrameInfo.NumMipLevels - 1.0f); // Clamp to ensure Allows >= Requested
+		const int32 LowestAllowedMipI = FMath::FloorToInt32(LowestAllowedMipF);
+		const float LowestMipToStreamF = FMath::Max(LowestAllowedMipF, LowestRequestedMipF);
+		const int32 LowestMipToStreamI = FMath::FloorToInt32(LowestMipToStreamF);
 
 		// Keep track of all tiles requested in this frame and also of the tiles requested only within the current priority
 		const uint32 NumTiles = FrameInfo.TileAllocations.Num();
@@ -852,9 +973,9 @@ void FStreamingManager::FilterRequests()
 		// StreamingTiles must be a super set of ResidentTiles
 		check(TBitArray<>::BitwiseAND(FrameInfo.StreamingTiles, FrameInfo.ResidentTiles, EBitwiseOperatorFlags::MaxSize) == FrameInfo.ResidentTiles);
 
-		const uint32 TargetRequestSize = (uint32)FMath::Clamp(GSVTStreamingRequestSize, 1, int32(UINT32_MAX / 1024)) * 1024u;
+		const uint32 TargetRequestSize = GSVTStreamingRequestSize <= 0 ? UINT32_MAX : (uint32)FMath::Clamp(GSVTStreamingRequestSize, 1, int32(UINT32_MAX / 1024)) * 1024u;
 
-		// Creates a FTileRange for each contiguous range of set bits in RequestedTilesInCurrentPriority
+		// Creates a FTileRange for each contiguous range of set bits in RequestedTilesInCurrentPriority and appends it to TileRangesToStream
 		auto MakeTileRanges = [&](uint8 Priority)
 		{
 			bool bResetCounters = true;
@@ -898,7 +1019,7 @@ void FStreamingManager::FilterRequests()
 		};
 
 		uint32 PrevPriority = 0xFFFFFFFFu;
-		for (int32 MipLevelIndex = FirstStreamingMipLevel; MipLevelIndex >= LowestRequestedMip; --MipLevelIndex)
+		for (int32 MipLevelIndex = FirstStreamingMipLevel; MipLevelIndex >= LowestRequestedMipI; --MipLevelIndex)
 		{
 			// Skip if mip level was not requested
 			if (((uint32)RequestValue.MipLevelMask & (1u << MipLevelIndex)) == 0)
@@ -906,8 +1027,16 @@ void FStreamingManager::FilterRequests()
 				continue;
 			}
 
-			// When the priority changes, we need to create tile ranges based on the currently set bets in RequestedTilesInCurrentPriority.
 			const uint8 Priority = RequestValue.Priorities[MipLevelIndex];
+			const bool bIsBlocking = Priority == uint8(-1);
+
+			// Skip if mip level is outside the allowed range for the bandwidth budget and is not a blocking request
+			if (MipLevelIndex < LowestAllowedMipI && !bIsBlocking)
+			{
+				continue;
+			}
+
+			// When the priority changes, we need to create tile ranges based on the currently set bets in RequestedTilesInCurrentPriority.
 			if (Priority != PrevPriority)
 			{
 				MakeTileRanges(Priority);
@@ -915,13 +1044,26 @@ void FStreamingManager::FilterRequests()
 			}
 			PrevPriority = Priority;
 
+			auto GetFractionalPageCount = [](uint32 PageCount, float FractionalMipLevel)
+			{
+				const float Fraction = 1.0f - FMath::Frac(FractionalMipLevel);
+				const uint32 ResultPageCount = FMath::CeilToInt32((float)PageCount * Fraction);
+				check(ResultPageCount <= PageCount);
+				return ResultPageCount;
+			};
+
 			// Get the range of pages to try and stream in
 			const FPageTopology::FMip& MipInfo = Topology.MipInfo[MipLevelIndex];
 			uint32 PageCount = MipInfo.PageCount;
-			if (MipLevelIndex == LowestRequestedMip)
+			if (MipLevelIndex == LowestRequestedMipI)
+			{
+				PageCount = GetFractionalPageCount(MipInfo.PageCount, LowestRequestedMipF);
+			}
+			if (MipLevelIndex == LowestAllowedMipI && !bIsBlocking) // Apply the fractional allowed mip limit only on non-blocking requests
 			{
-				PageCount = FMath::CeilToInt32((float)PageCount * FMath::Clamp(RequestValue.LowestMipFraction, 0.0f, 1.0f));
-				check(PageCount <= MipInfo.PageCount);
+				// If the lowest requested mip is the same integer mip as the lowest allowed mip, take the maximum fractional mip -> minimum bandwidth
+				const float FractionalMip = MipLevelIndex == LowestRequestedMipI ? FMath::Max(LowestRequestedMipF, LowestAllowedMipF) : LowestAllowedMipF;
+				PageCount = GetFractionalPageCount(MipInfo.PageCount, FractionalMip);
 			}
 			const uint32 PageBegin = MipInfo.PageOffset;
 			const uint32 PageEnd = MipInfo.PageOffset + PageCount;
@@ -941,16 +1083,13 @@ void FStreamingManager::FilterRequests()
 				{
 					check(FrameInfo.TileAllocations[TileIndex] != 0);
 					const uint32 TilePriority = FrameInfo.ResidentTiles.Num() - TileIndex; // Higher value means higher priority, but we want higher tiles to be streamed out first
-					SVTInfo->TileAllocator.UpdateUsage(NextUpdateIndex, FrameInfo.TileAllocations[TileIndex], TilePriority);
+					SVTInfo->TileAllocator.UpdateUsage(UpdateIndex, FrameInfo.TileAllocations[TileIndex], TilePriority);
 				}
 			}
-
-			// After processing the lowest/last mip level, we need to create one final batch of tile ranges
-			if (MipLevelIndex == LowestRequestedMip)
-			{
-				MakeTileRanges(Priority);
-			}
 		}
+
+		// After processing the lowest/last mip level, we need to create one final batch of tile ranges
+		MakeTileRanges(PrevPriority);
 	}
 
 	RequestsHashTable.Reset();
@@ -959,8 +1098,6 @@ void FStreamingManager::FilterRequests()
 	auto PriorityPredicate = [](const auto& A, const auto& B) { return A.Priority != B.Priority ? (A.Priority > B.Priority) : A.TileOffset < B.TileOffset; };
 	TileRangesToStream.Sort(PriorityPredicate);
 
-	// SVT_TODO: Here we need to shave off tiles depending on available bandwidth
-
 	const int32 MaxSelectedRequests = MaxPendingRequests - NumPendingRequests;
 	const int32 NumSelectedRequests = FMath::Min(MaxSelectedRequests, TileRangesToStream.Num());
 
@@ -1007,7 +1144,7 @@ void FStreamingManager::IssueRequests()
 			// Allocate tile
 			const uint32 TilePriority = FrameInfo.ResidentTiles.Num() - TileIndex; // Higher value means higher priority, but we want higher tiles to be streamed out first
 			FTileAllocator::FAllocation PreviousAllocation;
-			const uint32 TileCoord = SVTInfo->TileAllocator.Allocate(NextUpdateIndex, 0 /*FreeThreshold*/, TileRange.FrameIndex, TileIndex, TilePriority, false /*bLocked*/, PreviousAllocation);
+			const uint32 TileCoord = SVTInfo->TileAllocator.Allocate(UpdateIndex, 0 /*FreeThreshold*/, TileRange.FrameIndex, TileIndex, TilePriority, false /*bLocked*/, PreviousAllocation);
 			
 			// Allocation failed, cut the TileRange short and upload what we have.
 			if (TileCoord == INDEX_NONE)
@@ -1046,6 +1183,12 @@ void FStreamingManager::IssueRequests()
 					{
 						// None of the tiles in the tile range of the IO request is still requested, so we can cancel the request
 						PendingRequest.Reset();
+
+						if (GSVTStreamingLogVerbosity > 2)
+						{
+							UE_LOG(LogSparseVolumeTextureStreamingManager, Warning, TEXT("SVT '%s' frame %i: Cancelled IO request for frame %i, tile offset %i, tile count %i."),
+								*SVTInfo->SVTName.ToString(), TileRange.FrameIndex, PreviousAllocation.FrameIndex, TileRange.TileCount, TileRange.FrameIndex, NumSuccessfulAllocations);
+						}
 					}
 
 					// Remove the tile -> IO request mapping
@@ -1073,7 +1216,7 @@ void FStreamingManager::IssueRequests()
 			PendingRequest.FrameIndex = TileRange.FrameIndex;
 			PendingRequest.TileOffset = TileRange.TileOffset;
 			PendingRequest.TileCount = TileRange.TileCount;
-			PendingRequest.IssuedInFrame = NextUpdateIndex;
+			PendingRequest.IssuedInFrame = UpdateIndex;
 			PendingRequest.bBlocking = GSVTStreamingForceBlockingRequests || (TileRange.Priority == uint8(-1));
 
 			uint32 ReadOffset = 0;
@@ -1179,6 +1322,7 @@ int32 FStreamingManager::DetermineReadyRequests()
 			// Just mark as ready so it will be skipped later
 			PendingRequest.State = FPendingRequest::EState::DDC_Ready;
 #endif
+			++NumReadyRequests; // Don't forget to increment in this case too or we might miss some ready requests later!
 			continue; 
 		}
 		FStreamingInfo* SVTInfo = FindStreamingInfo(PendingRequest.SVTHandle);
@@ -1264,7 +1408,7 @@ int32 FStreamingManager::DetermineReadyRequests()
 void FStreamingManager::InstallReadyRequests()
 {
 	check(AsyncState.bUpdateActive);
-	check(AsyncState.NumReadyRequests <= PendingRequests.Num())
+	check(AsyncState.NumReadyRequests <= PendingRequests.Num());
 	if (AsyncState.NumReadyRequests <= 0)
 	{
 		return;
@@ -1679,6 +1823,27 @@ void FStreamingManager::RequestDDCData(TConstArrayView<UE::DerivedData::FCacheGe
 
 #endif // WITH_EDITORONLY_DATA
 
+FStreamingInstance* FStreamingManager::FStreamingInfo::GetAndUpdateStreamingInstance(uint64 StreamingInstanceKey, const FStreamingInstanceRequest& Request)
+{
+	// Try to find a FStreamingInstance around the requested frame index and matching the same key
+	FStreamingInstance* StreamingInstance = nullptr;
+	for (TUniquePtr<FStreamingInstance>& Instance : StreamingInstances)
+	{
+		if (Instance->IsFrameInWindow(Request.FrameIndex) && (StreamingInstanceKey == Instance->GetKey()))
+		{
+			StreamingInstance = Instance.Get();
+			StreamingInstance->AddRequest(Request);
+			break;
+		}
+	}
+	// No existing instance found -> create a new one
+	if (!StreamingInstance)
+	{
+		StreamingInstance = StreamingInstances[StreamingInstances.Emplace(MakeUnique<FStreamingInstance>(StreamingInstanceKey, PerFrameInfo.Num(), TArrayView<uint32>(MipLevelStreamingSize), Request))].Get();
+	}
+	return StreamingInstance;
+}
+
 void FTileAllocator::Init(const FIntVector3& InResolutionInTiles)
 {
 	check(InResolutionInTiles.X > 0 && InResolutionInTiles.Y > 0 && InResolutionInTiles.Z > 0);
diff --git a/Engine/Source/Runtime/Engine/Private/SparseVolumeTexture/SparseVolumeTextureStreamingManager.h b/Engine/Source/Runtime/Engine/Private/SparseVolumeTexture/SparseVolumeTextureStreamingManager.h
--- a/Engine/Source/Runtime/Engine/Private/SparseVolumeTexture/SparseVolumeTextureStreamingManager.h
+++ b/Engine/Source/Runtime/Engine/Private/SparseVolumeTexture/SparseVolumeTextureStreamingManager.h
@@ -9,12 +9,11 @@
 #include "Memory/SharedBuffer.h"
 #include "Serialization/BulkData.h"
 #include "RenderResource.h"
-#include "Containers/IntrusiveDoubleLinkedList.h"
 #include "Containers/Map.h"
 #include "Containers/StaticArray.h"
-#include "Containers/Union.h"
 #include "Containers/BinaryHeap.h"
 #include "RenderGraphBuilder.h"
+#include "Misc/DateTime.h"
 
 DECLARE_LOG_CATEGORY_EXTERN(LogSparseVolumeTextureStreamingManager, Log, All);
 
@@ -41,6 +40,8 @@ class FTextureRenderResources;
 class FTileDataTexture;
 class FTileUploader;
 class FPageTableUpdater;
+struct FStreamingInstanceRequest;
+class FStreamingInstance;
 
 // Helper class for managing slots in the tile data texture (SVT streaming pool). It uses a priority queue to reuse older slots that haven't been referenced for some (render) frames.
 class FTileAllocator
@@ -100,29 +101,18 @@ public:
 	//~ Begin IStreamingManager Interface.
 	virtual void Add_GameThread(UStreamableSparseVolumeTexture* SparseVolumeTexture) override;
 	virtual void Remove_GameThread(UStreamableSparseVolumeTexture* SparseVolumeTexture) override;
-	virtual void Request_GameThread(UStreamableSparseVolumeTexture* SparseVolumeTexture, float FrameIndex, int32 MipLevel, bool bBlocking) override;
+	virtual void Request_GameThread(UStreamableSparseVolumeTexture* SparseVolumeTexture, uint32 StreamingInstanceKey, float FrameRate, float FrameIndex, int32 MipLevel, EStreamingRequestFlags Flags) override;
 	virtual void Update_GameThread() override;
 	
-	virtual void Request(UStreamableSparseVolumeTexture* SparseVolumeTexture, float FrameIndex, int32 MipLevel, bool bBlocking) override;
-	virtual void BeginAsyncUpdate(FRDGBuilder& GraphBuilder, bool bBlocking) override;
+	virtual void Request(UStreamableSparseVolumeTexture* SparseVolumeTexture, uint32 StreamingInstanceKey, float FrameRate, float FrameIndex, int32 MipLevel, EStreamingRequestFlags Flags) override;
+	virtual void BeginAsyncUpdate(FRDGBuilder& GraphBuilder, bool bUseAsyncThread) override;
 	virtual void EndAsyncUpdate(FRDGBuilder& GraphBuilder) override;
+	virtual const FStreamingDebugInfo* GetStreamingDebugInfo(FRDGBuilder& GraphBuilder) const override;
 	//~ End IStreamingManager Interface.
 
 private:
 	friend class FStreamingUpdateTask;
 
-	// Represents a context or "window" into the frame sequence that moves along the playback direction. It is used to cache the prefetch direction.
-	struct FStreamingWindow
-	{
-		static constexpr int32 WindowSize = 5;
-		float CenterFrame = -1.0f; // Frame index this window is centered around
-		float LastCenterFrame = -1.0f;
-		int32 NumRequestsThisUpdate = 0;
-		uint32 LastRequested = 0;
-		bool bPlayForward = false;
-		bool bPlayBackward = false;
-	};
-
 	// One per frame of each SVT
 	struct FFrameInfo
 	{
@@ -148,11 +138,18 @@ private:
 		EPixelFormat FormatB = PF_Unknown;
 		FVector4f FallbackValueA = FVector4f();
 		FVector4f FallbackValueB = FVector4f();
+		int32 NumPrefetchFrames = 0;
+		float PrefetchPercentageStepSize = 0.0f;
+		float PrefetchPercentageBias = 0.0f;
 		TArray<FFrameInfo> PerFrameInfo;
-		TArray<FStreamingWindow> StreamingWindows;
+		TArray<TUniquePtr<FStreamingInstance>> StreamingInstances;
+		TArray<uint32, TInlineAllocator<16>> MipLevelStreamingSize; // MipLevelStreamingSize[MipLevel] * FrameRate is the required IO bandwidth in bytes/s to stream a given MipLevel.
 
 		FTileAllocator TileAllocator;
 		TUniquePtr<FTileDataTexture> TileDataTexture;
+
+		// Tries to look up an existing FStreamingInstance, creates a new one if it can't find one and finally updates it with the new request.
+		FStreamingInstance* GetAndUpdateStreamingInstance(uint64 StreamingInstanceKey, const FStreamingInstanceRequest& Request);
 	};
 
 	// Represents an IO request for tile(s)
@@ -235,6 +232,11 @@ private:
 		EPixelFormat FormatB = PF_Unknown;
 		FVector4f FallbackValueA = FVector4f(0.0f, 0.0f, 0.0f, 0.0f);
 		FVector4f FallbackValueB = FVector4f(0.0f, 0.0f, 0.0f, 0.0f);
+		int32 NumMipLevelsGlobal = 0;
+		float StreamingPoolSizeFactor = 0.0f;
+		int32 NumPrefetchFrames = 0;
+		float PrefetchPercentageStepSize = 0.0f;
+		float PrefetchPercentageBias = 0.0f;
 		TArray<FFrameInfo> FrameInfo; // Only Resources and TextureRenderResources are initialized
 	};
 
@@ -253,6 +255,7 @@ private:
 	// The payload associated with a request for a given frame of a SVT
 	struct FRequestPayload
 	{
+		FStreamingInstance* StreamingInstance = nullptr;
 		uint16 MipLevelMask = 0; // Bitmask of requested mip levels
 		float LowestMipFraction = 0.0f; // Setting this to a value between 0 and less than 1 signifies that only a certain percentage of pages in this mip level should be streamed
 		TStaticArray<uint8, 16> Priorities = TStaticArray<uint8, 16>(InPlace, 0); // Priority for each mip level with a set bit in MipLevelMask
@@ -289,9 +292,12 @@ private:
 	int32 MaxPendingRequests = 0;
 	int32 NumPendingRequests = 0;
 	int32 NextPendingRequestIndex = 0;
-	uint32 NextUpdateIndex = 1;
+	uint32 UpdateIndex = 1;
+	FDateTime InitTime = FDateTime::Now(); // The precise time doesn't matter, we only use it as a basis to compute time as a double.
+	int64 TotalRequestedBandwidth = 0;
 
 	// Transient lifetime
+	TSet<FStreamingInstance*> ActiveStreamingInstances; // Instances that have had a request in them since the last update
 	TSet<FTileDataTexture*> TileDataTexturesToUpdate;
 	TSet<FFrameInfo*> InvalidatedSVTFrames; // Set of SVT frames where tiles have been streamed in or out. Used in PatchPageTable().
 	TArray<FTileRange> TileRangesToStream; // Output of FilterRequests(), consumed in IssueRequests().
@@ -304,6 +310,7 @@ private:
 	void RemoveInternal(UStreamableSparseVolumeTexture* SparseVolumeTexture);
 	void AddRequest(const FStreamingRequest& Request);
 	void AsyncUpdate();
+	void ComputeBandwidthLimit();
 	void FilterRequests();
 	void IssueRequests();
 	int32 DetermineReadyRequests();
diff --git a/Engine/Source/Runtime/Engine/Public/SparseVolumeTexture/ISparseVolumeTextureStreamingManager.h b/Engine/Source/Runtime/Engine/Public/SparseVolumeTexture/ISparseVolumeTextureStreamingManager.h
--- a/Engine/Source/Runtime/Engine/Public/SparseVolumeTexture/ISparseVolumeTextureStreamingManager.h
+++ b/Engine/Source/Runtime/Engine/Public/SparseVolumeTexture/ISparseVolumeTextureStreamingManager.h
@@ -13,6 +13,44 @@ namespace UE
 namespace SVT
 {
 
+struct FStreamingDebugInfo
+{
+	struct FSVT
+	{
+		struct FInstance
+		{
+			uint32 Key;
+			float Frame;
+			float FrameRate;
+			float RequestedBandwidth;
+			float AllocatedBandwidth;
+			float RequestedMip;
+			float InBudgetMip;
+		};
+
+		const TCHAR* AssetName;
+		const float* FrameResidencyPercentages;
+		const float* FrameStreamingPercentages;
+		const FInstance* Instances;
+		int32 NumFrames;
+		int32 NumInstances;
+	};
+
+	const FSVT* SVTs;
+	int32 NumSVTs;
+	float RequestedBandwidth;
+	float BandwidthLimit;
+	float BandwidthScale;
+};
+
+enum class EStreamingRequestFlags : uint8
+{
+	None = 0,
+	Blocking = 1u << 0u,		// Use blocking IO requests.
+	HasFrameRate = 1u << 1u,	// The passed in FrameRate value can be used to more accurately predict bandwidth requirements.
+};
+ENUM_CLASS_FLAGS(EStreamingRequestFlags);
+
 // Interface for the SparseVolumeTexture streaming manager
 class IStreamingManager
 {
@@ -20,9 +58,13 @@ public:
 	//~ Begin game thread functions.
 	virtual void Add_GameThread(UStreamableSparseVolumeTexture* SparseVolumeTexture) = 0;
 	virtual void Remove_GameThread(UStreamableSparseVolumeTexture* SparseVolumeTexture) = 0;
-	// Request a frame to be streamed in. FrameIndex is of float type so that the fractional part can be used to better track the playback speed/direction.
+	// Request a frame to be streamed in.
+	// StreamingInstanceKey can be any arbitrary value that is suitable to keep track of the source of requests for a given SVT. 
+	// This key is used internally to associate incoming requests with prior requests issued for the same SVT. A good value to pass here might be a hash of the pointer of the component the SVT is used within.
+	// FrameRate is an optional argument which helps to more accurately predict the required bandwidth when using non-blocking requests. EStreamingRequestFlags::HasFrameRate must be set for this.
+	// FrameIndex is of float type so that the fractional part can be used to better track the playback speed/direction.
 	// This function automatically also requests all higher mip levels and adds prefetch requests for upcoming frames.
-	virtual void Request_GameThread(UStreamableSparseVolumeTexture* SparseVolumeTexture, float FrameIndex, int32 MipLevel, bool bBlocking) = 0;
+	virtual void Request_GameThread(UStreamableSparseVolumeTexture* SparseVolumeTexture, uint32 StreamingInstanceKey, float FrameRate, float FrameIndex, int32 MipLevel, EStreamingRequestFlags Flags) = 0;
 	// Issues a rendering command for updating the streaming manager. This is not normally needed, but is necessary in cases where blocking requests are required
 	// and the SVT is used in a rendering command (which is executed before the streaming manager would normally update). A call to this function is not needed when the SVT
 	// is used in a regular pass called from FDeferredShadingSceneRenderer::Render().
@@ -30,11 +72,13 @@ public:
 	//~ End game thread functions.
 	
 	//~ Begin rendering thread functions.
-	virtual void Request(UStreamableSparseVolumeTexture* SparseVolumeTexture, float FrameIndex, int32 MipLevel, bool bBlocking) = 0;
-	// Begins updating the streaming manager. If r.SparseVolumeTexture.Streaming.AsyncThread is 1 and bBlocking is false, most of the updating work is done in another thread.
-	virtual void BeginAsyncUpdate(FRDGBuilder& GraphBuilder, bool bBlocking = false) = 0;
+	virtual void Request(UStreamableSparseVolumeTexture* SparseVolumeTexture, uint32 StreamingInstanceKey, float FrameRate, float FrameIndex, int32 MipLevel, EStreamingRequestFlags Flags) = 0;
+	// Begins updating the streaming manager. If r.SparseVolumeTexture.Streaming.AsyncThread is 1 and bUseAsyncThread is true, most of the updating work is done in another thread.
+	virtual void BeginAsyncUpdate(FRDGBuilder& GraphBuilder, bool bUseAsyncThread = true) = 0;
 	// Waits for the job started in BeginAsyncUpdate() to complete, issues GPU work and does some cleanup.
 	virtual void EndAsyncUpdate(FRDGBuilder& GraphBuilder) = 0;
+	// Returns a FStreamingDebugInfo useful for debugging streaming performance. All data is allocated on the passed in GraphBuilder.
+	virtual const FStreamingDebugInfo* GetStreamingDebugInfo(FRDGBuilder& GraphBuilder) const = 0;
 	//~ End rendering thread functions.
 
 	virtual ~IStreamingManager() = default;
diff --git a/Engine/Source/Runtime/Renderer/Private/PostProcess/PostProcessing.cpp b/Engine/Source/Runtime/Renderer/Private/PostProcess/PostProcessing.cpp
--- a/Engine/Source/Runtime/Renderer/Private/PostProcess/PostProcessing.cpp
+++ b/Engine/Source/Runtime/Renderer/Private/PostProcess/PostProcessing.cpp
@@ -71,6 +71,7 @@
 #include "IESTextureManager.h"
 #include "UnrealEngine.h"
 #include "IlluminanceMeter.h"
+#include "SparseVolumeTexture/SparseVolumeTextureStreamingVisualize.h"
 
 bool IsMobileEyeAdaptationEnabled(const FViewInfo& View);
 
@@ -1774,6 +1775,13 @@ void AddPostProcessingPasses(
 			IESAtlas::AddDebugPass(GraphBuilder, View, SceneColor.Texture);
 		}
 
+		// Piggy back off of OnScreenDebug to avoid having to create a new show flag just for this simple debug visualization. Otherwise it might render into certain thumbnails.
+		// In the future it might be worth it to introduce a show flag?
+		if (EngineShowFlags.OnScreenDebug)
+		{
+			UE::SVT::AddStreamingDebugPass(GraphBuilder, View, SceneColor);
+		}
+
 		if (ShaderPrint::IsEnabled(View.ShaderPrintData))
 		{
 			ShaderPrint::DrawView(GraphBuilder, View, SceneColor, SceneDepth);
diff --git a/Engine/Source/Runtime/Renderer/Private/SparseVolumeTexture/SparseVolumeTextureStreamingVisualize.cpp b/Engine/Source/Runtime/Renderer/Private/SparseVolumeTexture/SparseVolumeTextureStreamingVisualize.cpp
new file mode 100644
--- /dev/null
+++ b/Engine/Source/Runtime/Renderer/Private/SparseVolumeTexture/SparseVolumeTextureStreamingVisualize.cpp
@@ -0,0 +1,89 @@
+// Copyright Epic Games, Inc. All Rights Reserved.
+
+#include "SparseVolumeTextureStreamingVisualize.h"
+#include "SparseVolumeTexture/ISparseVolumeTextureStreamingManager.h"
+#include "RenderGraph.h"
+#include "RenderGraphBuilder.h"
+#include "CanvasTypes.h"
+#include "CanvasItem.h"
+#include "Engine/Engine.h"
+
+static int32 GSVTStreamingShowDebugInfo = 0;
+static FAutoConsoleVariableRef CVarSVTStreamingShowDebugInfo(
+	TEXT("r.SparseVolumeTexture.Streaming.ShowDebugInfo"),
+	GSVTStreamingShowDebugInfo,
+	TEXT("Prints debug info about the streaming SVT instances to the screen."),
+	ECVF_RenderThreadSafe
+);
+
+void UE::SVT::AddStreamingDebugPass(FRDGBuilder& GraphBuilder, const FSceneView& View, FScreenPassTexture Output)
+{
+	if (GSVTStreamingShowDebugInfo == 0)
+	{
+		return;
+	}
+
+	const FStreamingDebugInfo* DebugInfo = UE::SVT::GetStreamingManager().GetStreamingDebugInfo(GraphBuilder);
+	FScreenPassRenderTarget OutputTarget = FScreenPassRenderTarget(Output, ERenderTargetLoadAction::ELoad);
+	AddDrawCanvasPass(GraphBuilder, RDG_EVENT_NAME("SVT::StreamingDebug"), View, OutputTarget, [DebugInfo](FCanvas& Canvas)
+	{
+		const double LineSpacing = 14.0;
+		double YOffset = 32.0;
+		TStringBuilder<128> StringBuilder;
+
+		StringBuilder.Appendf(TEXT("Sparse Volume Texture Streaming Stats\nRequested Peak Bandwidth: %#08.2f MiB/s, Bandwidth Limit: %#08.2f MiB/s, Bandwidth Throttle Scale: %6.2f%%\n"), 
+			DebugInfo->RequestedBandwidth, DebugInfo->BandwidthLimit, float(DebugInfo->BandwidthScale * 100.0));
+		FCanvasTextItem HeaderTextItem(FVector2D(8.0, YOffset), FText::FromStringView(StringBuilder.ToView()), GEngine->GetSmallFont(), FLinearColor::White);
+		HeaderTextItem.EnableShadow(FLinearColor::Black);
+		HeaderTextItem.Draw(&Canvas);
+		YOffset += LineSpacing * 2.0; // We just printed two lines
+		StringBuilder.Reset();
+
+		for (int32 SVTIndex = 0; SVTIndex < DebugInfo->NumSVTs; ++SVTIndex)
+		{
+			const FStreamingDebugInfo::FSVT& SVT = DebugInfo->SVTs[SVTIndex];
+
+			StringBuilder.Appendf(TEXT("    Asset: '%s', Num Frames: %i\n"), SVT.AssetName, SVT.NumFrames);
+			FCanvasTextItem SVTTextItem(FVector2D(8.0, YOffset), FText::FromStringView(StringBuilder.ToView()), GEngine->GetSmallFont(), FLinearColor::White);
+			SVTTextItem.EnableShadow(FLinearColor::Black);
+			SVTTextItem.Draw(&Canvas);
+			YOffset += LineSpacing;
+			StringBuilder.Reset();
+
+			const double TileYOffset = YOffset;
+			const double TileHeight = 24.0;
+			const double TileWidth = 8.0;
+			for (int32 FrameIndex = 0; FrameIndex < SVT.NumFrames; ++FrameIndex)
+			{
+				const double XOffset = 8.0 + FrameIndex * (TileWidth + 1.0);
+				const float Residency = SVT.FrameResidencyPercentages[FrameIndex];
+				const float Streaming = SVT.FrameStreamingPercentages[FrameIndex];
+				const FLinearColor ResidencyColor = FMath::Lerp(FLinearColor::Red, FLinearColor::Green, Residency);
+
+				FCanvasTileItem StreamingFrameTileItem(FVector2D(XOffset, YOffset + FMath::Min(TileHeight - 1.0, TileHeight * (1.0 - Streaming))), FVector2D(TileWidth, FMath::Max(1.0, TileHeight * Streaming)), FLinearColor(0.0f, 0.0f, 0.5f));
+				StreamingFrameTileItem.Draw(&Canvas);
+				FCanvasTileItem ResidencyFrameTileItem(FVector2D(XOffset, YOffset + FMath::Min(TileHeight - 1.0, TileHeight * (1.0 - Residency))), FVector2D(TileWidth, FMath::Max(1.0, TileHeight * Residency)), ResidencyColor);
+				ResidencyFrameTileItem.Draw(&Canvas);
+			}
+			YOffset += TileHeight + 2.0;
+
+			for (int32 InstanceIndex = 0; InstanceIndex < SVT.NumInstances; ++InstanceIndex)
+			{
+				const FStreamingDebugInfo::FSVT::FInstance& Instance = SVT.Instances[InstanceIndex];
+
+				StringBuilder.Appendf(TEXT("        Instance Key: %u, Frame: %#06.2f, Estimated Frame Rate: %#06.2f, Requested Peak Bandwidth: %#07.2f MiB/s, Allocated Bandwidth: %#07.2f MiB/s, Requested Mip: %#5.2f, Bandwidth Budget Mip: %#5.2f\n"),
+					Instance.Key, Instance.Frame, Instance.FrameRate, Instance.RequestedBandwidth, Instance.AllocatedBandwidth, Instance.RequestedMip, Instance.InBudgetMip);
+				FCanvasTextItem InstanceTextItem(FVector2D(8.0, YOffset), FText::FromStringView(StringBuilder.ToView()), GEngine->GetSmallFont(), FLinearColor::White);
+				InstanceTextItem.EnableShadow(FLinearColor::Black);
+				InstanceTextItem.Draw(&Canvas);
+				YOffset += LineSpacing;
+				StringBuilder.Reset();
+
+				const double FrameLineX = 8.0 + FMath::FloorToDouble(Instance.Frame) * (TileWidth + 1.0) + FMath::Frac(Instance.Frame) * TileWidth;
+				FCanvasLineItem InstanceLineItem(FVector2D(FrameLineX, TileYOffset - 1.0), FVector2D(FrameLineX, TileYOffset + TileHeight + 1.0));
+				InstanceLineItem.SetColor(FLinearColor::Blue);
+				InstanceLineItem.Draw(&Canvas);
+			}
+		}
+	});
+}
diff --git a/Engine/Source/Runtime/Renderer/Private/SparseVolumeTexture/SparseVolumeTextureStreamingVisualize.h b/Engine/Source/Runtime/Renderer/Private/SparseVolumeTexture/SparseVolumeTextureStreamingVisualize.h
new file mode 100644
--- /dev/null
+++ b/Engine/Source/Runtime/Renderer/Private/SparseVolumeTexture/SparseVolumeTextureStreamingVisualize.h
@@ -0,0 +1,17 @@
+// Copyright Epic Games, Inc. All Rights Reserved.
+
+#pragma once
+
+#include "CoreMinimal.h"
+#include "Engine/EngineTypes.h"
+#include "ScreenPass.h"
+
+class FRDGBuilder;
+
+namespace UE
+{
+namespace SVT
+{
+	void AddStreamingDebugPass(FRDGBuilder& GraphBuilder, const FSceneView& View, FScreenPassTexture Output);
+}
+}
\ No newline at end of file
diff --git a/Engine/Source/Runtime/Renderer/Private/SparseVolumeTexture/SparseVolumeTextureViewerComponent.cpp b/Engine/Source/Runtime/Renderer/Private/SparseVolumeTexture/SparseVolumeTextureViewerComponent.cpp
--- a/Engine/Source/Runtime/Renderer/Private/SparseVolumeTexture/SparseVolumeTextureViewerComponent.cpp
+++ b/Engine/Source/Runtime/Renderer/Private/SparseVolumeTexture/SparseVolumeTextureViewerComponent.cpp
@@ -223,7 +223,8 @@ void USparseVolumeTextureViewerComponent::TickComponent(float DeltaTime, enum EL
 			Frame = FMath::Clamp(Frame, 0.0f, (float)(NumFrames - 1));
 		}
 
-		SparseVolumeTextureFrame = USparseVolumeTextureFrame::GetFrameAndIssueStreamingRequest(SparseVolumeTexturePreview, Frame, MipLevel, bBlockingStreamingRequests);
+		const bool bHasValidFrameRate = bPlaying != 0;
+		SparseVolumeTextureFrame = USparseVolumeTextureFrame::GetFrameAndIssueStreamingRequest(SparseVolumeTexturePreview, GetTypeHash(this), FrameRate, Frame, MipLevel, bBlockingStreamingRequests, bHasValidFrameRate);
 	}
 	else
 	{
