World Partition Actor Descriptors:
- Cache the actor name in the actor descriptor instead of generating it from the path on every world outliner refresh.
- Make sure OnUnloadingInstance is marked as deprecated in 5.4 as it's removed from 5.5 (introduced in 5.4).

diff --git a/Engine/Source/Runtime/Engine/Private/WorldPartition/DataLayer/WorldDataLayersActorDesc.cpp b/Engine/Source/Runtime/Engine/Private/WorldPartition/DataLayer/WorldDataLayersActorDesc.cpp
--- a/Engine/Source/Runtime/Engine/Private/WorldPartition/DataLayer/WorldDataLayersActorDesc.cpp
+++ b/Engine/Source/Runtime/Engine/Private/WorldPartition/DataLayer/WorldDataLayersActorDesc.cpp
@@ -269,6 +269,7 @@ void FWorldDataLayersActorDesc::Serialize(FArchive& Ar)
 	}
 }
 
+PRAGMA_DISABLE_DEPRECATION_WARNINGS
 void FWorldDataLayersActorDesc::OnUnloadingInstance(const FWorldPartitionActorDescInstance* InActorDescInstance) const
 {
 	if (AWorldDataLayers* WorldDataLayers = Cast<AWorldDataLayers>(InActorDescInstance->GetActor()))
@@ -293,6 +294,7 @@ void FWorldDataLayersActorDesc::OnUnloadingInstance(const FWorldPartitionActorDe
 	}
 	FWorldPartitionActorDesc::OnUnloadingInstance(InActorDescInstance);
 }
+PRAGMA_ENABLE_DEPRECATION_WARNINGS
 
 bool FWorldDataLayersActorDesc::IsRuntimeRelevant(const FWorldPartitionActorDescInstance* InActorDescInstance) const 
 {
diff --git a/Engine/Source/Runtime/Engine/Private/WorldPartition/Landscape/LandscapeActorDesc.cpp b/Engine/Source/Runtime/Engine/Private/WorldPartition/Landscape/LandscapeActorDesc.cpp
--- a/Engine/Source/Runtime/Engine/Private/WorldPartition/Landscape/LandscapeActorDesc.cpp
+++ b/Engine/Source/Runtime/Engine/Private/WorldPartition/Landscape/LandscapeActorDesc.cpp
@@ -53,6 +53,7 @@ FBox FLandscapeActorDesc::GetEditorBounds() const
 	return FWorldPartitionActorDesc::GetEditorBounds();
 }
 
+PRAGMA_DISABLE_DEPRECATION_WARNINGS
 void FLandscapeActorDesc::OnUnloadingInstance(const FWorldPartitionActorDescInstance* InActorDescInstance) const
 {
 	if (ALandscapeProxy* LandscapeProxy = Cast<ALandscapeProxy>(InActorDescInstance->GetActor()))
@@ -62,6 +63,7 @@ void FLandscapeActorDesc::OnUnloadingInstance(const FWorldPartitionActorDescInst
 
 	FPartitionActorDesc::OnUnloadingInstance(InActorDescInstance);
 }
+PRAGMA_ENABLE_DEPRECATION_WARNINGS
 
 bool FLandscapeActorDesc::Equals(const FWorldPartitionActorDesc* Other) const
 {
diff --git a/Engine/Source/Runtime/Engine/Private/WorldPartition/WorldPartitionActorDesc.cpp b/Engine/Source/Runtime/Engine/Private/WorldPartition/WorldPartitionActorDesc.cpp
--- a/Engine/Source/Runtime/Engine/Private/WorldPartition/WorldPartitionActorDesc.cpp
+++ b/Engine/Source/Runtime/Engine/Private/WorldPartition/WorldPartitionActorDesc.cpp
@@ -158,6 +158,7 @@ void FWorldPartitionActorDesc::Init(const AActor* InActor)
 
 	ActorPackage = InActor->GetPackage()->GetFName();
 	ActorPath = bIsDefaultActorDesc ? *InActor->GetClass()->GetPathName() : *InActor->GetPathName();
+	ActorName = InActor->GetFName();
 
 	ContentBundleGuid = InActor->GetContentBundleGuid();
 
@@ -208,6 +209,7 @@ void FWorldPartitionActorDesc::Init(const FWorldPartitionActorDescInitData& Desc
 	ActorPath = DescData.ActorPath;
 	ActorNativeClass = DescData.NativeClass;
 	NativeClass = *DescData.NativeClass->GetPathName();
+	ActorName = *FPaths::GetExtension(ActorPath.ToString());
 
 	auto DeprecateClass = [this](FArchive& Archive)
 	{
@@ -761,7 +763,7 @@ FBox FWorldPartitionActorDesc::GetRuntimeBounds() const
 
 FName FWorldPartitionActorDesc::GetActorName() const
 {
-	return *FPaths::GetExtension(ActorPath.ToString());
+	return ActorName;
 }
 
 FName FWorldPartitionActorDesc::GetActorLabelOrName() const
diff --git a/Engine/Source/Runtime/Engine/Private/WorldPartition/WorldPartitionActorDescInstance.cpp b/Engine/Source/Runtime/Engine/Private/WorldPartition/WorldPartitionActorDescInstance.cpp
--- a/Engine/Source/Runtime/Engine/Private/WorldPartition/WorldPartitionActorDescInstance.cpp
+++ b/Engine/Source/Runtime/Engine/Private/WorldPartition/WorldPartitionActorDescInstance.cpp
@@ -77,7 +77,7 @@ FSoftObjectPath FWorldPartitionActorDescInstance::GetActorSoftPath() const
 
 FName FWorldPartitionActorDescInstance::GetActorName() const
 {
-	return *FPaths::GetExtension(GetActorSoftPath().ToString());
+	return ActorDesc->GetActorName();
 }
 
 bool FWorldPartitionActorDescInstance::IsValid() const
@@ -167,7 +167,9 @@ void FWorldPartitionActorDescInstance::MarkUnload()
 	FlushAsyncLoad();
 
 	// Notify Desc as it can have some custom code to run on the actor depending on type
+PRAGMA_DISABLE_DEPRECATION_WARNINGS
 	ActorDesc->OnUnloadingInstance(this);
+PRAGMA_ENABLE_DEPRECATION_WARNINGS
 
 	if (AActor* Actor = GetActor())
 	{
diff --git a/Engine/Source/Runtime/Engine/Public/WorldPartition/DataLayer/WorldDataLayersActorDesc.h b/Engine/Source/Runtime/Engine/Public/WorldPartition/DataLayer/WorldDataLayersActorDesc.h
--- a/Engine/Source/Runtime/Engine/Public/WorldPartition/DataLayer/WorldDataLayersActorDesc.h
+++ b/Engine/Source/Runtime/Engine/Public/WorldPartition/DataLayer/WorldDataLayersActorDesc.h
@@ -77,7 +77,9 @@ protected:
 	ENGINE_API virtual void Serialize(FArchive& Ar) override;
 	virtual bool IsResaveNeeded() const override { return !IsValid(); }
 	virtual bool IsRuntimeRelevant(const FWorldPartitionActorDescInstance* InActorDescInstance) const override;
+PRAGMA_DISABLE_DEPRECATION_WARNINGS
 	virtual void OnUnloadingInstance(const FWorldPartitionActorDescInstance* InActorDescInstance) const override;
+PRAGMA_ENABLE_DEPRECATION_WARNINGS
 	//~ End FWorldPartitionActorDesc Interface.
 
 private:
diff --git a/Engine/Source/Runtime/Engine/Public/WorldPartition/Landscape/LandscapeActorDesc.h b/Engine/Source/Runtime/Engine/Public/WorldPartition/Landscape/LandscapeActorDesc.h
--- a/Engine/Source/Runtime/Engine/Public/WorldPartition/Landscape/LandscapeActorDesc.h
+++ b/Engine/Source/Runtime/Engine/Public/WorldPartition/Landscape/LandscapeActorDesc.h
@@ -20,7 +20,9 @@ public:
 protected:
 	virtual uint32 GetSizeOf() const override { return sizeof(FLandscapeActorDesc); }
 	ENGINE_API virtual void Serialize(FArchive& Ar) override;
+PRAGMA_DISABLE_DEPRECATION_WARNINGS
 	ENGINE_API void OnUnloadingInstance(const FWorldPartitionActorDescInstance* InActorDescInstance) const override;
+PRAGMA_ENABLE_DEPRECATION_WARNINGS
 
 private:
 	FGuid LandscapeActorGuid;
diff --git a/Engine/Source/Runtime/Engine/Public/WorldPartition/WorldPartitionActorDesc.h b/Engine/Source/Runtime/Engine/Public/WorldPartition/WorldPartitionActorDesc.h
--- a/Engine/Source/Runtime/Engine/Public/WorldPartition/WorldPartitionActorDesc.h
+++ b/Engine/Source/Runtime/Engine/Public/WorldPartition/WorldPartitionActorDesc.h
@@ -378,6 +378,7 @@ public:
 	ENGINE_API bool IsMainWorldOnly() const;
 	ENGINE_API bool IsListedInSceneOutliner() const;
 protected:
+	UE_DEPRECATED(5.4, "OnUnloadingInstance is deprecated.")
 	virtual void OnUnloadingInstance(const FWorldPartitionActorDescInstance* InActorDescInstance) const {}
 
 	virtual bool GetChildContainerInstance(const FWorldPartitionActorDescInstance* InActorDescInstance, FContainerInstance& OutContainerInstance) const { return false; }
@@ -435,6 +436,7 @@ protected:
 	
 	// Transient
 	UClass*							ActorNativeClass;
+	FName							ActorName;
 	UActorDescContainer*			Container;
 	bool							bIsDefaultActorDesc;
 
