Add UWorld::EndPlay which properly calls RouteEndPlay and change shutdown code to use the new function. This must be called for any worlds that have a game mode and actually started play.
Clarify that BeginTearingDown should be called for all worlds that were initialized for play, which EndPlay will handle if it hasn't been called already.
Add FTestWorldWrapper struct which handles all of the logic required to create and destroy a gameplay world for tests.
Fix several automation tests to use the new wrapper or call EndPlay manually which fixes a crash with LODSyncTests.

diff --git a/Engine/Plugins/Runtime/GameplayAbilities/Source/GameplayAbilities/Private/Tests/AbilitySystemComponentTests.cpp b/Engine/Plugins/Runtime/GameplayAbilities/Source/GameplayAbilities/Private/Tests/AbilitySystemComponentTests.cpp
--- a/Engine/Plugins/Runtime/GameplayAbilities/Source/GameplayAbilities/Private/Tests/AbilitySystemComponentTests.cpp
+++ b/Engine/Plugins/Runtime/GameplayAbilities/Source/GameplayAbilities/Private/Tests/AbilitySystemComponentTests.cpp
@@ -317,6 +317,7 @@ protected:
 		}
 		GFrameCounter = InitialFrameCounter;
 
+		World->EndPlay(EEndPlayReason::Quit);
 		GEngine->DestroyWorldContext(World);
 		World->DestroyWorld(false);
 
diff --git a/Engine/Plugins/Runtime/GameplayAbilities/Source/GameplayAbilities/Private/Tests/GameplayEffectTests.cpp b/Engine/Plugins/Runtime/GameplayAbilities/Source/GameplayAbilities/Private/Tests/GameplayEffectTests.cpp
--- a/Engine/Plugins/Runtime/GameplayAbilities/Source/GameplayAbilities/Private/Tests/GameplayEffectTests.cpp
+++ b/Engine/Plugins/Runtime/GameplayAbilities/Source/GameplayAbilities/Private/Tests/GameplayEffectTests.cpp
@@ -774,6 +774,7 @@ protected:
 		}
 		GFrameCounter = InitialFrameCounter;
 
+		World->EndPlay(EEndPlayReason::Quit);
 		GEngine->DestroyWorldContext(World);
 		World->DestroyWorld(false);
 
diff --git a/Engine/Plugins/WorldMetrics/Source/WorldMetricsTest/Private/WorldMetricsTest.cpp b/Engine/Plugins/WorldMetrics/Source/WorldMetricsTest/Private/WorldMetricsTest.cpp
--- a/Engine/Plugins/WorldMetrics/Source/WorldMetricsTest/Private/WorldMetricsTest.cpp
+++ b/Engine/Plugins/WorldMetrics/Source/WorldMetricsTest/Private/WorldMetricsTest.cpp
@@ -611,6 +611,7 @@ static void ScopedWorldTest(EWorldType::Type WorldType, TFunctionRef<void(UWorld
 
 	WorldTest(World);
 
+	World->EndPlay(EEndPlayReason::Quit);
 	GEngine->DestroyWorldContext(World);
 	World->DestroyWorld(bInformEngineOfWorld);
 }
diff --git a/Engine/Source/Editor/UnrealEd/Private/PlayLevel.cpp b/Engine/Source/Editor/UnrealEd/Private/PlayLevel.cpp
--- a/Engine/Source/Editor/UnrealEd/Private/PlayLevel.cpp
+++ b/Engine/Source/Editor/UnrealEd/Private/PlayLevel.cpp
@@ -853,13 +853,10 @@ void UEditorEngine::TeardownPlaySession(FWorldContext& PieWorldContext)
 			}
 		}
 	}
-
-	// Go through and let all the PlayWorld Actor's know they are being destroyed
-	for (FActorIterator ActorIt(PlayWorld); ActorIt; ++ActorIt)
-	{
-		ActorIt->RouteEndPlay(EEndPlayReason::EndPlayInEditor);
-	}
-
+	
+	// End play for all actors in the world
+	PlayWorld->EndPlay(EEndPlayReason::EndPlayInEditor);
+	
 	PieWorldContext.OwningGameInstance->Shutdown();
 
 	// Move blueprint debugging pointers back to the objects in the editor world
diff --git a/Engine/Source/Runtime/Engine/Classes/Engine/World.h b/Engine/Source/Runtime/Engine/Classes/Engine/World.h
--- a/Engine/Source/Runtime/Engine/Classes/Engine/World.h
+++ b/Engine/Source/Runtime/Engine/Classes/Engine/World.h
@@ -1161,10 +1161,10 @@ public:
 	UE_DEPRECATED(5.4, "Public access to bBegunPlay is deprecated. Please update your code to use the public accessors GetBegunPlay() & SetBegunPlay().")
 	uint8 bBegunPlay:1;
 
-	/** Set whether BeginPlay has been called on actors */
+	/** Sets whether BeginPlay has been called for actors in the world. Use BeginPlay and EndPlay below to start the process of changing this. */
 	void SetBegunPlay(bool bHasBegunPlay);
 
-	/** Get whether BeginPlay has been called on actors */
+	/** Returns true if BeginPlay has been called on actors in the world (and EndPlay has not) */
 	bool GetBegunPlay() const;
 
 	DECLARE_EVENT_OneParam(UWorld, FOnBeginPlay, bool);
@@ -3249,9 +3249,6 @@ public:
 	/** Handle Exec/Console Commands related to the World */
 	bool Exec( UWorld* InWorld, const TCHAR* Cmd, FOutputDevice& Ar=*GLog );
 
-	/** Mark the world as being torn down */
-	void BeginTearingDown();
-
 private:
 	/** Internal version of CleanupWorld. */
 	void CleanupWorldInternal(bool bSessionEnded, bool bCleanupResources, bool bWorldChanged);
@@ -3547,11 +3544,27 @@ public:
 	 */
 	void InitializeActorsForPlay(const FURL& InURL, bool bResetTime = true, FRegisterComponentContext* Context = nullptr);
 
+	/** 
+	 * Mark a world that was initialized for play as starting to tear down in preparation for cleanup.
+	 * This will block the creation of new gameplay objects.
+	 */
+	void BeginTearingDown();
+
 	/**
-	 * Start gameplay. This will cause the game mode to transition to the correct state and call BeginPlay on all actors
+	 * Start gameplay. This will cause the game mode to transition to the correct state and call BeginPlay on all actors.
+	 * If this is called on a world with no game mode, it will execute world callbacks but will not set BegunPlay to true.
+	 * On networked clients, actor BeginPlay and SetBegunPlay are called from the game state replication.
 	 */
 	void BeginPlay();
 
+	/**
+	 * Tries to stop gameplay by sending EndPlay to all actors in the world.
+	 * This will call BeginTearingDown if it has not been called yet.
+	 * @param EndPlayReason The reason that play is ending
+	 * @return true if this actually stopped play (BegunPlay will return false)
+	 */
+	bool EndPlay(EEndPlayReason::Type EndPlayReason);
+
 	/** 
 	 * Looks for a PlayerController that was being swapped by the given NetConnection and, if found, destroys it
 	 * (because the swap is complete or the connection was closed)
@@ -4191,6 +4204,7 @@ public:
 	// Global Callback after actors have been initialized (on any world)
 	static UWorld::FOnWorldInitializedActors OnWorldInitializedActors;
 
+	// Global Callback when an initialized world begins to tear down before cleanup
 	static FWorldEvent OnWorldBeginTearDown;
 
 	DECLARE_MULTICAST_DELEGATE_TwoParams(FOnSeamlessTravelStart, UWorld*, const FString&);
diff --git a/Engine/Source/Runtime/Engine/Private/GameEngine.cpp b/Engine/Source/Runtime/Engine/Private/GameEngine.cpp
--- a/Engine/Source/Runtime/Engine/Private/GameEngine.cpp
+++ b/Engine/Source/Runtime/Engine/Private/GameEngine.cpp
@@ -1254,11 +1254,8 @@ void UGameEngine::PreExit()
 
 			// Make sure there are no pending visibility requests.
 			World->FlushLevelStreaming(EFlushLevelStreamingType::Visibility);
-						
-			for (FActorIterator ActorIt(World); ActorIt; ++ActorIt)
-			{
-				ActorIt->RouteEndPlay(EEndPlayReason::Quit);
-			}
+			
+			World->EndPlay(EEndPlayReason::Quit);
 
 			if (World->GetGameInstance() != nullptr)
 			{
diff --git a/Engine/Source/Runtime/Engine/Private/Tests/AutomationCommon.cpp b/Engine/Source/Runtime/Engine/Private/Tests/AutomationCommon.cpp
--- a/Engine/Source/Runtime/Engine/Private/Tests/AutomationCommon.cpp
+++ b/Engine/Source/Runtime/Engine/Private/Tests/AutomationCommon.cpp
@@ -5,6 +5,7 @@
 #include "Engine/World.h"
 #include "Misc/EngineVersion.h"
 #include "Engine/GameViewportClient.h"
+#include "Engine/GameInstance.h"
 #include "Engine/Engine.h"
 #include "HAL/PlatformFile.h"
 #include "HardwareInfo.h"
@@ -42,6 +43,229 @@ static TAutoConsoleVariable<int32> CVarAutomationAllowFrameTraceCapture(
 FOnEditorAutomationMapLoad AutomationCommon::OnEditorAutomationMapLoad;
 #endif
 
+///////////////////////////////////////////////////////////////////////
+// FTestWorldWrapper
+
+FTestWorldWrapper::~FTestWorldWrapper()
+{
+	// Make sure to destroy the test world unless this is a hard exit
+	if (UObjectInitialized() && !IsEngineExitRequested())
+	{
+		DestroyTestWorld(false);
+	}
+}
+
+bool FTestWorldWrapper::CreateTestWorld(EWorldType::Type WorldType)
+{
+	if (TestWorld != nullptr)
+	{
+		ReportFailure(TEXT("TestWorld already exists in CreateTestWorld!"));
+		return false;
+	}
+
+	UGameInstance* TestGameInstance = nullptr;
+
+	switch (WorldType)
+	{
+		case EWorldType::Game:
+			// Create something close to a game world that can be manually ticked
+			TestGameInstance = NewObject<UGameInstance>(GEngine);
+			TestGameInstance->Init();
+			break;
+
+		case EWorldType::Editor:
+		case EWorldType::EditorPreview:
+		case EWorldType::GamePreview:
+		case EWorldType::GameRPC:
+		case EWorldType::Inactive:
+			// Don't initialize game instance
+			break;
+
+		case EWorldType::PIE: // TODO Implement if needed
+		default:
+			ensureMsgf(false, TEXT("Unsupported world type in CreateTestWorld"));
+			ReportFailure(TEXT("Unsupported world type in CreateTestWorld"));
+			return false;
+	}
+
+	TestWorld = UWorld::CreateWorld(WorldType, false);
+
+	if (!TestWorld)
+	{
+		ReportFailure(TEXT("Failed to create world in CreateTestWorld"));
+		return false;
+	}
+
+	// Will be ticked manually
+	TestWorld->SetShouldTick(false);
+	TestWorld->AddToRoot();
+	
+	if (TestGameInstance)
+	{
+		TestWorld->SetGameInstance(TestGameInstance);
+	}
+
+	FWorldContext& WorldContext = GEngine->CreateNewWorldContext(WorldType);
+	WorldContext.SetCurrentWorld(TestWorld);
+	WorldContext.OwningGameInstance = TestGameInstance;
+
+	return true;
+}
+
+bool FTestWorldWrapper::DestroyTestWorld(bool bForceGarbageCollect)
+{
+	if (TestWorld == nullptr)
+	{
+		ReportFailure(TEXT("TestWorld does not exist in DestroyTestWorld!"));
+		return false;
+	}
+
+	// End play if it hasn't already happened, this will do nothing if it didn't spawn a game mode
+	EndPlayInTestWorld();
+
+	TestWorld->RemoveFromRoot();
+	if (TestWorld->GetGameInstance() != nullptr)
+	{
+		TestWorld->GetGameInstance()->Shutdown();
+	}
+
+	GEngine->DestroyWorldContext(TestWorld);
+	TestWorld->DestroyWorld(false);
+	TestWorld = nullptr;
+
+	if (bForceGarbageCollect)
+	{
+		CollectGarbage(GARBAGE_COLLECTION_KEEPFLAGS);
+
+		// Verify it was actually destroyed
+		GEngine->CheckAndHandleStaleWorldObjectReferences();
+	}
+
+	return true;
+}
+
+bool FTestWorldWrapper::BeginPlayInTestWorld()
+{
+	if (TestWorld == nullptr)
+	{
+		ReportFailure(TEXT("TestWorld does not exist in BeginPlayInTestWorld!"));
+		return false;
+	}
+
+	if (TestWorld->HasBegunPlay())
+	{
+		ReportFailure(TEXT("TestWorld has already begin play in BeginPlayInTestWorld!"));
+		return false;
+	}
+
+	CachedFrameCounter = GFrameCounter;
+
+	FURL URL;
+	if (TestWorld->GetGameInstance())
+	{
+		// This is required to actually forward actor BeginPlay
+		TestWorld->SetGameMode(URL);
+
+		if (!TestWorld->GetAuthGameMode())
+		{
+			ReportFailure(TEXT("BeginPlayInTestWorld failed to create GameMode"));
+			return false;
+		}
+
+		if (!TestWorld->GetWorldSettings())
+		{
+			ReportFailure(TEXT("BeginPlayInTestWorld failed to create WorldSettings"));
+			return false;
+		}
+	}
+
+	TestWorld->InitializeActorsForPlay(URL);
+	TestWorld->BeginPlay();
+
+	if (TestWorld->GetGameInstance())
+	{
+		if (!TestWorld->GetGameState())
+		{
+			ReportFailure(TEXT("BeginPlayInTestWorld failed to create GameState"));
+			return false;
+		}
+	}
+
+	return true;
+}
+
+bool FTestWorldWrapper::EndPlayInTestWorld()
+{
+	if (!TestWorld)
+	{
+		ReportFailure(TEXT("Failed to find world in EndPlayInTestWorld"));
+		return false;
+	}
+
+	// Ignore if this never started play
+	if (TestWorld->HasBegunPlay())
+	{
+		TestWorld->BeginTearingDown();
+		TestWorld->EndPlay(EEndPlayReason::Quit);
+
+		// Restore the frame counter
+		GFrameCounter = CachedFrameCounter;
+	}
+
+	return true;
+}
+
+bool FTestWorldWrapper::TickTestWorld(float DeltaTime)
+{
+	if (!TestWorld)
+	{
+		ReportFailure(TEXT("Failed to find world in TickTestWorld"));
+		return false;
+	}
+
+	TestWorld->Tick(LEVELTICK_All, DeltaTime);
+	if (TestWorld->HasBegunPlay())
+	{
+		// If this has begin play increment the global counter to emulate gameplay
+		GFrameCounter++;
+	}
+
+	return true;
+}
+
+void FTestWorldWrapper::ReportFailure(const TCHAR* ErrorMessage)
+{
+	FailureErrors.Emplace(ErrorMessage);
+}
+
+void FTestWorldWrapper::ClearFailureState()
+{
+	FailureErrors.Reset();
+}
+
+bool FTestWorldWrapper::HasFailed() const
+{
+	return FailureErrors.Num() > 0;
+}
+
+void FTestWorldWrapper::AppendErrorMessages(TArray<FString>& OutErrorMessages) const
+{
+	OutErrorMessages.Append(FailureErrors);
+}
+
+void FTestWorldWrapper::ForwardErrorMessages(FAutomationTestBase* AutomationTest) const
+{
+	if (HasFailed())
+	{
+		TArray<FString> FailureMessages;
+		AppendErrorMessages(FailureMessages);
+		for (const FString& Failure : FailureMessages)
+		{
+			AutomationTest->AddError(Failure);
+		}
+	}
+}
+
 ///////////////////////////////////////////////////////////////////////
 // Common Latent commands
 
diff --git a/Engine/Source/Runtime/Engine/Private/Tests/EngineAutomationTests.cpp b/Engine/Source/Runtime/Engine/Private/Tests/EngineAutomationTests.cpp
--- a/Engine/Source/Runtime/Engine/Private/Tests/EngineAutomationTests.cpp
+++ b/Engine/Source/Runtime/Engine/Private/Tests/EngineAutomationTests.cpp
@@ -1941,21 +1941,21 @@ void AttachmentTest_AttachWhenAttached(UWorld* World, FAutomationTestBase* Test)
 
 bool FAutomationAttachment::RunTest(const FString& Parameters)
 {
-	UWorld *World = UWorld::CreateWorld(EWorldType::Game, false);
-	FWorldContext& WorldContext = GEngine->CreateNewWorldContext(EWorldType::Game);
-	WorldContext.SetCurrentWorld(World);
+	// This will get cleaned up when it leaves scope
+	FTestWorldWrapper WorldWrapper;
+	WorldWrapper.CreateTestWorld(EWorldType::Game);
+	UWorld* World = WorldWrapper.GetTestWorld();
 
-	FURL URL;
-	World->InitializeActorsForPlay(URL);
-	World->BeginPlay();
-
-	AttachmentTest_AttachWhenNotAttached(World, this);
-	AttachmentTest_AttachWhenAttached(World, this);
-
-	GEngine->DestroyWorldContext(World);
-	World->DestroyWorld(false);
+	if (World)
+	{
+		WorldWrapper.BeginPlayInTestWorld();
+		AttachmentTest_AttachWhenNotAttached(World, this);
+		AttachmentTest_AttachWhenAttached(World, this);
+		WorldWrapper.ForwardErrorMessages(this);
 
-	return true;
+		return !HasAnyErrors();
+	}
+	return false;
 }
 
 #endif //WITH_DEV_AUTOMATION_TESTS
diff --git a/Engine/Source/Runtime/Engine/Private/TimerManagerTests.cpp b/Engine/Source/Runtime/Engine/Private/TimerManagerTests.cpp
--- a/Engine/Source/Runtime/Engine/Private/TimerManagerTests.cpp
+++ b/Engine/Source/Runtime/Engine/Private/TimerManagerTests.cpp
@@ -4,6 +4,7 @@
 #include "Misc/AutomationTest.h"
 #include "TimerManager.h"
 #include "Engine/Engine.h"
+#include "Tests/AutomationCommon.h"
 
 IMPLEMENT_SIMPLE_AUTOMATION_TEST(FTimerManagerTest, "System.Engine.TimerManager", EAutomationTestFlags::EditorContext | EAutomationTestFlags::EngineFilter)
 
@@ -265,24 +266,25 @@ bool TimerManagerTest_LoopingTimers_DifferentHandles(UWorld* World, FAutomationT
 
 bool FTimerManagerTest::RunTest(const FString& Parameters)
 {
-	UWorld *World = UWorld::CreateWorld(EWorldType::Game, false);
-	FWorldContext &WorldContext = GEngine->CreateNewWorldContext(EWorldType::Game);
-	WorldContext.SetCurrentWorld(World);
-	
-	FURL URL;
-	World->InitializeActorsForPlay(URL);
-	World->BeginPlay();
-
-	TimerManagerTest_InvalidTimers(World, this);
-	TimerManagerTest_MissingTimers(World, this);
-	TimerManagerTest_ValidTimer_HandleWithDelegate(World, this);
-	TimerManagerTest_ValidTimer_HandleLoopingSetDuringExecute(World, this);
-	TimerManagerTest_LoopingTimers_DifferentHandles(World, this);
+	// This will get cleaned up when it leaves scope
+	FTestWorldWrapper WorldWrapper;
+	WorldWrapper.CreateTestWorld(EWorldType::Game);
+	UWorld* World = WorldWrapper.GetTestWorld();
 
-	GEngine->DestroyWorldContext(World);
-	World->DestroyWorld(false);
+	if (World)
+	{
+		WorldWrapper.BeginPlayInTestWorld();
+		TimerManagerTest_InvalidTimers(World, this);
+		TimerManagerTest_MissingTimers(World, this);
+		TimerManagerTest_ValidTimer_HandleWithDelegate(World, this);
+		TimerManagerTest_ValidTimer_HandleLoopingSetDuringExecute(World, this);
+		TimerManagerTest_LoopingTimers_DifferentHandles(World, this);
+		WorldWrapper.ForwardErrorMessages(this);
+
+		return !HasAnyErrors();
+	}
 
-	return true;
+	return false;
 }
 
 
diff --git a/Engine/Source/Runtime/Engine/Private/UnrealEngine.cpp b/Engine/Source/Runtime/Engine/Private/UnrealEngine.cpp
--- a/Engine/Source/Runtime/Engine/Private/UnrealEngine.cpp
+++ b/Engine/Source/Runtime/Engine/Private/UnrealEngine.cpp
@@ -15237,10 +15237,7 @@ bool UEngine::LoadMap( FWorldContext& WorldContext, FURL URL, class UPendingNetG
 			}
 		}
 
-		for (FActorIterator ActorIt(WorldContext.World()); ActorIt; ++ActorIt)
-		{
-			ActorIt->RouteEndPlay(EEndPlayReason::LevelTransition);
-		}
+		WorldContext.World()->EndPlay(EEndPlayReason::LevelTransition);
 
 		// Do this after destroying pawns/playercontrollers, in case that spawns new things (e.g. dropped weapons)
 		WorldContext.World()->CleanupWorld();
diff --git a/Engine/Source/Runtime/Engine/Private/World.cpp b/Engine/Source/Runtime/Engine/Private/World.cpp
--- a/Engine/Source/Runtime/Engine/Private/World.cpp
+++ b/Engine/Source/Runtime/Engine/Private/World.cpp
@@ -3494,16 +3494,6 @@ void UWorld::AddToWorld( ULevel* Level, const FTransform& LevelTransform, bool b
 	}
 }
 
-void UWorld::BeginTearingDown()
-{
-	bIsTearingDown = true;
-	UE_LOG(LogWorld, Log, TEXT("BeginTearingDown for %s"), *GetOutermost()->GetName());
-
-	//Simultaneous similar edits that caused merge conflict. Taking both for now to unblock.
-	//Can likely be unified.
-	FWorldDelegates::OnWorldBeginTearDown.Broadcast(this);
-}
-
 // Cumulated time doing IncrementalUnregisterComponents in UWorld::RemoveFromWorld since last call to UWorld::UpdateLevelStreaming.
 static double GRemoveFromWorldUnregisterComponentTimeCumul = 0.0;
 
@@ -5300,6 +5290,21 @@ void UWorld::InitializeActorsForPlay(const FURL& InURL, bool bResetTime, FRegist
 	}
 }
 
+void UWorld::BeginTearingDown()
+{
+	if (bIsTearingDown)
+	{
+		UE_LOG(LogWorld, Warning, TEXT("BeginTearingDown called twice for %s!"), *GetOutermost()->GetName());
+	}
+	else
+	{
+		UE_LOG(LogWorld, Log, TEXT("BeginTearingDown for %s"), *GetOutermost()->GetName());
+	}
+
+	bIsTearingDown = true;
+	FWorldDelegates::OnWorldBeginTearDown.Broadcast(this);
+}
+
 void UWorld::BeginPlay()
 {
 	if (SupportsMakingVisibleTransactionRequests() && (IsNetMode(NM_DedicatedServer) || IsNetMode(NM_ListenServer)))
@@ -5334,6 +5339,30 @@ void UWorld::BeginPlay()
 	}
 }
 
+bool UWorld::EndPlay(EEndPlayReason::Type EndPlayReason)
+{
+	// If it hasn't already happened, mark the world as starting to tear down even if it did not fully start play
+	if (!bIsTearingDown)
+	{
+		BeginTearingDown();
+	}
+
+	if (!HasBegunPlay())
+	{
+		// Ignore requests on worlds that did not fully start play, such as test worlds with no game mode
+		return false;
+	}
+
+	for (FActorIterator ActorIt(this); ActorIt; ++ActorIt)
+	{
+		ActorIt->RouteEndPlay(EndPlayReason);
+	}
+
+	SetBegunPlay(false);
+
+	return true;
+}
+
 bool UWorld::IsNavigationRebuilt() const
 {
 	return GetNavigationSystem() == NULL || GetNavigationSystem()->IsNavigationBuilt(GetWorldSettings());
@@ -5350,6 +5379,11 @@ void UWorld::CleanupWorld(bool bSessionEnded, bool bCleanupResources, UWorld* Ne
 		bool bIsStreamingSubWorld = PersistentLevel && PersistentLevel->OwningWorld != this;
 		UE_CLOG(!bIsStreamingSubWorld, LogWorld, Warning, TEXT("UWorld::CleanupWorld called twice or called without InitWorld called first (%s)"), *GetName());
 	}
+	if (HasBegunPlay())
+	{
+		UE_LOG(LogWorld, Warning, TEXT("UWorld::CleanupWorld called on a world that has begun play, missing call to EndPlay (%s)"), *GetName());
+	}
+
 	const bool bWorldChanged = NewWorld != this;
 	CleanupWorldInternal(bSessionEnded, bCleanupResources, bWorldChanged);
 	bIsWorldInitialized = false;
@@ -5385,6 +5419,7 @@ void UWorld::CleanupWorldInternal(bool bSessionEnded, bool bCleanupResources, bo
 	if(FPhysScene* CurrPhysicsScene = GetPhysicsScene())
 	{
 		CurrPhysicsScene->WaitPhysScenes();
+		// @TODO: This function is misnamed for when it is called, cleanup is called after end play
 		CurrPhysicsScene->OnWorldEndPlay();
 	}
 
diff --git a/Engine/Source/Runtime/Engine/Public/Tests/AutomationCommon.h b/Engine/Source/Runtime/Engine/Public/Tests/AutomationCommon.h
--- a/Engine/Source/Runtime/Engine/Public/Tests/AutomationCommon.h
+++ b/Engine/Source/Runtime/Engine/Public/Tests/AutomationCommon.h
@@ -5,6 +5,7 @@
 #include "CoreMinimal.h"
 #include "Misc/AutomationTest.h"
 #include "Engine/GameViewportClient.h"
+#include "Engine/World.h"
 
 class AActor;
 class SWindow;
@@ -22,6 +23,54 @@ DECLARE_MULTICAST_DELEGATE_ThreeParams(FOnEditorAutomationMapLoad, const FString
 
 #endif
 
+/** 
+ * Utility class for creating and destroying a temporary test world, can be used for automation or performance testing.
+ * This can be used with any test framework but has utility functions for FAutomationTestBase.
+ */
+struct FTestWorldWrapper
+{
+	/** This will properly shut down and destroy the test world as needed */
+	ENGINE_API virtual ~FTestWorldWrapper();
+
+	/** Gets the wrapped world, can be null */
+	inline UWorld* GetTestWorld() const {return TestWorld;}
+
+	/** Creates a world of the appropriate world type, returns false on failure */
+	ENGINE_API virtual bool CreateTestWorld(EWorldType::Type WorldType);
+
+	/** Destroys the test world and handles any required cleanup */
+	ENGINE_API virtual bool DestroyTestWorld(bool bForceGarbageCollect);
+
+	/** Starts play in the test world to simulate gameplay */
+	ENGINE_API virtual bool BeginPlayInTestWorld();
+
+	/** Ticks the test world for one frame, defaults to 100 fps */
+	ENGINE_API virtual bool TickTestWorld(float DeltaTime = 0.01f);
+
+	/** Stops play properly */
+	ENGINE_API virtual bool EndPlayInTestWorld();
+
+	/** Registers an error message and marks test as failed, called by the functions above */
+	ENGINE_API virtual void ReportFailure(const TCHAR* ErrorMessage);
+
+	/** Clears any failures and initial state */
+	ENGINE_API virtual void ClearFailureState();
+
+	/** Returns true if there are any errors that should stop further execution */
+	ENGINE_API virtual bool HasFailed() const;
+
+	/** Gets the actual error messages for reporting to the automation framework */
+	ENGINE_API virtual void AppendErrorMessages(TArray<FString>& OutErrorMessages) const;
+
+	/** Reports error messages to a passed in automation test */
+	ENGINE_API virtual void ForwardErrorMessages(FAutomationTestBase* AutomationTest) const;
+
+protected:
+	UWorld* TestWorld = nullptr;
+	uint64 CachedFrameCounter = 0;
+	TArray<FString> FailureErrors;
+};
+
 /** Common automation functions */
 namespace AutomationCommon
 {
