Added missing case for enum EFunc1::Exp2 in FMaterialDerivativeAutogen::GetFunc1ReturnType

#rnx

diff --git a/Engine/Source/Runtime/Engine/Private/Materials/HLSLMaterialDerivativeAutogen.cpp b/Engine/Source/Runtime/Engine/Private/Materials/HLSLMaterialDerivativeAutogen.cpp
--- a/Engine/Source/Runtime/Engine/Private/Materials/HLSLMaterialDerivativeAutogen.cpp
+++ b/Engine/Source/Runtime/Engine/Private/Materials/HLSLMaterialDerivativeAutogen.cpp
@@ -351,6 +351,7 @@ FMaterialDerivativeAutogen::FOperationType1 FMaterialDerivativeAutogen::GetFunc1
 	case EFunc1::Log2:
 	case EFunc1::Log10:
 	case EFunc1::Exp:
+	case EFunc1::Exp2:
 	case EFunc1::Asin:
 	case EFunc1::AsinFast:
 	case EFunc1::Acos:
