Move the UMorph target creation to the game thread after the build. This will prevent two issues. The first one is we cannot call UObject::FindObjectSafe outside of the main thread since it take a gc scope lock (UE-208853). The second issue is we cannot modify the skeletalmesh morph target array outside of the game thread because FReferencerFinder will assert if the call is done during the skeletalmesh build with morph targets (UE-205185).

#rnx

diff --git a/Engine/Source/Runtime/Engine/Classes/Engine/SkinnedAssetAsyncCompileUtils.h b/Engine/Source/Runtime/Engine/Classes/Engine/SkinnedAssetAsyncCompileUtils.h
--- a/Engine/Source/Runtime/Engine/Classes/Engine/SkinnedAssetAsyncCompileUtils.h
+++ b/Engine/Source/Runtime/Engine/Classes/Engine/SkinnedAssetAsyncCompileUtils.h
@@ -6,6 +6,9 @@
  * Contains the helper functions to compile skinned asset asynchronously.
  */
 
+#if WITH_EDITOR
+#include "Animation/MorphTarget.h"
+#endif //WITH_EDITOR
 #include "Async/AsyncWork.h"
 #include "SkeletalMeshTypes.h"
 
@@ -24,6 +27,9 @@ public:
 	
 	//True if this compilation context is start from a serialize save
 	bool bIsSerializeSaving = false;
+#if WITH_EDITOR
+	TUniquePtr<FFinishBuildMorphTargetData> FinishBuildMorphTargetData;
+#endif //WITH_EDITOR
 };
 
 class FSkinnedAssetPostLoadContext : public FSkinnedAssetCompilationContext
diff --git a/Engine/Source/Runtime/Engine/Private/MorphTools.cpp b/Engine/Source/Runtime/Engine/Private/MorphTools.cpp
--- a/Engine/Source/Runtime/Engine/Private/MorphTools.cpp
+++ b/Engine/Source/Runtime/Engine/Private/MorphTools.cpp
@@ -219,6 +219,18 @@ TUniquePtr<FFinishBuildMorphTargetData> UMorphTarget::CreateFinishBuildMorphTarg
 
 void FFinishBuildMorphTargetData::ApplyEditorData(USkeletalMesh * SkeletalMesh, bool bIsSerializeSaving) const
 {
+	//List of potential issue with this function
+	// -This function change the array USkeletalMesh::MorphTargets, if a find reference call happen in same time for this asset(on the game thread)
+	//  The find reference can crash because it is accessing this array to find UObject reference.
+	//
+	// -If this function is call in the pre garbage collector event (we do finish compilation in this event), the call to FindObjectSafe will acquire a GC lock
+	//  which will cause a dead lock since the garbage collector already have the gc lock.
+	//
+	// TODO: remove the UMorphTarget object and replace the data in the import data of the skeletal mesh.
+	if (!IsInGameThread())
+	{
+		UE_ASSET_LOG(LogSkeletalMesh, Display, SkeletalMesh, TEXT("Calling function FFinishBuildMorphTargetData::ApplyEditorData outside of the game thread is not safe and can deadlock or crash."));
+	}
 	//Return if we do not need to apply data
 	if (!bApplyMorphTargetsData)
 	{
@@ -278,28 +290,28 @@ void FFinishBuildMorphTargetData::ApplyEditorData(USkeletalMesh * SkeletalMesh,
 				//which happen before the serialization of that cook skeletalmesh
 				if (!bIsSerializeSaving)
 				{
-				//Avoid recycling morphtarget with NewObject it cannot be done asynchronously
-				//Find the UMorphTarget and simply clear the data if it exist.
-				TArray<UObject*> SubObjects;
-				GetObjectsWithOuter(SkeletalMesh, SubObjects, true);
-				for (UObject* SubObject : SubObjects)
-				{
-					if (SubObject->GetFName() == MorphTargetName)
+					//Avoid recycling morphtarget with NewObject it cannot be done asynchronously
+					//Find the UMorphTarget and simply clear the data if it exist.
+					TArray<UObject*> SubObjects;
+					GetObjectsWithOuter(SkeletalMesh, SubObjects, true);
+					for (UObject* SubObject : SubObjects)
 					{
-						if (UMorphTarget* SubMorphTarget = Cast<UMorphTarget>(SubObject))
+						if (SubObject->GetFName() == MorphTargetName)
 						{
-							MorphTarget = SubMorphTarget;
-							MorphTarget->EmptyMorphLODModels();
-							MorphTarget->ClearGarbage();
-							break;
+							if (UMorphTarget* SubMorphTarget = Cast<UMorphTarget>(SubObject))
+							{
+								MorphTarget = SubMorphTarget;
+								MorphTarget->EmptyMorphLODModels();
+								MorphTarget->ClearGarbage();
+								break;
+							}
 						}
 					}
-				}
-				//Create a new morph target, if the object do not exist (creating a new uobject is ok to do asynchronously)
-				if (!MorphTarget)
-				{
-					MorphTarget = NewObject<UMorphTarget>(SkeletalMesh, MorphTargetName);
-				}
+					//Create a new morph target, if the object do not exist (creating a new uobject is ok to do asynchronously)
+					if (!MorphTarget)
+					{
+						MorphTarget = NewObject<UMorphTarget>(SkeletalMesh, MorphTargetName);
+					}
 					check(MorphTarget);
 				}
 				else
diff --git a/Engine/Source/Runtime/Engine/Private/SkeletalMesh.cpp b/Engine/Source/Runtime/Engine/Private/SkeletalMesh.cpp
--- a/Engine/Source/Runtime/Engine/Private/SkeletalMesh.cpp
+++ b/Engine/Source/Runtime/Engine/Private/SkeletalMesh.cpp
@@ -1526,6 +1526,11 @@ void CachePlatform(USkeletalMesh* Mesh, const ITargetPlatform* TargetPlatform, F
 	FSkinnedAssetBuildContext Context;
 	Context.bIsSerializeSaving = bIsSerializeSaving;
 	PlatformRenderData->Cache(TargetPlatform, Mesh, &Context);
+	if (Context.FinishBuildMorphTargetData.IsValid())
+	{
+		// Morph target is only supported on USkeletalMesh
+		Context.FinishBuildMorphTargetData->ApplyEditorData(Mesh, Context.bIsSerializeSaving);
+	}
 }
 
 static FSkeletalMeshRenderData& GetPlatformSkeletalMeshRenderData(USkeletalMesh* Mesh, const ITargetPlatform* TargetPlatform, const bool bIsSerializeSaving)
@@ -1968,10 +1973,7 @@ void USkeletalMesh::Build()
 	FSkinnedAssetBuildContext Context;
 	BeginBuildInternal(Context);
 	
-	//5.4 hack fix: if this skeletal mesh has some morph targets we disable asynchronous build.
-	//TODO: To remove this constraint we must get rid of the UMorphTarget sub object and directly create the render data from the FSkeletalMeshImportData.
-	const bool bBuildAsynchronous = FSkinnedAssetCompilingManager::Get().IsAsyncCompilationAllowed(this) && GetMorphTargets().IsEmpty();
-	if (bBuildAsynchronous)
+	if (FSkinnedAssetCompilingManager::Get().IsAsyncCompilationAllowed(this))
 	{
 		PrepareForAsyncCompilation();
 
@@ -2037,6 +2039,13 @@ void USkeletalMesh::ApplyFinishBuildInternalData(FSkinnedAssetCompilationContext
 	//We cannot execute this code outside of the game thread
 	checkf(IsInGameThread(), TEXT("Cannot execute function USkeletalMesh::ApplyFinishBuildInternalData asynchronously. Asset: %s"), *this->GetFullName());
 	check(ContextPtr);
+
+	//Apply the morphtargets change if any
+	if (ContextPtr->FinishBuildMorphTargetData.IsValid())
+	{
+		// Morph target is only supported on USkeletalMesh
+		ContextPtr->FinishBuildMorphTargetData->ApplyEditorData(this, ContextPtr->bIsSerializeSaving);
+	}
 }
 
 void USkeletalMesh::FinishBuildInternal(FSkinnedAssetBuildContext& Context)
@@ -3823,6 +3832,10 @@ void USkeletalMesh::GetAssetRegistryTags(FAssetRegistryTagsContext Context) cons
 
 		for(UMorphTarget* MorphTarget : GetMorphTargets())
 		{
+			if (!MorphTarget)
+			{
+				continue;
+			}
 			MorphTarget->GetFName().AppendString(MorphNamesBuilder);
 			MorphNamesBuilder.Append(MorphNamesTagDelimiter);
 		}
diff --git a/Engine/Source/Runtime/Engine/Private/SkeletalMeshRenderData.cpp b/Engine/Source/Runtime/Engine/Private/SkeletalMeshRenderData.cpp
--- a/Engine/Source/Runtime/Engine/Private/SkeletalMeshRenderData.cpp
+++ b/Engine/Source/Runtime/Engine/Private/SkeletalMeshRenderData.cpp
@@ -213,8 +213,7 @@ FString FSkeletalMeshRenderData::GetDerivedDataKey(const ITargetPlatform* Target
 void FSkeletalMeshRenderData::Cache(const ITargetPlatform* TargetPlatform, USkinnedAsset* Owner, FSkinnedAssetCompilationContext* ContextPtr)
 {
 	check(Owner);
-	// Disable ContextPtr check because only USkeletalMesh supports it.
-	//check(ContextPtr);
+	check(ContextPtr);
 
 	check(LODRenderData.Num() == 0); // Should only be called on new, empty RenderData
 	check(TargetPlatform);
@@ -257,24 +256,21 @@ void FSkeletalMeshRenderData::Cache(const ITargetPlatform* TargetPlatform, USkin
 
 			FLargeMemoryReader Ar(DerivedData.GetData(), DerivedData.Num(), ELargeMemoryReaderFlags::Persistent);
 
-			//Helper structure to change the morph targets
-			TUniquePtr<FFinishBuildMorphTargetData> FinishBuildMorphTargetData;
-
 			FSkeletalMeshModel* SkelMeshModel = Owner->GetImportedModel();
 			check(SkelMeshModel);
 
 			//Get the morph target data, we put it in the compilation context to apply them in the game thread before the InitResources
 			if (Owner->GetMorphTargets().Num() > 0)
 			{
-				FinishBuildMorphTargetData = Owner->GetMorphTargets()[0]->CreateFinishBuildMorphTargetData();
+				ContextPtr->FinishBuildMorphTargetData = Owner->GetMorphTargets()[0]->CreateFinishBuildMorphTargetData();
 			}
 			else
 			{
 				// Create and initialize the FinishBuildInternalData, use the class default object to call the virtual function
-				FinishBuildMorphTargetData = UMorphTarget::StaticClass()->GetDefaultObject<UMorphTarget>()->CreateFinishBuildMorphTargetData();
+				ContextPtr->FinishBuildMorphTargetData = UMorphTarget::StaticClass()->GetDefaultObject<UMorphTarget>()->CreateFinishBuildMorphTargetData();
 			}
-			check(FinishBuildMorphTargetData);
-			FinishBuildMorphTargetData->LoadFromMemoryArchive(Ar);
+			check(ContextPtr->FinishBuildMorphTargetData);
+			ContextPtr->FinishBuildMorphTargetData->LoadFromMemoryArchive(Ar);
 
 			//Serialize the LODModel sections since they are dependent on the reduction
 			for (int32 LODIndex = 0; LODIndex < SkelMeshModel->LODModels.Num(); LODIndex++)
@@ -298,16 +294,6 @@ void FSkeletalMeshRenderData::Cache(const ITargetPlatform* TargetPlatform, USkin
 				LODData.SerializeStreamedData(Ar, Owner, LODIndex, DummyStripFlags, bNeedsCPUAccess, bForceKeepCPUResources);
 			}
 
-			//Apply the morphtargets change if any
-			if (FinishBuildMorphTargetData.IsValid())
-			{
-				// Morph target is only supported on USkeletalMesh
-				if (USkeletalMesh* SkMesh = Cast<USkeletalMesh>(Owner))
-				{
-					FinishBuildMorphTargetData->ApplyEditorData(SkMesh, ContextPtr ? ContextPtr->bIsSerializeSaving : false);
-				}
-			}
-
 			int32 T1 = FPlatformTime::Cycles();
 			UE_LOG(LogSkeletalMesh, Verbose, TEXT("Skeletal Mesh found in DDC [%fms] %s"), FPlatformTime::ToMilliseconds(T1 - T0), *Owner->GetPathName());
 		}
