Fix incorrect deprecation code in FMeshApproximationSettings
* Prior to 27594466, the meshes were actually built with FallbackTarget == Auto by default, since the addition of FallbackTarget options by the Nanite team.
* Any value (incorrect or not) specified with the now deprecated FMeshApproximationSettings::NaniteProxyTrianglePercent was ignored.
* CL 27594466 incorrectly handled the deprecation by changing the default FallbackTarget to PercentTriangles.

diff --git a/Engine/Source/Runtime/Engine/Private/EngineTypes.cpp b/Engine/Source/Runtime/Engine/Private/EngineTypes.cpp
--- a/Engine/Source/Runtime/Engine/Private/EngineTypes.cpp
+++ b/Engine/Source/Runtime/Engine/Private/EngineTypes.cpp
@@ -87,7 +87,7 @@ void FMeshApproximationSettings::PostSerialize(const FArchive& Ar)
 		FMeshApproximationSettings DefaultObject;
 		if (NaniteProxyTrianglePercent_DEPRECATED != DefaultObject.NaniteProxyTrianglePercent_DEPRECATED)
 		{
-			NaniteFallbackTarget = ENaniteFallbackTarget::PercentTriangles;
+			NaniteFallbackTarget = ENaniteFallbackTarget::Auto;
 			NaniteFallbackPercentTriangles = NaniteProxyTrianglePercent_DEPRECATED / 100.0f;
 		}
 	}
