NetDriver will only process dynamic actor renames if the actor's outer changed. There is no meaningful work to do if the outer is the same.

diff --git a/Engine/Source/Runtime/Engine/Private/NetDriver.cpp b/Engine/Source/Runtime/Engine/Private/NetDriver.cpp
--- a/Engine/Source/Runtime/Engine/Private/NetDriver.cpp
+++ b/Engine/Source/Runtime/Engine/Private/NetDriver.cpp
@@ -3991,9 +3991,11 @@ void UNetDriver::NotifyActorRenamed(AActor* ThisActor, UObject* PreviousOuter, F
 		}
 		else
 		{
-			if (bIsServer)
+			const bool bActorChangedOuter = ThisActor->GetOuter() != PreviousOuter;
+
+			if (bIsServer && bActorChangedOuter)
 			{
-				UE_LOG(LogNet, Log, TEXT("NotifyActorRenamed on server, dynamic actor: %s PreviousOuter: %s PreviousName: %s"), *GetFullNameSafe(ThisActor), *GetFullNameSafe(PreviousOuter), *PreviousName.ToString());
+				UE_LOG(LogNet, Log, TEXT("NotifyActorRenamed on server, dynamic actor changed outer: %s PreviousOuter: %s PreviousName: %s"), *GetFullNameSafe(ThisActor), *GetFullNameSafe(PreviousOuter), *PreviousName.ToString());
 
 #if UE_WITH_IRIS
 				ensureMsgf(!ReplicationSystem, TEXT("Dynamic actor renaming not supported in Iris. Actor: %s PreviousOuter: %s"), *GetFullNameSafe(ThisActor), *GetFullNameSafe(PreviousOuter));
