Fixed body instance filter calculation when its SourceObject member is unset.

Source objects were added to allow non-actor workflows to use queries and bind their objects to the filter for a body. However it's currently assumed to always be valid. That's not always the case - such as with geoemtry collections and with various custom uses of body instance. Rather than just set the source object from geoemtry collection this change implements the correct fallback logic for the actor as owner if the check for the source object fails which could have been a potential breakage point in the body instance API.

diff --git a/Engine/Source/Runtime/Engine/Private/PhysicsEngine/BodyInstance.cpp b/Engine/Source/Runtime/Engine/Private/PhysicsEngine/BodyInstance.cpp
--- a/Engine/Source/Runtime/Engine/Private/PhysicsEngine/BodyInstance.cpp
+++ b/Engine/Source/Runtime/Engine/Private/PhysicsEngine/BodyInstance.cpp
@@ -4295,6 +4295,7 @@ void FBodyInstance::BuildBodyFilterData(FBodyCollisionFilterData& OutFilterData,
 	{
 		ActorOwner = OwnerComponentInst->GetOwner();
 	}
+
 	const bool bPhysicsStatic = !OwnerComponentInst || OwnerComponentInst->IsWorldGeometry();
 
 	// Grab collision setting from body instance.
@@ -4382,7 +4383,7 @@ void FBodyInstance::BuildBodyFilterData(FBodyCollisionFilterData& OutFilterData,
 		FCollisionFilterData SimFilterData;
 		FCollisionFilterData SimpleQueryData;
 
-		uint32 SourceObjectID = SourceObject.IsValid() ? SourceObject->GetUniqueID() : 0;
+		uint32 SourceObjectID = SourceObject.IsValid() ? SourceObject->GetUniqueID() : ActorOwner ? ActorOwner->GetUniqueID() : 0;
 		uint32 CompID = (OwnerComponentInst != nullptr) ? OwnerComponentInst->GetUniqueID() : 0;
 		CreateShapeFilterData(UseChannel, MaskFilter, SourceObjectID, UseResponse, CompID, InstanceBodyIndex, SimpleQueryData, SimFilterData, bRootCCD && !bPhysicsStatic, bUseNotifyRBCollision, bPhysicsStatic, bUseContactModification);
 
