Fixed USparseVolumeTextureFrame objects being created with the wrong NewObject function, leaving them without the RF_Public flag. This caused issues when the frame is set on a material and the user tries to save that material.

#rnx

diff --git a/Engine/Source/Runtime/Engine/Private/SparseVolumeTexture/SparseVolumeTexture.cpp b/Engine/Source/Runtime/Engine/Private/SparseVolumeTexture/SparseVolumeTexture.cpp
--- a/Engine/Source/Runtime/Engine/Private/SparseVolumeTexture/SparseVolumeTexture.cpp
+++ b/Engine/Source/Runtime/Engine/Private/SparseVolumeTexture/SparseVolumeTexture.cpp
@@ -1223,9 +1223,11 @@ bool UStreamableSparseVolumeTexture::AppendFrame(UE::SVT::FTextureData& Uncooked
 		VolumeResolution = VolumeBoundsMax;
 	}
 	
-
-	USparseVolumeTextureFrame* Frame = NewObject<USparseVolumeTextureFrame>(this);
-	if (Frame->Initialize(this, Frames.Num(), FrameTransform, UncookedFrame))
+	const int32 FrameIndex = Frames.Num();
+	const FName FrameBaseName = FName(FString::Printf(TEXT("%s_Frame%i"), *GetFName().ToString(), FrameIndex));
+	const FName FrameName = MakeUniqueObjectName(this, USparseVolumeTextureFrame::StaticClass(), FrameBaseName);
+	USparseVolumeTextureFrame* Frame = NewObject<USparseVolumeTextureFrame>(this, USparseVolumeTextureFrame::StaticClass(), FrameName, RF_Public);
+	if (Frame->Initialize(this, FrameIndex, FrameTransform, UncookedFrame))
 	{
 		Frames.Add(Frame);
 		return true;
