Fixed case where scene could be completely in shadow when using mobile renderer with and unbuilt lighting data

diff --git a/Engine/Source/Runtime/Engine/Classes/Components/LightComponent.h b/Engine/Source/Runtime/Engine/Classes/Components/LightComponent.h
--- a/Engine/Source/Runtime/Engine/Classes/Components/LightComponent.h
+++ b/Engine/Source/Runtime/Engine/Classes/Components/LightComponent.h
@@ -67,7 +67,8 @@ class ULightComponent : public ULightComponentBase
 	int32 ShadowMapChannel_DEPRECATED;
 
 	/** Transient shadowmap channel used to preview the results of stationary light shadowmap packing. */
-	int32 PreviewShadowMapChannel=0;
+	UPROPERTY(NonPIEDuplicateTransient)
+	int32 PreviewShadowMapChannel;
 	
 	/** Min roughness effective for this light. Used for softening specular highlights. */
 	UPROPERTY()
diff --git a/Engine/Source/Runtime/Engine/Private/Components/LightComponent.cpp b/Engine/Source/Runtime/Engine/Private/Components/LightComponent.cpp
--- a/Engine/Source/Runtime/Engine/Private/Components/LightComponent.cpp
+++ b/Engine/Source/Runtime/Engine/Private/Components/LightComponent.cpp
@@ -32,6 +32,7 @@
 #include "UObject/ICookInfo.h"
 #include "UObject/SoftObjectPath.h"
 #include "UObject/UnrealType.h"
+#include "UObject/Package.h"
 
 #include UE_INLINE_GENERATED_CPP_BY_NAME(LightComponent)
 
@@ -566,7 +567,12 @@ void ULightComponent::PostLoad()
 		ClearLightFunctionMaterial();
 	}
 
-	PreviewShadowMapChannel = INDEX_NONE;
+	// we want to make sure PreviewShadowMapChannel gets into PIE unchanged
+	if (!WITH_EDITOR || !GetOutermost()->HasAnyPackageFlags(PKG_PlayInEditor))
+	{
+		PreviewShadowMapChannel = INDEX_NONE;
+	}
+
 	Intensity = FMath::Max(0.0f, Intensity);
 
 	if (GetLinkerUEVersion() < VER_UE4_LIGHTCOMPONENT_USE_IES_TEXTURE_MULTIPLIER_ON_NON_IES_BRIGHTNESS)
