Prevent crashes due to locking the object hash while iterating on and deinitializing subsystems.

diff --git a/Engine/Source/Runtime/Engine/Private/Subsystems/SubsystemCollection.cpp b/Engine/Source/Runtime/Engine/Private/Subsystems/SubsystemCollection.cpp
--- a/Engine/Source/Runtime/Engine/Private/Subsystems/SubsystemCollection.cpp
+++ b/Engine/Source/Runtime/Engine/Private/Subsystems/SubsystemCollection.cpp
@@ -365,7 +365,10 @@ void FSubsystemCollectionBase::AddAllInstances(UClass* SubsystemClass)
 
 void FSubsystemCollectionBase::RemoveAllInstances(UClass* SubsystemClass)
 {
-	ForEachObjectOfClass(SubsystemClass, [](UObject* SubsystemObj)
+	TArray<UObject*> SubsystemsToRemove;
+	GetObjectsOfClass(SubsystemClass, SubsystemsToRemove);
+
+	for(UObject* SubsystemObj : SubsystemsToRemove)
 	{
 		USubsystem* Subsystem = CastChecked<USubsystem>(SubsystemObj);
 
@@ -373,7 +376,7 @@ void FSubsystemCollectionBase::RemoveAllInstances(UClass* SubsystemClass)
 		{
 			Subsystem->InternalOwningSubsystem->RemoveAndDeinitializeSubsystem(Subsystem);
 		}
-	});
+	}
 }
 
 
