[NetPhysics] Handle unregistering of rewindable ISimCallbackObject from IRewindCallback

diff --git a/Engine/Source/Runtime/Engine/Public/Physics/NetworkPhysicsComponent.h b/Engine/Source/Runtime/Engine/Public/Physics/NetworkPhysicsComponent.h
--- a/Engine/Source/Runtime/Engine/Public/Physics/NetworkPhysicsComponent.h
+++ b/Engine/Source/Runtime/Engine/Public/Physics/NetworkPhysicsComponent.h
@@ -112,7 +112,7 @@ struct TNetRewindHistory : public Chaos::TDataRewindHistory<DataType>
 		
 		if (Super::NumFrames > GetMaxArraySize())
 		{
-			UE_LOG(LogChaos, Warning, TEXT("TNetRewindHistory: serialized array of size %d exceeds maximum size %d."), Super::NumFrames, GetMaxArraySize());
+			UE_LOG(LogTemp, Warning, TEXT("TNetRewindHistory: serialized array of size %d exceeds maximum size %d."), Super::NumFrames, GetMaxArraySize());
 			Ar.SetError();
 			return;
 		}
@@ -170,7 +170,7 @@ private :
 		}
 		else
 		{
-			UE_LOG(LogChaos, Error, TEXT("TNetRewindHistory::NetSerializeData called on data struct %s without a native NetSerialize"), *ScriptStruct->GetName());
+			UE_LOG(LogTemp, Error, TEXT("TNetRewindHistory::NetSerializeData called on data struct %s without a native NetSerialize"), *ScriptStruct->GetName());
 
 			// Not working for now since the packagemap could be null
 			// UNetConnection* Connection = CastChecked<UPackageMapClient>(PackageMap)->GetConnection();
@@ -286,6 +286,11 @@ struct FNetworkPhysicsCallback : public Chaos::IRewindCallback
 		}
 	}
 
+	virtual void UnregisterRewindableSimCallback_Internal(Chaos::ISimCallbackObject* SimCallbackObject) override
+	{
+		RewindableCallbackObjects.Remove(SimCallbackObject);
+	}
+
 	// Updates the TMap on PhysScene that stores (non interpolated) physics data for replication.
 	// 
 	// Needs to be called from PT context to access fixed tick handle
diff --git a/Engine/Source/Runtime/Experimental/Chaos/Private/PBDRigidsSolver.cpp b/Engine/Source/Runtime/Experimental/Chaos/Private/PBDRigidsSolver.cpp
--- a/Engine/Source/Runtime/Experimental/Chaos/Private/PBDRigidsSolver.cpp
+++ b/Engine/Source/Runtime/Experimental/Chaos/Private/PBDRigidsSolver.cpp
@@ -1102,7 +1102,7 @@ namespace Chaos
 		
 		if (MRewindCallback) 
 		{
-			MRewindCallback->RewindData = MRewindData.Get();
+			MRewindCallback->RewindData = GetRewindData();
 		}
 		
 		UpdateIsDeterministic();
@@ -1830,6 +1830,14 @@ namespace Chaos
 		for (ISimCallbackObject* RemovedCallbackObject : PushData.SimCallbackObjectsToRemove)
 		{
 			RemovedCallbackObject->bPendingDelete = true;
+
+			if (RemovedCallbackObject->HasOption(ESimCallbackOptions::Rewind))
+			{
+				if (MRewindCallback)
+				{
+					MRewindCallback->UnregisterRewindableSimCallback_Internal(RemovedCallbackObject);
+				}
+			}
 		}
 
 		for (int32 Idx = SimCallbackObjects.Num() - 1; Idx >= 0; --Idx)
diff --git a/Engine/Source/Runtime/Experimental/Chaos/Public/RewindData.h b/Engine/Source/Runtime/Experimental/Chaos/Public/RewindData.h
--- a/Engine/Source/Runtime/Experimental/Chaos/Public/RewindData.h
+++ b/Engine/Source/Runtime/Experimental/Chaos/Public/RewindData.h
@@ -1836,6 +1836,9 @@ public:
 	/** Register a sim callback onto the rewind callback */
 	virtual void RegisterRewindableSimCallback_Internal(ISimCallbackObject* Callback) {}
 
+	/** Unregister a sim callback from the rewind callback */
+	virtual void UnregisterRewindableSimCallback_Internal(ISimCallbackObject* Callback) {}
+
 	/** Called When resim is finished with debug information about the resim */
 	virtual void SetResimDebugInfo_Internal(const FResimDebugInfo& ResimDebugInfo){}
 
