[NetPhysics] Fix crash from PopulateReplicationCache() where it accessed a pointer to an object in garbage collection

diff --git a/Engine/Source/Runtime/Engine/Private/PhysicsEngine/Experimental/PhysScene_Chaos.cpp b/Engine/Source/Runtime/Engine/Private/PhysicsEngine/Experimental/PhysScene_Chaos.cpp
--- a/Engine/Source/Runtime/Engine/Private/PhysicsEngine/Experimental/PhysScene_Chaos.cpp
+++ b/Engine/Source/Runtime/Engine/Private/PhysicsEngine/Experimental/PhysScene_Chaos.cpp
@@ -1315,6 +1315,13 @@ void FPhysScene_Chaos::AddToComponentMaps(UPrimitiveComponent* Component, IPhysi
 	}
 }
 
+// FReplicationCacheData constructor needs to be in .cpp due to UPrimitiveComponent being forward declared in the header which TWeakObjectPtr doesn't handle
+FPhysScene_Chaos::FReplicationCacheData::FReplicationCacheData(UPrimitiveComponent* InRootComponent, Chaos::FReal InAccessTime)
+	: RootComponent(InRootComponent)
+	, AccessTime(InAccessTime)
+	, bValidStateCached(false)
+{}
+
 const FRigidBodyState* FPhysScene_Chaos::GetStateFromReplicationCache(UPrimitiveComponent* RootComponent, int& ServerFrame)
 {
 	if (!GetSolver()->GetRewindCallback())
@@ -1379,7 +1386,7 @@ void FPhysScene_Chaos::PopulateReplicationCache(const int32 PhysicsStep)
 	{
 		StateWasCached = false;
 		FReplicationCacheData& ReplicationData = It.Value();
-		TObjectPtr<UPrimitiveComponent> RootComponent = ReplicationData.GetRootComponent();
+		UPrimitiveComponent* RootComponent = ReplicationData.GetRootComponent();
 		if (RootComponent)
 		{
 			if (FBodyInstanceAsyncPhysicsTickHandle BIHandle = RootComponent->GetBodyInstanceAsyncPhysicsTickHandle())
diff --git a/Engine/Source/Runtime/Engine/Public/Physics/Experimental/PhysScene_Chaos.h b/Engine/Source/Runtime/Engine/Public/Physics/Experimental/PhysScene_Chaos.h
--- a/Engine/Source/Runtime/Engine/Public/Physics/Experimental/PhysScene_Chaos.h
+++ b/Engine/Source/Runtime/Engine/Public/Physics/Experimental/PhysScene_Chaos.h
@@ -260,13 +260,9 @@ public:
 
 	struct FReplicationCacheData
 	{
-		FReplicationCacheData(UPrimitiveComponent* InRootComponent, Chaos::FReal InAccessTime)
-			: RootComponent(InRootComponent)
-			, AccessTime(InAccessTime)
-			, bValidStateCached(false)
-		{}
+		FReplicationCacheData(UPrimitiveComponent* InRootComponent, Chaos::FReal InAccessTime);
 
-		TObjectPtr<UPrimitiveComponent> GetRootComponent()	{ return RootComponent; }
+		UPrimitiveComponent* GetRootComponent()	{ return RootComponent.Get(); }
 		FRigidBodyState& GetState()	{ return State; }
 		void SetAccessTime(Chaos::FReal Time) { AccessTime = Time; }
 		Chaos::FReal GetAccessTime() { return AccessTime; }
@@ -274,7 +270,7 @@ public:
 		bool IsCached() { return bValidStateCached; }
 
 	private:
-		TObjectPtr<UPrimitiveComponent> RootComponent;
+		TWeakObjectPtr<UPrimitiveComponent> RootComponent;
 		Chaos::FReal AccessTime;
 		bool bValidStateCached;
 		FRigidBodyState State;
