[Material Editor] Fixed crash by removing debugging data that was used to  printing missing expressions on material function post-load.

diff --git a/Engine/Source/Runtime/Engine/Classes/Materials/MaterialExpression.h b/Engine/Source/Runtime/Engine/Classes/Materials/MaterialExpression.h
--- a/Engine/Source/Runtime/Engine/Classes/Materials/MaterialExpression.h
+++ b/Engine/Source/Runtime/Engine/Classes/Materials/MaterialExpression.h
@@ -168,11 +168,6 @@ struct FMaterialExpressionCollection
 	UPROPERTY()
 	TArray<TObjectPtr<UMaterialExpression>> Expressions;
 	
-	/** Temporary debug info strings for each expression to help debugging the sporadic
-	 *  unexpected occurrence of nulls in Expressions on PostLoad() (UE-198712) */
-	UPROPERTY()
-	TArray<FString> DebugExpressionInfos;
-
 	/** Array of comments associated with this material; viewed in the material editor. */
 	UPROPERTY()
 	TArray<TObjectPtr<UMaterialExpressionComment>> EditorComments;
diff --git a/Engine/Source/Runtime/Engine/Private/Materials/MaterialExpressions.cpp b/Engine/Source/Runtime/Engine/Private/Materials/MaterialExpressions.cpp
--- a/Engine/Source/Runtime/Engine/Private/Materials/MaterialExpressions.cpp
+++ b/Engine/Source/Runtime/Engine/Private/Materials/MaterialExpressions.cpp
@@ -15007,7 +15007,7 @@ void UMaterialFunction::Serialize(FArchive& Ar)
 {
 #if WITH_EDITORONLY_DATA
 	UMaterialFunctionEditorOnlyData* EditorOnly = GetEditorOnlyData();
-	if (EditorOnly && Ar.IsSaving() && !Ar.IsCooking() && !Ar.IsObjectReferenceCollector())
+	if (EditorOnly && Ar.IsSaving() && !Ar.IsCooking())
 	{
 		// If the collection of expressions got some null expressions remove them now, but warn the user about it.
 		if (EditorOnly->ExpressionCollection.Expressions.Remove(nullptr))
@@ -15017,18 +15017,6 @@ void UMaterialFunction::Serialize(FArchive& Ar)
 				"\n\nPlease close and repoen this Material Function and verify it is still valid."), FText::FromString(GetFullName()));
 			FMessageDialog::Open(EAppMsgType::Ok, Message);
 		}
-
-		// Temporary debugging code. This will populate the DebugExpressionInfos with information about each expression
-		// in ExpressionCollection.Expressions, to gather more information when some expression is null upon function
-		// PostLoad() to help solve UE-198712.
-		EditorOnly->ExpressionCollection.DebugExpressionInfos.Empty();
-		EditorOnly->ExpressionCollection.DebugExpressionInfos.Reserve(EditorOnly->ExpressionCollection.Expressions.Num());
-		for (UMaterialExpression* Expression : EditorOnly->ExpressionCollection.Expressions)
-		{
-			check(Expression);
-			FString Info = FString::Printf(TEXT("Name: '%s', Type: '%s'"), *Expression->GetFullName(), *Expression->GetClass()->GetFullName());
-			EditorOnly->ExpressionCollection.DebugExpressionInfos.Push(MoveTemp(Info));
-		}
 	}
 #endif
 
@@ -15125,9 +15113,6 @@ void UMaterialFunction::PostLoad()
 
 	if (GIsEditor && EditorOnly)
 	{
-		// We can display null expressions info if the DebugExpressionInfos was populated with data upon Serialize().
-		bool bDisplayNullExpressionInfo = EditorOnly->ExpressionCollection.DebugExpressionInfos.Num() == EditorOnly->ExpressionCollection.Expressions.Num();
-
 		// Go over all expressions in the collection and invalidate the material if a null expression is found. Then
 		// remove the null expression from the array.
 		for (int i = 0; i < EditorOnly->ExpressionCollection.Expressions.Num();)
@@ -15142,13 +15127,6 @@ void UMaterialFunction::PostLoad()
 			// Mark this function as invalid. This will cause the material containing an active call to it to fail translation.
 			bAllExpressionsLoadedCorrectly = false;
 
-			if (bDisplayNullExpressionInfo)
-			{
-				UE_LOG(LogMaterial, Log, TEXT("Expression in function expression collection with index %d was null. Expression Info: %s"), i, *EditorOnly->ExpressionCollection.DebugExpressionInfos[i]);
-					
-				EditorOnly->ExpressionCollection.DebugExpressionInfos.RemoveAt(i);
-			}
-
 			EditorOnly->ExpressionCollection.Expressions.RemoveAt(i);
 		}
 
