[Skeletal mesh & Deformer Graph] fix streaming for half edge buffers

diff --git a/Engine/Plugins/Animation/DeformerGraph/Source/OptimusCore/Private/DataInterfaces/OptimusDataInterfaceHalfEdge.cpp b/Engine/Plugins/Animation/DeformerGraph/Source/OptimusCore/Private/DataInterfaces/OptimusDataInterfaceHalfEdge.cpp
--- a/Engine/Plugins/Animation/DeformerGraph/Source/OptimusCore/Private/DataInterfaces/OptimusDataInterfaceHalfEdge.cpp
+++ b/Engine/Plugins/Animation/DeformerGraph/Source/OptimusCore/Private/DataInterfaces/OptimusDataInterfaceHalfEdge.cpp
@@ -240,7 +240,7 @@ void FOptimusHalfEdgeDataProviderProxy::AllocateResources(FRDGBuilder& GraphBuil
 	FSkeletalMeshRenderData const& SkeletalMeshRenderData = SkeletalMeshObject->GetSkeletalMeshRenderData();
 	FSkeletalMeshLODRenderData const* LodRenderData = &SkeletalMeshRenderData.LODRenderData[LodIndex];
 
-	if (LodRenderData->HalfEdgeBuffer.IsInitialized())
+	if (LodRenderData->HalfEdgeBuffer.IsReadyForRendering())
 	{
 		bUseBufferFromRenderData = true;
 
diff --git a/Engine/Source/Runtime/Engine/Internal/Streaming/SkeletalMeshUpdate.cpp b/Engine/Source/Runtime/Engine/Internal/Streaming/SkeletalMeshUpdate.cpp
--- a/Engine/Source/Runtime/Engine/Internal/Streaming/SkeletalMeshUpdate.cpp
+++ b/Engine/Source/Runtime/Engine/Internal/Streaming/SkeletalMeshUpdate.cpp
@@ -60,6 +60,7 @@ void FSkeletalMeshStreamIn::FIntermediateBuffers::CreateFromCPUData(FRHICommandL
 	SkinWeightVertexBuffer = LODResource.SkinWeightVertexBuffer.CreateRHIBuffer(RHICmdList);
 	ClothVertexBuffer = LODResource.ClothVertexBuffer.CreateRHIBuffer(RHICmdList);
 	IndexBuffer = LODResource.MultiSizeIndexContainer.CreateRHIBuffer(RHICmdList);
+	HalfEdgeBuffer = LODResource.HalfEdgeBuffer.CreateRHIBuffer(RHICmdList);
 }
 
 void FSkeletalMeshStreamIn::FIntermediateBuffers::TransferBuffers(FSkeletalMeshLODRenderData& LODResource, FRHIResourceReplaceBatcher& Batcher)
@@ -72,6 +73,7 @@ void FSkeletalMeshStreamIn::FIntermediateBuffers::TransferBuffers(FSkeletalMeshL
 	LODResource.ClothVertexBuffer.InitRHIForStreaming(ClothVertexBuffer, Batcher);
 	LODResource.MultiSizeIndexContainer.InitRHIForStreaming(IndexBuffer, Batcher);
 	LODResource.SkinWeightProfilesData.InitRHIForStreaming(AltSkinWeightVertexBuffers, Batcher);
+	LODResource.HalfEdgeBuffer.InitRHIForStreaming(HalfEdgeBuffer, Batcher);
 }
 
 #if RHI_RAYTRACING
@@ -364,6 +366,7 @@ void FSkeletalMeshStreamOut::ReleaseBuffers(const FContext& Context)
 				LODResource.ClothVertexBuffer.ReleaseRHIForStreaming(Batcher);
 				LODResource.MultiSizeIndexContainer.ReleaseRHIForStreaming(Batcher);
 				LODResource.SkinWeightProfilesData.ReleaseRHIForStreaming(Batcher);
+				LODResource.HalfEdgeBuffer.ReleaseRHIForStreaming(Batcher);
 
 				if (!FPlatformProperties::HasEditorOnlyData())
 				{
diff --git a/Engine/Source/Runtime/Engine/Internal/Streaming/SkeletalMeshUpdate.h b/Engine/Source/Runtime/Engine/Internal/Streaming/SkeletalMeshUpdate.h
--- a/Engine/Source/Runtime/Engine/Internal/Streaming/SkeletalMeshUpdate.h
+++ b/Engine/Source/Runtime/Engine/Internal/Streaming/SkeletalMeshUpdate.h
@@ -10,6 +10,7 @@ SkeletalMeshUpdate.h: Helpers to stream in and out skeletal mesh LODs.
 #include "Engine/SkeletalMesh.h"
 #include "IO/IoDispatcher.h"
 #include "RenderAssetUpdate.h"
+#include "Rendering/SkeletalMeshHalfEdgeBuffer.h"
 #include "Rendering/SkinWeightVertexBuffer.h"
 #include "Serialization/BulkData.h"
 #include "RayTracingGeometry.h"
@@ -84,6 +85,7 @@ protected:
 		FBufferRHIRef ClothVertexBuffer;
 		FBufferRHIRef IndexBuffer;
 		TArray<TPair<FName, FSkinWeightRHIInfo>> AltSkinWeightVertexBuffers;
+		FSkeletalMeshHalfEdgeBuffer::FRHIInfo HalfEdgeBuffer;
 
 		void CreateFromCPUData(FRHICommandListBase& RHICmdList, FSkeletalMeshLODRenderData& LODResource);
 
diff --git a/Engine/Source/Runtime/Engine/Private/Rendering/SkeletalMeshHalfEdgeBuffer.cpp b/Engine/Source/Runtime/Engine/Private/Rendering/SkeletalMeshHalfEdgeBuffer.cpp
--- a/Engine/Source/Runtime/Engine/Private/Rendering/SkeletalMeshHalfEdgeBuffer.cpp
+++ b/Engine/Source/Runtime/Engine/Private/Rendering/SkeletalMeshHalfEdgeBuffer.cpp
@@ -1,6 +1,8 @@
 ﻿// Copyright Epic Games, Inc. All Rights Reserved.
 
 #include "Rendering/SkeletalMeshHalfEdgeBuffer.h"
+
+#include "RHIResourceUpdates.h"
 #include "Rendering/SkeletalMeshLODRenderData.h"
 
 struct FEdgeKey
@@ -219,24 +221,24 @@ void FSkeletalMeshHalfEdgeBuffer::Init(const FSkeletalMeshLODRenderData& InLodRe
 	SkeletalMeshHalfEdgeUtility::BuildHalfEdgeBuffers(InLodRenderData, VertexToEdgeData, EdgeToTwinEdgeData);	
 }
 
-void FSkeletalMeshHalfEdgeBuffer::InitRHI(FRHICommandListBase& RHICmdList)
+FSkeletalMeshHalfEdgeBuffer::FRHIInfo FSkeletalMeshHalfEdgeBuffer::CreateRHIBuffer(FRHICommandListBase& RHICmdList)
 {
+	FRHIInfo Buffers;
+	
 	const static FLazyName ClassName(TEXT("FSkeletalMeshHalfEdgeBuffer"));
 
-	// InitRHI only called if CPU data is valid
-	check(VertexToEdgeData.Num() > 0 && EdgeToTwinEdgeData.Num() > 0)
+	uint32 MinBufferSize = GetMinBufferSize();
 	{
 		const uint32 BufferSize = VertexToEdgeData.GetResourceDataSize();
 		FRHIResourceCreateInfo CreateInfo(TEXT("VertexToEdgeData"));
 		CreateInfo.ClassName = ClassName;
 		CreateInfo.OwnerName = GetOwnerName();
-		VertexToEdgeBufferRHI = RHICmdList.CreateStructuredBuffer(sizeof(int32), BufferSize, BUF_Static | BUF_ShaderResource, ERHIAccess::SRVMask, CreateInfo);
-		VertexToEdgeBufferRHI->SetOwnerName(GetOwnerName());
-
-		void* BufferPtr = RHICmdList.LockBuffer(VertexToEdgeBufferRHI, 0, BufferSize, RLM_WriteOnly);
-		FMemory::ParallelMemcpy(BufferPtr, VertexToEdgeData.GetData(), BufferSize, EMemcpyCachePolicy::StoreUncached);
-		RHICmdList.UnlockBuffer(VertexToEdgeBufferRHI);
-		VertexToEdgeBufferSRV = RHICmdList.CreateShaderResourceView(VertexToEdgeBufferRHI);
+		if (BufferSize > 0)
+		{
+			CreateInfo.ResourceArray = &VertexToEdgeData;
+		}
+		Buffers.VertexToEdgeBufferRHI = RHICmdList.CreateStructuredBuffer(MinBufferSize, FMath::Max(BufferSize, MinBufferSize) , BUF_Static | BUF_ShaderResource, ERHIAccess::SRVMask, CreateInfo);
+		Buffers.VertexToEdgeBufferRHI->SetOwnerName(GetOwnerName());
 	}
 
 	{
@@ -244,12 +246,59 @@ void FSkeletalMeshHalfEdgeBuffer::InitRHI(FRHICommandListBase& RHICmdList)
 		FRHIResourceCreateInfo CreateInfo(TEXT("EdgeToTwinEdgeData"));
 		CreateInfo.ClassName = ClassName;
 		CreateInfo.OwnerName = GetOwnerName();
-		EdgeToTwinEdgeBufferRHI = RHICmdList.CreateStructuredBuffer(sizeof(int32), BufferSize, BUF_Static | BUF_ShaderResource, ERHIAccess::SRVMask, CreateInfo);
-		EdgeToTwinEdgeBufferRHI->SetOwnerName(GetOwnerName());
+		if (BufferSize > 0)
+		{
+			CreateInfo.ResourceArray = &EdgeToTwinEdgeData;
+		}
+		Buffers.EdgeToTwinEdgeBufferRHI = RHICmdList.CreateStructuredBuffer(sizeof(int32), FMath::Max(BufferSize, MinBufferSize), BUF_Static | BUF_ShaderResource, ERHIAccess::SRVMask, CreateInfo);
+		Buffers.EdgeToTwinEdgeBufferRHI->SetOwnerName(GetOwnerName());
+	}
+
+
+	return Buffers;
+}
+
+void FSkeletalMeshHalfEdgeBuffer::InitRHIForStreaming(FRHIInfo RHIInfo, FRHIResourceUpdateBatcher& Batcher)
+{
+	if (VertexToEdgeBufferRHI && RHIInfo.VertexToEdgeBufferRHI)
+	{
+		Batcher.QueueUpdateRequest(VertexToEdgeBufferRHI, RHIInfo.VertexToEdgeBufferRHI);
+	}
+	if (EdgeToTwinEdgeBufferRHI && RHIInfo.EdgeToTwinEdgeBufferRHI)
+	{
+		Batcher.QueueUpdateRequest(EdgeToTwinEdgeBufferRHI, RHIInfo.EdgeToTwinEdgeBufferRHI);
+	}
+}
 
-		void* BufferPtr = RHICmdList.LockBuffer(EdgeToTwinEdgeBufferRHI, 0, BufferSize, RLM_WriteOnly);
-		FMemory::ParallelMemcpy(BufferPtr, EdgeToTwinEdgeData.GetData(), BufferSize, EMemcpyCachePolicy::StoreUncached);
-		RHICmdList.UnlockBuffer(EdgeToTwinEdgeBufferRHI);
+void FSkeletalMeshHalfEdgeBuffer::ReleaseRHIForStreaming(FRHIResourceUpdateBatcher& Batcher)
+{
+	if (VertexToEdgeBufferRHI)
+	{
+		Batcher.QueueUpdateRequest(VertexToEdgeBufferRHI, nullptr);
+	}
+	if (EdgeToTwinEdgeBufferRHI )
+	{
+		Batcher.QueueUpdateRequest(EdgeToTwinEdgeBufferRHI, nullptr);
+	}	
+}
+
+void FSkeletalMeshHalfEdgeBuffer::InitRHI(FRHICommandListBase& RHICmdList)
+{
+	FRHIInfo Buffers = CreateRHIBuffer(RHICmdList);
+	
+	VertexToEdgeBufferRHI = Buffers.VertexToEdgeBufferRHI;
+	EdgeToTwinEdgeBufferRHI = Buffers.EdgeToTwinEdgeBufferRHI;
+	// We should always have a RHI and thus a SRV, even when CPU data is empty,
+	// which can happen in two cases:
+	// 1. It was not cooked with the skeletal mesh, in this case the buffer is never used.
+	// 2. It has not been streamed in for this LOD yet. But once streamed, 
+	//    the RHI will takeover ownership of the streamed buffer and becomes usable
+	if (ensure(VertexToEdgeBufferRHI))
+	{
+		VertexToEdgeBufferSRV = RHICmdList.CreateShaderResourceView(VertexToEdgeBufferRHI);
+	}
+	if (ensure(EdgeToTwinEdgeBufferRHI))
+	{
 		EdgeToTwinEdgeBufferSRV = RHICmdList.CreateShaderResourceView(EdgeToTwinEdgeBufferRHI);
 	}
 }
@@ -267,6 +316,14 @@ bool FSkeletalMeshHalfEdgeBuffer::IsCPUDataValid() const
 	return VertexToEdgeData.Num() > 0 && EdgeToTwinEdgeData.Num() > 0;
 }
 
+bool FSkeletalMeshHalfEdgeBuffer::IsReadyForRendering() const
+{
+	// The size of a buffer with valid data is definitely larger than the min buffer size
+	// Buffer not being ready does not always mean that the data was not cooked,
+	// it could also mean that the data hasn't been streamed in yet
+	return IsInitialized() && VertexToEdgeBufferRHI->GetSize() > GetMinBufferSize();
+}
+
 void FSkeletalMeshHalfEdgeBuffer::CleanUp()
 {
 	VertexToEdgeData.Reset();
@@ -284,6 +341,11 @@ void FSkeletalMeshHalfEdgeBuffer::Serialize(FArchive& Ar)
 	Ar << EdgeToTwinEdgeData;
 }
 
+uint32 FSkeletalMeshHalfEdgeBuffer::GetMinBufferSize() const
+{
+	return sizeof(int32);
+}
+
 FArchive& operator<<(FArchive& Ar, FSkeletalMeshHalfEdgeBuffer& HalfEdgeBuffer)
 {
 	HalfEdgeBuffer.Serialize(Ar);
diff --git a/Engine/Source/Runtime/Engine/Private/SkeletalMeshLODRenderData.cpp b/Engine/Source/Runtime/Engine/Private/SkeletalMeshLODRenderData.cpp
--- a/Engine/Source/Runtime/Engine/Private/SkeletalMeshLODRenderData.cpp
+++ b/Engine/Source/Runtime/Engine/Private/SkeletalMeshLODRenderData.cpp
@@ -303,11 +303,8 @@ void FSkeletalMeshLODRenderData::InitResources(bool bNeedsVertexColors, int32 LO
 
 	VertexAttributeBuffers.InitResources();
 	
-	if (HalfEdgeBuffer.IsCPUDataValid())
-	{
-		HalfEdgeBuffer.SetOwnerName(OwnerName);
-		BeginInitResource(&HalfEdgeBuffer, &UE::RenderCommandPipe::SkeletalMesh);	
-	}
+	HalfEdgeBuffer.SetOwnerName(OwnerName);
+	BeginInitResource(&HalfEdgeBuffer, &UE::RenderCommandPipe::SkeletalMesh);	
 
 #if RHI_RAYTRACING
 	if (IsRayTracingAllowed())
diff --git a/Engine/Source/Runtime/Engine/Public/Rendering/SkeletalMeshHalfEdgeBuffer.h b/Engine/Source/Runtime/Engine/Public/Rendering/SkeletalMeshHalfEdgeBuffer.h
--- a/Engine/Source/Runtime/Engine/Public/Rendering/SkeletalMeshHalfEdgeBuffer.h
+++ b/Engine/Source/Runtime/Engine/Public/Rendering/SkeletalMeshHalfEdgeBuffer.h
@@ -7,6 +7,7 @@
 
 class FSkeletalMeshLODRenderData;
 class FRHIShaderResourceView;
+struct FRHIResourceUpdateBatcher;
 
 namespace SkeletalMeshHalfEdgeUtility
 {
@@ -20,12 +21,25 @@ class FSkeletalMeshHalfEdgeBuffer : public FRenderResource
 {
 public:
 	void Init(const FSkeletalMeshLODRenderData& InLodRenderData);
-		
+
+	struct FRHIInfo
+	{
+		FBufferRHIRef VertexToEdgeBufferRHI;
+		FBufferRHIRef EdgeToTwinEdgeBufferRHI;
+	};
+
+	FRHIInfo CreateRHIBuffer(FRHICommandListBase& RHICmdList);
+
+	void InitRHIForStreaming(FRHIInfo RHIInfo, FRHIResourceUpdateBatcher& Batcher);
+	void ReleaseRHIForStreaming(FRHIResourceUpdateBatcher& Batcher);
+	
 	void ENGINE_API InitRHI(FRHICommandListBase& RHICmdList) override;
 
 	void ENGINE_API ReleaseRHI() override;
 
 	bool IsCPUDataValid() const;
+
+	bool ENGINE_API IsReadyForRendering() const;
 	
 	void CleanUp();
 
@@ -47,6 +61,8 @@ public:
 	}
 
 private:
+	uint32 GetMinBufferSize() const;
+	
 	TResourceArray<int32> VertexToEdgeData;
 	TResourceArray<int32> EdgeToTwinEdgeData;
 	
