[SkeletalMesh & Deformer Graph] follow up fix to merge conflict for CL32132480

hode link https://horde.devtools.epicgames.com/job/65eb886aa60fac6697945898?step=8862

diff --git a/Engine/Source/Runtime/Engine/Private/Rendering/SkeletalMeshHalfEdgeBuffer.cpp b/Engine/Source/Runtime/Engine/Private/Rendering/SkeletalMeshHalfEdgeBuffer.cpp
--- a/Engine/Source/Runtime/Engine/Private/Rendering/SkeletalMeshHalfEdgeBuffer.cpp
+++ b/Engine/Source/Runtime/Engine/Private/Rendering/SkeletalMeshHalfEdgeBuffer.cpp
@@ -2,7 +2,7 @@
 
 #include "Rendering/SkeletalMeshHalfEdgeBuffer.h"
 
-#include "RHIResourceUpdates.h"
+#include "RHIResourceReplace.h"
 #include "Rendering/SkeletalMeshLODRenderData.h"
 
 struct FEdgeKey
@@ -258,27 +258,27 @@ FSkeletalMeshHalfEdgeBuffer::FRHIInfo FSkeletalMeshHalfEdgeBuffer::CreateRHIBuff
 	return Buffers;
 }
 
-void FSkeletalMeshHalfEdgeBuffer::InitRHIForStreaming(FRHIInfo RHIInfo, FRHIResourceUpdateBatcher& Batcher)
+void FSkeletalMeshHalfEdgeBuffer::InitRHIForStreaming(FRHIInfo RHIInfo, FRHIResourceReplaceBatcher& Batcher)
 {
 	if (VertexToEdgeBufferRHI && RHIInfo.VertexToEdgeBufferRHI)
 	{
-		Batcher.QueueUpdateRequest(VertexToEdgeBufferRHI, RHIInfo.VertexToEdgeBufferRHI);
+		Batcher.EnqueueReplace(VertexToEdgeBufferRHI, RHIInfo.VertexToEdgeBufferRHI);
 	}
 	if (EdgeToTwinEdgeBufferRHI && RHIInfo.EdgeToTwinEdgeBufferRHI)
 	{
-		Batcher.QueueUpdateRequest(EdgeToTwinEdgeBufferRHI, RHIInfo.EdgeToTwinEdgeBufferRHI);
+		Batcher.EnqueueReplace(EdgeToTwinEdgeBufferRHI, RHIInfo.EdgeToTwinEdgeBufferRHI);
 	}
 }
 
-void FSkeletalMeshHalfEdgeBuffer::ReleaseRHIForStreaming(FRHIResourceUpdateBatcher& Batcher)
+void FSkeletalMeshHalfEdgeBuffer::ReleaseRHIForStreaming(FRHIResourceReplaceBatcher& Batcher)
 {
 	if (VertexToEdgeBufferRHI)
 	{
-		Batcher.QueueUpdateRequest(VertexToEdgeBufferRHI, nullptr);
+		Batcher.EnqueueReplace(VertexToEdgeBufferRHI, nullptr);
 	}
 	if (EdgeToTwinEdgeBufferRHI )
 	{
-		Batcher.QueueUpdateRequest(EdgeToTwinEdgeBufferRHI, nullptr);
+		Batcher.EnqueueReplace(EdgeToTwinEdgeBufferRHI, nullptr);
 	}	
 }
 
diff --git a/Engine/Source/Runtime/Engine/Public/Rendering/SkeletalMeshHalfEdgeBuffer.h b/Engine/Source/Runtime/Engine/Public/Rendering/SkeletalMeshHalfEdgeBuffer.h
--- a/Engine/Source/Runtime/Engine/Public/Rendering/SkeletalMeshHalfEdgeBuffer.h
+++ b/Engine/Source/Runtime/Engine/Public/Rendering/SkeletalMeshHalfEdgeBuffer.h
@@ -7,7 +7,7 @@
 
 class FSkeletalMeshLODRenderData;
 class FRHIShaderResourceView;
-struct FRHIResourceUpdateBatcher;
+class FRHIResourceReplaceBatcher;
 
 namespace SkeletalMeshHalfEdgeUtility
 {
@@ -30,8 +30,8 @@ public:
 
 	FRHIInfo CreateRHIBuffer(FRHICommandListBase& RHICmdList);
 
-	void InitRHIForStreaming(FRHIInfo RHIInfo, FRHIResourceUpdateBatcher& Batcher);
-	void ReleaseRHIForStreaming(FRHIResourceUpdateBatcher& Batcher);
+	void InitRHIForStreaming(FRHIInfo RHIInfo, FRHIResourceReplaceBatcher& Batcher);
+	void ReleaseRHIForStreaming(FRHIResourceReplaceBatcher& Batcher);
 	
 	void ENGINE_API InitRHI(FRHICommandListBase& RHICmdList) override;
 
