Fixed an issue where attached children to a scene component were visually teleporting off to space because the parent's children array OnRep ran before the child actor's AttachementReplication OnRep ran.
- Note: CVar "s.CheckRootComponentReplicationOnAttachedChildren" was added to disable this if needed
- Now if an attached child's root component does not replicate and uses AttachementReplication, then the parent's scene component will not do a parent fixup, since the Actor's OnRep will do it when it eventually fires and allow the update child transforms to properly work.

#rnx

diff --git a/Engine/Source/Runtime/Engine/Private/Components/SceneComponent.cpp b/Engine/Source/Runtime/Engine/Private/Components/SceneComponent.cpp
--- a/Engine/Source/Runtime/Engine/Private/Components/SceneComponent.cpp
+++ b/Engine/Source/Runtime/Engine/Private/Components/SceneComponent.cpp
@@ -46,6 +46,17 @@ namespace SceneComponentStatics
 	static const FName PhysicsVolumeTraceName(TEXT("PhysicsVolumeTrace"));
 }
 
+namespace SceneComponentCVars
+{
+	bool bCheckRootComponentReplicationOnAttachedChildren = true;
+	static FAutoConsoleVariableRef CVarCheckRootComponentReplicationOnAttachedChildren(
+		TEXT("s.CheckRootComponentReplicationOnAttachedChildren"),
+		bCheckRootComponentReplicationOnAttachedChildren,
+		TEXT("Scene Component OnRep_AttachedChildren:\n")
+		TEXT("false: fix up any children that are missing the parent, true: fixes up all children except non-replicated root components (default)"),
+		ECVF_Default);
+}
+
 DEFINE_LOG_CATEGORY_STATIC(LogSceneComponent, Log, All);
 
 DECLARE_CYCLE_STAT(TEXT("UpdateComponentToWorld"), STAT_UpdateComponentToWorld, STATGROUP_Component);
@@ -3453,6 +3464,18 @@ void USceneComponent::OnRep_AttachChildren()
 	{
 		if (ChildComponent)
 		{
+			// @note: When a actor's root component is flagged not to replicate it uses AActor::AttachmentReplication and
+			// when using that we'll want that to be the authority in making sure the parent component it's attached to
+			// is properly set up.
+			if (SceneComponentCVars::bCheckRootComponentReplicationOnAttachedChildren)
+			{
+				AActor* ChildOwner = ChildComponent->GetOwner();
+				if (ChildOwner && ChildOwner->GetRootComponent() == ChildComponent && !ChildComponent->GetIsReplicated())
+				{
+					continue;
+				}
+			}
+
 			if (ChildComponent->GetAttachParent() != this)
 			{
 				ChildComponent->SetAttachParent(this);
