Fix cancelled tasks incorrectly reporting they finished successfully. That was affecting the result in UStreamableRenderAsset::TickStreaming, setting CachedSRRState.NumResidentLODs incorrectly to CachedSRRState.NumRequestedLODs.
We would then end up with holes in the LOD chain, and try to issue draws with these empty meshes. That could end up with a GPU hang on some platforms.
We used to call MarkAsSuccessfullyFinished in ::DoFinishUpdate only !IsCancelled() so that CL restores this behaviour

diff --git a/Engine/Source/Runtime/Engine/Internal/Streaming/SkeletalMeshUpdate.cpp b/Engine/Source/Runtime/Engine/Internal/Streaming/SkeletalMeshUpdate.cpp
--- a/Engine/Source/Runtime/Engine/Internal/Streaming/SkeletalMeshUpdate.cpp
+++ b/Engine/Source/Runtime/Engine/Internal/Streaming/SkeletalMeshUpdate.cpp
@@ -247,7 +247,6 @@ void FSkeletalMeshStreamIn::DoCancel(const FContext& Context)
 	}
 
 	check(!StreamingRHICmdList);
-	MarkAsSuccessfullyFinished();
 }
 
 FSkeletalMeshStreamOut::FSkeletalMeshStreamOut(const USkeletalMesh* InMesh)
diff --git a/Engine/Source/Runtime/Engine/Private/Streaming/StaticMeshUpdate.cpp b/Engine/Source/Runtime/Engine/Private/Streaming/StaticMeshUpdate.cpp
--- a/Engine/Source/Runtime/Engine/Private/Streaming/StaticMeshUpdate.cpp
+++ b/Engine/Source/Runtime/Engine/Private/Streaming/StaticMeshUpdate.cpp
@@ -327,7 +327,6 @@ void FStaticMeshStreamIn::DoCancel(const FContext& Context)
 	}
 
 	check(!StreamingRHICmdList);
-	MarkAsSuccessfullyFinished();
 }
 
 FStaticMeshStreamOut::FStaticMeshStreamOut(const UStaticMesh* InMesh, bool InDiscardCPUData)
