World Partition: report warnings during map checks for spatially loaded actors that returns invalid streaming bounds.

diff --git a/Engine/Source/Runtime/Engine/Public/WorldPartition/WorldPartitionActorDesc.h b/Engine/Source/Runtime/Engine/Public/WorldPartition/WorldPartitionActorDesc.h
--- a/Engine/Source/Runtime/Engine/Public/WorldPartition/WorldPartitionActorDesc.h
+++ b/Engine/Source/Runtime/Engine/Public/WorldPartition/WorldPartitionActorDesc.h
@@ -210,7 +210,7 @@ public:
 	ENGINE_API FGuid GetContentBundleGuid() const;
 
 	virtual const FGuid& GetSceneOutlinerParent() const { return GetParentActor(); }
-	virtual bool IsResaveNeeded() const { return false; }
+	virtual bool IsResaveNeeded() const { return bIsSpatiallyLoaded && !bIsBoundsValid; }
 
 	ENGINE_API virtual void CheckForErrors(const IWorldPartitionActorDescInstanceView* InActorDescView, IStreamingGenerationErrorHandler* ErrorHandler) const;
 
