[NNE] Switch to NNERuntimeORTDml (RDG) as we will soon remove NNERuntimeRDGDml.

#tests Compile Files (Non-Unity)

#rnx

diff --git a/Engine/Plugins/Experimental/NeuralRendering/NeuralRendering.uplugin b/Engine/Plugins/Experimental/NeuralRendering/NeuralRendering.uplugin
--- a/Engine/Plugins/Experimental/NeuralRendering/NeuralRendering.uplugin
+++ b/Engine/Plugins/Experimental/NeuralRendering/NeuralRendering.uplugin
@@ -30,7 +30,7 @@
 	],
 	"Plugins": [
 		{
-			"Name": "NNERuntimeRDG",
+			"Name": "NNERuntimeORT",
 			"Enabled": true
 		}
 	]
diff --git a/Engine/Source/Runtime/Engine/Classes/Engine/NeuralProfile.h b/Engine/Source/Runtime/Engine/Classes/Engine/NeuralProfile.h
--- a/Engine/Source/Runtime/Engine/Classes/Engine/NeuralProfile.h
+++ b/Engine/Source/Runtime/Engine/Classes/Engine/NeuralProfile.h
@@ -51,7 +51,7 @@ enum class ETileOverlapResolveType : uint8
 UENUM(BlueprintType)
 enum class ENeuralProfileRuntimeType : uint8
 {
-	NNERuntimeRDGDml UMETA(DisplayName = "NNERuntimeRDGDml"),
+	NNERuntimeORTDml UMETA(DisplayName = "NNERuntimeORTDml"),
 
 	/** Does not have full operator support*/
 	NNERuntimeRDGHlsl UMETA(DisplayName = "NNERuntimeRDGHlsl"),
@@ -79,7 +79,7 @@ struct FNeuralProfileStruct
 	UPROPERTY(Category = "Common", EditAnywhere, BlueprintReadOnly, meta = (DisplayName = "Output Format", editcondition = "false", EditConditionHides))
 	ENeuralProfileFormat OutputFormat;
 
-	//runtime type (support "NNERuntimeRDGDml" only at this moment)
+	//runtime type (support "NNERuntimeORTDml" only at this moment)
 	UPROPERTY(Category = "Model", EditAnywhere, BlueprintReadOnly)
 	ENeuralProfileRuntimeType RuntimeType;
 
@@ -114,7 +114,7 @@ struct FNeuralProfileStruct
 	{
 		InputFormat = ENeuralProfileFormat::Type32;
 		OutputFormat = ENeuralProfileFormat::Type32;
-		RuntimeType = ENeuralProfileRuntimeType::NNERuntimeRDGDml;
+		RuntimeType = ENeuralProfileRuntimeType::NNERuntimeORTDml;
 		NNEModelData = nullptr;
 		InputDimension = FIntVector4(0);
 		OutputDimension = FIntVector4(0);
diff --git a/Engine/Source/Runtime/Engine/Private/Rendering/NeuralProfile.cpp b/Engine/Source/Runtime/Engine/Private/Rendering/NeuralProfile.cpp
--- a/Engine/Source/Runtime/Engine/Private/Rendering/NeuralProfile.cpp
+++ b/Engine/Source/Runtime/Engine/Private/Rendering/NeuralProfile.cpp
@@ -23,7 +23,7 @@ TUniquePtr<NeuralProfile::INeuralProfileManager> GNeuralProfileManager;
 static const TCHAR* GetNeuralProfileRuntimeName(ENeuralProfileRuntimeType NeuralProfileRuntimeType)
 {
 	static const TCHAR* const kRuntimeNames[] = {
-		TEXT("NNERuntimeRDGDml"),
+		TEXT("NNERuntimeORTDml"),
 		TEXT("NNERuntimeRDGHlsl")
 	};
 
