World Partition: added new virtual functions on AActor::OnLoadedActorAddedToLevel/OnLoadedActorRemovedFromLevel called when actors are dynamically added or removed from the persistent level during world partition loading operations + deprecated FWorldPartitionActorDesc::OnUnloadingInstance and moved existing code to new actor callbacks instead.

diff --git a/Engine/Source/Runtime/Engine/Classes/GameFramework/Actor.h b/Engine/Source/Runtime/Engine/Classes/GameFramework/Actor.h
--- a/Engine/Source/Runtime/Engine/Classes/GameFramework/Actor.h
+++ b/Engine/Source/Runtime/Engine/Classes/GameFramework/Actor.h
@@ -1071,6 +1071,12 @@ public:
 	 * Creates an uninitialized actor descriptor from a specific class.
 	 */
 	static ENGINE_API TUniquePtr<class FWorldPartitionActorDesc> StaticCreateClassActorDesc(const TSubclassOf<AActor>& ActorClass);
+
+	/** Called when this actor gets added to the level in the editor by world partition loading. */
+	ENGINE_API virtual void OnLoadedActorAddedToLevel() {}
+
+	/** Called when this actor gets removed from the level in the editor by world partition unloading. */
+	ENGINE_API virtual void OnLoadedActorRemovedFromLevel() {}
 #endif // WITH_EDITOR
 
 private:
diff --git a/Engine/Source/Runtime/Engine/Private/Level.cpp b/Engine/Source/Runtime/Engine/Private/Level.cpp
--- a/Engine/Source/Runtime/Engine/Private/Level.cpp
+++ b/Engine/Source/Runtime/Engine/Private/Level.cpp
@@ -924,6 +924,8 @@ void ULevel::AddLoadedActors(const TArray<AActor*>& ActorList, const FTransform*
 				Actor->FixupActorFolder();
 			}
 
+			Actor->OnLoadedActorAddedToLevel();
+
 			OnLoadedActorAddedToLevelEvent.Broadcast(*Actor);
 		}
 
@@ -1024,6 +1026,7 @@ void ULevel::RemoveLoadedActors(const TArray<AActor*>& ActorList, const FTransfo
 		{
 			Actor->UnregisterAllComponents();
 			Actor->RegisterAllActorTickFunctions(false, true);
+			Actor->OnLoadedActorRemovedFromLevel();
 
 			OnLoadedActorRemovedFromLevelEvent.Broadcast(*Actor);
 
diff --git a/Engine/Source/Runtime/Engine/Private/WorldPartition/DataLayer/WorldDataLayers.cpp b/Engine/Source/Runtime/Engine/Private/WorldPartition/DataLayer/WorldDataLayers.cpp
--- a/Engine/Source/Runtime/Engine/Private/WorldPartition/DataLayer/WorldDataLayers.cpp
+++ b/Engine/Source/Runtime/Engine/Private/WorldPartition/DataLayer/WorldDataLayers.cpp
@@ -632,12 +632,34 @@ void AWorldDataLayers::DumpDataLayers(FOutputDevice& OutputDevice) const
 }
 
 #if WITH_EDITOR
-
 TUniquePtr<class FWorldPartitionActorDesc> AWorldDataLayers::CreateClassActorDesc() const
 {
 	return TUniquePtr<FWorldPartitionActorDesc>(new FWorldDataLayersActorDesc());
 }
 
+void AWorldDataLayers::OnLoadedActorRemovedFromLevel()
+{
+	Super::OnLoadedActorRemovedFromLevel();
+
+	if (IsUsingExternalPackageDataLayerInstances())
+	{
+		ForEachDataLayerInstance([this](UDataLayerInstance* DataLayerInstance)
+		{
+			check(DataLayerInstance->IsPackageExternal())
+			ForEachObjectWithPackage(DataLayerInstance->GetPackage(), [](UObject* Object)
+			{
+				if (Object->HasAnyFlags(RF_Public | RF_Standalone))
+				{
+					CastChecked<UMetaData>(Object)->ClearFlags(RF_Public | RF_Standalone);
+				}
+				return true;
+			}, false);
+
+			return true;
+		});
+	}
+}
+
 AWorldDataLayers* AWorldDataLayers::Create(UWorld* World, FName InWorldDataLayerName)
 {
 	FActorSpawnParameters SpawnParameters;
diff --git a/Engine/Source/Runtime/Engine/Private/WorldPartition/DataLayer/WorldDataLayersActorDesc.cpp b/Engine/Source/Runtime/Engine/Private/WorldPartition/DataLayer/WorldDataLayersActorDesc.cpp
--- a/Engine/Source/Runtime/Engine/Private/WorldPartition/DataLayer/WorldDataLayersActorDesc.cpp
+++ b/Engine/Source/Runtime/Engine/Private/WorldPartition/DataLayer/WorldDataLayersActorDesc.cpp
@@ -269,33 +269,6 @@ void FWorldDataLayersActorDesc::Serialize(FArchive& Ar)
 	}
 }
 
-PRAGMA_DISABLE_DEPRECATION_WARNINGS
-void FWorldDataLayersActorDesc::OnUnloadingInstance(const FWorldPartitionActorDescInstance* InActorDescInstance) const
-{
-	if (AWorldDataLayers* WorldDataLayers = Cast<AWorldDataLayers>(InActorDescInstance->GetActor()))
-	{
-		if (WorldDataLayers->IsUsingExternalPackageDataLayerInstances())
-		{
-			WorldDataLayers->ForEachDataLayerInstance([this](UDataLayerInstance* DataLayerInstance)
-			{
-				check(DataLayerInstance->IsPackageExternal())
-				ForEachObjectWithPackage(DataLayerInstance->GetPackage(), [](UObject* Object)
-				{
-					if (Object->HasAnyFlags(RF_Public | RF_Standalone))
-					{
-						CastChecked<UMetaData>(Object)->ClearFlags(RF_Public | RF_Standalone);
-					}
-					return true;
-				}, false);
-
-				return true;
-			});
-		}
-	}
-	FWorldPartitionActorDesc::OnUnloadingInstance(InActorDescInstance);
-}
-PRAGMA_ENABLE_DEPRECATION_WARNINGS
-
 bool FWorldDataLayersActorDesc::IsRuntimeRelevant(const FWorldPartitionActorDescInstance* InActorDescInstance) const 
 {
 	if (!FWorldPartitionActorDesc::IsRuntimeRelevant(InActorDescInstance))
diff --git a/Engine/Source/Runtime/Engine/Private/WorldPartition/Landscape/LandscapeActorDesc.cpp b/Engine/Source/Runtime/Engine/Private/WorldPartition/Landscape/LandscapeActorDesc.cpp
--- a/Engine/Source/Runtime/Engine/Private/WorldPartition/Landscape/LandscapeActorDesc.cpp
+++ b/Engine/Source/Runtime/Engine/Private/WorldPartition/Landscape/LandscapeActorDesc.cpp
@@ -53,18 +53,6 @@ FBox FLandscapeActorDesc::GetEditorBounds() const
 	return FWorldPartitionActorDesc::GetEditorBounds();
 }
 
-PRAGMA_DISABLE_DEPRECATION_WARNINGS
-void FLandscapeActorDesc::OnUnloadingInstance(const FWorldPartitionActorDescInstance* InActorDescInstance) const
-{
-	if (ALandscapeProxy* LandscapeProxy = Cast<ALandscapeProxy>(InActorDescInstance->GetActor()))
-	{
-		LandscapeProxy->ActorDescReferences.Empty();
-	}
-
-	FPartitionActorDesc::OnUnloadingInstance(InActorDescInstance);
-}
-PRAGMA_ENABLE_DEPRECATION_WARNINGS
-
 bool FLandscapeActorDesc::Equals(const FWorldPartitionActorDesc* Other) const
 {
 	if (FPartitionActorDesc::Equals(Other))
diff --git a/Engine/Source/Runtime/Engine/Private/WorldPartition/WorldPartitionActorDescInstance.cpp b/Engine/Source/Runtime/Engine/Private/WorldPartition/WorldPartitionActorDescInstance.cpp
--- a/Engine/Source/Runtime/Engine/Private/WorldPartition/WorldPartitionActorDescInstance.cpp
+++ b/Engine/Source/Runtime/Engine/Private/WorldPartition/WorldPartitionActorDescInstance.cpp
@@ -166,11 +166,6 @@ void FWorldPartitionActorDescInstance::MarkUnload()
 {
 	FlushAsyncLoad();
 
-	// Notify Desc as it can have some custom code to run on the actor depending on type
-PRAGMA_DISABLE_DEPRECATION_WARNINGS
-	ActorDesc->OnUnloadingInstance(this);
-PRAGMA_ENABLE_DEPRECATION_WARNINGS
-
 	if (AActor* Actor = GetActor())
 	{
 		// At this point, it can happen that an actor isn't in an external package:
diff --git a/Engine/Source/Runtime/Engine/Public/WorldPartition/DataLayer/WorldDataLayers.h b/Engine/Source/Runtime/Engine/Public/WorldPartition/DataLayer/WorldDataLayers.h
--- a/Engine/Source/Runtime/Engine/Public/WorldPartition/DataLayer/WorldDataLayers.h
+++ b/Engine/Source/Runtime/Engine/Public/WorldPartition/DataLayer/WorldDataLayers.h
@@ -56,11 +56,14 @@ class AWorldDataLayers : public AInfo, public IDataLayerInstanceProvider
 	GENERATED_UCLASS_BODY()
 
 public:
+	//~ Begin UObject Interface
 	ENGINE_API virtual void PostLoad() override;
-	ENGINE_API virtual void RewindForReplay() override;
-	ENGINE_API virtual void BeginPlay() override;
 	virtual void Serialize(FArchive& Ar) override;
+	//~ End UObject Interface
 
+	//~ Begin AActor Interface
+	ENGINE_API virtual void RewindForReplay() override;
+	ENGINE_API virtual void BeginPlay() override;
 #if WITH_EDITOR
 	ENGINE_API virtual void PreEditUndo() override;
 	ENGINE_API virtual void PostEditUndo() override;
@@ -72,6 +75,8 @@ public:
 	virtual bool ActorTypeSupportsDataLayer() const override { return false; }
 	virtual bool ActorTypeSupportsExternalDataLayer() const override { return false; }
 	ENGINE_API virtual TUniquePtr<class FWorldPartitionActorDesc> CreateClassActorDesc() const override;
+	virtual void OnLoadedActorRemovedFromLevel() override;
+	//~ End AActor Interface
 
 	static ENGINE_API AWorldDataLayers* Create(UWorld* World, FName InWorldDataLayerName = NAME_None);
 	static ENGINE_API AWorldDataLayers* Create(const FActorSpawnParameters& SpawnParameters);
diff --git a/Engine/Source/Runtime/Engine/Public/WorldPartition/DataLayer/WorldDataLayersActorDesc.h b/Engine/Source/Runtime/Engine/Public/WorldPartition/DataLayer/WorldDataLayersActorDesc.h
--- a/Engine/Source/Runtime/Engine/Public/WorldPartition/DataLayer/WorldDataLayersActorDesc.h
+++ b/Engine/Source/Runtime/Engine/Public/WorldPartition/DataLayer/WorldDataLayersActorDesc.h
@@ -77,9 +77,6 @@ protected:
 	ENGINE_API virtual void Serialize(FArchive& Ar) override;
 	virtual bool IsResaveNeeded() const override { return !IsValid(); }
 	virtual bool IsRuntimeRelevant(const FWorldPartitionActorDescInstance* InActorDescInstance) const override;
-PRAGMA_DISABLE_DEPRECATION_WARNINGS
-	virtual void OnUnloadingInstance(const FWorldPartitionActorDescInstance* InActorDescInstance) const override;
-PRAGMA_ENABLE_DEPRECATION_WARNINGS
 	//~ End FWorldPartitionActorDesc Interface.
 
 private:
diff --git a/Engine/Source/Runtime/Engine/Public/WorldPartition/Landscape/LandscapeActorDesc.h b/Engine/Source/Runtime/Engine/Public/WorldPartition/Landscape/LandscapeActorDesc.h
--- a/Engine/Source/Runtime/Engine/Public/WorldPartition/Landscape/LandscapeActorDesc.h
+++ b/Engine/Source/Runtime/Engine/Public/WorldPartition/Landscape/LandscapeActorDesc.h
@@ -20,9 +20,6 @@ public:
 protected:
 	virtual uint32 GetSizeOf() const override { return sizeof(FLandscapeActorDesc); }
 	ENGINE_API virtual void Serialize(FArchive& Ar) override;
-PRAGMA_DISABLE_DEPRECATION_WARNINGS
-	ENGINE_API void OnUnloadingInstance(const FWorldPartitionActorDescInstance* InActorDescInstance) const override;
-PRAGMA_ENABLE_DEPRECATION_WARNINGS
 
 private:
 	FGuid LandscapeActorGuid;
diff --git a/Engine/Source/Runtime/Engine/Public/WorldPartition/WorldPartitionActorDesc.h b/Engine/Source/Runtime/Engine/Public/WorldPartition/WorldPartitionActorDesc.h
--- a/Engine/Source/Runtime/Engine/Public/WorldPartition/WorldPartitionActorDesc.h
+++ b/Engine/Source/Runtime/Engine/Public/WorldPartition/WorldPartitionActorDesc.h
@@ -377,10 +377,8 @@ public:
 
 	ENGINE_API bool IsMainWorldOnly() const;
 	ENGINE_API bool IsListedInSceneOutliner() const;
-protected:
-	UE_DEPRECATED(5.4, "OnUnloadingInstance is deprecated.")
-	virtual void OnUnloadingInstance(const FWorldPartitionActorDescInstance* InActorDescInstance) const {}
 
+protected:
 	virtual bool GetChildContainerInstance(const FWorldPartitionActorDescInstance* InActorDescInstance, FContainerInstance& OutContainerInstance) const { return false; }
 	virtual UActorDescContainerInstance* CreateChildContainerInstance(const FWorldPartitionActorDescInstance* InActorDescInstance) const { return nullptr; }
 	virtual UWorldPartition* GetLoadedChildWorldPartition(const FWorldPartitionActorDescInstance* InActorDescInstance) const { return nullptr; }
diff --git a/Engine/Source/Runtime/Landscape/Classes/LandscapeProxy.h b/Engine/Source/Runtime/Landscape/Classes/LandscapeProxy.h
--- a/Engine/Source/Runtime/Landscape/Classes/LandscapeProxy.h
+++ b/Engine/Source/Runtime/Landscape/Classes/LandscapeProxy.h
@@ -979,7 +979,9 @@ public:
 	virtual void EditorApplyScale(const FVector& DeltaScale, const FVector* PivotLocation, bool bAltDown, bool bShiftDown, bool bCtrlDown) override;
 	virtual void EditorApplyMirror(const FVector& MirrorScale, const FVector& PivotLocation) override;
 	virtual void PostEditMove(bool bFinished) override;
+	virtual void OnLoadedActorRemovedFromLevel() override;
 	//~ End AActor Interface
+
 	virtual uint32 GetDefaultGridSize(UWorld* InWorld) const override { return 1; }
 	virtual FGuid GetGridGuid() const override { return LandscapeGuid; }
 	virtual TUniquePtr<class FWorldPartitionActorDesc> CreateClassActorDesc() const override;
diff --git a/Engine/Source/Runtime/Landscape/Private/LandscapeEdit.cpp b/Engine/Source/Runtime/Landscape/Private/LandscapeEdit.cpp
--- a/Engine/Source/Runtime/Landscape/Private/LandscapeEdit.cpp
+++ b/Engine/Source/Runtime/Landscape/Private/LandscapeEdit.cpp
@@ -5011,6 +5011,12 @@ void ALandscapeProxy::PostEditMove(bool bFinished)
 	}
 }
 
+void ALandscapeProxy::OnLoadedActorRemovedFromLevel()
+{
+	Super::OnLoadedActorRemovedFromLevel();
+	ActorDescReferences.Empty();
+}
+
 void ALandscapeProxy::PostEditImport()
 {
 	Super::PostEditImport();
