Fixed AActor::IsPropertyChangedAffectingDataLayers ignoring when the actor's ExternalDataLayerAsset changes (ground work for upcoming support of changing external data layer)

diff --git a/Engine/Source/Runtime/Engine/Private/ActorEditor.cpp b/Engine/Source/Runtime/Engine/Private/ActorEditor.cpp
--- a/Engine/Source/Runtime/Engine/Private/ActorEditor.cpp
+++ b/Engine/Source/Runtime/Engine/Private/ActorEditor.cpp
@@ -1809,12 +1809,15 @@ bool AActor::IsPropertyChangedAffectingDataLayers(FPropertyChangedEvent& Propert
 	if (PropertyChangedEvent.Property != nullptr)
 	{
 		static const FName NAME_DataLayerAssets = GET_MEMBER_NAME_CHECKED(AActor, DataLayerAssets);
+		static const FName NAME_ExternalDataLayerAsset = GET_MEMBER_NAME_CHECKED(AActor, ExternalDataLayerAsset);
 
 		const FName PropertyName = PropertyChangedEvent.GetPropertyName();
-		if (PropertyName == NAME_DataLayerAssets &&
-			((PropertyChangedEvent.ChangeType == EPropertyChangeType::ValueSet) ||
-				(PropertyChangedEvent.ChangeType == EPropertyChangeType::ArrayClear) ||
-				(PropertyChangedEvent.ChangeType == EPropertyChangeType::Duplicate)))
+		if (PropertyName == NAME_ExternalDataLayerAsset ||
+			(PropertyName == NAME_DataLayerAssets &&
+			 ((PropertyChangedEvent.ChangeType == EPropertyChangeType::ValueSet) ||
+			 (PropertyChangedEvent.ChangeType == EPropertyChangeType::ArrayClear) ||
+			 (PropertyChangedEvent.ChangeType == EPropertyChangeType::Duplicate)))
+			)
 		{
 			return true;
 		}
