World Partition: fix BP library functions not updated since actor descriptor instances now returns editor bounds in world-sapce.

diff --git a/Engine/Source/Runtime/Engine/Private/WorldPartitionBlueprintLibrary.cpp b/Engine/Source/Runtime/Engine/Private/WorldPartitionBlueprintLibrary.cpp
--- a/Engine/Source/Runtime/Engine/Private/WorldPartitionBlueprintLibrary.cpp
+++ b/Engine/Source/Runtime/Engine/Private/WorldPartitionBlueprintLibrary.cpp
@@ -22,7 +22,7 @@ FActorDesc::FActorDesc()
 {}
 
 #if WITH_EDITOR
-FActorDesc::FActorDesc(const FWorldPartitionActorDesc& InActorDesc, const FTransform& InTransform)
+FActorDesc::FActorDesc(const FWorldPartitionActorDesc& InActorDesc)
 {
 	Guid = InActorDesc.GetGuid();
 
@@ -39,7 +39,7 @@ FActorDesc::FActorDesc(const FWorldPartitionActorDesc& InActorDesc, const FTrans
 
 	Name = InActorDesc.GetActorName();
 	Label = InActorDesc.GetActorLabel();
-	Bounds = InActorDesc.GetEditorBounds().TransformBy(InTransform);
+	Bounds = InActorDesc.GetEditorBounds();
 	RuntimeGrid = InActorDesc.GetRuntimeGrid();
 	bIsSpatiallyLoaded = InActorDesc.GetIsSpatiallyLoaded();
 	bActorIsEditorOnly = InActorDesc.GetActorIsEditorOnly();
@@ -88,13 +88,13 @@ bool UWorldPartitionBlueprintLibrary::GetActorDescs(const UWorldPartition* InWor
 
 	InWorldPartition->ForEachActorDescContainerInstance([&](const UActorDescContainerInstance* InContainerInstance)
 	{
-		bResult &= GetActorDescs(InContainerInstance, FTransform::Identity, OutActorDescs);
+		bResult &= GetActorDescs(InContainerInstance, OutActorDescs);
 	});
 	
 	return bResult;
 }
 
-bool UWorldPartitionBlueprintLibrary::GetActorDescs(const UActorDescContainerInstance* InContainerInstance, const FTransform& InTransform, TArray<FActorDesc>& OutActorDescs)
+bool UWorldPartitionBlueprintLibrary::GetActorDescs(const UActorDescContainerInstance* InContainerInstance, TArray<FActorDesc>& OutActorDescs)
 {
 	bool bResult = true;
 
@@ -105,7 +105,7 @@ bool UWorldPartitionBlueprintLibrary::GetActorDescs(const UActorDescContainerIns
 			FWorldPartitionActorDesc::FContainerInstance ContainerInstance;
 			if (Iterator->GetChildContainerInstance(ContainerInstance))
 			{
-				bResult &= GetActorDescs(ContainerInstance.ContainerInstance, ContainerInstance.Transform * InTransform, OutActorDescs);
+				bResult &= GetActorDescs(ContainerInstance.ContainerInstance, OutActorDescs);
 			}
 			else
 			{
@@ -114,14 +114,14 @@ bool UWorldPartitionBlueprintLibrary::GetActorDescs(const UActorDescContainerIns
 		}
 		else
 		{
-			OutActorDescs.Emplace(*Iterator->GetActorDesc(), InTransform);
+			OutActorDescs.Emplace(*Iterator->GetActorDesc());
 		}
 	}
 
 	return bResult;
 }
 
-bool UWorldPartitionBlueprintLibrary::HandleIntersectingActorDesc(const FWorldPartitionActorDescInstance* ActorDescInstance, const FBox& InBox, const FTransform& InTransform, TArray<FActorDesc>& OutActorDescs)
+bool UWorldPartitionBlueprintLibrary::HandleIntersectingActorDesc(const FWorldPartitionActorDescInstance* ActorDescInstance, const FBox& InBox, TArray<FActorDesc>& OutActorDescs)
 {
 	bool bResult = true;
 
@@ -130,7 +130,7 @@ bool UWorldPartitionBlueprintLibrary::HandleIntersectingActorDesc(const FWorldPa
 		FWorldPartitionActorDesc::FContainerInstance ContainerInstance;
 		if (ActorDescInstance->GetChildContainerInstance(ContainerInstance))
 		{
-			bResult &= GetIntersectingActorDescs(ContainerInstance.ContainerInstance, InBox, ContainerInstance.Transform * InTransform, OutActorDescs);
+			bResult &= GetIntersectingActorDescs(ContainerInstance.ContainerInstance, InBox, OutActorDescs);
 		}
 		else
 		{
@@ -139,7 +139,7 @@ bool UWorldPartitionBlueprintLibrary::HandleIntersectingActorDesc(const FWorldPa
 	}
 	else
 	{
-		OutActorDescs.Emplace(*ActorDescInstance->GetActorDesc(), InTransform);
+		OutActorDescs.Emplace(*ActorDescInstance->GetActorDesc());
 	}
 
 	return bResult;
@@ -151,22 +151,22 @@ bool UWorldPartitionBlueprintLibrary::GetIntersectingActorDescs(UWorldPartition*
 
 	FWorldPartitionHelpers::ForEachIntersectingActorDescInstance(WorldPartition, InBox, [&bResult, &InBox, &OutActorDescs](const FWorldPartitionActorDescInstance* ActorDesc)
 	{
-		bResult &= HandleIntersectingActorDesc(ActorDesc, InBox, FTransform::Identity, OutActorDescs);
+		bResult &= HandleIntersectingActorDesc(ActorDesc, InBox, OutActorDescs);
 		return true;
 	});
 
 	return bResult;
 }
 
-bool UWorldPartitionBlueprintLibrary::GetIntersectingActorDescs(const UActorDescContainerInstance* InContainerInstance, const FBox& InBox, const FTransform& InTransform, TArray<FActorDesc>& OutActorDescs)
+bool UWorldPartitionBlueprintLibrary::GetIntersectingActorDescs(const UActorDescContainerInstance* InContainerInstance, const FBox& InBox, TArray<FActorDesc>& OutActorDescs)
 {
 	bool bResult = true;
 
 	for (UActorDescContainerInstance::TConstIterator<> Iterator(InContainerInstance); Iterator; ++Iterator)
 	{
-		if (Iterator->GetEditorBounds().TransformBy(InTransform).Intersect(InBox))
+		if (Iterator->GetEditorBounds().Intersect(InBox))
 		{
-			bResult &= HandleIntersectingActorDesc(*Iterator, InBox, InTransform, OutActorDescs);
+			bResult &= HandleIntersectingActorDesc(*Iterator, InBox, OutActorDescs);
 		}
 	}
 
diff --git a/Engine/Source/Runtime/Engine/Public/WorldPartition/WorldPartitionBlueprintLibrary.h b/Engine/Source/Runtime/Engine/Public/WorldPartition/WorldPartitionBlueprintLibrary.h
--- a/Engine/Source/Runtime/Engine/Public/WorldPartition/WorldPartitionBlueprintLibrary.h
+++ b/Engine/Source/Runtime/Engine/Public/WorldPartition/WorldPartitionBlueprintLibrary.h
@@ -25,7 +25,7 @@ struct FActorDesc
 	FActorDesc();
 
 #if WITH_EDITOR
-	FActorDesc(const FWorldPartitionActorDesc& InActorDesc, const FTransform& InTransform);
+	FActorDesc(const FWorldPartitionActorDesc& InActorDesc);
 #endif
 
 	/** The actor GUID of this descriptor. */
@@ -91,10 +91,10 @@ class UWorldPartitionBlueprintLibrary : public UBlueprintFunctionLibrary
 	static FDelegateHandle OnWorldPartitionUninitializedHandle;
 
 	static bool GetActorDescs(const UWorldPartition* WorldPartition, TArray<FActorDesc>& OutActorDescs);
-	static bool GetActorDescs(const UActorDescContainerInstance* InContainer, const FTransform& InTransform, TArray<FActorDesc>& OutActorDescs);
+	static bool GetActorDescs(const UActorDescContainerInstance* InContainer, TArray<FActorDesc>& OutActorDescs);
 	static bool GetIntersectingActorDescs(UWorldPartition* WorldPartition, const FBox& InBox, TArray<FActorDesc>& OutActorDescs);
-	static bool GetIntersectingActorDescs(const UActorDescContainerInstance* InContainerInstance, const FBox& InBox, const FTransform& InTransform, TArray<FActorDesc>& OutActorDescs);
-	static bool HandleIntersectingActorDesc(const FWorldPartitionActorDescInstance* ActorDescInstance, const FBox& InBox, const FTransform& InTransform, TArray<FActorDesc>& OutActorDescs);
+	static bool GetIntersectingActorDescs(const UActorDescContainerInstance* InContainerInstance, const FBox& InBox, TArray<FActorDesc>& OutActorDescs);
+	static bool HandleIntersectingActorDesc(const FWorldPartitionActorDescInstance* ActorDescInstance, const FBox& InBox, TArray<FActorDesc>& OutActorDescs);
 #endif
 
 public:
