Don't remap Content Bundle Softobjectpaths at cook time as there is no world package remapping available.
Rely on runtime softobjectpath remapping for Content Bundles.

diff --git a/Engine/Source/Runtime/Engine/Private/WorldPartition/WorldPartitionRuntimeLevelStreamingCell.cpp b/Engine/Source/Runtime/Engine/Private/WorldPartition/WorldPartitionRuntimeLevelStreamingCell.cpp
--- a/Engine/Source/Runtime/Engine/Private/WorldPartition/WorldPartitionRuntimeLevelStreamingCell.cpp
+++ b/Engine/Source/Runtime/Engine/Private/WorldPartition/WorldPartitionRuntimeLevelStreamingCell.cpp
@@ -511,8 +511,10 @@ bool UWorldPartitionRuntimeLevelStreamingCell::OnPopulateGeneratedPackageForCook
 		FWorldPartitionLevelHelper::MoveExternalActorsToLevel(Packages, NewLevel, OutModifiedPackages);
 
 		// Push temporarily the cooking ExternalStreamingObject in the policy for RemapLevelSoftObjectPaths to use it to resolve softobjectpaths
+		// Do this only if the ExternalStreamingObject has a valid root external data layer asset, as Content Bundle soft object remapping is not supported at cook time (there is no world package remapping)
 		const URuntimeHashExternalStreamingObjectBase* ExternalStreamingObject = GetTypedOuter<URuntimeHashExternalStreamingObjectBase>();
-		FScopedCookingExternalStreamingObject ScopeCookingExternalStreamingObject(ExternalStreamingObject);
+		const URuntimeHashExternalStreamingObjectBase* CookingExternalStreamingObject = ExternalStreamingObject && ExternalStreamingObject->GetRootExternalDataLayerAsset() ? ExternalStreamingObject : nullptr;
+		FScopedCookingExternalStreamingObject ScopeCookingExternalStreamingObject(CookingExternalStreamingObject);
 		
 		// Remap Level's SoftObjectPaths
 		FWorldPartitionLevelHelper::RemapLevelSoftObjectPaths(NewLevel, WorldPartition);
