World Partition Blueprint Library: construct actor descriptors from the actor descriptor instance instead of the descriptor, to make sure they are in world space.

diff --git a/Engine/Source/Runtime/Engine/Private/WorldPartitionBlueprintLibrary.cpp b/Engine/Source/Runtime/Engine/Private/WorldPartitionBlueprintLibrary.cpp
--- a/Engine/Source/Runtime/Engine/Private/WorldPartitionBlueprintLibrary.cpp
+++ b/Engine/Source/Runtime/Engine/Private/WorldPartitionBlueprintLibrary.cpp
@@ -22,7 +22,7 @@ FActorDesc::FActorDesc()
 {}
 
 #if WITH_EDITOR
-FActorDesc::FActorDesc(const FWorldPartitionActorDesc& InActorDesc)
+FActorDesc::FActorDesc(const FWorldPartitionActorDescInstance& InActorDesc)
 {
 	Guid = InActorDesc.GetGuid();
 
@@ -114,7 +114,7 @@ bool UWorldPartitionBlueprintLibrary::GetActorDescs(const UActorDescContainerIns
 		}
 		else
 		{
-			OutActorDescs.Emplace(*Iterator->GetActorDesc());
+			OutActorDescs.Emplace(**Iterator);
 		}
 	}
 
@@ -139,7 +139,7 @@ bool UWorldPartitionBlueprintLibrary::HandleIntersectingActorDesc(const FWorldPa
 	}
 	else
 	{
-		OutActorDescs.Emplace(*ActorDescInstance->GetActorDesc());
+		OutActorDescs.Emplace(*ActorDescInstance);
 	}
 
 	return bResult;
diff --git a/Engine/Source/Runtime/Engine/Public/WorldPartition/WorldPartitionBlueprintLibrary.h b/Engine/Source/Runtime/Engine/Public/WorldPartition/WorldPartitionBlueprintLibrary.h
--- a/Engine/Source/Runtime/Engine/Public/WorldPartition/WorldPartitionBlueprintLibrary.h
+++ b/Engine/Source/Runtime/Engine/Public/WorldPartition/WorldPartitionBlueprintLibrary.h
@@ -25,7 +25,7 @@ struct FActorDesc
 	FActorDesc();
 
 #if WITH_EDITOR
-	FActorDesc(const FWorldPartitionActorDesc& InActorDesc);
+	FActorDesc(const FWorldPartitionActorDescInstance& InActorDesc);
 #endif
 
 	/** The actor GUID of this descriptor. */
