Texture : fix GetMipData should use Read Lock

diff --git a/Engine/Source/Runtime/Engine/Private/Texture.cpp b/Engine/Source/Runtime/Engine/Private/Texture.cpp
--- a/Engine/Source/Runtime/Engine/Private/Texture.cpp
+++ b/Engine/Source/Runtime/Engine/Private/Texture.cpp
@@ -2651,8 +2651,7 @@ bool FTextureSource::GetMipImage(FImage & OutImage, int32 BlockIndex, int32 Laye
 		return false;
 	}
 	
-	const int64 MipSizeBytes = CalcMipSize(BlockIndex, LayerIndex, MipIndex);
-	check( MipData.Num() == MipSizeBytes );
+	// code dupe of FMipLock ; consider deleting this and use FMipLock.Image instead
 	
 	FTextureSourceBlock Block;
 	GetBlock(BlockIndex, Block);
@@ -2665,7 +2664,8 @@ bool FTextureSource::GetMipImage(FImage & OutImage, int32 BlockIndex, int32 Laye
 	OutImage.Format = FImageCoreUtils::ConvertToRawImageFormat(GetFormat(LayerIndex));
 	OutImage.GammaSpace = GetGammaSpace(LayerIndex);
 
-	check( OutImage.GetImageSizeBytes() == MipSizeBytes );
+	check( OutImage.GetImageSizeBytes() == OutImage.RawData.Num() );
+
 	return true;
 }
 
@@ -2673,63 +2673,51 @@ bool FTextureSource::GetMipData(TArray64<uint8>& OutMipData, int32 BlockIndex, i
 {
 	TRACE_CPUPROFILER_EVENT_SCOPE(FTextureSource::GetMipData (TArray64));
 	
-#if WITH_EDITOR
-	FScopeLock BulkDataExclusiveScope(&BulkDataLock.Get());
-#endif
+	// used by GetMipImage()
 
-	CheckTextureIsUnlocked(TEXT("GetMipData (TArray64)"));
+	// note: do not use this to get all mips by calling GetMipData repeatedly, it's very inefficient, as it may decompress the source each time
+	//	instead use the GetMipData that returns all mips in one call
 
-	bool bSuccess = false;
+	FMipLock MipLock(FTextureSource::ELockState::ReadOnly,this,BlockIndex,LayerIndex,MipIndex);
 
-	if (IsValid() && BlockIndex < GetNumBlocks() && LayerIndex < NumLayers && MipIndex < NumMips && HasPayloadData())
+	if( ! MipLock.IsValid() )
 	{
-		checkf(NumLockedMips == 0, TEXT("Attempting to access a locked FTextureSource"));
-		// LockedMipData should only be allocated if NumLockedMips > 0 so the following assert should have been caught
-		// by the one above. If it fires then it indicates that there is a lock/unlock mismatch as well as invalid access!
-		checkf(LockedMipData.IsNull(), TEXT("Attempting to access mip data while locked mip data is still allocated"));
+		return false;
+	}
 
-		FSharedBuffer DecompressedData = Decompress();
+	const void * MipData = MipLock.Image.RawData;
+	const int64 MipSize = MipLock.Image.GetImageSizeBytes();
 
-		if (!DecompressedData.IsNull())
-		{
-			const int64 MipOffset = CalcMipOffset(BlockIndex, LayerIndex, MipIndex);
-			const int64 MipSize = CalcMipSize(BlockIndex, LayerIndex, MipIndex);
+	check( MipSize > 0 );
 
-			if ((int64)DecompressedData.GetSize() >= MipOffset + MipSize)
-			{
-				OutMipData.Empty(MipSize);
-				OutMipData.AddUninitialized(MipSize);
-				FMemory::Memcpy(
-					OutMipData.GetData(),
-					(const uint8*)DecompressedData.GetData() + MipOffset,
-					MipSize
-				);
-
-				bSuccess = true;
-			}
-		}	
-	}
-	
-	return bSuccess;
+	OutMipData.SetNumUninitialized(MipSize);
+	FMemory::Memcpy(
+		OutMipData.GetData(),
+		MipData,
+		MipSize);
+
+	return true;
 }
 
 FTextureSource::FMipData FTextureSource::GetMipData(IImageWrapperModule* )
 {
 	TRACE_CPUPROFILER_EVENT_SCOPE(FTextureSource::GetMipData (FMipData));
 
-#if WITH_EDITOR
-	// We can end up waiting here a lot as the bulk data gets serialized for entry in to the transaction buffer.
-	TRACE_CPUPROFILER_EVENT_SCOPE(FTextureSource::GetMipData (StartLock) );
-	FScopeLock _(&BulkDataLock.Get());
-#endif //WITH_EDITOR
+	if ( LockMipReadOnly(0,0,0) == nullptr )
+	{
+		// failed!
+		UE_LOG(LogTexture, Error, TEXT("LockMipReadOnly failed in GetMipData"));
 
-	CheckTextureIsUnlocked(TEXT("GetMipData (FMipData)"));
-	
-	check(LockedMipData.IsNull());
-	check(NumLockedMips == 0);
+		return FMipData(*this,FSharedBuffer());
+	}
+	else
+	{
+		FSharedBuffer DecompressedData = LockedMipData.GetDataReadOnly();
 
-	FSharedBuffer DecompressedData = Decompress();
-	return FMipData(*this, DecompressedData);
+		UnlockMip(0,0,0);
+
+		return FMipData(*this, DecompressedData);
+	}
 }
 
 int64 FTextureSource::CalcMipSize(int32 BlockIndex, int32 LayerIndex, int32 MipIndex) const
