Add material function name to error messages.

diff --git a/Engine/Source/Runtime/Engine/Private/Materials/MaterialExpressions.cpp b/Engine/Source/Runtime/Engine/Private/Materials/MaterialExpressions.cpp
--- a/Engine/Source/Runtime/Engine/Private/Materials/MaterialExpressions.cpp
+++ b/Engine/Source/Runtime/Engine/Private/Materials/MaterialExpressions.cpp
@@ -16953,7 +16953,7 @@ int32 UMaterialExpressionMaterialFunctionCall::Compile(class FMaterialCompiler*
 	{
 		if (!FunctionInputs[i].ExpressionInput)
 		{
-			return Compiler->Errorf(TEXT("Function call input with index %d is unset."), i);
+			return Compiler->Errorf(TEXT("Function (%s) call input with index %d is unset."), *MaterialFunction->GetPathName(), i);
 		}
 	}
 
@@ -16961,13 +16961,13 @@ int32 UMaterialExpressionMaterialFunctionCall::Compile(class FMaterialCompiler*
 	{
 		if (!FunctionOutputs[i].ExpressionOutput)
 		{
-			return Compiler->Errorf(TEXT("Function call output with index %d is unset."), i);
+			return Compiler->Errorf(TEXT("Function (%s) call output with index %d is unset."), *MaterialFunction->GetPathName(), i);
 		}
 	}
 
 	if (!FunctionOutputs.IsValidIndex(OutputIndex))
 	{
-		return Compiler->Errorf(TEXT("Invalid function output"));
+		return Compiler->Errorf(TEXT("Invalid function (%s) output"), *MaterialFunction->GetPathName());
 	}
 
 	// Link the function's inputs into the caller graph before entering
