Fixed static analysis warning

diff --git a/Engine/Source/Runtime/Engine/Private/Components/LightComponent.cpp b/Engine/Source/Runtime/Engine/Private/Components/LightComponent.cpp
--- a/Engine/Source/Runtime/Engine/Private/Components/LightComponent.cpp
+++ b/Engine/Source/Runtime/Engine/Private/Components/LightComponent.cpp
@@ -568,7 +568,7 @@ void ULightComponent::PostLoad()
 	}
 
 	// we want to make sure PreviewShadowMapChannel gets into PIE unchanged
-	if (!WITH_EDITOR || !GetOutermost()->HasAnyPackageFlags(PKG_PlayInEditor))
+	if (!GIsEditor || !GetOutermost()->HasAnyPackageFlags(PKG_PlayInEditor))
 	{
 		PreviewShadowMapChannel = INDEX_NONE;
 	}
