[Material Editor] Check that all Material Parameter Collections loaded succesfully in the translation results retrieved from the DDC before using them, other wise continue translation.

diff --git a/Engine/Source/Runtime/Engine/Private/Materials/HLSLMaterialTranslator.cpp b/Engine/Source/Runtime/Engine/Private/Materials/HLSLMaterialTranslator.cpp
--- a/Engine/Source/Runtime/Engine/Private/Materials/HLSLMaterialTranslator.cpp
+++ b/Engine/Source/Runtime/Engine/Private/Materials/HLSLMaterialTranslator.cpp
@@ -14614,6 +14614,17 @@ void FHLSLMaterialTranslator::AsyncQueryDDC(
 					ResultsMemoryReader << MaterialSourceTemplateParams;
 
 					STAT(DDCRequestSerializeTime = FPlatformTime::Seconds() - SerializeTime);
+
+					// Verify that the Parameter Collections referenced by the Defines buffer are all valid.
+					for (UMaterialParameterCollection* ParameterCollection : EnvironmentDefines->ParameterCollections)
+					{
+						// If the parameter collection is null (it failed to load), do not use these cached results and continue translating the material
+						if (!ParameterCollection)
+						{
+							return;
+						}
+					}
+					
 					DDCQueryHit = true;
 				}
 			}
