FSceneCapturePS has highp floats forced now to fix a driver bug
FullPrecisionInPS can now be enabled on a per global shader basis and not just globally controlled by r.Mobile.FloatPrecisionMode

diff --git a/Engine/Source/Developer/ShaderFormatOpenGL/Private/ShaderFormatOpenGL.cpp b/Engine/Source/Developer/ShaderFormatOpenGL/Private/ShaderFormatOpenGL.cpp
--- a/Engine/Source/Developer/ShaderFormatOpenGL/Private/ShaderFormatOpenGL.cpp
+++ b/Engine/Source/Developer/ShaderFormatOpenGL/Private/ShaderFormatOpenGL.cpp
@@ -26,7 +26,7 @@ extern void CompileOpenGLShader(
 	GLSLVersion Version);
 
 /** Version for shader format, this becomes part of the DDC key. */
-static const FGuid UE_SHADER_GLSL_VER = FGuid("15A0CD3A-20CC-4F15-95CC-04E631068DFE");
+static const FGuid UE_SHADER_GLSL_VER = FGuid("33988E19-4962-4762-8219-F15483DA764A");
 
 class FShaderFormatGLSL : public UE::ShaderCompilerCommon::FBaseShaderFormat 
 {
diff --git a/Engine/Source/Runtime/Engine/Private/ShaderCompiler/ShaderCompiler.cpp b/Engine/Source/Runtime/Engine/Private/ShaderCompiler/ShaderCompiler.cpp
--- a/Engine/Source/Runtime/Engine/Private/ShaderCompiler/ShaderCompiler.cpp
+++ b/Engine/Source/Runtime/Engine/Private/ShaderCompiler/ShaderCompiler.cpp
@@ -8130,7 +8130,7 @@ void GlobalBeginCompileShader(
 	if (bUsingMobileRenderer)
     {
         static IConsoleVariable* CVar = IConsoleManager::Get().FindConsoleVariable(TEXT("r.Mobile.FloatPrecisionMode"));
-        Input.Environment.FullPrecisionInPS = CVar ? (CVar->GetInt() == EMobileFloatPrecisionMode::Full) : false;
+        Input.Environment.FullPrecisionInPS |= CVar ? (CVar->GetInt() == EMobileFloatPrecisionMode::Full) : false;
     }
 	
 	{
diff --git a/Engine/Source/Runtime/Renderer/Private/SceneCaptureRendering.cpp b/Engine/Source/Runtime/Renderer/Private/SceneCaptureRendering.cpp
--- a/Engine/Source/Runtime/Renderer/Private/SceneCaptureRendering.cpp
+++ b/Engine/Source/Runtime/Renderer/Private/SceneCaptureRendering.cpp
@@ -172,6 +172,12 @@ public:
 		{
 			OutEnvironment.SetRenderTargetOutputFormat(0, PF_A32B32G32R32F);
 		}
+
+		if (IsMobilePlatform(Parameters.Platform))
+		{
+			OutEnvironment.FullPrecisionInPS = 1;
+		}
+
 	}
 };
 
