[Material Editor] Material editor now only puts results to DDC after succesful translation if the DDC query finished in time, that is before the full concurrent translation finished. This is to avoid polluting the DDC with the results of small materials that will never be fetched by the DDC because they're faster to directly translate instead.

diff --git a/Engine/Source/Runtime/Engine/Private/Materials/HLSLMaterialTranslator.cpp b/Engine/Source/Runtime/Engine/Private/Materials/HLSLMaterialTranslator.cpp
--- a/Engine/Source/Runtime/Engine/Private/Materials/HLSLMaterialTranslator.cpp
+++ b/Engine/Source/Runtime/Engine/Private/Materials/HLSLMaterialTranslator.cpp
@@ -1166,7 +1166,9 @@ bool FHLSLMaterialTranslator::Translate()
 #endif
 
 	// DDC query local data
+	DDCQueryCompleted = false;
 	DDCQueryHit = false;
+
 	UE::DerivedData::FRequestOwner DDCRequestOwner{ UE::DerivedData::EPriority::Highest };
 	FSharedBuffer MaterialCompilationOutputBuffer;
 	FSharedBuffer TranslationResultsBuffer;
@@ -1192,7 +1194,7 @@ bool FHLSLMaterialTranslator::Translate()
 	{
 		if (bVerbose)
 		{
-			UE_LOG(LogMaterial, Display, TEXT("Material '%s' translation results were in the DDC"), *Material->GetMaterialInterface()->GetFullName());
+			UE_LOG(LogMaterial, Display, TEXT("Material '%s' translation results retrieved from DDC."), *Material->GetMaterialInterface()->GetFullName());
 		}
 
 		STAT(double SerializeTime = FPlatformTime::Seconds());
@@ -1212,15 +1214,20 @@ bool FHLSLMaterialTranslator::Translate()
 	{
 		if (bVerbose)
 		{
-			UE_LOG(LogMaterial, Display, TEXT("Material '%s' results not found in the DDC, material was fully translated"), *Material->GetMaterialInterface()->GetFullName());
+			UE_LOG(LogMaterial, Display, TEXT("Material '%s' completed full translation."), *Material->GetMaterialInterface()->GetFullName());
 		}
 
 		// Material not in cache. If translation was succesful, finalize the results and push them to the DDC
 		PrepareEnvironmentDefines();
 		PrepareMaterialSourceStringParameters();
 
-		if (!bDisableTranslationDDC)
+		if (!bDisableTranslationDDC && DDCQueryCompleted)
 		{
+			if (bVerbose)
+			{
+				UE_LOG(LogMaterial, Display, TEXT("Pushing material '%s' translation results to DDCn."), *Material->GetMaterialInterface()->GetFullName());
+			}
+
 			PushResultsToDDC();
 		}
 	}
@@ -14589,45 +14596,53 @@ void FHLSLMaterialTranslator::AsyncQueryDDC(
 		DDCRequestOwner,
 		[&](UE::DerivedData::FCacheGetResponse&& Response)
 		{
-			if (Response.Status == UE::DerivedData::EStatus::Ok)
+			if (Response.Status != UE::DerivedData::EStatus::Ok)
 			{
-				// Try fetching the translation results from the DDC using the generated key hash
-				FSharedBuffer MaterialCompilationOutputBuffer = Response.Record.GetValue(MaterialCompilationOutputId).GetData().Decompress();
-				FSharedBuffer TranslationResultsBuffer = Response.Record.GetValue(MaterialResultsOutputId).GetData().Decompress();
-				EnvironmentDefinesBuffer = Response.Record.GetValue(EnvironmentDefinesId).GetData().Decompress();
+				return;
+			}
+			
+			// The DDC request was fullfilled and we got data from it.
+			DDCQueryCompleted = true;
 
-				// Load the results if we hit the cache.
-				if (!MaterialCompilationOutputBuffer.IsNull() && !TranslationResultsBuffer.IsNull() && !EnvironmentDefinesBuffer.IsNull())
-				{
-					STAT(double SerializeTime = FPlatformTime::Seconds());
+			// Try fetching the translation results from the DDC using the generated key hash
+			FSharedBuffer MaterialCompilationOutputBuffer = Response.Record.GetValue(MaterialCompilationOutputId).GetData().Decompress();
+			FSharedBuffer TranslationResultsBuffer = Response.Record.GetValue(MaterialResultsOutputId).GetData().Decompress();
+			EnvironmentDefinesBuffer = Response.Record.GetValue(EnvironmentDefinesId).GetData().Decompress();
+
+			// Load the results if we hit the cache.
+			if (MaterialCompilationOutputBuffer.IsNull() || TranslationResultsBuffer.IsNull() || EnvironmentDefinesBuffer.IsNull())
+			{
+				return;
+			}
 
-					// Read the material compilation output
-					FShaderMapPointerTable PointerTable;
-					FPlatformTypeLayoutParameters LayoutParams;
-					LayoutParams.InitializeForPlatform(GetTargetPlatform());
-					FMemoryReaderView MaterialCompilationOutputReader{ TArrayView<uint8>{ (uint8*)MaterialCompilationOutputBuffer.GetData(), (int)MaterialCompilationOutputBuffer.GetSize() } };
-					FMemoryImageObject LoadedContent = FMemoryImageResult::LoadFromArchive(MaterialCompilationOutputReader, StaticGetTypeLayoutDesc<FMaterialCompilationOutput>(), &PointerTable, LayoutParams);
-					DDCMaterialCompilationOutput = *(FMaterialCompilationOutput*)LoadedContent.Object;
+			STAT(double SerializeTime = FPlatformTime::Seconds());
 
-					// Read the array of material string parameters
-					FMemoryReaderView ResultsMemoryReader{ TArrayView<uint8>{ (uint8*)TranslationResultsBuffer.GetData(), (int)TranslationResultsBuffer.GetSize() } };
-					ResultsMemoryReader << MaterialSourceTemplateParams;
+			// Read the material compilation output
+			FShaderMapPointerTable PointerTable;
+			FPlatformTypeLayoutParameters LayoutParams;
+			LayoutParams.InitializeForPlatform(GetTargetPlatform());
+			FMemoryReaderView MaterialCompilationOutputReader{ TArrayView<uint8>{ (uint8*)MaterialCompilationOutputBuffer.GetData(), (int)MaterialCompilationOutputBuffer.GetSize() } };
+			FMemoryImageObject LoadedContent = FMemoryImageResult::LoadFromArchive(MaterialCompilationOutputReader, StaticGetTypeLayoutDesc<FMaterialCompilationOutput>(), &PointerTable, LayoutParams);
+			DDCMaterialCompilationOutput = *(FMaterialCompilationOutput*)LoadedContent.Object;
 
-					STAT(DDCRequestSerializeTime = FPlatformTime::Seconds() - SerializeTime);
+			// Read the array of material string parameters
+			FMemoryReaderView ResultsMemoryReader{ TArrayView<uint8>{ (uint8*)TranslationResultsBuffer.GetData(), (int)TranslationResultsBuffer.GetSize() } };
+			ResultsMemoryReader << MaterialSourceTemplateParams;
 
-					// Verify that the Parameter Collections referenced by the Defines buffer are all valid.
-					for (UMaterialParameterCollection* ParameterCollection : EnvironmentDefines->ParameterCollections)
-					{
-						// If the parameter collection is null (it failed to load), do not use these cached results and continue translating the material
-						if (!ParameterCollection)
-						{
-							return;
-						}
-					}
-					
-					DDCQueryHit = true;
+			STAT(DDCRequestSerializeTime = FPlatformTime::Seconds() - SerializeTime);
+
+			// Verify that the Parameter Collections referenced by the Defines buffer are all valid.
+			for (UMaterialParameterCollection* ParameterCollection : EnvironmentDefines->ParameterCollections)
+			{
+				// If the parameter collection is null (it failed to load), do not use these cached results and continue translating the material
+				if (!ParameterCollection)
+				{
+					return;
 				}
 			}
+			
+			// The DDC request hit the cache and usable data was retrieved from it.
+			DDCQueryHit = true;
 		}
 	);	
 }
diff --git a/Engine/Source/Runtime/Engine/Private/Materials/HLSLMaterialTranslator.h b/Engine/Source/Runtime/Engine/Private/Materials/HLSLMaterialTranslator.h
--- a/Engine/Source/Runtime/Engine/Private/Materials/HLSLMaterialTranslator.h
+++ b/Engine/Source/Runtime/Engine/Private/Materials/HLSLMaterialTranslator.h
@@ -1408,6 +1408,9 @@ protected:
 	/** The output material shader defines */
 	TUniquePtr<FEnvironmentDefines> EnvironmentDefines;
 
+	/** Whether the async DDC query was completed (instead of terminating because canceled) */
+	TAtomic<bool> DDCQueryCompleted;
+
 	/** Signals when the async DDC query task has completed AND there was a hit */
 	TAtomic<bool> DDCQueryHit;
 
