Networking - Fixed issue where some exports was not properly exported with recorded demo frame if exported by checkpoint causing bunch errors due to not all mustbemapped exports being registered.

diff --git a/Engine/Source/Runtime/Engine/Private/ReplayHelper.cpp b/Engine/Source/Runtime/Engine/Private/ReplayHelper.cpp
--- a/Engine/Source/Runtime/Engine/Private/ReplayHelper.cpp
+++ b/Engine/Source/Runtime/Engine/Private/ReplayHelper.cpp
@@ -767,7 +767,7 @@ void FReplayHelper::ProcessCheckpointActors(UNetConnection* Connection, TArrayVi
 {
 	UPackageMapClient* PackageMapClient = Cast<UPackageMapClient>(Connection->PackageMap);
 
-	// Save package map ack status in case we export stuff during the checkpoint (so we can restore the connection back to what it was before we saved the checkpoint)
+	// Set package map ack status override in case we export stuff during the checkpoint
 	PackageMapClient->OverridePackageMapExportAckStatus(&CheckpointSaveContext.CheckpointAckState);
 
 	Connection->SetReserveDestroyedChannels(false);
@@ -815,11 +815,12 @@ void FReplayHelper::ProcessCheckpointActors(UNetConnection* Connection, TArrayVi
 		}
 
 		FlushNetChecked(*Connection);
-
-		PackageMapClient->OverridePackageMapExportAckStatus(nullptr);
 	}
 
 	Connection->SetReserveDestroyedChannels(true);
+
+	// Restore package map ack status
+	PackageMapClient->OverridePackageMapExportAckStatus(nullptr);
 }
 
 void FReplayHelper::TickCheckpoint(UNetConnection* Connection)
@@ -861,9 +862,6 @@ void FReplayHelper::TickCheckpoint(UNetConnection* Connection)
 
 		UPackageMapClient* PackageMapClient = Cast<UPackageMapClient>(Connection->PackageMap);
 
-		// Save package map ack status in case we export stuff during the checkpoint (so we can restore the connection back to what it was before we saved the checkpoint)
-		PackageMapClient->OverridePackageMapExportAckStatus(&CheckpointSaveContext.CheckpointAckState);
-
 		const bool bDeltaCheckpoint = HasDeltaCheckpoints();
 
 		while (bExecuteNextState && (CheckpointSaveContext.CheckpointSaveState != ECheckpointSaveState::Finalize) && !(Params.CheckpointMaxUploadTimePerFrame > 0 && CurrentTime - Params.StartCheckpointTime > Params.CheckpointMaxUploadTimePerFrame))
