World Partition Blueprint Utilities: implement GetActorDescsForActors to gather actor descriptors from a list of actor pointers.

diff --git a/Engine/Source/Runtime/Engine/Private/WorldPartitionBlueprintLibrary.cpp b/Engine/Source/Runtime/Engine/Private/WorldPartitionBlueprintLibrary.cpp
--- a/Engine/Source/Runtime/Engine/Private/WorldPartitionBlueprintLibrary.cpp
+++ b/Engine/Source/Runtime/Engine/Private/WorldPartitionBlueprintLibrary.cpp
@@ -268,6 +268,20 @@ bool UWorldPartitionBlueprintLibrary::GetIntersectingActorDescs(const FBox& InBo
 	return false;
 }
 
+bool UWorldPartitionBlueprintLibrary::GetActorDescsForActors(const TArray<AActor*>& InActors, TArray<FActorDesc>& OutActorDescs)
+{
+#if WITH_EDITOR
+	if (UWorldPartition* WorldPartition = GetWorldPartition())
+	{
+		for (AActor* Actor : InActors)
+		{
+			OutActorDescs.Emplace(*WorldPartition->GetActorDescInstance(Actor->GetActorGuid()));
+		}
+	}
+#endif
+	return false;
+}
+
 UDataLayerManager* UWorldPartitionBlueprintLibrary::GetDataLayerManager(class UObject* WorldContextObject)
 {
 	return UDataLayerManager::GetDataLayerManager(WorldContextObject);
diff --git a/Engine/Source/Runtime/Engine/Public/WorldPartition/WorldPartitionBlueprintLibrary.h b/Engine/Source/Runtime/Engine/Public/WorldPartition/WorldPartitionBlueprintLibrary.h
--- a/Engine/Source/Runtime/Engine/Public/WorldPartition/WorldPartitionBlueprintLibrary.h
+++ b/Engine/Source/Runtime/Engine/Public/WorldPartition/WorldPartitionBlueprintLibrary.h
@@ -150,6 +150,14 @@ public:
 	UFUNCTION(BlueprintCallable, Category="World Partition", meta=(DevelopmentOnly))
 	static bool GetIntersectingActorDescs(const FBox& InBox, TArray<FActorDesc>& OutActorDescs);
 
+	/**
+	 * Gets all the actor descriptors from the provided actor pointers, which represents descriptors on disk, e.g. will not 
+	 * reflect properties of unsaved actors.
+	 * @return True if the operation was successful.
+	 */
+	UFUNCTION(BlueprintCallable, Category="World Partition", meta=(DevelopmentOnly))
+	static bool GetActorDescsForActors(const TArray<AActor*>& InActors, TArray<FActorDesc>& OutActorDescs);
+
 	/**
 	 * Returns the Data Layer Manager for this object.
 	 */
