Add operator bool on FActorInstanceHandle to allow declaration of variable inside ifs' selection-statements

diff --git a/Engine/Source/Runtime/Engine/Classes/Engine/ActorInstanceHandle.h b/Engine/Source/Runtime/Engine/Classes/Engine/ActorInstanceHandle.h
--- a/Engine/Source/Runtime/Engine/Classes/Engine/ActorInstanceHandle.h
+++ b/Engine/Source/Runtime/Engine/Classes/Engine/ActorInstanceHandle.h
@@ -97,6 +97,8 @@ struct FActorInstanceHandle
 	ENGINE_API bool operator==(const AActor* OtherActor) const;
 	ENGINE_API bool operator!=(const AActor* OtherActor) const;
 
+	explicit operator bool() const { return IsValid(); }
+
 	friend ENGINE_API uint32 GetTypeHash(const FActorInstanceHandle& Handle);
 
 	friend ENGINE_API FArchive& operator<<(FArchive& Ar, FActorInstanceHandle& Handle);
