Fix bad edit condition on WorldSetttings.VolumetricLightmapLoadingRange.

diff --git a/Engine/Source/Runtime/Engine/Classes/GameFramework/WorldSettings.h b/Engine/Source/Runtime/Engine/Classes/GameFramework/WorldSettings.h
--- a/Engine/Source/Runtime/Engine/Classes/GameFramework/WorldSettings.h
+++ b/Engine/Source/Runtime/Engine/Classes/GameFramework/WorldSettings.h
@@ -756,7 +756,7 @@ public:
 #endif
 
 	/**  Range in which volumetric lightmaps will be loaded.  */
-	UPROPERTY(BlueprintReadWrite, EditAnywhere, Category = LightmassVolumeLighting, meta = (UIMin = "1600", UIMax = "102400"), meta = (EditConditionHides, EditCondition = "bWorldPartition"))
+	UPROPERTY(BlueprintReadWrite, EditAnywhere, Category = LightmassVolumeLighting, meta = (UIMin = "1600", UIMax = "102400"), meta = (EditConditionHides, EditCondition = "WorldPartition != nullptr"))
 	float VolumetricLightmapLoadingRange;
 
 
