"Exclude From HLOD Levels" has no effect in a World Partition level set up with HLODs
* Hide property when working in a WP world as it's not supported - HLODs are always built from the previous HLOD level

diff --git a/Engine/Source/Runtime/Engine/Classes/Components/PrimitiveComponent.h b/Engine/Source/Runtime/Engine/Classes/Components/PrimitiveComponent.h
--- a/Engine/Source/Runtime/Engine/Classes/Components/PrimitiveComponent.h
+++ b/Engine/Source/Runtime/Engine/Classes/Components/PrimitiveComponent.h
@@ -1952,13 +1952,17 @@ public:
 
 private:
 	/** Which specific HLOD levels this component should be excluded from */
-	UPROPERTY(EditAnywhere, AdvancedDisplay, Category = HLOD, meta = (Bitmask, BitmaskEnum = "/Script/Engine.EHLODLevelExclusion", DisplayName = "Exclude from HLOD Levels", DisplayAfter = "bEnableAutoLODGeneration", EditConditionHides, EditCondition = "bEnableAutoLODGeneration"))
+	UPROPERTY(EditAnywhere, AdvancedDisplay, Category = HLOD, meta = (Bitmask, BitmaskEnum = "/Script/Engine.EHLODLevelExclusion", DisplayName = "Exclude from HLOD Levels", DisplayAfter = "bEnableAutoLODGeneration", EditConditionHides, EditCondition = "AllowHLODLevelsExclusion()"))
 	uint8 ExcludeFromHLODLevels;
 
 	/** LOD parent primitive to draw instead of this one (multiple UPrim's will point to the same LODParent ) */
 	UPROPERTY(NonPIEDuplicateTransient)
 	TObjectPtr<class UPrimitiveComponent> LODParentPrimitive;
 
+	/* Returns true if component can be excluded from specific HLOD levels */
+	UFUNCTION()
+	bool AllowHLODLevelsExclusion() const;
+
 public:
 
 	/** Set the LOD parent component */
diff --git a/Engine/Source/Runtime/Engine/Private/Components/PrimitiveComponent.cpp b/Engine/Source/Runtime/Engine/Private/Components/PrimitiveComponent.cpp
--- a/Engine/Source/Runtime/Engine/Private/Components/PrimitiveComponent.cpp
+++ b/Engine/Source/Runtime/Engine/Private/Components/PrimitiveComponent.cpp
@@ -4549,6 +4549,11 @@ void UPrimitiveComponent::SetLODParentPrimitive(UPrimitiveComponent * InLODParen
 	}
 }
 
+bool UPrimitiveComponent::AllowHLODLevelsExclusion() const
+{
+	return bEnableAutoLODGeneration && !UWorld::IsPartitionedWorld(GetWorld());
+}
+
 UPrimitiveComponent* UPrimitiveComponent::GetLODParentPrimitive() const
 {
 	return LODParentPrimitive;
