Add proper custom gravity support to crouching/uncrouching in the character movement component
#rnx

diff --git a/Engine/Source/Runtime/Engine/Private/Components/CharacterMovementComponent.cpp b/Engine/Source/Runtime/Engine/Private/Components/CharacterMovementComponent.cpp
--- a/Engine/Source/Runtime/Engine/Private/Components/CharacterMovementComponent.cpp
+++ b/Engine/Source/Runtime/Engine/Private/Components/CharacterMovementComponent.cpp
@@ -3088,7 +3088,7 @@ void UCharacterMovementComponent::Crouch(bool bClientSimulation)
 			FCollisionQueryParams CapsuleParams(SCENE_QUERY_STAT(CrouchTrace), false, CharacterOwner);
 			FCollisionResponseParams ResponseParam;
 			InitCollisionParams(CapsuleParams, ResponseParam);
-			const bool bEncroached = GetWorld()->OverlapBlockingTestByChannel(UpdatedComponent->GetComponentLocation() - FVector(0.f,0.f,ScaledHalfHeightAdjust), FQuat::Identity,
+			const bool bEncroached = GetWorld()->OverlapBlockingTestByChannel(UpdatedComponent->GetComponentLocation() - RotateGravityToWorld(FVector(0.f,0.f,ScaledHalfHeightAdjust)), GetWorldToGravityTransform(),
 				UpdatedComponent->GetCollisionObjectType(), GetPawnCapsuleCollisionShape(SHRINK_None), CapsuleParams, ResponseParam);
 
 			// If encroached, cancel
@@ -3102,7 +3102,7 @@ void UCharacterMovementComponent::Crouch(bool bClientSimulation)
 		if (bCrouchMaintainsBaseLocation)
 		{
 			// Intentionally not using MoveUpdatedComponent, where a horizontal plane constraint would prevent the base of the capsule from staying at the same spot.
-			UpdatedComponent->MoveComponent(FVector(0.f, 0.f, -ScaledHalfHeightAdjust), UpdatedComponent->GetComponentQuat(), true, nullptr, EMoveComponentFlags::MOVECOMP_NoFlags, ETeleportType::TeleportPhysics);
+			UpdatedComponent->MoveComponent(RotateGravityToWorld(FVector(0.f, 0.f, -ScaledHalfHeightAdjust)), UpdatedComponent->GetComponentQuat(), true, nullptr, EMoveComponentFlags::MOVECOMP_NoFlags, ETeleportType::TeleportPhysics);
 		}
 
 		CharacterOwner->bIsCrouched = true;
@@ -3125,7 +3125,7 @@ void UCharacterMovementComponent::Crouch(bool bClientSimulation)
 		FNetworkPredictionData_Client_Character* ClientData = GetPredictionData_Client_Character();
 		if (ClientData)
 		{
-			ClientData->MeshTranslationOffset -= FVector(0.f, 0.f, MeshAdjust);
+			ClientData->MeshTranslationOffset -= RotateGravityToWorld(FVector(0.f, 0.f, MeshAdjust));
 			ClientData->OriginalMeshTranslationOffset = ClientData->MeshTranslationOffset;
 		}
 	}
@@ -3180,7 +3180,7 @@ void UCharacterMovementComponent::UnCrouch(bool bClientSimulation)
 		if (!bCrouchMaintainsBaseLocation)
 		{
 			// Expand in place
-			bEncroached = MyWorld->OverlapBlockingTestByChannel(PawnLocation, FQuat::Identity, CollisionChannel, StandingCapsuleShape, CapsuleParams, ResponseParam);
+			bEncroached = MyWorld->OverlapBlockingTestByChannel(PawnLocation, GetWorldToGravityTransform(), CollisionChannel, StandingCapsuleShape, CapsuleParams, ResponseParam);
 		
 			if (bEncroached)
 			{
@@ -3192,11 +3192,11 @@ void UCharacterMovementComponent::UnCrouch(bool bClientSimulation)
 					CharacterOwner->GetCapsuleComponent()->GetScaledCapsuleSize(PawnRadius, PawnHalfHeight);
 					const float ShrinkHalfHeight = PawnHalfHeight - PawnRadius;
 					const float TraceDist = PawnHalfHeight - ShrinkHalfHeight;
-					const FVector Down = FVector(0.f, 0.f, -TraceDist);
+					const FVector Down = RotateGravityToWorld(FVector(0.f, 0.f, -TraceDist));
 
 					FHitResult Hit(1.f);
 					const FCollisionShape ShortCapsuleShape = GetPawnCapsuleCollisionShape(SHRINK_HeightCustom, ShrinkHalfHeight);
-					const bool bBlockingHit = MyWorld->SweepSingleByChannel(Hit, PawnLocation, PawnLocation + Down, FQuat::Identity, CollisionChannel, ShortCapsuleShape, CapsuleParams);
+					const bool bBlockingHit = MyWorld->SweepSingleByChannel(Hit, PawnLocation, PawnLocation + Down, GetWorldToGravityTransform(), CollisionChannel, ShortCapsuleShape, CapsuleParams);
 					if (Hit.bStartPenetrating)
 					{
 						bEncroached = true;
@@ -3205,8 +3205,9 @@ void UCharacterMovementComponent::UnCrouch(bool bClientSimulation)
 					{
 						// Compute where the base of the sweep ended up, and see if we can stand there
 						const float DistanceToBase = (Hit.Time * TraceDist) + ShortCapsuleShape.Capsule.HalfHeight;
-						const FVector NewLoc = FVector(PawnLocation.X, PawnLocation.Y, PawnLocation.Z - DistanceToBase + StandingCapsuleShape.Capsule.HalfHeight + SweepInflation + MIN_FLOOR_DIST / 2.f);
-						bEncroached = MyWorld->OverlapBlockingTestByChannel(NewLoc, FQuat::Identity, CollisionChannel, StandingCapsuleShape, CapsuleParams, ResponseParam);
+						const FVector GravityRelativeAdjustment = FVector(0, 0, -DistanceToBase + StandingCapsuleShape.Capsule.HalfHeight + SweepInflation + MIN_FLOOR_DIST / 2.f);
+						const FVector NewLoc = PawnLocation + RotateGravityToWorld(GravityRelativeAdjustment);
+						bEncroached = MyWorld->OverlapBlockingTestByChannel(NewLoc, GetWorldToGravityTransform(), CollisionChannel, StandingCapsuleShape, CapsuleParams, ResponseParam);
 						if (!bEncroached)
 						{
 							// Intentionally not using MoveUpdatedComponent, where a horizontal plane constraint would prevent the base of the capsule from staying at the same spot.
@@ -3219,8 +3220,8 @@ void UCharacterMovementComponent::UnCrouch(bool bClientSimulation)
 		else
 		{
 			// Expand while keeping base location the same.
-			FVector StandingLocation = PawnLocation + FVector(0.f, 0.f, StandingCapsuleShape.GetCapsuleHalfHeight() - CurrentCrouchedHalfHeight);
-			bEncroached = MyWorld->OverlapBlockingTestByChannel(StandingLocation, FQuat::Identity, CollisionChannel, StandingCapsuleShape, CapsuleParams, ResponseParam);
+			FVector StandingLocation = PawnLocation + RotateGravityToWorld(FVector(0.f, 0.f, StandingCapsuleShape.GetCapsuleHalfHeight() - CurrentCrouchedHalfHeight));
+			bEncroached = MyWorld->OverlapBlockingTestByChannel(StandingLocation, GetWorldToGravityTransform(), CollisionChannel, StandingCapsuleShape, CapsuleParams, ResponseParam);
 
 			if (bEncroached)
 			{
@@ -3230,8 +3231,8 @@ void UCharacterMovementComponent::UnCrouch(bool bClientSimulation)
 					const float MinFloorDist = UE_KINDA_SMALL_NUMBER * 10.f;
 					if (CurrentFloor.bBlockingHit && CurrentFloor.FloorDist > MinFloorDist)
 					{
-						StandingLocation.Z -= CurrentFloor.FloorDist - MinFloorDist;
-						bEncroached = MyWorld->OverlapBlockingTestByChannel(StandingLocation, FQuat::Identity, CollisionChannel, StandingCapsuleShape, CapsuleParams, ResponseParam);
+						StandingLocation -= RotateGravityToWorld(FVector(0, 0, CurrentFloor.FloorDist - MinFloorDist));
+						bEncroached = MyWorld->OverlapBlockingTestByChannel(StandingLocation, GetWorldToGravityTransform(), CollisionChannel, StandingCapsuleShape, CapsuleParams, ResponseParam);
 					}
 				}				
 			}
@@ -3270,7 +3271,7 @@ void UCharacterMovementComponent::UnCrouch(bool bClientSimulation)
 		FNetworkPredictionData_Client_Character* ClientData = GetPredictionData_Client_Character();
 		if (ClientData)
 		{
-			ClientData->MeshTranslationOffset += FVector(0.f, 0.f, MeshAdjust);
+			ClientData->MeshTranslationOffset += RotateGravityToWorld(FVector(0.f, 0.f, MeshAdjust));
 			ClientData->OriginalMeshTranslationOffset = ClientData->MeshTranslationOffset;
 		}
 	}
