CharacterMovementComponent: simulated proxy can go into falling due to linear velocity quantization when using custom gravity
#rnx

diff --git a/Engine/Source/Runtime/Engine/Private/Components/CharacterMovementComponent.cpp b/Engine/Source/Runtime/Engine/Private/Components/CharacterMovementComponent.cpp
--- a/Engine/Source/Runtime/Engine/Private/Components/CharacterMovementComponent.cpp
+++ b/Engine/Source/Runtime/Engine/Private/Components/CharacterMovementComponent.cpp
@@ -2146,7 +2146,7 @@ void UCharacterMovementComponent::SimulateMovement(float DeltaSeconds)
 				bool bShouldFindFloor = Velocity.Z <= 0.f;
 				if (HasCustomGravity())
 				{
-					bShouldFindFloor = RotateWorldToGravity(Velocity).Z <= 0.0;
+					bShouldFindFloor = RotateWorldToGravity(Velocity).Z <= UE_KINDA_SMALL_NUMBER;
 				}
 
 				if (StepDownResult.bComputedFloor)
@@ -2188,7 +2188,7 @@ void UCharacterMovementComponent::SimulateMovement(float DeltaSeconds)
 						// No floor, must fall.
 						if (HasCustomGravity())
 						{
-							if (RotateWorldToGravity(Velocity).Z <= 0.f || bApplyGravityWhileJumping || !CharacterOwner->IsJumpProvidingForce())
+							if (RotateWorldToGravity(Velocity).Z <= UE_KINDA_SMALL_NUMBER || bApplyGravityWhileJumping || !CharacterOwner->IsJumpProvidingForce())
 							{
 								Velocity = NewFallVelocity(Velocity, -GetGravityDirection() * GetGravityZ(), DeltaSeconds);
 							}
