Texture.h change display name of bUseNewMipFilter to reflect its real meaning as a toggle for all modern image processing

diff --git a/Engine/Source/Runtime/Engine/Classes/Engine/Texture.h b/Engine/Source/Runtime/Engine/Classes/Engine/Texture.h
--- a/Engine/Source/Runtime/Engine/Classes/Engine/Texture.h
+++ b/Engine/Source/Runtime/Engine/Classes/Engine/Texture.h
@@ -1272,8 +1272,12 @@ public:
 	UPROPERTY(EditAnywhere, BlueprintReadWrite, Category = Texture, meta=(ClampMin = "0", ClampMax = "1.0", EditCondition="bDoScaleMipsForAlphaCoverage"), AdvancedDisplay)
 	FVector4 AlphaCoverageThresholds = FVector4(0,0,0,0);
 
-	/** Use faster mip generation filter, usually the same result but occasionally causes color shift in high contrast areas. */
-	UPROPERTY(EditAnywhere, BlueprintReadWrite, Category=Texture, meta=(DisplayName = "Use Fast MipGen Filter"), AdvancedDisplay)
+	/** Use faster mip generation filter, usually the same result but occasionally causes color shift in high contrast areas.
+	This is now used for things other than just the mip filter.  It's a bool to toggle the legacy texture processing pipe vs. the modern improved path.
+	It is turned on automatically by SetModernSettingsForNewOrChangedTexture (eg. reimport, and all new textures). */
+	// for GUI hover text : 
+	/** Disable for legacy compatibility.  New and changed textures should set this to use modern improved image processing. */
+	UPROPERTY(EditAnywhere, BlueprintReadWrite, Category=Texture, meta=(DisplayName = "Use Improved Image Processing"), AdvancedDisplay)
 	bool bUseNewMipFilter = false;
 
 	/** When true the texture's border will be preserved during mipmap generation. */
