Fix multiple issues with large virtual texture imports.
* fix condition that checks need for tail mip block for large textures - do not create tail block if max count of mips are already used for the blocks
* limit max size builder can process - depends on runtime limitations how VT page table is handled
* generate VT thumbnail for editor using layer fallback color if calculated thumbnail size is not a multiple of BC block size (happens for greatly skewed image sizes)

diff --git a/Engine/Source/Runtime/Engine/Private/Texture2D.cpp b/Engine/Source/Runtime/Engine/Private/Texture2D.cpp
--- a/Engine/Source/Runtime/Engine/Private/Texture2D.cpp
+++ b/Engine/Source/Runtime/Engine/Private/Texture2D.cpp
@@ -1615,101 +1615,137 @@ void FVirtualTexture2DResource::InitializeEditorResources(IVirtualTexture* InVir
 			MipHeight = FMath::DivideAndRoundUp(MipHeight, 2u);
 		}
 
-		const EPixelFormat PixelFormat = VTData->LayerTypes[0];
-		const bool bCopyUnwantedBordersForAlignment = VTData->TileBorderSize <= 2 && IsBlockCompressedFormat(PixelFormat);
-		const uint32 MipScaleFactor = (1u << MipLevel);
-		const uint32 MipWidthInTiles = FMath::DivideAndRoundUp(GetNumTilesX(), MipScaleFactor);
-		const uint32 MipHeightInTiles = FMath::DivideAndRoundUp(GetNumTilesY(), MipScaleFactor);
-		const uint32 TileSizeInPixels = bCopyUnwantedBordersForAlignment ? GetTileSize() + 2 * GetBorderSize() : GetTileSize();
-		const uint32 LayerMask = 1u; // FVirtualTexture2DResource should only have a single layer
-
-		TArray<FPageToProduce> PagesToProduce;
-		PagesToProduce.Reserve(MipWidthInTiles * MipHeightInTiles);
-		for (uint32 TileY = 0u; TileY < MipHeightInTiles; ++TileY)
+		FRHICommandListImmediate& RHICommandList = FRHICommandListExecutor::GetImmediateCommandList();
+
+		EPixelFormat PixelFormat = VTData->LayerTypes[0];
+
+		FString Name = TextureName.ToString();
+
+		FTextureRHIRef Texture2DRHI;
+
+		if (IsBlockCompressedFormat(PixelFormat) && ((MipWidth % 4 != 0) || (MipHeight % 4 != 0)))
 		{
-			for (uint32 TileX = 0u; TileX < MipWidthInTiles; ++TileX)
+			// This means that VT texture is very skewed on one dimension, thus making other dimension
+			// very small - sometimes less than VT block size (128) so potentially it can be non-multiple of
+			// compressed texture block size. In such case create 1x1 texture thumbnail from fallback color
+
+			uint32 TexWidth = 1;
+			uint32 TexHeight = 1;
+			PixelFormat = EPixelFormat::PF_B8G8R8A8;
+
+			// create thumbnail from fallback color
+
+			FRHITextureCreateDesc StagingDesc = FRHITextureCreateDesc::Create2D(*Name, TexWidth, TexHeight, PixelFormat);
+			StagingDesc.SetFlags(TexCreateFlags);
+			if (!IsRunningRHIInSeparateThread())
 			{
-				const uint32 vAddress = FMath::MortonCode2(TileX) | (FMath::MortonCode2(TileY) << 1);
-				const FVTRequestPageResult RequestResult = InVirtualTexture->RequestPageData(FRHICommandListExecutor::GetImmediateCommandList(), ProducerHandle, LayerMask, MipLevel, vAddress, EVTRequestPagePriority::High);
-				
-				// High priority request should never be Saturated
-				// It's possible for status to be Invalid, if requesting data from a mip level that doesn't exist for the given producer (when using sparse UDIMs)
-				// Technically could try to handle this, by check LocalMipBias, grabbing lower resolution tile, and resizing...but that would make this code much more complex for very little gain
-				ensure(RequestResult.Status != EVTRequestPageStatus::Saturated);
+				StagingDesc.AddFlags(ETextureCreateFlags::CPUWritable);
+			}
+
+			FTextureRHIRef StagingTexture2DRHI = RHICreateTexture(StagingDesc);
+
+			uint32 Stride;
+			void* TexMemory = RHICommandList.LockTexture2D(StagingTexture2DRHI, 0, RLM_WriteOnly, Stride, false, false);
+			*static_cast<FColor*>(TexMemory) = VTData->LayerFallbackColors[0].ToFColor(true);
+			RHICommandList.UnlockTexture2D(StagingTexture2DRHI, 0u, false, false);
 
-				if (VTRequestPageStatus_HasData(RequestResult.Status))
+			Texture2DRHI = MoveTemp(StagingTexture2DRHI);
+		}
+		else
+		{
+			const bool bCopyUnwantedBordersForAlignment = VTData->TileBorderSize <= 2 && IsBlockCompressedFormat(PixelFormat);
+			const uint32 MipScaleFactor = (1u << MipLevel);
+			const uint32 MipWidthInTiles = FMath::DivideAndRoundUp(GetNumTilesX(), MipScaleFactor);
+			const uint32 MipHeightInTiles = FMath::DivideAndRoundUp(GetNumTilesY(), MipScaleFactor);
+			const uint32 TileSizeInPixels = bCopyUnwantedBordersForAlignment ? GetTileSize() + 2 * GetBorderSize() : GetTileSize();
+
+			const uint32 LayerMask = 1u; // FVirtualTexture2DResource should only have a single layer
+
+			TArray<FPageToProduce> PagesToProduce;
+			PagesToProduce.Reserve(MipWidthInTiles * MipHeightInTiles);
+			for (uint32 TileY = 0u; TileY < MipHeightInTiles; ++TileY)
+			{
+				for (uint32 TileX = 0u; TileX < MipWidthInTiles; ++TileX)
 				{
-					PagesToProduce.Add({ RequestResult.Handle, TileX, TileY });
+					const uint32 vAddress = FMath::MortonCode2(TileX) | (FMath::MortonCode2(TileY) << 1);
+					const FVTRequestPageResult RequestResult = InVirtualTexture->RequestPageData(FRHICommandListExecutor::GetImmediateCommandList(), ProducerHandle, LayerMask, MipLevel, vAddress, EVTRequestPagePriority::High);
+
+					// High priority request should never be Saturated
+					// It's possible for status to be Invalid, if requesting data from a mip level that doesn't exist for the given producer (when using sparse UDIMs)
+					// Technically could try to handle this, by check LocalMipBias, grabbing lower resolution tile, and resizing...but that would make this code much more complex for very little gain
+					ensure(RequestResult.Status != EVTRequestPageStatus::Saturated);
+
+					if (VTRequestPageStatus_HasData(RequestResult.Status))
+					{
+						PagesToProduce.Add({ RequestResult.Handle, TileX, TileY });
+					}
 				}
 			}
-		}
 
-		FString Name = TextureName.ToString();
-		FRHITextureCreateDesc Desc = FRHITextureCreateDesc::Create2D(*Name, MipWidthInTiles * TileSizeInPixels, MipHeightInTiles * TileSizeInPixels, PixelFormat);
-		Desc.AddFlags(TexCreateFlags);
-
-		FTextureRHIRef Texture2DRHI = RHICreateTexture(Desc);
+			FRHITextureCreateDesc Desc = FRHITextureCreateDesc::Create2D(*Name, MipWidthInTiles * TileSizeInPixels, MipHeightInTiles * TileSizeInPixels, PixelFormat);
+			Desc.AddFlags(TexCreateFlags);
 
-		FRHICommandListImmediate& RHICommandList = FRHICommandListExecutor::GetImmediateCommandList();
+			Texture2DRHI = RHICreateTexture(Desc);
 
-		// We want to strip borders when compositing tiles since we're just laying out tiles in a regular texture.
-		// But if we have block compressed formats with border less than 4 then doing this will lead to an unaligned copy. We keep the small unwanted borders in that case.
-		const EVTProducePageFlags ProducePageFlags = bCopyUnwantedBordersForAlignment ? EVTProducePageFlags::None : EVTProducePageFlags::SkipPageBorders;
+			// We want to strip borders when compositing tiles since we're just laying out tiles in a regular texture.
+			// But if we have block compressed formats with border less than 4 then doing this will lead to an unaligned copy. We keep the small unwanted borders in that case.
+			const EVTProducePageFlags ProducePageFlags = bCopyUnwantedBordersForAlignment ? EVTProducePageFlags::None : EVTProducePageFlags::SkipPageBorders;
 
-		TArray<IVirtualTextureFinalizer*> Finalizers;
-		for (const FPageToProduce& Page : PagesToProduce)
-		{
-			const uint32 vAddress = FMath::MortonCode2(Page.TileX) | (FMath::MortonCode2(Page.TileY) << 1);
-
-			FVTProduceTargetLayer TargetLayer;
-			TargetLayer.TextureRHI = Texture2DRHI;
-			TargetLayer.pPageLocation = FIntVector(Page.TileX, Page.TileY, 0);
-
-			IVirtualTextureFinalizer* Finalizer = InVirtualTexture->ProducePageData(RHICommandList,
-				GMaxRHIFeatureLevel,
-				ProducePageFlags,
-				ProducerHandle, LayerMask, MipLevel, vAddress,
-				Page.Handle,
-				&TargetLayer);
-			if (Finalizer)
+			TArray<IVirtualTextureFinalizer*> Finalizers;
+			for (const FPageToProduce& Page : PagesToProduce)
 			{
-				Finalizers.AddUnique(Finalizer);
+				const uint32 vAddress = FMath::MortonCode2(Page.TileX) | (FMath::MortonCode2(Page.TileY) << 1);
+
+				FVTProduceTargetLayer TargetLayer;
+				TargetLayer.TextureRHI = Texture2DRHI;
+				TargetLayer.pPageLocation = FIntVector(Page.TileX, Page.TileY, 0);
+
+				IVirtualTextureFinalizer* Finalizer = InVirtualTexture->ProducePageData(RHICommandList,
+					GMaxRHIFeatureLevel,
+					ProducePageFlags,
+					ProducerHandle, LayerMask, MipLevel, vAddress,
+					Page.Handle,
+					&TargetLayer);
+				if (Finalizer)
+				{
+					Finalizers.AddUnique(Finalizer);
+				}
 			}
-		}
 
-		{
-			FRDGBuilder GraphBuilder(RHICommandList);
-			for (IVirtualTextureFinalizer* Finalizer : Finalizers)
 			{
-				Finalizer->Finalize(GraphBuilder);
+				FRDGBuilder GraphBuilder(RHICommandList);
+				for (IVirtualTextureFinalizer* Finalizer : Finalizers)
+				{
+					Finalizer->Finalize(GraphBuilder);
+				}
+				GraphBuilder.Execute();
 			}
-			GraphBuilder.Execute();
-		}
 
-		if (MipWidthInTiles * TileSizeInPixels != MipWidth || MipHeightInTiles * TileSizeInPixels != MipHeight)
-		{
-			// Logical dimensions of mip image may be smaller than tile size (in this case tile will contain mirrored/wrapped padding)
-			// In this case, copy the proper sub-image from the tiled texture we produced into a new texture of the correct size
-			check(MipWidth <= MipWidthInTiles * TileSizeInPixels);
-			check(MipHeight <= MipHeightInTiles * TileSizeInPixels);
+			if (MipWidthInTiles * TileSizeInPixels != MipWidth || MipHeightInTiles * TileSizeInPixels != MipHeight)
+			{
+				// Logical dimensions of mip image may be smaller than tile size (in this case tile will contain mirrored/wrapped padding)
+				// In this case, copy the proper sub-image from the tiled texture we produced into a new texture of the correct size
+				check(MipWidth <= MipWidthInTiles * TileSizeInPixels);
+				check(MipHeight <= MipHeightInTiles * TileSizeInPixels);
 
-			const FRHITextureCreateDesc ResizedDesc =
-				FRHITextureCreateDesc::Create2D(*Name, MipWidth, MipHeight, PixelFormat)
-				.SetFlags(Desc.Flags)
-				.SetInitialState(ERHIAccess::CopyDest);
+				const FRHITextureCreateDesc ResizedDesc =
+					FRHITextureCreateDesc::Create2D(*Name, MipWidth, MipHeight, PixelFormat)
+					.SetFlags(Desc.Flags)
+					.SetInitialState(ERHIAccess::CopyDest);
 
-			FTextureRHIRef ResizedTexture2DRHI = RHICreateTexture(ResizedDesc);
+				FTextureRHIRef ResizedTexture2DRHI = RHICreateTexture(ResizedDesc);
 
-			FRHICopyTextureInfo CopyInfo;
-			CopyInfo.Size = FIntVector(MipWidth, MipHeight, 1);
+				FRHICopyTextureInfo CopyInfo;
+				CopyInfo.Size = FIntVector(MipWidth, MipHeight, 1);
 
-			// Put the source texture in CopySrc mode. The destination texture is already in CopyDest mode because we created it that way.
-			RHICommandList.Transition(FRHITransitionInfo(Texture2DRHI, ERHIAccess::SRVMask, ERHIAccess::CopySrc));
-			RHICommandList.CopyTexture(Texture2DRHI, ResizedTexture2DRHI, CopyInfo);
-			// Make the destination texture SRVMask again. We don't care about the source texture after this, so we won't bother transitioning it.
-			RHICommandList.Transition(FRHITransitionInfo(ResizedTexture2DRHI, ERHIAccess::CopyDest, ERHIAccess::SRVMask));
+				// Put the source texture in CopySrc mode. The destination texture is already in CopyDest mode because we created it that way.
+				RHICommandList.Transition(FRHITransitionInfo(Texture2DRHI, ERHIAccess::SRVMask, ERHIAccess::CopySrc));
+				RHICommandList.CopyTexture(Texture2DRHI, ResizedTexture2DRHI, CopyInfo);
+				// Make the destination texture SRVMask again. We don't care about the source texture after this, so we won't bother transitioning it.
+				RHICommandList.Transition(FRHITransitionInfo(ResizedTexture2DRHI, ERHIAccess::CopyDest, ERHIAccess::SRVMask));
 
-			Texture2DRHI = MoveTemp(ResizedTexture2DRHI);
+				Texture2DRHI = MoveTemp(ResizedTexture2DRHI);
+			}
 		}
 
 		TextureRHI = Texture2DRHI;
diff --git a/Engine/Source/Runtime/Engine/Private/VT/VirtualTextureDataBuilder.cpp b/Engine/Source/Runtime/Engine/Private/VT/VirtualTextureDataBuilder.cpp
--- a/Engine/Source/Runtime/Engine/Private/VT/VirtualTextureDataBuilder.cpp
+++ b/Engine/Source/Runtime/Engine/Private/VT/VirtualTextureDataBuilder.cpp
@@ -323,26 +323,46 @@ bool FVirtualTextureBuilderDerivedInfo::InitializeFromBuildSettings(const FTextu
 	SizeInBlocksX = InSourceData.SizeInBlocksX;
 	SizeInBlocksY = InSourceData.SizeInBlocksY;
 
+	const int64 FullSizeX = (int64)BlockSizeX * SizeInBlocksX;
+	const int64 FullSizeY = (int64)BlockSizeY * SizeInBlocksY;
+
+	// make sure virtual texture dimensions are valid for runtime usage
+	// this should match calculation in FVirtualTextureAllocator::Alloc in VirtualTextureAllocator.cpp file
+	const int64 WidthInTiles = FullSizeX / TileSize;
+	const int64 HeightInTiles = FullSizeY / TileSize;
+	const uint64 MaxSizeInTiles = FMath::Max(WidthInTiles, HeightInTiles);
+	const int64 vLogMaxSize = FMath::CeilLogTwo64(MaxSizeInTiles);
+	if (vLogMaxSize > VIRTUALTEXTURE_LOG2_MAX_PAGETABLE_SIZE)
+	{
+		// max VT size on pixels that runtime supports is TileSize<<MAX_PAGETABLE_SIZE pixels for max dimension
+		// for 128 tile size that is 128<<12 = 524288 pixels
+		UE_LOG(LogVirtualTexturing, Warning,
+			TEXT("InitializeFromBuildSettings failed: VT dimensions (%lld x %lld) are too large - too many tiles (%lld x %lld) [%s]"),
+			FullSizeX, FullSizeY,
+			WidthInTiles, HeightInTiles,
+			*InSourceData.TextureFullName);
+		return false;
+	}
+
 	// total dimensions (of virtual canvas of UDIM blocks) must fit in INT32 on each axis
 	//	  there is a limit of 16 bits of the tile index for the U32 morton code, maybe that's stricter?
 	//	  that's something like 128*65536 maximum virtual dimension?
 	//	  in practice that's hard to hit because the total pixel count will limit you first
-	const int64 VTCanvasMaxDimension = 128*65536; // must fit in INT32_MAX
-	if ( (int64)BlockSizeX * SizeInBlocksX > VTCanvasMaxDimension ||
-		 (int64)BlockSizeY * SizeInBlocksY > VTCanvasMaxDimension )
+	const int64 VTCanvasMaxDimension = TileSize * 65536; // must fit in INT32_MAX
+	if (FullSizeX > VTCanvasMaxDimension || FullSizeY > VTCanvasMaxDimension )
 	{
 		UE_LOG(LogVirtualTexturing,Warning,TEXT("InitializeFromBuildSettings failed : dimensions exceed VTCanvasMaxDimension "
 			"(%d x %d = %lld) (%d x %d = %lld) [%s]"),
-			BlockSizeX,SizeInBlocksX,(int64)BlockSizeX * SizeInBlocksX,
-			BlockSizeY,SizeInBlocksY,(int64)BlockSizeY * SizeInBlocksY,
+			BlockSizeX,SizeInBlocksX,FullSizeX,
+			BlockSizeY,SizeInBlocksY,FullSizeY,
 			*InSourceData.TextureFullName);
 
 		return false;
 	}
 
-	SizeX = BlockSizeX * SizeInBlocksX;
-	SizeY = BlockSizeY * SizeInBlocksY;
-	
+	SizeX = (int32)FullSizeX;
+	SizeY = (int32)FullSizeY;
+
 	// there is no strict limit on total pixel count
 	//	but output must fit in 4 GB
 	//	so as a sanity check, test if pixel count is over 4G
@@ -827,9 +847,6 @@ void FVirtualTextureDataBuilder::BuildLayerBlocks(FSlowTask& BuildTask, uint32 L
 	const int32 NumLayers = SourceData.Layers.Num();
 	const int32 NumBlocks = SourceData.Blocks.Num();
 
-	// If we have more than 1 block, need to create miptail that contains mips made from multiple blocks
-	const bool bNeedsMiptailBlock = (NumBlocks > 1);
-
 	// Miptail
 	TArray<FImage> MiptailInputImages;
 	FPixelDataRectangle MiptailPixelData{ LayerData.SourceFormat, 0, 0, 0 };
@@ -842,6 +859,11 @@ void FVirtualTextureDataBuilder::BuildLayerBlocks(FSlowTask& BuildTask, uint32 L
 	const uint32 MipInputSizeY = FMath::RoundUpToPowerOfTwo(DerivedInfo.SizeInBlocksY * MipHeightInBlock);
 	const uint32 MipInputSize = FMath::Max(MipInputSizeX, MipInputSizeY);
 
+	// If we have more than 1 block and we can produce more mips than each block has
+	// then need to create miptail that contains mips made from multiple blocks2
+	// be aware of mip limit (OutData.NumMips) - no need for miptail block if all mips are already used
+	const bool bNeedsMiptailBlock = (NumBlocks > 1) && (OutData.NumMips > (MaxMipInBlock + 1));
+
 	if (bNeedsMiptailBlock)
 	{
 		MiptailInputImages.Reset(1);
