ActorInstanceHandle: Implement equality operator using only cached Actors or InstanceIndex, and not FetchActor.
This prevents operator== from spiking in CPU usage when it called FetchActor and spawned actors inline

diff --git a/Engine/Source/Runtime/Engine/Private/Engine/ActorInstanceHandle.cpp b/Engine/Source/Runtime/Engine/Private/Engine/ActorInstanceHandle.cpp
--- a/Engine/Source/Runtime/Engine/Private/Engine/ActorInstanceHandle.cpp
+++ b/Engine/Source/Runtime/Engine/Private/Engine/ActorInstanceHandle.cpp
@@ -193,9 +193,9 @@ bool FActorInstanceHandle::DoesRepresentClass(const UClass* OtherClass) const
 UClass* FActorInstanceHandle::GetRepresentedClass() const
 {
 	// Calling IsActorValid will resolve the handle if necessary
-	if (IsActorValid())
+	if (const AActor* CachedActor = GetCachedActor())
 	{
-		return ReferenceObject->GetClass();
+		return CachedActor->GetClass();
 	}
 
 	IActorInstanceManagerInterface* ManagerInterfacePtr = ManagerInterface.Get();
@@ -381,17 +381,33 @@ bool FActorInstanceHandle::operator==(const FActorInstanceHandle& Other) const
 		return false;
 	}
 
-	// try to compare managers and indices first if we have them
-	if (ManagerInterface.IsValid() && Other.ManagerInterface.IsValid() && InstanceIndex != INDEX_NONE && Other.InstanceIndex != INDEX_NONE)
+	// Try to compare the cached actors if available
+	const AActor* MyActor = GetCachedActor();
+	const AActor* OtherActor = Other.GetCachedActor();
+
+	// If both Actors are set, we can compare the pointers.
+	// But if one Actor is missing, there's a possibility that the Actor was spawned after the handle was initially setup.
+	// In that case we defer to operator==(AActor*) which will try to update the cached actor by calling FindActor on the manager interface.
+	// If one of the actor is set and FindActor returns null, we know both handles can't point to the same Actor.
+	if (MyActor)
 	{
-		return ManagerInterface == Other.ManagerInterface && InstanceIndex == Other.InstanceIndex;
+		if (OtherActor)
+		{
+			return Other == MyActor;
+		}
+		else
+		{
+			return Other.operator==(MyActor);
+		}
+	}
+	else if (OtherActor)
+	{
+		return operator==(OtherActor);
 	}
 
-	// try to compare the actors
-	const AActor* MyActor = FetchActor();
-	const AActor* OtherActor = Other.FetchActor();
-
-	return MyActor == OtherActor;
+	// The actors are null, try to compare managers and indices first if we have them.
+	// If the ManagerInterface and InstanceIndex of both handles are invalid, we consider them null and equal to each other.
+	return ManagerInterface == Other.ManagerInterface && InstanceIndex == Other.InstanceIndex;
 }
 
 bool FActorInstanceHandle::operator!=(const FActorInstanceHandle& Other) const
@@ -415,8 +431,14 @@ bool FActorInstanceHandle::operator==(const AActor* OtherActor) const
 		return !ManagerInterface.IsValid() && InstanceIndex == INDEX_NONE;
 	}
 
-	IActorInstanceManagerInterface* ManagerInterfacePtr = ManagerInterface.Get();
-	return ManagerInterfacePtr && (ManagerInterfacePtr->FindActor(*this) == OtherActor);
+	if (IActorInstanceManagerInterface* ManagerInterfacePtr = ManagerInterface.Get())
+	{
+		// Since we need to call FindActor, we might as well save the result.
+		// Setting ReferencedObject without calling SetCachedActor to avoid having to const_cast this.
+		ReferenceObject = ManagerInterfacePtr->FindActor(*this);
+	}
+
+	return GetCachedActor() == OtherActor;
 }
 
 bool FActorInstanceHandle::operator!=(const AActor* OtherActor) const
