Make ShouldIgnoreOverlapResult use InstancedActorHandles to avoid calling FetchActor/SpawnActor.
This avoids spawning InstancedActors that were hit in the overlap checks, just to ask them whether they are initialized.

diff --git a/Engine/Source/Runtime/Engine/Private/Components/PrimitiveComponent.cpp b/Engine/Source/Runtime/Engine/Private/Components/PrimitiveComponent.cpp
--- a/Engine/Source/Runtime/Engine/Private/Components/PrimitiveComponent.cpp
+++ b/Engine/Source/Runtime/Engine/Private/Components/PrimitiveComponent.cpp
@@ -2553,7 +2553,7 @@ static bool ShouldIgnoreHitResult(const UWorld* InWorld, FHitResult const& TestH
 	return false;
 }
 
-static FORCEINLINE_DEBUGGABLE bool ShouldIgnoreOverlapResult(const UWorld* World, const AActor* ThisActor, const UPrimitiveComponent& ThisComponent, const AActor* OtherActor, const UPrimitiveComponent& OtherComponent, bool bCheckOverlapFlags)
+static FORCEINLINE_DEBUGGABLE bool ShouldIgnoreOverlapResult(const UWorld* World, const AActor* ThisActor, const UPrimitiveComponent& ThisComponent, const FActorInstanceHandle& OtherActor, const UPrimitiveComponent& OtherComponent, bool bCheckOverlapFlags)
 {
 	// Don't overlap with self
 	if (&ThisComponent == &OtherComponent)
@@ -2575,7 +2575,7 @@ static FORCEINLINE_DEBUGGABLE bool ShouldIgnoreOverlapResult(const UWorld* World
 		return true;
 	}
 
-	if (!World || OtherActor == World->GetWorldSettings() || !OtherActor->IsActorInitialized())
+	if (!World || OtherActor == World->GetWorldSettings() || (OtherActor.GetCachedActor() && !OtherActor.GetCachedActor()->IsActorInitialized()))
 	{
 		return true;
 	}
@@ -2791,7 +2791,7 @@ bool UPrimitiveComponent::MoveComponentImpl( const FVector& Delta, const FQuat&
 						UPrimitiveComponent* OverlapComponent = TestHit.Component.Get();
 						if (OverlapComponent && (OverlapComponent->GetGenerateOverlapEvents() || bForceGatherOverlaps))
 						{
-							if (!ShouldIgnoreOverlapResult(MyWorld, Actor, *this, TestHit.HitObjectHandle.FetchActor(), *OverlapComponent, /*bCheckOverlapFlags=*/ !bForceGatherOverlaps))
+							if (!ShouldIgnoreOverlapResult(MyWorld, Actor, *this, TestHit.HitObjectHandle, *OverlapComponent, /*bCheckOverlapFlags=*/ !bForceGatherOverlaps))
 							{
 								// don't process touch events after initial blocking hits
 								if (BlockingHitIndex >= 0 && TestHit.Time > Hits[BlockingHitIndex].Time)
@@ -3941,7 +3941,7 @@ bool UPrimitiveComponent::UpdateOverlapsImpl(const TOverlapArrayView* NewPending
 						if (HitComp && (HitComp != this) && HitComp->GetGenerateOverlapEvents())
 						{
 							const bool bCheckOverlapFlags = false; // Already checked above
-							if (!ShouldIgnoreOverlapResult(MyWorld, MyActor, *this, Result.OverlapObjectHandle.FetchActor(), *HitComp, bCheckOverlapFlags))
+							if (!ShouldIgnoreOverlapResult(MyWorld, MyActor, *this, Result.OverlapObjectHandle, *HitComp, bCheckOverlapFlags))
 							{
 								OverlapMultiResult.Emplace(HitComp, Result.ItemIndex);		// don't need to add unique unless the overlap check can return dupes
 							}
