[Material Editor] Small improvement to my last change to make the translator push results to the DDC with a more correct heuristic.

diff --git a/Engine/Source/Runtime/Engine/Private/Materials/HLSLMaterialTranslator.cpp b/Engine/Source/Runtime/Engine/Private/Materials/HLSLMaterialTranslator.cpp
--- a/Engine/Source/Runtime/Engine/Private/Materials/HLSLMaterialTranslator.cpp
+++ b/Engine/Source/Runtime/Engine/Private/Materials/HLSLMaterialTranslator.cpp
@@ -14596,13 +14596,17 @@ void FHLSLMaterialTranslator::AsyncQueryDDC(
 		DDCRequestOwner,
 		[&](UE::DerivedData::FCacheGetResponse&& Response)
 		{
+			// If the request was't canceled it means it came back with an answer. Mark the query complete.
+			if (Response.Status != UE::DerivedData::EStatus::Canceled)
+			{
+				DDCQueryCompleted = true;
+			}
+
+			// If the status isn't Ok, the request either missed or was canceled.
 			if (Response.Status != UE::DerivedData::EStatus::Ok)
 			{
 				return;
 			}
-			
-			// The DDC request was fullfilled and we got data from it.
-			DDCQueryCompleted = true;
 
 			// Try fetching the translation results from the DDC using the generated key hash
 			FSharedBuffer MaterialCompilationOutputBuffer = Response.Record.GetValue(MaterialCompilationOutputId).GetData().Decompress();
