Fix for compiler warning.

diff --git a/Engine/Source/Runtime/Engine/Private/NetDriver.cpp b/Engine/Source/Runtime/Engine/Private/NetDriver.cpp
--- a/Engine/Source/Runtime/Engine/Private/NetDriver.cpp
+++ b/Engine/Source/Runtime/Engine/Private/NetDriver.cpp
@@ -258,10 +258,7 @@ namespace UE::Net::Private
 				OutParams.MaxReplicatedObjectCount = ReplicationSystemConfig.MaxReplicatedObjectServerCount;
 			}
 
-			if (ReplicationSystemConfig.PreAllocatedReplicatedObjectServerCount >= 0)
-			{
-				OutParams.PreAllocatedReplicatedObjectCount = ReplicationSystemConfig.PreAllocatedReplicatedObjectServerCount;
-			}
+			OutParams.PreAllocatedReplicatedObjectCount = ReplicationSystemConfig.PreAllocatedReplicatedObjectServerCount;
 
 			OutParams.MaxReplicatedWriterObjectCount = OutParams.MaxReplicatedObjectCount;
 		}
@@ -272,10 +269,7 @@ namespace UE::Net::Private
 				OutParams.MaxReplicatedObjectCount = ReplicationSystemConfig.MaxReplicatedObjectClientCount;
 			}
 
-			if (ReplicationSystemConfig.PreAllocatedReplicatedObjectClientCount >= 0)
-			{
-				OutParams.PreAllocatedReplicatedObjectCount = ReplicationSystemConfig.PreAllocatedReplicatedObjectClientCount;
-			}
+			OutParams.PreAllocatedReplicatedObjectCount = ReplicationSystemConfig.PreAllocatedReplicatedObjectClientCount;
 
 			if (ReplicationSystemConfig.MaxReplicatedWriterObjectClientCount != 0)
 			{
