Texture.cpp clarify LockMip , mainly a nop comment, behavior unchanged

diff --git a/Engine/Source/Runtime/Engine/Private/Texture.cpp b/Engine/Source/Runtime/Engine/Private/Texture.cpp
--- a/Engine/Source/Runtime/Engine/Private/Texture.cpp
+++ b/Engine/Source/Runtime/Engine/Private/Texture.cpp
@@ -2544,10 +2544,60 @@ FMutableMemoryView FTextureSource::LockMipInternal(int32 BlockIndex, int32 Layer
 
 	checkf(RequestedLockState != ELockState::None, TEXT("Cannot call FTextureSource::LockMipInternal with a RequestedLockState of type ELockState::None"));
 
+	// This "lock" is not a critical section; it does not block multi-threaded access
+	//  rather it is a way to scope access to the decompressed mip data
+	// it does work for multi-threaded reads, but NOT for multi-threaded writes
+	//
+	// this is a sort of RW lock, not really.  eg. if you try to lock for write when previously locked for read,
+	//	it does not block on the read locks being released the way a RW lock would
+	// this lock is recursive (can have multiple locks from the same thread)
+	// we do allow locking for read inside a lock write on the same thread, but not vice-versa (can not lock for write from inside a read lock)
+	// locking for write from multiple threads is *allowed* by this code, which is wrong of course
+	//	this all works only if the rules of texture threading are followed
+	// that is :
+	//	1. only the main thread should mutate textures
+	//	2. before mutating textures, always use PreEditChange/PostEditChange, this blocks any async builds
+	//	3. the only multi-threaded access of textures should be for *read* (eg. the builder can run on tasks, but only reads textures)
+	//	4. writing should always be single threaded and no multi-threaded reads can happen while one thread is writing
+	//
+	// so things you would normally expect a RW lock to protect against are not allowed to happen by the texture threading rules
+	//  and this RW lock does not enforce them or provide protection!
+	//
+	// note: actually locks the whole texture, not one mip at a time
+	//
+	// note: if you are using this to access mips one at a time, that is very inefficient unless you hold one lock ref throughout
+
 	FMutableMemoryView MipView;
 
 	if (BlockIndex < GetNumBlocks() && LayerIndex < NumLayers && MipIndex < NumMips)
 	{
+		if (NumLockedMips > 0 && RequestedLockState != LockState)
+		{
+			// previously locked, and requested lock is not the same as previous
+			if ( LockState == ELockState::ReadWrite )
+			{
+				// previous lock was for write
+				// we're requesting a read
+				//	allow it, promote our request to write
+				// this must be happening due to recursive locking, NOT from different threads
+				// if anyone has a write lock, texture multi-threading is not allowed
+				RequestedLockState = ELockState::ReadWrite;
+			}
+			else
+			{
+				// was previously locked for read, now wants to write
+				// that is not allowed, will fail
+				check( LockState == ELockState::ReadOnly );
+				check( RequestedLockState == ELockState::ReadWrite );
+
+				UE_LOG(LogTexture,Error, TEXT("LockMip cannot lock for write when previously locked for read [%s]"), 
+					Owner ? *Owner->GetFullName() : TEXT("unowned"));	
+				
+				// no data, you did not get the lock, do not call Unlock
+				return MipView;
+			}
+		}
+
 		if (LockedMipData.IsNull())
 		{
 			checkf(NumLockedMips == 0, TEXT("Texture mips are locked but the LockedMipData is missing"));
@@ -4322,17 +4372,9 @@ bool FTextureSource::ComputeChannelLinearMinMax(int32 InLayerIndex, FLinearColor
 	}
 
 	TRACE_CPUPROFILER_EVENT_SCOPE(FTextureSource::ComputeChannelLinearMinMax);
-
-	// If we're already locked then just use what the existing lock type was since we aren't changing anything
-	// and if we try ReadOnly when we are locked ReadWrite we'll get a lock mismatch error.
-	FTextureSource::ELockState UseLockType = FTextureSource::ELockState::ReadOnly;
-	if (NumLockedMips)
-	{
-		UseLockType = LockState;
-	}
-
-	// have to strip const for the lock state
-	FTextureSource::FMipLock LockedMip0(UseLockType, (FTextureSource*)this, 0);
+	
+	// we hold a lock throughout so that we don't unlock multiple times
+	FTextureSource::FMipLock LockedMip0(ELockState::ReadOnly, const_cast<FTextureSource*>(this), 0);
 	if (LockedMip0.IsValid() == false)
 	{
 		return false;
@@ -4343,10 +4385,9 @@ bool FTextureSource::ComputeChannelLinearMinMax(int32 InLayerIndex, FLinearColor
 
 	for (int32 BlockIndex = 0; BlockIndex < GetNumBlocks(); BlockIndex++)
 	{
-		// The data is already present and locked from the mip0 lock above, this just gets use the
-		// imageview.
-		// have to strip const for the lock state
-		FTextureSource::FMipLock LockedBlock(UseLockType, (FTextureSource*)this, BlockIndex, InLayerIndex, 0);
+		// The data is already present and locked from the mip0 lock above, this just gets us the imageview
+		// Note we only look at mip 0 ; it is possible that other mips go out of the MinMax bound we find.
+		FTextureSource::FMipLock LockedBlock(ELockState::ReadOnly, const_cast<FTextureSource*>(this), BlockIndex, InLayerIndex, 0);
 		check(LockedBlock.IsValid()); // should be same as validity check above!!
 
 		FLinearColor MinColor, MaxColor;
@@ -4440,36 +4481,27 @@ void FTextureSource::UpdateChannelMinMaxFromIncomingTextureData(FMemoryView InNe
 
 bool FTextureSource::UpdateChannelLinearMinMax()
 {
-	LayerColorInfo.Empty();
-
-	// If we're already locked then just use what the existing lock type was since we aren't changing anything
-	// and if we try ReadOnly when we are locked ReadWrite we'll get a lock mismatch error.
-	FTextureSource::ELockState UseLockType = FTextureSource::ELockState::ReadOnly;
-	if (NumLockedMips)
-	{
-		UseLockType = LockState;
-	}
-
-	// have to strip const for the lock state.
-	// we take a lock here so that we don't do a separate lock for each layer - its OK to nest locks.
-	FTextureSource::FMipLock LockedMip0(UseLockType, (FTextureSource*)this, 0);
+	
+	// we hold a lock throughout so that we don't unlock multiple times
+	FTextureSource::FMipLock LockedMip0(ELockState::ReadOnly, const_cast<FTextureSource*>(this), 0);
 	if (LockedMip0.IsValid() == false)
 	{
+		LayerColorInfo.Empty();
 		return false;
 	}
 
+	LayerColorInfo.SetNum(NumLayers);
 	for (int32 LayerIndex = 0; LayerIndex < NumLayers; LayerIndex++)
 	{
-		FTextureSourceLayerColorInfo& LayerInfo = LayerColorInfo.AddDefaulted_GetRef();
+		FTextureSourceLayerColorInfo& LayerInfo = LayerColorInfo[LayerIndex];
 
-		bool GotMinMax = ComputeChannelLinearMinMax(LayerIndex, LayerInfo.ColorMin, LayerInfo.ColorMax);
-		check(GotMinMax); // should be the same check as above
-		if (GotMinMax == false)
+		if ( ! ComputeChannelLinearMinMax(LayerIndex, LayerInfo.ColorMin, LayerInfo.ColorMax) )
 		{
 			LayerColorInfo.Empty();
 			return false;
 		}
 	}
+
 	return true;
 }
 
