Fixed off-by-one bug leading to an out-of-bounds array access on empty SVT frames.

#rnx

diff --git a/Engine/Source/Runtime/Engine/Private/SparseVolumeTexture/SparseVolumeTextureStreamingManager.cpp b/Engine/Source/Runtime/Engine/Private/SparseVolumeTexture/SparseVolumeTextureStreamingManager.cpp
--- a/Engine/Source/Runtime/Engine/Private/SparseVolumeTexture/SparseVolumeTextureStreamingManager.cpp
+++ b/Engine/Source/Runtime/Engine/Private/SparseVolumeTexture/SparseVolumeTextureStreamingManager.cpp
@@ -680,8 +680,9 @@ void FStreamingManager::AddInternal(FRDGBuilder& GraphBuilder, FNewSparseVolumeT
 			NumRootVoxelsB += RootTileInfo.NumVoxels[1];
 		}
 
-		// Determine the high water mark of per frame streaming size when streaming at any given mip level
-		uint32 MaxReferencedTileIndex = 0;
+		// Determine the high water mark of per frame streaming size when streaming at any given mip level.
+		// Find the maximum referenced tile index + 1, which is the index in TileDataOffsets (a prefix sum with N+1 elements) at which the cumulative size of this tile and all prior tiles is stored.
+		uint32 MaxReferencedTileIndexPlusOne = 0;
 		for (int32 MipLevel = FrameInfo.NumMipLevels - 2; MipLevel >= 0; --MipLevel)
 		{
 			if (Topology.MipInfo.IsValidIndex(MipLevel))
@@ -690,9 +691,9 @@ void FStreamingManager::AddInternal(FRDGBuilder& GraphBuilder, FNewSparseVolumeT
 				for (uint32 PageIndex = MipInfo.PageOffset; PageIndex < (MipInfo.PageOffset + MipInfo.PageCount); ++PageIndex)
 				{
 					const uint32 TileIndex = Topology.TileIndices[PageIndex];
-					MaxReferencedTileIndex = FMath::Max(MaxReferencedTileIndex, TileIndex);
+					MaxReferencedTileIndexPlusOne = FMath::Max(MaxReferencedTileIndexPlusOne, TileIndex + 1);
 				}
-				const uint32 MipStreamingSize = Resources->StreamingMetaData.TileDataOffsets[MaxReferencedTileIndex + 1] - Resources->StreamingMetaData.GetRootTileSize();
+				const uint32 MipStreamingSize = Resources->StreamingMetaData.TileDataOffsets[MaxReferencedTileIndexPlusOne] - Resources->StreamingMetaData.GetRootTileSize();
 				SVTInfo.MipLevelStreamingSize[MipLevel] = FMath::Max(SVTInfo.MipLevelStreamingSize[MipLevel], MipStreamingSize);
 			}
 		}
