Submixes. Bring back deleted properties so they can still be migrated for older assets. Add custom version for handling migration upgrades.

#tests Sondra tested this in 5.4 against older assets to verify migration

diff --git a/Engine/Source/Runtime/Engine/Classes/Sound/SoundSubmix.h b/Engine/Source/Runtime/Engine/Classes/Sound/SoundSubmix.h
--- a/Engine/Source/Runtime/Engine/Classes/Sound/SoundSubmix.h
+++ b/Engine/Source/Runtime/Engine/Classes/Sound/SoundSubmix.h
@@ -463,6 +463,7 @@ public:
 protected:
 
 	ENGINE_API virtual void Serialize(FArchive& Ar) override;
+	ENGINE_API virtual void PostLoad() override;
 
 #if WITH_EDITOR
 	ENGINE_API virtual void PostEditChangeProperty(struct FPropertyChangedEvent& PropertyChangedEvent) override;
@@ -483,6 +484,22 @@ protected:
 
 	// State handling for bouncing output.
 	TUniquePtr<Audio::FAudioRecordingData> RecordingData;
+
+#if WITH_EDITORONLY_DATA
+
+	// Forever deprecated properties.
+	// These must be kept to always be able to migrate older assets.
+	UPROPERTY()
+	float OutputVolume_DEPRECATED;
+	UPROPERTY()
+	float WetLevel_DEPRECATED;
+	UPROPERTY()
+	float DryLevel_DEPRECATED;
+
+	void InitDeprecatedDefaults();
+	void HandleVersionMigration(const int32 Version);
+
+#endif //WITH_EDITORONLY_DATA
 };
 	
 
diff --git a/Engine/Source/Runtime/Engine/Private/SoundSubmix.cpp b/Engine/Source/Runtime/Engine/Private/SoundSubmix.cpp
--- a/Engine/Source/Runtime/Engine/Private/SoundSubmix.cpp
+++ b/Engine/Source/Runtime/Engine/Private/SoundSubmix.cpp
@@ -19,6 +19,8 @@
 #include "Async/Async.h"
 #endif // WITH_EDITOR
 
+#include "SoundSubmixCustomVersion.h"
+
 static int32 ClearBrokenSubmixAssetsCVar = 0;
 FAutoConsoleVariableRef CVarFixUpBrokenSubmixAssets(
 	TEXT("au.submix.clearbrokensubmixassets"),
@@ -27,6 +29,14 @@ FAutoConsoleVariableRef CVarFixUpBrokenSubmixAssets(
 	TEXT("0: Disable, >0: Enable"),
 	ECVF_Default);
 
+namespace SoundSubmixPrivate
+{
+	// Modulators default. 
+	static const float Default_OutputVolumeModulation = 0.f;
+	static const float Default_WetLevelModulation = 0.f;
+	static const float Default_DryLevelModulation = -96.f;
+}
+
 USoundSubmixWithParentBase::USoundSubmixWithParentBase(const FObjectInitializer& ObjectInitializer)
 	: Super(ObjectInitializer)
 	, ParentSubmix(nullptr)
@@ -46,42 +56,33 @@ USoundSubmix::USoundSubmix(const FObjectInitializer& ObjectInitializer)
 	, EnvelopeFollowerAttackTime(10)
 	, EnvelopeFollowerReleaseTime(500)
 {
-	OutputVolumeModulation.Value = 0.f;
-	WetLevelModulation.Value = 0.f;
-	DryLevelModulation.Value = -96.f;
+	using namespace SoundSubmixPrivate;
+	OutputVolumeModulation.Value	= Default_OutputVolumeModulation;
+	WetLevelModulation.Value		= Default_WetLevelModulation;
+	DryLevelModulation.Value		= Default_DryLevelModulation;
+
+#if WITH_EDITORONLY_DATA
+	InitDeprecatedDefaults();
+#endif //WITH_EDITORONLY_DATA
 }
 
 void USoundSubmix::Serialize(FArchive& Ar)
 {
 	Super::Serialize(Ar);
+	Ar.UsingCustomVersion(FSoundSubmixCustomVersion::GUID);
+}
 
-#if WITH_EDITORONLY_DATA
-	if (Ar.IsLoading() || Ar.IsSaving())
-	{
-		// use -96dB as a noise floor when fixing up linear volume settings
-		static constexpr float LinearNeg96dB = 0.0000158489319f;
-
-		// fix values previously saved as linear values
-		if (OutputVolumeModulation.Value > 0.0f)
-		{
-			OutputVolumeModulation.Value = Audio::ConvertToDecibels(OutputVolumeModulation.Value);
-		}
+void USoundSubmix::PostLoad()
+{
+	Super::PostLoad();
 
-		if (WetLevelModulation.Value > 0.0f)
-		{
-			WetLevelModulation.Value = Audio::ConvertToDecibels(WetLevelModulation.Value);
-		}
+#if WITH_EDITORONLY_DATA
 
-		if (DryLevelModulation.Value > 0.0f)
-		{
-			DryLevelModulation.Value = Audio::ConvertToDecibels(DryLevelModulation.Value);
-		}
+	const int32 Version = GetLinkerCustomVersion(FSoundSubmixCustomVersion::GUID);
+	HandleVersionMigration(Version);
 
-		OutputVolumeModulation.VersionModulators();
-		WetLevelModulation.VersionModulators();
-		DryLevelModulation.VersionModulators();
-	}
 #endif // WITH_EDITORONLY_DATA
+
 }
 
 UEndpointSubmix::UEndpointSubmix(const FObjectInitializer& ObjectInitializer)
@@ -1487,3 +1488,90 @@ ENGINE_API void SubmixUtils::RefreshEditorForSubmix(const USoundSubmixBase* InSu
 
 #endif // WITH_EDITOR
 
+// Versioning and Deprecated Property Migration.
+// --------------------------------------------
+
+#if WITH_EDITORONLY_DATA
+
+namespace SoundSubmixMigration
+{
+	// Old defaults.
+	static const float OldDefault_OutputVolume(-1.0f);
+	static const float OldDefault_WetLevel(-1.0f);
+	static const float OldDefault_DryLevel(-1.0f);
+}
+
+void USoundSubmix::InitDeprecatedDefaults()
+{
+	using namespace SoundSubmixMigration;
+	
+	// We must init these to their old defaults prior to serialization to test if they are in fact serialized.
+	OutputVolume_DEPRECATED = OldDefault_OutputVolume;
+	DryLevel_DEPRECATED = OldDefault_DryLevel;
+	WetLevel_DEPRECATED = OldDefault_WetLevel;
+}
+
+void USoundSubmix::HandleVersionMigration(const int32 Version)
+{
+	if (Version < FSoundSubmixCustomVersion::MigrateModulatedSendProperties)
+	{
+		using namespace SoundSubmixPrivate;
+		using namespace SoundSubmixMigration;
+
+		auto ConvertToModulatedDb = [this](const float InValue, const float InDefault, const float InDefaultModulationValue, FSoundModulationDestinationSettings& OutModulator, const TCHAR* InParamName) 
+		{
+			// IF after load this old property has non-default value.
+			// AND the newer form is still at a default value.
+			// THEN we can safely convert the value over.
+
+			if (!FMath::IsNearlyEqual(InValue, InDefault) &&
+				FMath::IsNearlyEqual(OutModulator.Value, InDefaultModulationValue))
+			{
+				// use -96dB as a noise floor when fixing up linear volume settings
+				static constexpr float LinearNeg96dB = 0.0000158489319f;
+
+				if (InValue <= LinearNeg96dB)
+				{
+					OutModulator.Value = -96.f;
+				}
+				else
+				{
+					OutModulator.Value = Audio::ConvertToDecibels(InValue);
+				}
+
+				UE_LOG(LogAudio, Display, TEXT("SoundSubmix::HandleVersionMigration, ConvertToModulatedDb, Asset = %s, %s = %2.2f dB from %2.2f"), *GetName(), InParamName, OutModulator.Value, InValue);
+			}
+		};
+
+		// Convert.
+		ConvertToModulatedDb(OutputVolume_DEPRECATED, OldDefault_OutputVolume, Default_OutputVolumeModulation, OutputVolumeModulation, TEXT("OutputVoluime"));
+		ConvertToModulatedDb(WetLevel_DEPRECATED, OldDefault_WetLevel, Default_WetLevelModulation, WetLevelModulation, TEXT("WetLevel"));
+		ConvertToModulatedDb(DryLevel_DEPRECATED, OldDefault_DryLevel, Default_DryLevelModulation, DryLevelModulation, TEXT("DryLevel"));
+	}
+
+	// Convert linear modulators to dB. (this has most likely happened as part of above update, but just in case we have some outliers, handle this separately).
+	if (Version < FSoundSubmixCustomVersion::ConvertLinearModulatorsToDb)
+	{
+		auto ConvertToDb = [this](FSoundModulationDestinationSettings& OutValue, const TCHAR* InName) 
+		{
+			// Assume anything > 0.f is in a linear scale and convert, otherwise ignore it.
+			if (OutValue.Value > 0.0f)
+			{
+				const float dbValue = Audio::ConvertToDecibels(OutValue.Value);
+				UE_LOG(LogAudio, Display, TEXT("SoundSubmix::HandleVersionMigration, ConvertToDb, Asset = %s, %s = %2.2f dB from %2.2f"), *GetName(), InName, dbValue, OutValue.Value)
+				OutValue.Value = dbValue;
+			}
+		};
+
+		// Convert.
+		ConvertToDb(OutputVolumeModulation, TEXT("OutputVolume"));
+		ConvertToDb(WetLevelModulation, TEXT("WetLevel"));
+		ConvertToDb(DryLevelModulation, TEXT("DryLevel"));
+
+		OutputVolumeModulation.VersionModulators();
+		WetLevelModulation.VersionModulators();
+		DryLevelModulation.VersionModulators();
+	}
+}
+
+#endif //WITH_EDITORONLY_DATA
diff --git a/Engine/Source/Runtime/Engine/Private/SoundSubmixCustomVersion.cpp b/Engine/Source/Runtime/Engine/Private/SoundSubmixCustomVersion.cpp
new file mode 100644
--- /dev/null
+++ b/Engine/Source/Runtime/Engine/Private/SoundSubmixCustomVersion.cpp
@@ -0,0 +1,9 @@
+// Copyright Epic Games, Inc. All Rights Reserved.
+
+#include "SoundSubmixCustomVersion.h"
+#include "Serialization/CustomVersion.h"
+
+const FGuid FSoundSubmixCustomVersion::GUID(0x79D25301, 0x47AF41F8, 0xADCF10E7, 0x97415C16);
+
+// Register the custom version with core
+FCustomVersionRegistration GSoundSubmixCustomVersion(FSoundSubmixCustomVersion::GUID, FSoundSubmixCustomVersion::LatestVersion, TEXT("SoundSubmixVer"));
diff --git a/Engine/Source/Runtime/Engine/Private/SoundSubmixCustomVersion.h b/Engine/Source/Runtime/Engine/Private/SoundSubmixCustomVersion.h
new file mode 100644
--- /dev/null
+++ b/Engine/Source/Runtime/Engine/Private/SoundSubmixCustomVersion.h
@@ -0,0 +1,31 @@
+// Copyright Epic Games, Inc. All Rights Reserved.
+
+#pragma once
+
+#include "Misc/Guid.h"
+
+// Custom serialization version for SoundSubmixes.
+struct ENGINE_API FSoundSubmixCustomVersion
+{
+	enum Type
+	{
+		// Before any version changes were made in the plugin
+		BeforeCustomVersionWasAdded = 0,
+			
+		// Migrated deprecated properties OutputVolume, WetLevel, DryLevel 
+		MigrateModulatedSendProperties, 
+
+		// Convert modulated properties to dB.
+		ConvertLinearModulatorsToDb,
+
+		// -----<new versions can be added above this line>-------------------------------------------------
+		VersionPlusOne,
+		LatestVersion = VersionPlusOne - 1
+	};
+
+	// The GUID for this custom version number
+	const static FGuid GUID;
+
+private:
+	FSoundSubmixCustomVersion() {}
+};
