Fix incorrect color override streams for manual vertex fetch that occur because the spline mesh did not bind the correct vertex factory uniform buffer.

diff --git a/Engine/Source/Runtime/Engine/Private/Components/SplineMeshComponent.cpp b/Engine/Source/Runtime/Engine/Private/Components/SplineMeshComponent.cpp
--- a/Engine/Source/Runtime/Engine/Private/Components/SplineMeshComponent.cpp
+++ b/Engine/Source/Runtime/Engine/Private/Components/SplineMeshComponent.cpp
@@ -215,7 +215,9 @@ void FSplineMeshVertexFactoryShaderParameters::GetElementShaderBindings(
 	}
 	if (LocalVertexFactory->SupportsManualVertexFetch(FeatureLevel) || bUseGPUScene)
 	{
-		ShaderBindings.Add(Shader->GetUniformBufferParameter<FLocalVertexFactoryUniformShaderParameters>(), LocalVertexFactory->GetUniformBuffer());
+		FRHIUniformBuffer* VertexFactoryUniformBuffer = static_cast<FRHIUniformBuffer*>(BatchElement.VertexFactoryUserData);
+
+		ShaderBindings.Add(Shader->GetUniformBufferParameter<FLocalVertexFactoryUniformShaderParameters>(), VertexFactoryUniformBuffer != nullptr ? VertexFactoryUniformBuffer : LocalVertexFactory->GetUniformBuffer());
 	}
 
 	// If we can't use GPU Scene instance data, we have to bind the params to the VS loosely
