Added new class UDeletedObjectPlaceholder to properly handle when an external package contains no more object.
- Ground work for full support of moving an actor from one External Data Layer to another
- Fixes the Save Package dialog not showing the last valid actor's display name (instead was showing the package name)
- Fixes empty packages not properly cleaned up by ULevel::CleanupLevel
- Changing level's bUseExternalActors from true to false and reloading the level without saving now properly works

diff --git a/Engine/Source/Developer/AssetTools/Private/AssetTools.cpp b/Engine/Source/Developer/AssetTools/Private/AssetTools.cpp
--- a/Engine/Source/Developer/AssetTools/Private/AssetTools.cpp
+++ b/Engine/Source/Developer/AssetTools/Private/AssetTools.cpp
@@ -42,6 +42,7 @@
 #include "AssetTypeActions/AssetTypeActions_ForceFeedbackEffect.h"
 #include "AssetTypeActions/AssetTypeActions_ParticleSystem.h"
 #include "AssetTypeActions/AssetTypeActions_PhysicalMaterialMask.h"
+#include "AssetTypeActions/AssetTypeActions_DeletedObjectPlaceholder.h"
 #include "WorldPartition/WorldPartition.h"
 #include "SDiscoveringAssetsDialog.h"
 #include "AssetFixUpRedirectors.h"
@@ -1384,6 +1385,7 @@ UAssetToolsImpl::UAssetToolsImpl(const FObjectInitializer& ObjectInitializer)
 	RegisterAssetTypeActions(MakeShareable(new FAssetTypeActions_ForceFeedbackEffect(InputCategoryBit)));
 	RegisterAssetTypeActions(MakeShareable(new FAssetTypeActions_ParticleSystem));
 	RegisterAssetTypeActions(MakeShareable(new FAssetTypeActions_PhysicalMaterialMask));
+	RegisterAssetTypeActions(MakeShareable(new FAssetTypeActions_DeletedObjectPlaceholder));
 
 	// Note: Please don't add any more actions here!  They belong in an editor-only module that is more tightly
 	// coupled to your new system, and you should not create a dependency on your new system from AssetTools.
diff --git a/Engine/Source/Developer/AssetTools/Private/AssetTypeActions/AssetTypeActions_DeletedObjectPlaceholder.cpp b/Engine/Source/Developer/AssetTools/Private/AssetTypeActions/AssetTypeActions_DeletedObjectPlaceholder.cpp
new file mode 100644
--- /dev/null
+++ b/Engine/Source/Developer/AssetTools/Private/AssetTypeActions/AssetTypeActions_DeletedObjectPlaceholder.cpp
@@ -0,0 +1,16 @@
+// Copyright Epic Games, Inc. All Rights Reserved.
+
+#include "AssetTypeActions_DeletedObjectPlaceholder.h"
+
+#if WITH_EDITOR
+
+FString FAssetTypeActions_DeletedObjectPlaceholder::GetObjectDisplayName(UObject* InObject) const
+{
+	if (UDeletedObjectPlaceholder* Object = Cast<UDeletedObjectPlaceholder>(InObject))
+	{
+		return Object->GetDisplayName();
+	}
+	return FAssetTypeActions_Base::GetObjectDisplayName(InObject);
+}
+
+#endif
\ No newline at end of file
diff --git a/Engine/Source/Developer/AssetTools/Private/AssetTypeActions/AssetTypeActions_DeletedObjectPlaceholder.h b/Engine/Source/Developer/AssetTools/Private/AssetTypeActions/AssetTypeActions_DeletedObjectPlaceholder.h
new file mode 100644
--- /dev/null
+++ b/Engine/Source/Developer/AssetTools/Private/AssetTypeActions/AssetTypeActions_DeletedObjectPlaceholder.h
@@ -0,0 +1,22 @@
+// Copyright Epic Games, Inc. All Rights Reserved.
+
+#pragma once
+
+#if WITH_EDITOR
+#include "AssetTypeActions_Base.h"
+#include "AssetTypeCategories.h"
+#include "DeletedObjectPlaceholder.h"
+
+class FAssetTypeActions_DeletedObjectPlaceholder : public FAssetTypeActions_Base
+{
+public:
+	// IAssetTypeActions interface
+	virtual FText GetName() const override { return NSLOCTEXT("AssetTypeActions", "AssetTypeActions_DeletedObjectPlaceholder", "Deleted Object Placeholder"); }
+	virtual FColor GetTypeColor() const override { return FColor::Red; }
+	virtual UClass* GetSupportedClass() const override { return UDeletedObjectPlaceholder::StaticClass(); }
+	virtual uint32 GetCategories() override { return EAssetTypeCategories::Misc; }
+	virtual bool CanLocalize() const override { return false; }
+	virtual FString GetObjectDisplayName(UObject* InObject) const override;
+	// End of IAssetTypeActions interface
+};
+#endif
\ No newline at end of file
diff --git a/Engine/Source/Runtime/Engine/Private/ActorEditor.cpp b/Engine/Source/Runtime/Engine/Private/ActorEditor.cpp
--- a/Engine/Source/Runtime/Engine/Private/ActorEditor.cpp
+++ b/Engine/Source/Runtime/Engine/Private/ActorEditor.cpp
@@ -41,6 +41,7 @@
 #include "Algo/Transform.h"
 #include "Algo/RemoveIf.h"
 #include "Misc/DataValidation.h"
+#include "DeletedObjectPlaceholder.h"
 
 #define LOCTEXT_NAMESPACE "ErrorChecking"
 
@@ -1082,18 +1083,29 @@ void AActor::SetPackageExternal(bool bExternal, bool bShouldDirty, UPackage* Act
     // Mark the current actor & package as dirty
 	Modify(bShouldDirty);
 
-	UPackage* LevelPackage = GetLevel()->GetPackage(); 
 	if (bExternal)
 	{
-		UPackage* NewActorPackage = ActorExternalPackage ? ActorExternalPackage : ULevel::CreateActorPackage(LevelPackage, GetLevel()->GetActorPackagingScheme(), GetPathName(), this);
-		SetExternalPackage(NewActorPackage);
+		UPackage* LevelPackage = GetLevel()->GetPackage();
+		UPackage* ActorPackage = ActorExternalPackage ? ActorExternalPackage : ULevel::CreateActorPackage(LevelPackage, GetLevel()->GetActorPackagingScheme(), GetPathName(), this);
+		// Cleanup the package from its UDeletedObjectPlaceholder if any (CreateActorPackage can return an existing package)
+		UDeletedObjectPlaceholder::RemoveFromPackage(ActorPackage);
+		SetExternalPackage(ActorPackage);
 	}
 	else
 	{
+		const bool bWasMainActorInPackage = IsMainPackageActor();
 		UPackage* ActorPackage = GetExternalPackage();
 		// Detach the linker exports so it doesn't resolve to this actor anymore
 		ResetLinkerExports(ActorPackage);
 		SetExternalPackage(nullptr);
+
+		// If the old external package is empty, create a dummy object outered to the level part of the original external package
+		// so that ULevel::CleanupLevel will be able to visit & process this empty external package.
+		// This object will allow the save package dialog to display the label of the previously removed actor (see FAssetTypeActions_DeletedObjectPlaceholder::GetObjectDisplayName)
+		if (!ActorPackage->HasAnyPackageFlags(PKG_NewlyCreated) && bWasMainActorInPackage && !AActor::FindActorInPackage(ActorPackage))
+		{
+			UDeletedObjectPlaceholder::Create(GetLevel(), ActorPackage, this);
+		}
 	}
 
 	for (UActorComponent* ActorComponent : GetComponents())
diff --git a/Engine/Source/Runtime/Engine/Private/DeletedObjectPlaceholder.cpp b/Engine/Source/Runtime/Engine/Private/DeletedObjectPlaceholder.cpp
new file mode 100644
--- /dev/null
+++ b/Engine/Source/Runtime/Engine/Private/DeletedObjectPlaceholder.cpp
@@ -0,0 +1,85 @@
+// Copyright Epic Games, Inc. All Rights Reserved.
+
+#include "DeletedObjectPlaceholder.h"
+
+#if WITH_EDITOR
+
+#include "Engine/Level.h"
+#include "WorldPartition/ActorDescContainer.h"
+#include "WorldPartition/WorldPartition.h"
+#include "WorldPartition/DataLayer/ExternalDataLayerAsset.h"
+
+UDeletedObjectPlaceholder::FObjectCreated UDeletedObjectPlaceholder::OnObjectCreated;
+
+bool UDeletedObjectPlaceholder::IsAsset() const
+{
+	return IsPackageExternal() && !GetPackage()->HasAnyFlags(RF_Transient) && !HasAnyFlags(RF_ClassDefaultObject) && !GetPackage()->HasAnyPackageFlags(PKG_PlayInEditor) && !AActor::FindActorInPackage(GetExternalPackage());
+}
+
+UDeletedObjectPlaceholder* UDeletedObjectPlaceholder::Create(ULevel* InLevel, UPackage* InPackage, const UObject* InOriginalObject)
+{
+	if (IsRunningCookCommandlet())
+	{
+		return nullptr;
+	}
+	check(InLevel);
+	check(InOriginalObject);
+	check(!UDeletedObjectPlaceholder::FindInPackage(InPackage));
+	const AActor* OriginalActor = Cast<AActor>(InOriginalObject);
+	UDeletedObjectPlaceholder* DeletedObjectPlaceholder = NewObject<UDeletedObjectPlaceholder>(InLevel, NAME_None, RF_Standalone | RF_Transactional | RF_Transient);
+	DeletedObjectPlaceholder->SetExternalPackage(InPackage);
+	DeletedObjectPlaceholder->OriginalObject = InOriginalObject;
+	DeletedObjectPlaceholder->DisplayName = OriginalActor ? OriginalActor->GetActorLabel() : InOriginalObject->GetName();
+	DeletedObjectPlaceholder->ExternalDataLayerUID = OriginalActor && OriginalActor->GetExternalDataLayerAsset() ? OriginalActor->GetExternalDataLayerAsset()->GetUID() : FExternalDataLayerUID();
+
+	OnObjectCreated.Broadcast(DeletedObjectPlaceholder);
+
+	return DeletedObjectPlaceholder;
+}
+
+UDeletedObjectPlaceholder* UDeletedObjectPlaceholder::RemoveFromPackage(UPackage* InPackage)
+{
+	if (IsRunningCookCommandlet())
+	{
+		return nullptr;
+	}
+
+	if (UDeletedObjectPlaceholder* DeletedObjectPlaceholder = InPackage ? UDeletedObjectPlaceholder::FindInPackage(InPackage) : nullptr)
+	{
+		DeletedObjectPlaceholder->Modify(false);
+		DeletedObjectPlaceholder->ClearFlags(RF_Standalone);
+		const FName NewName = MakeUniqueObjectName(nullptr, UDeletedObjectPlaceholder::StaticClass(), FName(*FString::Printf(TEXT("%s_Trashed"), *DeletedObjectPlaceholder->GetName())));
+		DeletedObjectPlaceholder->Rename(*NewName.ToString(), GetTransientPackage(), REN_ForceNoResetLoaders | REN_DontCreateRedirectors | REN_NonTransactional | REN_DoNotDirty);
+		DeletedObjectPlaceholder->SetExternalPackage(nullptr);
+		return DeletedObjectPlaceholder;
+	}
+	return nullptr;
+}
+
+UDeletedObjectPlaceholder* UDeletedObjectPlaceholder::FindInPackage(UPackage* InPackage)
+{
+	if (IsRunningCookCommandlet())
+	{
+		return nullptr;
+	}
+
+	UDeletedObjectPlaceholder* DeletedObjectPlaceholder = nullptr;
+	ForEachObjectWithPackage(InPackage, [&DeletedObjectPlaceholder](UObject* Object)
+	{
+		DeletedObjectPlaceholder = Cast<UDeletedObjectPlaceholder>(Object);
+		return !DeletedObjectPlaceholder;
+	}, false);
+	return DeletedObjectPlaceholder;
+}
+
+void UDeletedObjectPlaceholder::PostEditUndo()
+{
+	Super::PostEditUndo();
+
+	if (OriginalObject.IsValid() && (GetPackage() != GetTransientPackage()))
+	{
+		OnObjectCreated.Broadcast(this);
+	}
+}
+
+#endif
\ No newline at end of file
diff --git a/Engine/Source/Runtime/Engine/Private/WorldPartition/ActorDescContainer.cpp b/Engine/Source/Runtime/Engine/Private/WorldPartition/ActorDescContainer.cpp
--- a/Engine/Source/Runtime/Engine/Private/WorldPartition/ActorDescContainer.cpp
+++ b/Engine/Source/Runtime/Engine/Private/WorldPartition/ActorDescContainer.cpp
@@ -18,8 +18,26 @@
 #include "WorldPartition/WorldPartitionClassDescRegistry.h"
 #include "WorldPartition/DataLayer/DataLayerManager.h"
 #include "WorldPartition/DataLayer/ExternalDataLayerAsset.h"
+#include "WorldPartition/ActorDescContainerSubsystem.h"
+#include "DeletedObjectPlaceholder.h"
 
 UActorDescContainer::FActorDescContainerInitializeDelegate UActorDescContainer::OnActorDescContainerInitialized;
+
+FUObjectAnnotationSparse<FDeletedObjectPlaceholderAnnotation, true> UActorDescContainer::DeletedObjectPlaceholdersAnnotation;
+
+FDeletedObjectPlaceholderAnnotation::FDeletedObjectPlaceholderAnnotation(const UDeletedObjectPlaceholder* InDeletedObjectPlaceholder, const FString& InActorDescContainerName)
+	: DeletedObjectPlaceholder(InDeletedObjectPlaceholder)
+	, ActorDescContainerName(InActorDescContainerName)
+{
+}
+
+UActorDescContainer* FDeletedObjectPlaceholderAnnotation::GetActorDescContainer() const
+{
+	UActorDescContainerSubsystem* ActorDescContainerSubsystem = UActorDescContainerSubsystem::Get();
+	UActorDescContainer* ActorDescContainer = ActorDescContainerSubsystem ? ActorDescContainerSubsystem->GetActorDescContainer(ActorDescContainerName) : nullptr;
+	return ActorDescContainer;
+}
+
 #endif
 
 UActorDescContainer::UActorDescContainer(const FObjectInitializer& ObjectInitializer)
@@ -291,6 +309,31 @@ const FWorldPartitionActorDesc* UActorDescContainer::GetActorDescByName(FName Ac
 	return nullptr;
 }
 
+bool UActorDescContainer::ShouldHandleDeletedObjectPlaceholderEvent(const UDeletedObjectPlaceholder* InDeletedObjectPlaceholder) const
+{
+	const FExternalDataLayerUID ContainerExternalDataLayerUID = ExternalDataLayerAsset ? ExternalDataLayerAsset->GetUID() : FExternalDataLayerUID();
+	if (ContainerExternalDataLayerUID == InDeletedObjectPlaceholder->GetExternalDataLayerUID())
+	{
+		const FString PackageName = InDeletedObjectPlaceholder->GetPackage()->GetName();
+		const FString ContainerExternalActorPath = GetExternalActorPath() / TEXT("");
+		return PackageName.StartsWith(ContainerExternalActorPath);
+	}
+	return false;
+};
+
+void UActorDescContainer::OnDeletedObjectPlaceholderCreated(const UDeletedObjectPlaceholder* InDeletedObjectPlaceholder)
+{
+	const UObject* OriginalObject = InDeletedObjectPlaceholder->GetOriginalObject();
+	if (const AActor* Actor = Cast<AActor>(OriginalObject))
+	{
+		if (ShouldHandleDeletedObjectPlaceholderEvent(InDeletedObjectPlaceholder))
+		{
+			check(GetActorDescriptor(Actor->GetActorGuid()));
+			DeletedObjectPlaceholdersAnnotation.AddAnnotation(Actor, FDeletedObjectPlaceholderAnnotation(InDeletedObjectPlaceholder, GetContainerName()));
+		}
+	}
+}
+
 void UActorDescContainer::OnObjectPreSave(UObject* Object, FObjectPreSaveContext SaveContext)
 {
 	if (!SaveContext.IsProceduralSave() && !(SaveContext.GetSaveFlags() & SAVE_FromAutosave))
@@ -300,6 +343,23 @@ void UActorDescContainer::OnObjectPreSave(UObject* Object, FObjectPreSaveContext
 			if (ShouldHandleActorEvent(Actor))
 			{
 				check(IsValidChecked(Actor));
+
+				// Handle the case where the actor changed package but the old/empty package has not been processed/deleted
+				// One case where this can happen is if the user choses to save the new package but unchecks the deleted package
+				// Remove(unhash) the corresponding original actor (guid) from its original container before adding the new one (unhash before hashing)
+				if (FDeletedObjectPlaceholderAnnotation Annotation = DeletedObjectPlaceholdersAnnotation.GetAndRemoveAnnotation(Actor); Annotation.IsValid())
+				{
+					// In the case where the object changed to a new container and created a new package, then changed back to its original location,
+					// OnDeletedObjectPlaceholderCreated will not be called for the newly created package
+					// This is why we need to validate that the annotation's container is still relevant by the annotation's DeletedObjectPlaceholder using ShouldHandleDeletedObjectPlaceholderEvent
+					UActorDescContainer* ActorDescContainer = Annotation.GetActorDescContainer();
+					if (ActorDescContainer && ActorDescContainer->ShouldHandleDeletedObjectPlaceholderEvent(Annotation.GetDeletedObjectPlaceholder()))
+					{
+						check(Annotation.GetDeletedObjectPlaceholder()->GetOriginalObject() == Actor);
+						verify(ActorDescContainer->RemoveActor(Actor->GetActorGuid()));
+					}
+				}
+				
 				if (TUniquePtr<FWorldPartitionActorDesc>* ExistingActorDesc = GetActorDescriptor(Actor->GetActorGuid()))
 				{
 					// Existing actor
@@ -321,11 +381,26 @@ void UActorDescContainer::OnObjectPreSave(UObject* Object, FObjectPreSaveContext
 
 void UActorDescContainer::OnPackageDeleted(UPackage* Package)
 {
-	AActor* Actor = AActor::FindActorInPackage(Package);
-
-	if (ShouldHandleActorEvent(Actor))
+	if (const AActor* Actor = AActor::FindActorInPackage(Package))
 	{
-		RemoveActor(Actor->GetActorGuid());
+		if (ShouldHandleActorEvent(Actor))
+		{
+			RemoveActor(Actor->GetActorGuid());
+		}
+	}
+	else if (UDeletedObjectPlaceholder* DeletedObjectPlaceholder = UDeletedObjectPlaceholder::FindInPackage(Package))
+	{
+		if (ShouldHandleDeletedObjectPlaceholderEvent(DeletedObjectPlaceholder))
+		{
+			// Here we validate that we didn't already processed the DeletedObjectPlaceholder in OnObjectPreSave
+			const AActor* OriginalActor = CastChecked<AActor>(DeletedObjectPlaceholder->GetOriginalObject());
+			if (FDeletedObjectPlaceholderAnnotation Annotation = DeletedObjectPlaceholdersAnnotation.GetAndRemoveAnnotation(OriginalActor); Annotation.IsValid())
+			{
+				check(Annotation.GetDeletedObjectPlaceholder() == DeletedObjectPlaceholder);
+				check(Annotation.GetActorDescContainer() == this);
+				verify(RemoveActor(OriginalActor->GetActorGuid()));
+			}
+		}
 	}
 }
 
@@ -405,6 +480,8 @@ void UActorDescContainer::RegisterEditorDelegates()
 
 		FWorldPartitionClassDescRegistry& ClassDescRegistry = FWorldPartitionClassDescRegistry::Get();
 		ClassDescRegistry.OnClassDescriptorUpdated().AddUObject(this, &UActorDescContainer::OnClassDescriptorUpdated);
+
+		UDeletedObjectPlaceholder::OnObjectCreated.AddUObject(this, &UActorDescContainer::OnDeletedObjectPlaceholderCreated);
 	}
 }
 
@@ -417,6 +494,8 @@ void UActorDescContainer::UnregisterEditorDelegates()
 
 		FWorldPartitionClassDescRegistry& ClassDescRegistry = FWorldPartitionClassDescRegistry::Get();
 		ClassDescRegistry.OnClassDescriptorUpdated().RemoveAll(this);
+
+		UDeletedObjectPlaceholder::OnObjectCreated.RemoveAll(this);
 	}
 }
 
diff --git a/Engine/Source/Runtime/Engine/Public/DeletedObjectPlaceholder.h b/Engine/Source/Runtime/Engine/Public/DeletedObjectPlaceholder.h
new file mode 100644
--- /dev/null
+++ b/Engine/Source/Runtime/Engine/Public/DeletedObjectPlaceholder.h
@@ -0,0 +1,55 @@
+// Copyright Epic Games, Inc. All Rights Reserved.
+
+#pragma once
+
+#include "CoreFwd.h"
+#include "UObject/ObjectMacros.h"
+
+#if WITH_EDITOR
+#include "WorldPartition/DataLayer/ExternalDataLayerUID.h"
+#endif
+
+#include "DeletedObjectPlaceholder.generated.h"
+
+class UActorDescContainer;
+class UDeletedObjectPlaceholder;
+
+UCLASS(Transient, MinimalAPI)
+class UDeletedObjectPlaceholder : public UObject
+{
+	GENERATED_BODY()
+
+public:
+
+#if WITH_EDITOR
+	virtual bool IsAsset() const override;
+	virtual void PostEditUndo() override;
+	
+	const UObject* GetOriginalObject() const { FSoftObjectPtr ObjPtr(OriginalObject); return ObjPtr.Get(); }
+	const FString& GetDisplayName() const { return DisplayName; }
+	const FExternalDataLayerUID GetExternalDataLayerUID() const { return FExternalDataLayerUID(ExternalDataLayerUID); }
+
+	static UDeletedObjectPlaceholder* Create(ULevel* InLevel, UPackage* InPackage, const UObject* InOriginalObject);
+	static UDeletedObjectPlaceholder* FindInPackage(UPackage* InPackage);
+	static UDeletedObjectPlaceholder* RemoveFromPackage(UPackage* InPackage);
+
+	DECLARE_MULTICAST_DELEGATE_OneParam(FObjectCreated, const UDeletedObjectPlaceholder*);
+	static ENGINE_API FObjectCreated OnObjectCreated;
+#endif
+
+private:
+
+#if WITH_EDITORONLY_DATA
+	/** Original object's display name */
+	UPROPERTY()
+	FString DisplayName;
+
+	/** Original object's external data layer UID */
+	UPROPERTY()
+	uint32 ExternalDataLayerUID;
+
+	/** Original object that is replaced by the placeholder */
+	UPROPERTY()
+	FSoftObjectPath OriginalObject;
+#endif
+};
\ No newline at end of file
diff --git a/Engine/Source/Runtime/Engine/Public/WorldPartition/ActorDescContainer.h b/Engine/Source/Runtime/Engine/Public/WorldPartition/ActorDescContainer.h
--- a/Engine/Source/Runtime/Engine/Public/WorldPartition/ActorDescContainer.h
+++ b/Engine/Source/Runtime/Engine/Public/WorldPartition/ActorDescContainer.h
@@ -4,6 +4,7 @@
 #include "CoreMinimal.h"
 #include "UObject/ObjectMacros.h"
 #include "UObject/Object.h"
+#include "UObject/UObjectAnnotation.h"
 #include "WorldPartition/ActorDescList.h"
 #include "WorldPartition/WorldPartitionHandle.h"
 #include "WorldPartition/WorldPartitionActorDesc.h"
@@ -11,8 +12,28 @@
 #include "ActorDescContainer.generated.h"
 
 class FLinkerInstancingContext;
+class UDeletedObjectPlaceholder;
 class UWorldPartition;
 
+#if WITH_EDITOR
+struct FDeletedObjectPlaceholderAnnotation
+{
+public:
+	FDeletedObjectPlaceholderAnnotation(const UDeletedObjectPlaceholder* InDeletedObjectPlaceholder = nullptr, const FString& InActorDescContainerName = FString());
+	bool IsDefault() const { return DeletedObjectPlaceholder.IsExplicitlyNull() && ActorDescContainerName.IsEmpty(); }
+	bool IsValid() const { return DeletedObjectPlaceholder.IsValid() && !ActorDescContainerName.IsEmpty(); }
+	const UDeletedObjectPlaceholder* GetDeletedObjectPlaceholder() const { return DeletedObjectPlaceholder.Get(); }
+	UActorDescContainer* GetActorDescContainer() const;
+
+private:
+	TWeakObjectPtr<const UDeletedObjectPlaceholder> DeletedObjectPlaceholder;
+	// We store the container name instead of keeping a WeakObjectPtr to properly handle the case where the container 
+	// is unregistered/re-registered between usage of annotation (this can happen if a plugin is unregistered/re-registered).
+	FString ActorDescContainerName;
+};
+#endif
+
+
 UCLASS(MinimalAPI)
 class UActorDescContainer : public UObject, public FActorDescList
 {
@@ -174,9 +195,14 @@ private:
 	// GetWorld() should never be called on an ActorDescContainer to avoid any confusion as it can be used as a template
 	UWorld* GetWorld() const override { return nullptr; }
 
+	bool ShouldHandleDeletedObjectPlaceholderEvent(const UDeletedObjectPlaceholder* InDeletedObjectPlaceholder) const;
+	void OnDeletedObjectPlaceholderCreated(const UDeletedObjectPlaceholder* InDeletedObjectPlaceholder);
+
 	ENGINE_API void RegisterEditorDelegates();
 	ENGINE_API void UnregisterEditorDelegates();
 
+	static FUObjectAnnotationSparse<FDeletedObjectPlaceholderAnnotation, true> DeletedObjectPlaceholdersAnnotation;
+
 protected:
 	TObjectPtr<const UExternalDataLayerAsset> ExternalDataLayerAsset;
 #endif
diff --git a/Engine/Source/Runtime/Engine/Public/WorldPartition/ActorDescContainerSubsystem.h b/Engine/Source/Runtime/Engine/Public/WorldPartition/ActorDescContainerSubsystem.h
--- a/Engine/Source/Runtime/Engine/Public/WorldPartition/ActorDescContainerSubsystem.h
+++ b/Engine/Source/Runtime/Engine/Public/WorldPartition/ActorDescContainerSubsystem.h
@@ -41,6 +41,8 @@ public:
 	
 	template <class ContainerType = UActorDescContainer>
 	ContainerType* RegisterContainer(const typename UActorDescContainer::FInitializeParams& InitParams) { return ContainerManager.RegisterContainer<ContainerType>(InitParams); }
+	UActorDescContainer* GetActorDescContainer(const FString& InContainerName) { return ContainerManager.GetActorDescContainer(InContainerName); }
+	const UActorDescContainer* GetActorDescContainer(const FString& InContainerName) const { return ContainerManager.GetActorDescContainer(InContainerName); }
 	void RegisterContainer(UActorDescContainer* Container) { ContainerManager.RegisterContainer(Container); }
 	void UnregisterContainer(UActorDescContainer* Container) { ContainerManager.UnregisterContainer(Container); }
 	FBox GetContainerBounds(const FString& ContainerName) const { return ContainerManager.GetContainerBounds(ContainerName); }
@@ -129,6 +131,17 @@ private:
 			FRegisteredContainer& RegisteredContainer = RegisteredContainers.FindChecked(Container->GetContainerName());
 			RegisteredContainer.RefCount++;
 		}
+
+		UActorDescContainer* GetActorDescContainer(const FString& InContainerName)
+		{ 
+			FRegisteredContainer* RegisteredContainer = RegisteredContainers.Find(InContainerName);
+			return RegisteredContainer ? RegisteredContainer->Container : nullptr;
+		}
+
+		const UActorDescContainer* GetActorDescContainer(const FString& InContainerName) const 
+		{
+			return const_cast<FContainerManager*>(this)->GetActorDescContainer(InContainerName);
+		}
 				
 		void UnregisterContainer(UActorDescContainer* Container);
 
diff --git a/Engine/Source/Runtime/Engine/Public/WorldPartition/DataLayer/ExternalDataLayerUID.h b/Engine/Source/Runtime/Engine/Public/WorldPartition/DataLayer/ExternalDataLayerUID.h
--- a/Engine/Source/Runtime/Engine/Public/WorldPartition/DataLayer/ExternalDataLayerUID.h
+++ b/Engine/Source/Runtime/Engine/Public/WorldPartition/DataLayer/ExternalDataLayerUID.h
@@ -17,6 +17,7 @@ public:
 
 	FString ToString() const;
 #if WITH_EDITOR
+	FExternalDataLayerUID(uint32 InValue = 0) { Value = InValue; }
 	static bool Parse(const FString& InUIDString, FExternalDataLayerUID& OutUID);
 #endif
 
