- Extends ULevelStreaming to have a protected member AsyncRequestIDs which is used to keep track of async loads required for loading a level.
- During OnLoadingFinished we clear AsyncRequestIDs
- Modifies the `ULevelStreaming::RequestLevel` implementation of  to track its async load requests in AsyncRequestIDs
- Modifies `UWorldPartitionLevelStreamingDynamic::RequestLevel` to track which async load requests it makes in AsyncRequestIDs
- Changed `UWorld::BlockTillLevelStreamingCompleted` implementation to no longer flush all async loads globally and instead only flush outstanding streaming level async requests
-- If a child class of `ULevelStreaming` has not recorded any async loads in AsyncRequestIDs, we fallback to flushing all async loads as before since we can't know if implementers are relying on the past behaviour of a forced flush of all async loads. Ideally this fallback behaviour is removed in UE6.0
- Adds `s.World.ForceFlushAllAsyncLoadsDuringLevelStreaming` allowing licensees to revert back to old flushing behaviour should they rely on async loads they don't control.

PIE.Startup 3m44s -> 3m19s

diff --git a/Engine/Source/Runtime/Engine/Classes/Engine/LevelStreaming.h b/Engine/Source/Runtime/Engine/Classes/Engine/LevelStreaming.h
--- a/Engine/Source/Runtime/Engine/Classes/Engine/LevelStreaming.h
+++ b/Engine/Source/Runtime/Engine/Classes/Engine/LevelStreaming.h
@@ -267,6 +267,9 @@ protected:
 	/** Whether the streaming level can safely skip making visible transaction request from the client to the server */
 	uint8 bSkipClientUseMakingVisibleTransactionRequest:1;
 
+	/** IDs used to track the underlying async load requests for the level */
+	TArray<int32> AsyncRequestIDs;
+
 private:
 	/** What the current streamed state of the streaming level is */
 	ELevelStreamingState CurrentState;
@@ -455,6 +458,12 @@ public:
 		return (LoadedLevel || PendingUnloadLevel);
 	}
 
+	/** Returns the RequestIDs for the async load requests required to load the level. These IDs are cleared during OnLoadingFinished */
+	const TArray<int32>& GetAsyncRequestIDs() const
+	{
+		return AsyncRequestIDs; 
+	}
+
 	/** Returns if the streaming level has requested to be unloaded and removed. */
 	UFUNCTION(BlueprintPure, Category = LevelStreaming)
 	bool GetIsRequestingUnloadAndRemoval() const { return bIsRequestingUnloadAndRemoval; }
diff --git a/Engine/Source/Runtime/Engine/Classes/Engine/World.h b/Engine/Source/Runtime/Engine/Classes/Engine/World.h
--- a/Engine/Source/Runtime/Engine/Classes/Engine/World.h
+++ b/Engine/Source/Runtime/Engine/Classes/Engine/World.h
@@ -2960,6 +2960,13 @@ public:
 
 private:
 	UMaterialParameterCollectionInstance* CreateParameterCollectionInstance(int32 ExistingIndex, UMaterialParameterCollection* Collection, bool bUpdateScene);
+
+	/** Returns true if streaming levels have not loaded and are waiting on outstanding async load requests */
+	bool HasAsyncLevelRequests();
+
+	/** Waits for all streaming levels with outstanding async load requests */
+	void FlushAsyncLevelRequests();
+
 public:
 
 	/** Gets the canvas object for rendering to a render target.  Will allocate one if needed. */
diff --git a/Engine/Source/Runtime/Engine/Private/LevelStreaming.cpp b/Engine/Source/Runtime/Engine/Private/LevelStreaming.cpp
--- a/Engine/Source/Runtime/Engine/Private/LevelStreaming.cpp
+++ b/Engine/Source/Runtime/Engine/Private/LevelStreaming.cpp
@@ -1698,7 +1698,7 @@ bool ULevelStreaming::RequestLevel(UWorld* PersistentWorld, bool bAllowLevelLoad
 				InstancingContextPtr = &InstancingContext;
 			}
 #endif
-			LoadPackageAsync(PackagePath, DesiredPackageName, FLoadPackageAsyncDelegate::CreateUObject(this, &ULevelStreaming::AsyncLevelLoadComplete), PackageFlags, PIEInstanceID, GetPriority(), InstancingContextPtr);
+			AsyncRequestIDs.Add(LoadPackageAsync(PackagePath, DesiredPackageName, FLoadPackageAsyncDelegate::CreateUObject(this, &ULevelStreaming::AsyncLevelLoadComplete), PackageFlags, PIEInstanceID, GetPriority(), InstancingContextPtr));
 
 			// streamingServer: server loads everything?
 			// Editor immediately blocks on load and we also block if background level streaming is disabled.
@@ -1710,7 +1710,7 @@ bool ULevelStreaming::RequestLevel(UWorld* PersistentWorld, bool bAllowLevelLoad
 				}
 
 				// Finish all async loading.
-				FlushAsyncLoading();
+				FlushAsyncLoading(AsyncRequestIDs);
 			}
 		}
 		else
@@ -2043,6 +2043,9 @@ void ULevelStreaming::OnLoadingStarted()
 
 void ULevelStreaming::OnLoadingFinished()
 {
+	// Clear our request ids as we are done loading
+	AsyncRequestIDs.Empty();
+
 	UWorld* World = GetWorld();
 	if (World && World->IsGameWorld())
 	{
diff --git a/Engine/Source/Runtime/Engine/Private/World.cpp b/Engine/Source/Runtime/Engine/Private/World.cpp
--- a/Engine/Source/Runtime/Engine/Private/World.cpp
+++ b/Engine/Source/Runtime/Engine/Private/World.cpp
@@ -150,6 +150,11 @@ FAutoConsoleVariableRef CVarDisableRemapScriptActors(TEXT("net.DisableRemapScrip
 static bool bDisableInGamePerfTrackersForUninitializedWorlds = true;
 FAutoConsoleVariableRef CVarDisableInGamePerfTrackersForUninitializedWorlds(TEXT("s.World.SkipPerfTrackerForUninitializedWorlds"), bDisableInGamePerfTrackersForUninitializedWorlds, TEXT("When set, disables allocation of InGamePerformanceTrackers for Worlds that aren't initialized."));
 
+// Now that it's possible for subclasses of ULevelStreaming to indicate which async loads are necessary for loading,
+// it's possible existing subclasses haven't added their required loads to their StreamingLevel->GetAsyncRequestIDs() array. As a fallback,
+// allow users to force flushing of all async loads during level streaming, as was done in UE 5.4 and lower.
+static bool bForceFlushAllAsyncLoadsDuringLevelStreaming = false;
+FAutoConsoleVariableRef CVarForceFlushAllAsyncLoadsDuringLevelStreaming(TEXT("s.World.ForceFlushAllAsyncLoadsDuringLevelStreaming"), bForceFlushAllAsyncLoadsDuringLevelStreaming, TEXT("When set, level streaming will wait for all outstanding async loads globally."));
 
 static TAutoConsoleVariable<int32> CVarPurgeEditorSceneDuringPIE(
 	TEXT("r.PurgeEditorSceneDuringPIE"),
@@ -157,6 +162,25 @@ static TAutoConsoleVariable<int32> CVarPurgeEditorSceneDuringPIE(
 	TEXT("0 to keep editor scene fully initialized during PIE (default)\n")
 	TEXT("1 to purge editor scene from memory during PIE and restore when the session finishes."));
 
+namespace UE::Private::World
+{
+	struct FStreamingLevelsToConsiderIterationScope
+	{
+		FStreamingLevelsToConsiderIterationScope(FStreamingLevelsToConsider& InStreamingLevelsToConsider)
+			: StreamingLevelsToConsider(InStreamingLevelsToConsider)
+		{
+			StreamingLevelsToConsider.BeginConsideration();
+		}
+
+		~FStreamingLevelsToConsiderIterationScope()
+		{
+			StreamingLevelsToConsider.EndConsideration();
+		}
+
+		FStreamingLevelsToConsider& StreamingLevelsToConsider;
+	};
+}
+
 /*-----------------------------------------------------------------------------
 	FAdaptiveAddToWorld implementation.
 -----------------------------------------------------------------------------*/
@@ -4177,8 +4201,8 @@ void UWorld::BlockTillLevelStreamingCompleted()
 		// Probe if we have anything to do
 		UpdateLevelStreaming();
 		
-		// Everytime we have work to do, add an extra loop to handle FlushAsyncLoading calls
-		if (IsVisibilityRequestPending() || IsAsyncLoading())
+		// Everytime we have work to do, add an extra loop to handle outstanding async loads to stream
+		if (IsVisibilityRequestPending() || HasAsyncLevelRequests())
 		{
 			WorkToDo = 2;
 		}
@@ -4264,6 +4288,60 @@ PRAGMA_ENABLE_DEPRECATION_WARNINGS
 
 extern ENGINE_API bool GIsLowMemory;
 
+bool UWorld::HasAsyncLevelRequests()
+{
+	const TArray<TObjectPtr<ULevelStreaming>>& ConsideredStreamingLevels = StreamingLevelsToConsider.GetStreamingLevels();
+	if (ConsideredStreamingLevels.IsEmpty())
+	{
+		return false;
+	}
+
+	UE::Private::World::FStreamingLevelsToConsiderIterationScope Scope(StreamingLevelsToConsider);
+	for (const ULevelStreaming* StreamingLevel : ConsideredStreamingLevels)
+	{
+		if (!StreamingLevel->GetAsyncRequestIDs().IsEmpty())
+		{
+			return true;
+		}
+	}
+
+	return false;
+}
+
+void UWorld::FlushAsyncLevelRequests()
+{
+	if (bForceFlushAllAsyncLoadsDuringLevelStreaming)
+	{
+		// Flushing an empty list implies flushing all async loads globally.
+		FlushAsyncLoading();
+		return;
+	}
+
+	TArray<int32> LevelStreamingRequestIDs;
+	{
+		const TArray<TObjectPtr<ULevelStreaming>>& ConsideredStreamingLevels = StreamingLevelsToConsider.GetStreamingLevels();
+		if (ConsideredStreamingLevels.IsEmpty())
+		{
+			return;
+		}
+
+		UE::Private::World::FStreamingLevelsToConsiderIterationScope Scope(StreamingLevelsToConsider);
+		for (ULevelStreaming* StreamingLevel : ConsideredStreamingLevels)
+		{
+			const TArray<int32>& RequestIDs = StreamingLevel->GetAsyncRequestIDs();
+			if (!RequestIDs.IsEmpty())
+			{
+				LevelStreamingRequestIDs.Append(RequestIDs);
+			}
+		}
+	}
+
+	if(!LevelStreamingRequestIDs.IsEmpty())
+	{
+		FlushAsyncLoading(LevelStreamingRequestIDs);
+	}
+}
+
 void UWorld::UpdateLevelStreaming()
 {
 	TRACE_CPUPROFILER_EVENT_SCOPE(UWorld::UpdateLevelStreaming);
@@ -4291,44 +4369,45 @@ void UWorld::UpdateLevelStreaming()
 	// Store current number of pending unload levels, it may change in loop bellow
 	const int32 NumLevelsPendingPurge = FLevelStreamingGCHelper::GetNumLevelsPendingPurge();
 
-	StreamingLevelsToConsider.BeginConsideration();
-
-	for (int32 Index = StreamingLevelsToConsider.GetStreamingLevels().Num() - 1; Index >= 0; --Index)
 	{
-		// Call the blocking tick on the movie player periodically.
-		if ((Index & 0x7) == 7)
-		{
-			FMoviePlayerProxy::BlockingTick();
-		}
+		UE::Private::World::FStreamingLevelsToConsiderIterationScope Scope(StreamingLevelsToConsider);
 
-		if (ULevelStreaming* StreamingLevel = StreamingLevelsToConsider.GetStreamingLevels()[Index])
+		for (int32 Index = StreamingLevelsToConsider.GetStreamingLevels().Num() - 1; Index >= 0; --Index)
 		{
-			bool bUpdateAgain = true;
-			bool bShouldContinueToConsider = true;
-			while (bUpdateAgain && bShouldContinueToConsider)
+			// Call the blocking tick on the movie player periodically.
+			if ((Index & 0x7) == 7)
 			{
-				bool bRedetermineTarget = false;
-				FStreamingLevelPrivateAccessor::UpdateStreamingState(StreamingLevel, bUpdateAgain, bRedetermineTarget);
+				FMoviePlayerProxy::BlockingTick();
+			}
 
-				if (bRedetermineTarget)
+			if (ULevelStreaming* StreamingLevel = StreamingLevelsToConsider.GetStreamingLevels()[Index])
+			{
+				bool bUpdateAgain = true;
+				bool bShouldContinueToConsider = true;
+				while (bUpdateAgain && bShouldContinueToConsider)
 				{
-					bShouldContinueToConsider = FStreamingLevelPrivateAccessor::UpdateTargetState(StreamingLevel);
+					bool bRedetermineTarget = false;
+					FStreamingLevelPrivateAccessor::UpdateStreamingState(StreamingLevel, bUpdateAgain, bRedetermineTarget);
+
+					if (bRedetermineTarget)
+					{
+						bShouldContinueToConsider = FStreamingLevelPrivateAccessor::UpdateTargetState(StreamingLevel);
+					}
 				}
-			}
 
-			if (!bShouldContinueToConsider)
+				if (!bShouldContinueToConsider)
+				{
+					StreamingLevelsToConsider.RemoveAt(Index);
+				}
+			}
+			else
 			{
 				StreamingLevelsToConsider.RemoveAt(Index);
 			}
 		}
-		else
-		{
-			StreamingLevelsToConsider.RemoveAt(Index);
-		}
-	}
 
-	AllLevelsChangedEvent.Broadcast();
-	StreamingLevelsToConsider.EndConsideration();
+		AllLevelsChangedEvent.Broadcast();
+	}
 
 	const int32 CurrentNumLevelsPendingPurge = FLevelStreamingGCHelper::GetNumLevelsPendingPurge();
 	const int32 LevelStreamingContinuouslyIncrementalGCWhileLevelsPendingPurge = GLevelStreamingContinuouslyIncrementalGCWhileLevelsPendingPurgeOverride ? 1 : GLevelStreamingContinuouslyIncrementalGCWhileLevelsPendingPurge;
@@ -4571,7 +4650,7 @@ void UWorld::FlushLevelStreaming(EFlushLevelStreamingType FlushType)
 		if (FlushLevelStreamingType == EFlushLevelStreamingType::Full)
 		{
 			// Make sure all outstanding loads are taken care of, other than ones associated with the excluded type
-			FlushAsyncLoading();
+			FlushAsyncLevelRequests();
 		}
 
 		// Kick off making levels visible if loading finished by flushing.
diff --git a/Engine/Source/Runtime/Engine/Private/WorldPartition/WorldPartitionLevelHelper.cpp b/Engine/Source/Runtime/Engine/Private/WorldPartition/WorldPartitionLevelHelper.cpp
--- a/Engine/Source/Runtime/Engine/Private/WorldPartition/WorldPartitionLevelHelper.cpp
+++ b/Engine/Source/Runtime/Engine/Private/WorldPartition/WorldPartitionLevelHelper.cpp
@@ -495,7 +495,7 @@ bool FWorldPartitionLevelHelper::LoadActorsWithPropertyOverridesInternal(FLoadAc
 	}
 
 	LoadProgress->NumPendingLoadRequests = PropertyOverridesToLoad.Num();
-	
+
 	// Do Loading
 	for (auto const&[AssetPath, PackageName] : PropertyOverridesToLoad)
 	{
@@ -550,7 +550,11 @@ bool FWorldPartitionLevelHelper::LoadActorsWithPropertyOverridesInternal(FLoadAc
 		{
 			FPackagePath PackagePath = FPackagePath::FromPackageNameChecked(PackageToLoad);
 
-			::LoadPackageAsync(PackagePath, RemappedPackageName, CompletionCallback, PKG_None, -1, 0, &InstancingContext);
+			int32 RequestID = ::LoadPackageAsync(PackagePath, RemappedPackageName, CompletionCallback, PKG_None, -1, 0, &InstancingContext);
+			if (LoadProgress->Params.AsyncRequestIDs)
+			{
+				LoadProgress->Params.AsyncRequestIDs->Add(RequestID);
+			}
 		}
 		else
 		{
@@ -576,7 +580,7 @@ bool FWorldPartitionLevelHelper::LoadActors(const FLoadActorsParams& InParams)
 		.SetCompletionCallback(InParams.CompletionCallback)
 		.SetDestLevel(InParams.DestLevel)
 		.SetInstancingContext(MoveTemp(InParams.InstancingContext))
-		.SetLoadAsync(InParams.bLoadAsync)
+		.SetLoadAsync(InParams.bLoadAsync, InParams.AsyncRequestIDs)
 		.SetOuterWorld(InParams.OuterWorld)
 		.SetPackageReferencer(InParams.PackageReferencer);
 	return LoadActors(MoveTemp(ParamsCopy));
@@ -850,7 +854,11 @@ bool FWorldPartitionLevelHelper::LoadActorsInternal(FLoadActorsParams&& InParams
 			const UPackage* DestPackage = InParams.DestLevel->GetPackage();
 			const EPackageFlags PackageFlags = InParams.DestLevel->GetPackage()->HasAnyPackageFlags(PKG_PlayInEditor) ? PKG_PlayInEditor : PKG_None;
 			const FPackagePath PackagePath = FPackagePath::FromPackageNameChecked(PackageToLoad);
-			::LoadPackageAsync(PackagePath, PackageName, CompletionCallback, PackageFlags, DestPackage->GetPIEInstanceID(), 0, &ContainerInstancingContext);
+			int32 RequestID = ::LoadPackageAsync(PackagePath, PackageName, CompletionCallback, PackageFlags, DestPackage->GetPIEInstanceID(), 0, &ContainerInstancingContext);
+			if (InParams.AsyncRequestIDs)
+			{
+				InParams.AsyncRequestIDs->Add(RequestID);
+			}
 		}
 		else
 		{
diff --git a/Engine/Source/Runtime/Engine/Private/WorldPartition/WorldPartitionLevelStreamingDynamic.cpp b/Engine/Source/Runtime/Engine/Private/WorldPartition/WorldPartitionLevelStreamingDynamic.cpp
--- a/Engine/Source/Runtime/Engine/Private/WorldPartition/WorldPartitionLevelStreamingDynamic.cpp
+++ b/Engine/Source/Runtime/Engine/Private/WorldPartition/WorldPartitionLevelStreamingDynamic.cpp
@@ -333,7 +333,7 @@ bool UWorldPartitionLevelStreamingDynamic::RequestLevel(UWorld* InPersistentWorl
 					}
 
 					// Finish all async loading.
-					FlushAsyncLoading();
+					FlushAsyncLoading(AsyncRequestIDs);
 				}
 				else
 				{
@@ -466,7 +466,7 @@ bool UWorldPartitionLevelStreamingDynamic::IssueLoadRequests()
 			.SetActorPackages(ChildPackagesToLoad)
 			.SetPackageReferencer(&PackageReferencer)
 			.SetCompletionCallback(FinalizeLoading)
-			.SetLoadAsync(World->IsGameWorld())
+			.SetLoadAsync(World->IsGameWorld(), &AsyncRequestIDs)
 			.SetInstancingContext(MoveTemp(InstancingContext));
 
 		FWorldPartitionLevelHelper::LoadActors(MoveTemp(Params));
diff --git a/Engine/Source/Runtime/Engine/Public/WorldPartition/WorldPartitionLevelHelper.h b/Engine/Source/Runtime/Engine/Public/WorldPartition/WorldPartitionLevelHelper.h
--- a/Engine/Source/Runtime/Engine/Public/WorldPartition/WorldPartitionLevelHelper.h
+++ b/Engine/Source/Runtime/Engine/Public/WorldPartition/WorldPartitionLevelHelper.h
@@ -52,6 +52,7 @@ public:
 			, DestLevel(nullptr)
 			, PackageReferencer(nullptr)
 			, bLoadAsync(false)
+			, AsyncRequestIDs(nullptr)
 		{}
 
 		UWorld* OuterWorld;
@@ -60,6 +61,7 @@ public:
 		FPackageReferencer* PackageReferencer;
 		TFunction<void(bool)> CompletionCallback;
 		bool bLoadAsync;
+		TArray<int32>* AsyncRequestIDs;
 		mutable FLinkerInstancingContext InstancingContext;
 
 		FLoadActorsParams& SetOuterWorld(UWorld* InOuterWorld) { OuterWorld = InOuterWorld; return *this; }
@@ -67,7 +69,12 @@ public:
 		FLoadActorsParams& SetActorPackages(TArrayView<FWorldPartitionRuntimeCellObjectMapping> InActorPackages) { ActorPackages = InActorPackages; return *this; }
 		FLoadActorsParams& SetPackageReferencer(FPackageReferencer* InPackageReferencer) { PackageReferencer = InPackageReferencer; return *this; }
 		FLoadActorsParams& SetCompletionCallback(TFunction<void(bool)> InCompletionCallback) { CompletionCallback = InCompletionCallback; return *this; }
-		FLoadActorsParams& SetLoadAsync(bool bInLoadAsync) { bLoadAsync = bInLoadAsync; return *this; }
+		FLoadActorsParams& SetLoadAsync(bool bInLoadAsync, TArray<int32>* OutAsyncRequestIDs = nullptr)
+		{ 
+			bLoadAsync = bInLoadAsync;
+			AsyncRequestIDs = OutAsyncRequestIDs;
+			return *this;
+		}
 		FLoadActorsParams& SetInstancingContext(FLinkerInstancingContext InInstancingContext) { InstancingContext = InInstancingContext; return *this; }
 	};
 
