World Partition Runtime HashSet: follow up IsSourceAffectingGrid calls optimization when iterating through streaming sources leftover from CL 31943210.

diff --git a/Engine/Source/Runtime/Engine/Private/WorldPartition/RuntimeHashSet/WorldPartitionRuntimeHashSet.cpp b/Engine/Source/Runtime/Engine/Private/WorldPartition/RuntimeHashSet/WorldPartitionRuntimeHashSet.cpp
--- a/Engine/Source/Runtime/Engine/Private/WorldPartition/RuntimeHashSet/WorldPartitionRuntimeHashSet.cpp
+++ b/Engine/Source/Runtime/Engine/Private/WorldPartition/RuntimeHashSet/WorldPartitionRuntimeHashSet.cpp
@@ -485,7 +485,7 @@ void UWorldPartitionRuntimeHashSet::ForEachStreamingCellsSources(const TArray<FW
 
 				for (const FRuntimePartitionStreamingData* StreamingData : *StreamingDataList)
 				{
-					Source.ForEachShape(StreamingData->LoadingRange, StreamingData->Name, false, [this, &Source, StreamingData, &Func](const FSphericalSector& Shape)
+					Source.ForEachShape(StreamingData->LoadingRange, false, [this, &Source, StreamingData, &Func](const FSphericalSector& Shape)
 					{
 						const FSphere ShapeSphere(Shape.GetCenter(), Shape.GetRadius());
 
diff --git a/Engine/Source/Runtime/Engine/Public/WorldPartition/WorldPartitionStreamingSource.h b/Engine/Source/Runtime/Engine/Public/WorldPartition/WorldPartitionStreamingSource.h
--- a/Engine/Source/Runtime/Engine/Public/WorldPartition/WorldPartitionStreamingSource.h
+++ b/Engine/Source/Runtime/Engine/Public/WorldPartition/WorldPartitionStreamingSource.h
@@ -466,7 +466,7 @@ struct FWorldPartitionStreamingSource
 	bool bRemote;
 
 	/** Returns a box encapsulating all shapes. */
-	FBox CalcBounds(float InGridLoadingRange, FName InGridName, bool bCalcIn2D = false) const
+	inline FBox CalcBounds(float InGridLoadingRange, FName InGridName, bool bCalcIn2D = false) const
 	{
 		FBox OutBounds(ForceInit);
 		ForEachShape(InGridLoadingRange, InGridName, bCalcIn2D, [&OutBounds](const FSphericalSector& Sector)
@@ -477,7 +477,13 @@ struct FWorldPartitionStreamingSource
 	}
 
 	/** Helper method that iterates over all shapes. If none is provided, it will still pass a sphere shape using grid's loading range. */
-	void ForEachShape(float InGridLoadingRange, FName InGridName, bool bInProjectIn2D, TFunctionRef<void(const FSphericalSector&)> InOperation) const
+	inline void ForEachShape(float InGridLoadingRange, bool bInProjectIn2D, TFunctionRef<void(const FSphericalSector&)> InOperation) const
+	{
+		FStreamingSourceShapeHelper::ForEachShape(InGridLoadingRange, InGridLoadingRange, bInProjectIn2D, Location, Rotation, Shapes, InOperation, ExtraRadius, ExtraAngle);
+	}
+
+	/** Helper method that iterates over all shapes affecting a specific grid. If none is provided, it will still pass a sphere shape using grid's loading range. */
+	inline void ForEachShape(float InGridLoadingRange, FName InGridName, bool bInProjectIn2D, TFunctionRef<void(const FSphericalSector&)> InOperation) const
 	{
 		if (FStreamingSourceShapeHelper::IsSourceAffectingGrid(TargetGrids, TargetBehavior, InGridName))
 		{
