World Partition Runtime HashSet: in PIE, create streaming datas spatial indexes and runtime grid data maps right after world duplication to allow making spatial queries from calls to BlockTillLevelStreamingCompleted before the world has issues OnBeginPlay.

diff --git a/Engine/Source/Runtime/Engine/Private/WorldPartition/RuntimeHashSet/WorldPartitionRuntimeHashSet.cpp b/Engine/Source/Runtime/Engine/Private/WorldPartition/RuntimeHashSet/WorldPartitionRuntimeHashSet.cpp
--- a/Engine/Source/Runtime/Engine/Private/WorldPartition/RuntimeHashSet/WorldPartitionRuntimeHashSet.cpp
+++ b/Engine/Source/Runtime/Engine/Private/WorldPartition/RuntimeHashSet/WorldPartitionRuntimeHashSet.cpp
@@ -666,6 +666,19 @@ void UWorldPartitionRuntimeHashSet::PostEditChangeChainProperty(FPropertyChanged
 	}
 }
 
+void UWorldPartitionRuntimeHashSet::PostDuplicate(EDuplicateMode::Type DuplicateMode)
+{
+	Super::PostDuplicate(DuplicateMode);
+	
+	// In PIE, create streaming datas spatial indexes and runtime grid data maps right after world duplication to allow making spatial queries from calls 
+	// to BlockTillLevelStreamingCompleted before the world has issued OnBeginPlay (see UGameInstance::StartPlayInEditorGameInstance).
+	if (DuplicateMode == EDuplicateMode::PIE)
+	{
+		ForEachStreamingData([this](const FRuntimePartitionStreamingData& StreamingData) { StreamingData.CreatePartitionsSpatialIndex(); return true; });
+		UpdateRuntimeDataGridMap();
+	}
+}
+
 UWorldPartitionRuntimeHashSet::FCellUniqueId UWorldPartitionRuntimeHashSet::GetCellUniqueId(const URuntimePartition::FCellDescInstance& InCellDescInstance) const
 {
 	FCellUniqueId CellUniqueId;
diff --git a/Engine/Source/Runtime/Engine/Public/WorldPartition/RuntimeHashSet/WorldPartitionRuntimeHashSet.h b/Engine/Source/Runtime/Engine/Public/WorldPartition/RuntimeHashSet/WorldPartitionRuntimeHashSet.h
--- a/Engine/Source/Runtime/Engine/Public/WorldPartition/RuntimeHashSet/WorldPartitionRuntimeHashSet.h
+++ b/Engine/Source/Runtime/Engine/Public/WorldPartition/RuntimeHashSet/WorldPartitionRuntimeHashSet.h
@@ -126,6 +126,7 @@ class UWorldPartitionRuntimeHashSet final : public UWorldPartitionRuntimeHash
 	//~ Begin UObject Interface
 #if WITH_EDITOR
 	ENGINE_API virtual void PostEditChangeChainProperty(FPropertyChangedChainEvent& PropertyChangedEvent) override;
+	ENGINE_API virtual void PostDuplicate(EDuplicateMode::Type DuplicateMode) override;
 #endif
 	//~ End UObject Interface
 
