Fix Actors dissapearing from the map when we remove the WorldPartition through the WorldSettings

diff --git a/Engine/Source/Runtime/Engine/Private/WorldPartition/WorldPartition.cpp b/Engine/Source/Runtime/Engine/Private/WorldPartition/WorldPartition.cpp
--- a/Engine/Source/Runtime/Engine/Private/WorldPartition/WorldPartition.cpp
+++ b/Engine/Source/Runtime/Engine/Private/WorldPartition/WorldPartition.cpp
@@ -1095,6 +1095,8 @@ bool UWorldPartition::RemoveWorldPartition(AWorldSettings* WorldSettings)
 					else if(Actor->GetExternalPackage())
 					{
 						ActorReferences.Emplace(WorldPartition, Actor->GetActorGuid());
+						// Remove the actor from the partition so it doesn't get removed from the World when we uninitialize the partition 
+						WorldPartition->RemoveActor(Actor->GetActorGuid());
 					}
 				}
 			}
