[shaders] skip material shadermap puts for materials for which gets were skipped (i.e. when ODSC is on and material is not required complete)

diff --git a/Engine/Source/Runtime/Engine/Private/ShaderCompiler/ShaderCompiler.cpp b/Engine/Source/Runtime/Engine/Private/ShaderCompiler/ShaderCompiler.cpp
--- a/Engine/Source/Runtime/Engine/Private/ShaderCompiler/ShaderCompiler.cpp
+++ b/Engine/Source/Runtime/Engine/Private/ShaderCompiler/ShaderCompiler.cpp
@@ -6768,10 +6768,13 @@ void FShaderCompilingManager::ProcessCompiledShaderMaps(
 			check(MaterialDependencies.Num() > 0);
 			FMaterialShaderParameters ShaderParameters(MaterialDependencies[0]);
 
+			bool bRequiredComplete = false;
+
 			while (MaterialIndex < MaterialDependencies.Num())
 			{
 				FMaterial* Material = MaterialDependencies[MaterialIndex];
 				check(Material->GetGameThreadCompilingShaderMapId() == CompilingShaderMap->GetCompilingId());
+				bRequiredComplete |= Material->IsRequiredComplete();
 
 #if DEBUG_INFINITESHADERCOMPILE
 				UE_LOG(LogTemp, Display, TEXT("Shader map %s complete, GameThreadShaderMap 0x%08X%08X, marking material %s as finished"), *ShaderMap->GetFriendlyName(), (int)((int64)(ShaderMap.GetReference()) >> 32), (int)((int64)(ShaderMap.GetReference())), *Material->GetFriendlyName());
@@ -6887,7 +6890,7 @@ void FShaderCompilingManager::ProcessCompiledShaderMaps(
 				}
 			}
 
-			if (NumIncompleteMaterials == 0)
+			if (NumIncompleteMaterials == 0 && (IsMaterialMapDDCEnabled() || bRequiredComplete))
 			{
 				CompilingShaderMap->bCompiledSuccessfully = bSuccess;
 				CompilingShaderMap->bCompilationFinalized = true;
