[Shaders] Don't crash with fatal error when SCW crashed.

When we detect a crash inside the platform compiler, we should not crash the cooker so failed compile jobs can be dumped to the debug output.
This has been tested with the CompileShadersTestBed commandlet and a simulated crash in the shader backend, resulting in the desired debug info.

#rnx

diff --git a/Engine/Source/Runtime/Engine/Private/ShaderCompiler/ShaderCompiler.cpp b/Engine/Source/Runtime/Engine/Private/ShaderCompiler/ShaderCompiler.cpp
--- a/Engine/Source/Runtime/Engine/Private/ShaderCompiler/ShaderCompiler.cpp
+++ b/Engine/Source/Runtime/Engine/Private/ShaderCompiler/ShaderCompiler.cpp
@@ -2036,7 +2036,7 @@ static constexpr int32 GSingleThreadedRunsDisabled = -2;
 static constexpr int32 GSingleThreadedRunsIncreaseFactor = 8;
 static constexpr int32 GSingleThreadedRunsMaxCount = (1 << 24);
 
-static void ModalErrorOrLog(const FString& Title, const FString& Text, int64 CurrentFilePos = 0, int64 ExpectedFileSize = 0)
+static void ModalErrorOrLog(const FString& Title, const FString& Text, int64 CurrentFilePos = 0, int64 ExpectedFileSize = 0, bool bIsErrorFatal = true)
 {
 	static FThreadSafeBool bModalReported;
 
@@ -2062,10 +2062,14 @@ static void ModalErrorOrLog(const FString& Title, const FString& Text, int64 Cur
 			FPlatformProcess::SleepInfinite();
 		}
 	}
-	else
+	else if (bIsErrorFatal)
 	{
 		UE_LOG(LogShaderCompilers, Fatal, TEXT("%s\n%s\n%s"), *Title, *Text, *BadFile);
 	}
+	else
+	{
+		UE_LOG(LogShaderCompilers, Error, TEXT("%s\n%s\n%s"), *Title, *Text, *BadFile);
+	}
 }
 
 template<class EnumType>
@@ -2448,7 +2452,10 @@ namespace ShaderCompileWorkerError
 
 	void HandleCrashInsidePlatformCompiler(const TCHAR* Data)
 	{
-		ModalErrorOrLog(TEXT("ShaderCompileWorker failed"), FString::Printf(TEXT("Crash inside the platform compiler:\n%s"), Data));
+		// If the crash originates from a platform compiler, the error code must have been reported and we don't have to assume a corrupted output file.
+		// In that case, don't crash the cooker with a fatal error, just report the error so the cooker can dump debug info.
+		constexpr bool bIsErrorFatal = false;
+		ModalErrorOrLog(TEXT("ShaderCompileWorker failed"), FString::Printf(TEXT("Crash inside the platform compiler:\n%s"), Data), 0, 0, bIsErrorFatal);
 	}
 
 	void HandleBadInputFile(const TCHAR* Data)
