[shaders] clear out existing worker files when a worker relaunches. this prevents job batch mismatches that occur when we abandon workers due to an unexpected crash (some of the abandoned workers will still finish their batch and write an output file, which erroneously gets associated with a new batch on relaunch)

diff --git a/Engine/Source/Runtime/Engine/Private/ShaderCompiler/ShaderCompiler.cpp b/Engine/Source/Runtime/Engine/Private/ShaderCompiler/ShaderCompiler.cpp
--- a/Engine/Source/Runtime/Engine/Private/ShaderCompiler/ShaderCompiler.cpp
+++ b/Engine/Source/Runtime/Engine/Private/ShaderCompiler/ShaderCompiler.cpp
@@ -3962,6 +3962,9 @@ bool FShaderCompileThreadRunnable::LaunchWorkersIfNeeded()
 				FString InputFileName(TEXT("WorkerInputOnly.in"));
 				FString OutputFileName(TEXT("WorkerOutputOnly.out"));
 
+				// Delete any potential stale output files; these can persist if we had previously abandoned running with workers due to an unexpected SCW termination.
+				IFileManager::Get().Delete(*(WorkingDirectory / OutputFileName));
+
 				// Store the handle with this thread so that we will know not to launch it again
 				CurrentWorkerInfo.WorkerProcess = Manager->LaunchWorker(WorkingDirectory, Manager->ProcessId, WorkerIndex, InputFileName, OutputFileName);
 				CurrentWorkerInfo.bLaunchedWorker = true;
