Mark as garbage UnnecessaryComponents in editor, this prevents them from being reinstanced

diff --git a/Engine/Source/Runtime/Engine/Private/InheritableComponentHandler.cpp b/Engine/Source/Runtime/Engine/Private/InheritableComponentHandler.cpp
--- a/Engine/Source/Runtime/Engine/Private/InheritableComponentHandler.cpp
+++ b/Engine/Source/Runtime/Engine/Private/InheritableComponentHandler.cpp
@@ -256,6 +256,9 @@ void UInheritableComponentHandler::ValidateTemplates()
 #if WITH_EDITOR
 						// in editor, move the component template aside so its name is free:
 						Record.ComponentTemplate->Rename(nullptr, GetTransientPackage(), REN_DoNotDirty | REN_DontCreateRedirectors | REN_ForceNoResetLoaders | REN_NonTransactional);
+						Record.ComponentTemplate->ClearFlags(RF_Standalone);
+						Record.ComponentTemplate->RemoveFromRoot();
+						Record.ComponentTemplate->MarkAsGarbage();
 #endif // WITH_EDITOR
 						UnnecessaryComponents.AddUnique(Record.ComponentTemplate);
 					}
