Heterogeneous Volumes: Prevent SVT streaming request when component is not actively being rendered.

diff --git a/Engine/Source/Runtime/Engine/Private/Components/HeterogeneousVolumeComponent.cpp b/Engine/Source/Runtime/Engine/Private/Components/HeterogeneousVolumeComponent.cpp
--- a/Engine/Source/Runtime/Engine/Private/Components/HeterogeneousVolumeComponent.cpp
+++ b/Engine/Source/Runtime/Engine/Private/Components/HeterogeneousVolumeComponent.cpp
@@ -546,7 +546,7 @@ void UHeterogeneousVolumeComponent::TickComponent(float DeltaTime, ELevelTick Ti
 {
 	Super::TickComponent(DeltaTime, TickType, ThisTickFunction);
 
-	if (MaterialInstanceDynamic)
+	if (ShouldRender() && MaterialInstanceDynamic)
 	{
 		const int32 SVTParameterIndex = 0;
 		FName SVTParameterName;
