[shaders] remove check that compiling material shadermap has at least one material dependency; this isn't always the case, however it is necessarily the case that at least one dependency is required to save the material shadermap to DDC (since ShaderMapToUseForRendering is only set inside the material loop), so an equivalent check can occur there instead.

diff --git a/Engine/Source/Runtime/Engine/Private/ShaderCompiler/ShaderCompiler.cpp b/Engine/Source/Runtime/Engine/Private/ShaderCompiler/ShaderCompiler.cpp
--- a/Engine/Source/Runtime/Engine/Private/ShaderCompiler/ShaderCompiler.cpp
+++ b/Engine/Source/Runtime/Engine/Private/ShaderCompiler/ShaderCompiler.cpp
@@ -6774,10 +6774,7 @@ void FShaderCompilingManager::ProcessCompiledShaderMaps(
 			int32 NumIncompleteMaterials = 0;
 			int32 MaterialIndex = 0;
 			
-			// all materials using the same shader map must necessarily have the same set of FMaterialShaderParameters, so just initialize from the first
-			check(MaterialDependencies.Num() > 0);
-			FMaterialShaderParameters ShaderParameters(MaterialDependencies[0]);
-
+			FMaterial* SingleMaterial = MaterialDependencies.Num() > 0 ? MaterialDependencies[0] : nullptr;
 			bool bRequiredComplete = false;
 
 			while (MaterialIndex < MaterialDependencies.Num())
@@ -6906,11 +6903,15 @@ void FShaderCompilingManager::ProcessCompiledShaderMaps(
 				CompilingShaderMap->bCompilationFinalized = true;
 				if (ShaderMapToUseForRendering)
 				{
+					// ShaderMapToUseForRendering is only initialized inside the loop over material dependencies,
+					// so it's safe to assume that SingleMaterial has been set (a material is needed to construct 
+					// the FMaterialShaderParameters struct which is in turn needed to build the DDC key).
+					check(SingleMaterial != nullptr);
 					ShaderMapToUseForRendering->bCompiledSuccessfully = true;
 					ShaderMapToUseForRendering->bCompilationFinalized = true;
 					if (ShaderMapToUseForRendering->bIsPersistent)
 					{
-						ShaderMapToUseForRendering->SaveToDerivedDataCache(ShaderParameters);
+						ShaderMapToUseForRendering->SaveToDerivedDataCache(FMaterialShaderParameters(SingleMaterial));
 					}
 				}
 
