Level Instance Subsystem: New FLevelInstancesUpdated that fires with an TArray<ILevelInstanceInterface*> when a Level Instance is committed and causes it and other instances to reload. (will be used by PCG in upcoming CL)

Partial resubmit of CL 32204837 (split in 2)

#rnx

diff --git a/Engine/Source/Runtime/Engine/Private/LevelInstance/LevelInstanceSubsystem.cpp b/Engine/Source/Runtime/Engine/Private/LevelInstance/LevelInstanceSubsystem.cpp
--- a/Engine/Source/Runtime/Engine/Private/LevelInstance/LevelInstanceSubsystem.cpp
+++ b/Engine/Source/Runtime/Engine/Private/LevelInstance/LevelInstanceSubsystem.cpp
@@ -2647,7 +2647,7 @@ bool ULevelInstanceSubsystem::CommitLevelInstanceInternal(TUniquePtr<FLevelInsta
 	UActorDescContainerSubsystem::GetChecked().NotifyContainerUpdated(WorldAssetPackage);
 
 	
-	TArray<TPair<ULevelInstanceSubsystem*, FLevelInstanceID>> LevelInstancesToUpdate;
+	TMap<ULevelInstanceSubsystem*, TArray<FLevelInstanceID>> LevelInstancesToUpdate;
 	// Gather list to update
 	for (TObjectIterator<UWorld> It(RF_ClassDefaultObject | RF_ArchetypeObject, true); It; ++It)
 	{
@@ -2661,7 +2661,7 @@ bool ULevelInstanceSubsystem::CommitLevelInstanceInternal(TUniquePtr<FLevelInsta
 				{
 					if (CurrentLevelInstance == LevelInstance || bChangesCommitted)
 					{
-						LevelInstancesToUpdate.Add({ LevelInstanceSubsystem, CurrentLevelInstance->GetLevelInstanceID() });
+						LevelInstancesToUpdate.FindOrAdd(LevelInstanceSubsystem).Add(CurrentLevelInstance->GetLevelInstanceID());
 					}
 				}
 			}
@@ -2671,9 +2671,12 @@ bool ULevelInstanceSubsystem::CommitLevelInstanceInternal(TUniquePtr<FLevelInsta
 	// Do update
 	for (const auto& KeyValuePair : LevelInstancesToUpdate)
 	{
-		if (ILevelInstanceInterface* LevelInstanceToUpdate = KeyValuePair.Key->GetLevelInstance(KeyValuePair.Value))
+		for (const FLevelInstanceID& LevelInstanceIDToUpdate : KeyValuePair.Value)
 		{
-			LevelInstanceToUpdate->UpdateLevelInstanceFromWorldAsset();
+			if (ILevelInstanceInterface* LevelInstanceToUpdate = KeyValuePair.Key->GetLevelInstance(LevelInstanceIDToUpdate))
+			{
+				LevelInstanceToUpdate->UpdateLevelInstanceFromWorldAsset();
+			}
 		}
 	}
 
@@ -2708,6 +2711,24 @@ bool ULevelInstanceSubsystem::CommitLevelInstanceInternal(TUniquePtr<FLevelInsta
 	if (bChangesCommitted)
 	{
 		LevelInstanceChangedEvent.Broadcast(WorldAssetPackage);
+
+		// Send event per World (per LevelInstanceSubsystem)
+		for (const auto& KeyValuePair : LevelInstancesToUpdate)
+		{
+			TArray<ILevelInstanceInterface*> UpdatedLevelInstances;
+			for (const FLevelInstanceID& UpdatedLevelInstanceID : KeyValuePair.Value)
+			{
+				if (ILevelInstanceInterface* UpdatedLevelInstance = KeyValuePair.Key->GetLevelInstance(UpdatedLevelInstanceID))
+				{
+					UpdatedLevelInstances.Add(UpdatedLevelInstance);
+				}
+			}
+
+			if (UpdatedLevelInstances.Num() > 0)
+			{
+				KeyValuePair.Key->LevelInstancesUpdatedEvent.Broadcast(UpdatedLevelInstances);
+			}
+		}
 	}
 
 	GEngine->BroadcastLevelActorListChanged();
diff --git a/Engine/Source/Runtime/Engine/Public/LevelInstance/LevelInstanceSubsystem.h b/Engine/Source/Runtime/Engine/Public/LevelInstance/LevelInstanceSubsystem.h
--- a/Engine/Source/Runtime/Engine/Public/LevelInstance/LevelInstanceSubsystem.h
+++ b/Engine/Source/Runtime/Engine/Public/LevelInstance/LevelInstanceSubsystem.h
@@ -161,6 +161,10 @@ public:
 	DECLARE_EVENT_OneParam(ULevelInstanceSubsystem, FLevelInstanceChanged, FName);
 	FLevelInstanceChanged& OnLevelInstanceChanged() { return LevelInstanceChangedEvent; }
 
+	/** Editor-only event triggered when level instances are reloaded after a change */
+	DECLARE_EVENT_OneParam(ULevelInstanceSubsystem, FLevelInstancesUpdated, const TArray<ILevelInstanceInterface*>&);
+	FLevelInstancesUpdated& OnLevelInstancesUpdated() { return LevelInstancesUpdatedEvent; }
+
 	static ENGINE_API void ResetLoadersForWorldAsset(const FString& WorldAsset);
 
 	ENGINE_API bool PassLevelInstanceFilter(UWorld* World, const FWorldPartitionHandle& Actor) const;
@@ -311,6 +315,8 @@ private:
 	bool bIsCommittingLevelInstance;
 
 	FLevelInstanceChanged LevelInstanceChangedEvent;
+	FLevelInstancesUpdated LevelInstancesUpdatedEvent;
+
 #endif
 
 	struct FLevelInstance
