texture ForceRebuild should never try to fetch data before building. This fixes an issue where a force rebuild would not actually rebuild in some cases.

diff --git a/Engine/Source/Runtime/Engine/Private/TextureDerivedDataTask.cpp b/Engine/Source/Runtime/Engine/Private/TextureDerivedDataTask.cpp
--- a/Engine/Source/Runtime/Engine/Private/TextureDerivedDataTask.cpp
+++ b/Engine/Source/Runtime/Engine/Private/TextureDerivedDataTask.cpp
@@ -1467,6 +1467,23 @@ static void DDC1_FetchAndFillDerivedData(
 	TRACE_CPUPROFILER_EVENT_SCOPE(Texture.DDC1_FetchAndFillDerivedData);
 
 	bool bForceRebuild = EnumHasAnyFlags(CacheFlags, ETextureCacheFlags::ForceRebuild);
+	FString FetchOrBuildKeySuffix;
+	GetTextureDerivedDataKeySuffix(Texture, BuildSettingsPerLayerFetchOrBuild.GetData(), FetchOrBuildKeySuffix);
+
+	if (bForceRebuild)
+	{
+		// If we know we are rebuilding, don't touch the cache.
+		bSucceeded = false;
+		bInvalidVirtualTextureCompression = false;
+		KeySuffix = MoveTemp(FetchOrBuildKeySuffix);
+
+		FString FetchOrBuildKey;
+		GetTextureDerivedDataKeyFromSuffix(KeySuffix, FetchOrBuildKey);
+		DerivedData->DerivedDataKey.Emplace<FString>(MoveTemp(FetchOrBuildKey));
+		DerivedData->ResultMetadata = FetchOrBuildMetadata;
+		return;
+	}
+		
 	bool bForVirtualTextureStreamingBuild = EnumHasAnyFlags(CacheFlags, ETextureCacheFlags::ForVirtualTextureStreamingBuild);
 
 	FSharedBuffer RawDerivedData;
@@ -1478,8 +1495,6 @@ static void DDC1_FetchAndFillDerivedData(
 	FString LocalDerivedDataKeySuffix;
 	FString LocalDerivedDataKey;
 
-	FString FetchOrBuildKeySuffix;
-	GetTextureDerivedDataKeySuffix(Texture, BuildSettingsPerLayerFetchOrBuild.GetData(), FetchOrBuildKeySuffix);
 
 	bool bGotDDCData = false;
 	bool bUsedFetchFirst = false;
