Don't initialize/parse encoding settings if we aren't the editor.

diff --git a/Engine/Source/Runtime/Engine/Private/TextureEncodingSettings.cpp b/Engine/Source/Runtime/Engine/Private/TextureEncodingSettings.cpp
--- a/Engine/Source/Runtime/Engine/Private/TextureEncodingSettings.cpp
+++ b/Engine/Source/Runtime/Engine/Private/TextureEncodingSettings.cpp
@@ -39,6 +39,7 @@ enum class ReadyMask : uint32
 	Everyone = 0x3
 };
 
+#if WITH_EDITORONLY_DATA
 static void ConstructResolvedSettings(ReadyMask InReadyMask, UDeveloperSettings* InSettings, FResolvedTextureEncodingSettings* OutResolvedSettings)
 {
 	// Settings is the last one to get initialized - the others have all had PostInitProperties
@@ -219,21 +220,33 @@ static void TryToResolveSettings(ReadyMask InReadyMask, UDeveloperSettings* InSe
 	}
 }
 
+#endif // WITH_EDITORONLY_DATA
+
 FResolvedTextureEncodingSettings const& FResolvedTextureEncodingSettings::Get()
 {
+#if WITH_EDITORONLY_DATA
 	FResolvedTextureEncodingSettings* ResolvedSettings = nullptr;
 	TryToResolveSettings(ReadyMask::Query, nullptr, &ResolvedSettings);
 	return *ResolvedSettings;
+#else
+	checkNoEntry();
+	static FResolvedTextureEncodingSettings DefaultSettings;
+	return DefaultSettings;
+#endif
 }
 
 void UTextureEncodingUserSettings::PostInitProperties()
 {
 	Super::PostInitProperties();
+#if WITH_EDITORONLY_DATA
 	TryToResolveSettings(ReadyMask::User, this, nullptr);
+#endif
 }
 
 void UTextureEncodingProjectSettings::PostInitProperties()
 {
 	Super::PostInitProperties();
+#if WITH_EDITORONLY_DATA
 	TryToResolveSettings(ReadyMask::Project, this, nullptr);
+#endif
 }
