Add support for mirroring root motion

diff --git a/Engine/Source/Runtime/Engine/Classes/Animation/MirrorDataTable.h b/Engine/Source/Runtime/Engine/Classes/Animation/MirrorDataTable.h
--- a/Engine/Source/Runtime/Engine/Classes/Animation/MirrorDataTable.h
+++ b/Engine/Source/Runtime/Engine/Classes/Animation/MirrorDataTable.h
@@ -37,6 +37,7 @@ namespace EMirrorFindReplaceMethod
     };
 }
 
+
 /**  Base Mirror Table containing all data required by the animation mirroring system. */
 USTRUCT()
 struct FMirrorTableRow : public FTableRowBase
@@ -64,6 +65,7 @@ struct FMirrorTableRow : public FTableRowBase
 	ENGINE_API bool operator<(FMirrorTableRow const& Other) const;
 };
 
+
 /** Find and Replace expressions used to generate mirror tables*/
 USTRUCT()
 struct FMirrorFindReplaceExpression
@@ -205,6 +207,9 @@ public:
 	UPROPERTY(EditAnywhere, Category = Mirroring)
 	TEnumAsByte<EAxis::Type> MirrorAxis;
 
+	UPROPERTY(EditAnywhere, Category = Mirroring)
+	bool  bMirrorRootMotion = true;
+	
 	UPROPERTY(EditAnywhere, AssetRegistrySearchable, Category = Skeleton)
 	TObjectPtr<USkeleton> Skeleton; 
 
diff --git a/Engine/Source/Runtime/Engine/Private/Animation/AttributesRuntime.cpp b/Engine/Source/Runtime/Engine/Private/Animation/AttributesRuntime.cpp
--- a/Engine/Source/Runtime/Engine/Private/Animation/AttributesRuntime.cpp
+++ b/Engine/Source/Runtime/Engine/Private/Animation/AttributesRuntime.cpp
@@ -5,13 +5,13 @@
 #include "Animation/AnimationSettings.h"
 #include "Animation/AnimationAsset.h"
 #include "Animation/MirrorDataTable.h"
-
 #include "Stats/Stats.h"
 
+#include "Animation/AnimRootMotionProvider.h"
 #include "Animation/IAttributeBlendOperator.h" 
 #include "Animation/AttributeBlendData.h"
-
 #include "Animation/AttributeTypes.h"
+#include "AnimationRuntime.h"
 
 namespace UE { namespace Anim {
 
@@ -444,6 +444,8 @@ void Attributes::MirrorAttributes(FStackAttributeContainer& CustomAttributes, co
 	Attributes::MirrorAttributes(CustomAttributes, MirrorTable, CompactPoseMirrorBones);
 }
 
+
+
 void Attributes::MirrorAttributes(FStackAttributeContainer& CustomAttributes, const UMirrorDataTable& MirrorTable, const TArray<FCompactPoseBoneIndex>& CompactPoseMirrorBones)
 {
 	struct FMirrorSet
@@ -570,6 +572,30 @@ void Attributes::MirrorAttributes(FStackAttributeContainer& CustomAttributes, co
 			}
 		} 
 	}
+
+	const IAnimRootMotionProvider* RootMotionProvider = IAnimRootMotionProvider::Get();
+	
+	if (RootMotionProvider != nullptr && MirrorTable.bMirrorRootMotion)
+	{
+		FTransform RootMotionTransformDelta = FTransform::Identity; 
+		if (RootMotionProvider->ExtractRootMotion(CustomAttributes, RootMotionTransformDelta)) 
+		{
+			auto MirrorTransform = [](const FTransform& SourceTransform, EAxis::Type MirrorAxis) 
+			{
+				FVector T = SourceTransform.GetTranslation(); 
+				T = FAnimationRuntime::MirrorVector(T, MirrorAxis); 
+				FQuat Q = SourceTransform.GetRotation();
+				Q = FAnimationRuntime::MirrorQuat(Q, MirrorAxis); 
+
+				FVector S = SourceTransform.GetScale3D(); 
+
+				return FTransform(Q, T, S); 
+			};
+
+			RootMotionTransformDelta = MirrorTransform(RootMotionTransformDelta, MirrorTable.MirrorAxis); 
+			RootMotionProvider->OverrideRootMotion(RootMotionTransformDelta, CustomAttributes);
+		}
+	}
 }
 
 ECustomAttributeBlendType Attributes::GetAttributeBlendType(const FName& InName)
