Only print SCW memory stats on OOM if lock can be taken immediately

It's useful to get the memory usage of SCWs when an OOM appears. However, when multiple threads encounter an OOM condition at the same time, it is possible that one of them goes to perpertual sleep while holding the lock guarding the SCW data. This means that OOM reporting itself can stall indefinitely trying to take that lock.

diff --git a/Engine/Source/Runtime/Engine/Private/ShaderCompiler/ShaderCompiler.cpp b/Engine/Source/Runtime/Engine/Private/ShaderCompiler/ShaderCompiler.cpp
--- a/Engine/Source/Runtime/Engine/Private/ShaderCompiler/ShaderCompiler.cpp
+++ b/Engine/Source/Runtime/Engine/Private/ShaderCompiler/ShaderCompiler.cpp
@@ -37,6 +37,7 @@
 #include "Misc/FileHelper.h"
 #include "Misc/MessageDialog.h"
 #include "Misc/ScopedSlowTask.h"
+#include "Misc/ScopeTryLock.h"
 #include "Modules/ModuleManager.h"
 #include "ObjectCacheContext.h"
 #include "ProfilingDebugging/CookStats.h"
@@ -4103,9 +4104,8 @@ void FShaderCompileThreadRunnable::CompileDirectlyThroughDll()
 	}
 }
 
-void FShaderCompileThreadRunnable::PrintWorkerMemoryUsage()
+void FShaderCompileThreadRunnable::PrintWorkerMemoryUsageWithLockTaken()
 {
-	FScopeLock WorkerScopeLock(&WorkerInfosLock);
 	FPlatformProcessMemoryStats TotalMemoryStats{};
 	int32 NumValidWorkers = 0;
 	constexpr int64 Gibibyte = 1024 * 1024 * 1024;
@@ -4159,6 +4159,26 @@ void FShaderCompileThreadRunnable::PrintWorkerMemoryUsage()
 	}
 }
 
+bool FShaderCompileThreadRunnable::PrintWorkerMemoryUsage(bool bAllowToWaitForLock)
+{
+	if (bAllowToWaitForLock)
+	{
+		FScopeLock WorkerScopeLock(&WorkerInfosLock);
+		PrintWorkerMemoryUsageWithLockTaken();
+		return true;
+	}
+	else
+	{
+		FScopeTryLock WorkerScopeLock(&WorkerInfosLock);
+		if (WorkerScopeLock.IsLocked())
+		{
+			PrintWorkerMemoryUsageWithLockTaken();
+			return true;
+		}
+		return false;
+	}
+}
+
 FShaderCompileMemoryUsage FShaderCompileThreadRunnable::GetExternalWorkerMemoryUsage()
 {
 	FScopeLock WorkerScopeLock(&WorkerInfosLock);
@@ -5570,9 +5590,11 @@ IDistributedBuildController* FShaderCompilingManager::FindRemoteCompilerControll
 
 void FShaderCompilingManager::ReportMemoryUsage()
 {
+	// This function runs from within an OOM callback. It should not take locks, as much as possible.
+	constexpr bool bAllowToWaitForLock = false;
 	for (const TUniquePtr<FShaderCompileThreadRunnableBase>& ThreadPtr : Threads)
 	{
-		ThreadPtr->PrintWorkerMemoryUsage();
+		ThreadPtr->PrintWorkerMemoryUsage(bAllowToWaitForLock);
 	}
 }
 
diff --git a/Engine/Source/Runtime/Engine/Public/ShaderCompiler.h b/Engine/Source/Runtime/Engine/Public/ShaderCompiler.h
--- a/Engine/Source/Runtime/Engine/Public/ShaderCompiler.h
+++ b/Engine/Source/Runtime/Engine/Public/ShaderCompiler.h
@@ -140,7 +140,12 @@ protected:
 	
 	TAtomic<bool> bForceFinish;
 
-	virtual void PrintWorkerMemoryUsage() {}
+	/**
+	* Tries to print out the memory usage of all shader compile workers. When called during an out-of-memory event, it is useful to allow this process
+	* to wait for any locks, so we can rule out deadlocks while reporting out-of-memory errors.
+	* Returns whether the memory usage was successfully printed.
+	*/
+	virtual bool PrintWorkerMemoryUsage(bool bAllowToWaitForLock=true) { return false; }
 	/** Returns the amount of memory (in bytes) used by external processes related to this, if any. */
 	virtual FShaderCompileMemoryUsage GetExternalWorkerMemoryUsage() { return FShaderCompileMemoryUsage{}; }
 
@@ -207,7 +212,7 @@ public:
 	virtual ~FShaderCompileThreadRunnable();
 
 protected:
-	virtual void PrintWorkerMemoryUsage() override;
+	virtual bool PrintWorkerMemoryUsage(bool bAllowToWaitForLock = true) override;
 	virtual FShaderCompileMemoryUsage GetExternalWorkerMemoryUsage() override;
 
 private:
@@ -238,6 +243,8 @@ private:
 	virtual int32 CompilingLoop() override;
 
 	virtual void OnMachineResourcesChanged() override;
+
+	void PrintWorkerMemoryUsageWithLockTaken();
 };
 
 class FShaderCompileUtilities
