World Partition: add support for IStaticSpatialIndexDataInterface.GetAllocatedSize.

diff --git a/Engine/Source/Runtime/Engine/Private/WorldPartition/RuntimeHashSet/StaticSpatialIndex.cpp b/Engine/Source/Runtime/Engine/Private/WorldPartition/RuntimeHashSet/StaticSpatialIndex.cpp
--- a/Engine/Source/Runtime/Engine/Private/WorldPartition/RuntimeHashSet/StaticSpatialIndex.cpp
+++ b/Engine/Source/Runtime/Engine/Private/WorldPartition/RuntimeHashSet/StaticSpatialIndex.cpp
@@ -21,7 +21,7 @@ bool FStaticSpatialIndex::FListImpl::ForEachElement(TFunctionRef<bool(uint32 InV
 	return true;
 }
 
-bool  FStaticSpatialIndex::FListImpl::ForEachIntersectingElement(const FBox& InBox, TFunctionRef<bool(uint32 InValueIndex)> Func) const
+bool FStaticSpatialIndex::FListImpl::ForEachIntersectingElement(const FBox& InBox, TFunctionRef<bool(uint32 InValueIndex)> Func) const
 {
 	for (uint32 ValueIndex : Elements)
 	{
@@ -55,6 +55,11 @@ bool FStaticSpatialIndex::FListImpl::ForEachIntersectingElement(const FSphere& I
 	return true;
 }
 
+uint32 FStaticSpatialIndex::FListImpl::GetAllocatedSize() const
+{
+	return sizeof(*this) + Elements.GetAllocatedSize();
+}
+
 void FStaticSpatialIndex::FRTreeImpl::Init(const TArray<TPair<FBox, uint32>>& InElements)
 {
 	if (InElements.Num())
@@ -213,4 +218,30 @@ bool FStaticSpatialIndex::FRTreeImpl::ForEachIntersectingElementRecursive(const
 		}
 	}
 	return true;
-}
\ No newline at end of file
+}
+
+uint32 FStaticSpatialIndex::FRTreeImpl::GetAllocatedSize() const
+{
+	TFunction<uint32(const FNode*, uint32)> GetAllocatedSizeRecursive = [this, &GetAllocatedSizeRecursive](const FNode* Node, uint32 BaseSize)
+	{
+		uint32 AllocatedSize = BaseSize;
+
+		if (Node->Content.IsType<FNode::FNodeType>())
+		{
+			AllocatedSize += Node->Content.Get<FNode::FNodeType>().GetAllocatedSize();
+
+			for (auto& ChildNode : Node->Content.Get<FNode::FNodeType>())
+			{
+				AllocatedSize += GetAllocatedSizeRecursive(&ChildNode, sizeof(FNode));
+			}
+		}
+		else
+		{
+			AllocatedSize += Node->Content.Get<FNode::FLeafType>().GetAllocatedSize();
+		}
+
+		return AllocatedSize;
+	};
+	
+	return sizeof(*this) + GetAllocatedSizeRecursive(&RootNode, 0);
+}
diff --git a/Engine/Source/Runtime/Engine/Public/WorldPartition/RuntimeHashSet/StaticSpatialIndex.h b/Engine/Source/Runtime/Engine/Public/WorldPartition/RuntimeHashSet/StaticSpatialIndex.h
--- a/Engine/Source/Runtime/Engine/Public/WorldPartition/RuntimeHashSet/StaticSpatialIndex.h
+++ b/Engine/Source/Runtime/Engine/Public/WorldPartition/RuntimeHashSet/StaticSpatialIndex.h
@@ -11,7 +11,8 @@
 struct IStaticSpatialIndexDataInterface
 {
 	virtual ~IStaticSpatialIndexDataInterface() {}
-	virtual const FBox& GetBox(uint32 InIndex) const =0;
+	virtual const FBox& GetBox(uint32 InIndex) const = 0;
+	virtual uint32 GetAllocatedSize() const = 0;
 };
 
 template <typename ValueType, class SpatialIndexType>
@@ -80,6 +81,11 @@ public:
 		return Elements[InIndex].Key;
 	}
 
+	virtual uint32 GetAllocatedSize() const override
+	{
+		return sizeof(*this) + Elements.GetAllocatedSize() + SpatialIndex.GetAllocatedSize();
+	}
+
 private:
 	void InitSpatialIndex()
 	{		
@@ -122,6 +128,8 @@ namespace FStaticSpatialIndex
 		bool ForEachIntersectingElement(const FBox& InBox, TFunctionRef<bool(uint32 InValueIndex)> Func) const;
 		bool ForEachIntersectingElement(const FSphere& InSphere, TFunctionRef<bool(uint32 InValueIndex)> Func) const;
 
+		uint32 GetAllocatedSize() const;
+
 	private:
 		TArray<uint32> Elements;
 	};
@@ -139,6 +147,8 @@ namespace FStaticSpatialIndex
 		bool ForEachIntersectingElement(const FBox& InBox, TFunctionRef<bool(uint32 InValueIndex)> Func) const;
 		bool ForEachIntersectingElement(const FSphere& InSphere, TFunctionRef<bool(uint32 InValueIndex)> Func) const;
 
+		uint32 GetAllocatedSize() const;
+
 	private:
 		struct FNode
 		{
