World Partition: added support for 2D queries in streaming sources.

diff --git a/Engine/Source/Runtime/Engine/Private/WorldPartition/RuntimeHashSet/WorldPartitionRuntimeHashSet.cpp b/Engine/Source/Runtime/Engine/Private/WorldPartition/RuntimeHashSet/WorldPartitionRuntimeHashSet.cpp
--- a/Engine/Source/Runtime/Engine/Private/WorldPartition/RuntimeHashSet/WorldPartitionRuntimeHashSet.cpp
+++ b/Engine/Source/Runtime/Engine/Private/WorldPartition/RuntimeHashSet/WorldPartitionRuntimeHashSet.cpp
@@ -16,19 +16,34 @@ void FRuntimePartitionStreamingData::CreatePartitionsSpatialIndex() const
 	if (!SpatialIndex)
 	{
 		SpatialIndex = MakeUnique<FStaticSpatialIndexType>();
-
-		TArray<TPair<FBox, TObjectPtr<UWorldPartitionRuntimeCell>>> PartitionsElements;
-		Algo::Transform(SpatiallyLoadedCells, PartitionsElements, [](UWorldPartitionRuntimeCell* Cell)
 		{
-			return TPair<FBox, TObjectPtr<UWorldPartitionRuntimeCell>>(Cell->GetContentBounds(), Cell);
-		});
-		SpatialIndex->Init(PartitionsElements);
+			TArray<TPair<FBox, TObjectPtr<UWorldPartitionRuntimeCell>>> PartitionsElements;
+			Algo::Transform(SpatiallyLoadedCells, PartitionsElements, [](UWorldPartitionRuntimeCell* Cell)
+			{
+				return TPair<FBox, TObjectPtr<UWorldPartitionRuntimeCell>>(Cell->GetStreamingBounds(), Cell);
+			});
+			SpatialIndex->Init(PartitionsElements);
+		}
+		
+		SpatialIndex2D = MakeUnique<FStaticSpatialIndexType>();
+		{
+			TArray<TPair<FBox, TObjectPtr<UWorldPartitionRuntimeCell>>> PartitionsElements;
+			Algo::Transform(SpatiallyLoadedCells, PartitionsElements, [](UWorldPartitionRuntimeCell* Cell)
+			{
+				FBox CellBounds = Cell->GetStreamingBounds();
+				CellBounds.Min.Z = -HALF_WORLD_MAX;
+				CellBounds.Max.Z = HALF_WORLD_MAX;
+				return TPair<FBox, TObjectPtr<UWorldPartitionRuntimeCell>>(CellBounds, Cell);
+			});
+			SpatialIndex2D->Init(PartitionsElements);
+		}
 	}
 }
 
 void FRuntimePartitionStreamingData::DestroyPartitionsSpatialIndex() const
 {
 	SpatialIndex.Reset();
+	SpatialIndex2D.Reset();
 }
 
 void URuntimeHashSetExternalStreamingObject::CreatePartitionsSpatialIndex() const
@@ -57,6 +72,11 @@ void URuntimeHashSetExternalStreamingObject::AddReferencedObjects(UObject* InThi
 		{
 			StreamingData.SpatialIndex->AddReferencedObjects(Collector);
 		}
+		
+		if (StreamingData.SpatialIndex2D.IsValid())
+		{
+			StreamingData.SpatialIndex2D->AddReferencedObjects(Collector);
+		}
 	}
 #endif
 
@@ -489,7 +509,7 @@ void UWorldPartitionRuntimeHashSet::ForEachStreamingCellsSources(const TArray<FW
 					{
 						const FSphere ShapeSphere(Shape.GetCenter(), Shape.GetRadius());
 
-						StreamingData->SpatialIndex->ForEachIntersectingElement(ShapeSphere, [this, &Source, &Shape, &Func](UWorldPartitionRuntimeCell* Cell)
+						auto ForEachIntersectingElementFunc = [this, &Source, &Shape, &Func](UWorldPartitionRuntimeCell* Cell)
 						{
 #if WITH_EDITOR
 							if (!IsCellRelevantFor(Cell->GetClientOnlyVisible()))
@@ -505,7 +525,16 @@ void UWorldPartitionRuntimeHashSet::ForEachStreamingCellsSources(const TArray<FW
 								Cell->AppendStreamingSourceInfo(Source, Shape);
 								Func(Cell, ((CellEffectiveWantedState == EDataLayerRuntimeState::Loaded) || (Source.TargetState == EStreamingSourceTargetState::Loaded)) ? EStreamingSourceTargetState::Loaded : EStreamingSourceTargetState::Activated);
 							}
-						});
+						};
+
+						if (Source.bForce2D)
+						{
+							StreamingData->SpatialIndex2D->ForEachIntersectingElement(ShapeSphere, ForEachIntersectingElementFunc);
+						}
+						else
+						{
+							StreamingData->SpatialIndex->ForEachIntersectingElement(ShapeSphere, ForEachIntersectingElementFunc);
+						}
 					});
 				}
 			}
diff --git a/Engine/Source/Runtime/Engine/Private/WorldPartition/WorldPartitionStreamingSource.cpp b/Engine/Source/Runtime/Engine/Private/WorldPartition/WorldPartitionStreamingSource.cpp
--- a/Engine/Source/Runtime/Engine/Private/WorldPartition/WorldPartitionStreamingSource.cpp
+++ b/Engine/Source/Runtime/Engine/Private/WorldPartition/WorldPartitionStreamingSource.cpp
@@ -57,6 +57,11 @@ FString FWorldPartitionStreamingSource::ToString() const
 		(int32)(Velocity.Size() * 0.0223694f)
 	);
 
+	if (bForce2D)
+	{
+		StringBuilder += TEXT(" | Force2D");
+	}
+
 	if (Shapes.Num())
 	{
 		StringBuilder += TEXT(" | ");
diff --git a/Engine/Source/Runtime/Engine/Public/WorldPartition/RuntimeHashSet/WorldPartitionRuntimeHashSet.h b/Engine/Source/Runtime/Engine/Public/WorldPartition/RuntimeHashSet/WorldPartitionRuntimeHashSet.h
--- a/Engine/Source/Runtime/Engine/Public/WorldPartition/RuntimeHashSet/WorldPartitionRuntimeHashSet.h
+++ b/Engine/Source/Runtime/Engine/Public/WorldPartition/RuntimeHashSet/WorldPartitionRuntimeHashSet.h
@@ -87,6 +87,7 @@ struct FRuntimePartitionStreamingData
 
 	// Transient
 	mutable TUniquePtr<FStaticSpatialIndexType> SpatialIndex;
+	mutable TUniquePtr<FStaticSpatialIndexType> SpatialIndex2D;
 };
 
 template<>
diff --git a/Engine/Source/Runtime/Engine/Public/WorldPartition/WorldPartitionStreamingSource.h b/Engine/Source/Runtime/Engine/Public/WorldPartition/WorldPartitionStreamingSource.h
--- a/Engine/Source/Runtime/Engine/Public/WorldPartition/WorldPartitionStreamingSource.h
+++ b/Engine/Source/Runtime/Engine/Public/WorldPartition/WorldPartitionStreamingSource.h
@@ -346,6 +346,7 @@ struct FWorldPartitionStreamingSource
 		, TargetBehavior(EStreamingSourceTargetBehavior::Include)
 		, bReplay(false)
 		, bRemote(false)
+		, bForce2D(false)
 		, Hash2D(0)
 		, Hash3D(0)
 		, OldLocation(FVector::ZeroVector)
@@ -367,6 +368,7 @@ struct FWorldPartitionStreamingSource
 		, TargetBehavior(EStreamingSourceTargetBehavior::Include)
 		, bReplay(false)
 		, bRemote(bRemote)
+		, bForce2D(false)
 		, Hash2D(0)
 		, Hash3D(0)
 		, OldLocation(InLocation)
@@ -397,6 +399,7 @@ struct FWorldPartitionStreamingSource
 		Shapes = Other.Shapes;
 		bReplay = Other.bReplay;
 		bRemote = Other.bRemote;
+		bForce2D = Other.bForce2D;
 		Hash2D = Other.Hash2D;
 		Hash3D = Other.Hash3D;
 		OldLocation = Other.OldLocation;
@@ -465,6 +468,9 @@ struct FWorldPartitionStreamingSource
 	/** If true, this streaming source is from a remote session */
 	bool bRemote;
 
+	/** If true, this streaming source will force intersecting cells in 2D */
+	bool bForce2D;
+
 	/** Returns a box encapsulating all shapes. */
 	inline FBox CalcBounds(float InGridLoadingRange, FName InGridName, bool bCalcIn2D = false) const
 	{
