World Partition Runtime Hashset: debug draw will now show the fixed size grid with low opacity under the streaming cells bounds.

diff --git a/Engine/Source/Runtime/Engine/Private/WorldPartition/RuntimeHashSet/WorldPartitionRuntimeHashSetDebugDraw.cpp b/Engine/Source/Runtime/Engine/Private/WorldPartition/RuntimeHashSet/WorldPartitionRuntimeHashSetDebugDraw.cpp
--- a/Engine/Source/Runtime/Engine/Private/WorldPartition/RuntimeHashSet/WorldPartitionRuntimeHashSetDebugDraw.cpp
+++ b/Engine/Source/Runtime/Engine/Private/WorldPartition/RuntimeHashSet/WorldPartitionRuntimeHashSetDebugDraw.cpp
@@ -135,8 +135,15 @@ bool UWorldPartitionRuntimeHashSet::Draw2D(FWorldPartitionDraw2DContext& DrawCon
 		{
 			for (const UWorldPartitionRuntimeCell* Cell : FilteredCells) //-V1078
 			{
+				// Draw fixed cell bounds
 				const FVector2D CellBoundsSize = FVector2D(Cell->GetCellBounds().GetSize());
 				const FVector2D CellBoundsMin = FVector2D(Cell->GetCellBounds().Min);
+				DrawContext.LocalDrawTile(GridScreenBounds, CellBoundsMin, CellBoundsSize, FLinearColor::Gray.CopyWithNewOpacity(0.05f), WorldToScreen);
+				DrawContext.LocalDrawBox(GridScreenBounds, CellBoundsMin, CellBoundsSize, FLinearColor::Black.CopyWithNewOpacity(0.05f), 1, WorldToScreen);
+
+				// Draw streaming cell bounds
+				const FVector2D CellStreamingBoundsSize = FVector2D(Cell->GetStreamingBounds().GetSize());
+				const FVector2D CellStreamingBoundsMin = FVector2D(Cell->GetStreamingBounds().Min);
 
 				float CellOpacity = 0.0f;
 				TArray<FLinearColor> CellColors;
@@ -177,16 +184,16 @@ bool UWorldPartitionRuntimeHashSet::Draw2D(FWorldPartitionDraw2DContext& DrawCon
 				}
 
 				FVector2D::FReal BoundsOffsetX = 0;
-				const FVector2D::FReal BoundsOffsetStepX = CellBoundsSize.X / CellColors.Num();
+				const FVector2D::FReal BoundsOffsetStepX = CellStreamingBoundsSize.X / CellColors.Num();
 				for (const FLinearColor& CellColor : CellColors)
 				{
-					const FVector2D EffectiveCellBoundsMin(CellBoundsMin.X + BoundsOffsetX, CellBoundsMin.Y);
-					const FVector2D EffectiveCellBoundsSize(BoundsOffsetStepX, CellBoundsSize.Y);
-					DrawContext.LocalDrawTile(GridScreenBounds, EffectiveCellBoundsMin, EffectiveCellBoundsSize, CellColor.CopyWithNewOpacity(CellOpacity), WorldToScreen);
+					const FVector2D EffectiveCellStreamingBoundsMin(CellStreamingBoundsMin.X + BoundsOffsetX, CellStreamingBoundsMin.Y);
+					const FVector2D EffectiveCellStreamingBoundsSize(BoundsOffsetStepX, CellStreamingBoundsSize.Y);
+					DrawContext.LocalDrawTile(GridScreenBounds, EffectiveCellStreamingBoundsMin, EffectiveCellStreamingBoundsSize, CellColor.CopyWithNewOpacity(CellOpacity), WorldToScreen);
 					BoundsOffsetX += BoundsOffsetStepX;
 				}
 				
-				DrawContext.LocalDrawBox(GridScreenBounds, CellBoundsMin, CellBoundsSize, FLinearColor::Black, 1, WorldToScreen);
+				DrawContext.LocalDrawBox(GridScreenBounds, CellStreamingBoundsMin, CellStreamingBoundsSize, FLinearColor::Black, 1, WorldToScreen);
 			}
 		}
 
@@ -349,8 +356,8 @@ void UWorldPartitionRuntimeHashSet::Draw3D(const TArray<FWorldPartitionStreaming
 
 				StreamingData.SpatialIndex.Get()->ForEachIntersectingElement(ShapeSphere, [this, VisualizeMode, &DataLayerDebugColors, ContentBundleManager, OwningWorld, &WorldPartitionTransform](UWorldPartitionRuntimeCell* Cell)
 				{
-					const FVector2D CellBoundsSize = FVector2D(Cell->GetCellBounds().GetSize());
-					const FVector2D CellBoundsMin = FVector2D(Cell->GetCellBounds().Min);
+					const FVector2D CellStreamingBoundsSize = FVector2D(Cell->GetStreamingBounds().GetSize());
+					const FVector2D CellStreamingBoundsMin = FVector2D(Cell->GetStreamingBounds().Min);
 
 					float CellOpacity = 0.0f;
 					TArray<FLinearColor> CellColors;
@@ -391,7 +398,7 @@ void UWorldPartitionRuntimeHashSet::Draw3D(const TArray<FWorldPartitionStreaming
 					}
 
 					// Draw Cell using its debug color
-					const FBox Box(Cell->GetCellBounds());
+					const FBox Box(Cell->GetStreamingBounds());
 					const FVector BoxCenter(Box.GetCenter());
 					const FColor BoxColor(CellColors[0].CopyWithNewOpacity(CellOpacity).ToFColor(true));
 					const FVector CellPos = WorldPartitionTransform.TransformPosition(BoxCenter);
