Fixed case where overriden MinLOD on a skinned mesh component could be higher than asset MinLOD resulting in rendering corruption

diff --git a/Engine/Source/Runtime/Engine/Private/Components/SkinnedMeshComponent.cpp b/Engine/Source/Runtime/Engine/Private/Components/SkinnedMeshComponent.cpp
--- a/Engine/Source/Runtime/Engine/Private/Components/SkinnedMeshComponent.cpp
+++ b/Engine/Source/Runtime/Engine/Private/Components/SkinnedMeshComponent.cpp
@@ -3883,7 +3883,9 @@ void USkinnedMeshComponent::OverrideMinLOD(int32 InNewMinLOD)
 
 int32 USkinnedMeshComponent::ComputeMinLOD() const
 {
-	int32 MinLodIndex = bOverrideMinLod ? MinLodModel : GetSkinnedAsset()->GetMinLodIdx();
+	int32 AssetMinLod = GetSkinnedAsset()->GetMinLodIdx();
+	// overriden MinLOD can't be higher than asset MinLOD
+	int32 MinLodIndex = bOverrideMinLod ? FMath::Max(MinLodModel, AssetMinLod) : AssetMinLod;
 	int32 NumLODs = GetNumLODs();
 	// want to make sure MinLOD stays within the valid range
 	MinLodIndex = FMath::Min(MinLodIndex, NumLODs - 1);
