Change the default capture source for WorldPartition Minimap captures to be full SceneColor instead of just Base Color. Some shading models (such as water) don't render into the BaseColor gbuffer texture which prevents them from showing up in minimap captures.
Since the capture is now taking SceneColor, we also need to disable most lighting/postprocess effects via the capture showflags or they will be included in the minimap capture.

diff --git a/Engine/Source/Runtime/Engine/Private/WorldPartition/WorldPartitionMiniMap.cpp b/Engine/Source/Runtime/Engine/Private/WorldPartition/WorldPartitionMiniMap.cpp
--- a/Engine/Source/Runtime/Engine/Private/WorldPartition/WorldPartitionMiniMap.cpp
+++ b/Engine/Source/Runtime/Engine/Private/WorldPartition/WorldPartitionMiniMap.cpp
@@ -24,7 +24,7 @@ AWorldPartitionMiniMap::AWorldPartitionMiniMap(const FObjectInitializer& ObjectI
 	, MiniMapTexture(nullptr)
 	, WorldUnitsPerPixel(50)
 	, BuilderCellSize(102400)
-	, CaptureSource(ESceneCaptureSource::SCS_BaseColor)
+	, CaptureSource(ESceneCaptureSource::SCS_SceneColorHDR)
 	, CaptureWarmupFrames(5)
 {
 }
diff --git a/Engine/Source/Runtime/Engine/Private/WorldPartition/WorldPartitionMiniMapHelper.cpp b/Engine/Source/Runtime/Engine/Private/WorldPartition/WorldPartitionMiniMapHelper.cpp
--- a/Engine/Source/Runtime/Engine/Private/WorldPartition/WorldPartitionMiniMapHelper.cpp
+++ b/Engine/Source/Runtime/Engine/Private/WorldPartition/WorldPartitionMiniMapHelper.cpp
@@ -89,6 +89,15 @@ void FWorldPartitionMiniMapHelper::CaptureBoundsMiniMapToTexture(UWorld* InWorld
 	CaptureComponent->bCaptureEveryFrame = false;
 	CaptureComponent->bCaptureOnMovement = false;
 
+	// We only want a basic color view of the world which isn't affected by any lighting or atmospheric effects which might obscure the minimap view:
+	CaptureComponent->ShowFlags.Lighting = false;
+	CaptureComponent->ShowFlags.Atmosphere = false;
+	CaptureComponent->ShowFlags.PostProcessing = false;
+	CaptureComponent->ShowFlags.AntiAliasing = false;
+	CaptureComponent->ShowFlags.Fog = false;
+	CaptureComponent->ShowFlags.VolumetricFog = false;
+	CaptureComponent->ShowFlags.DynamicShadows = false;
+
 	// Disable vignetting, otherwise we'll see it as a pattern between each captured tiles
 	CaptureComponent->PostProcessSettings.bOverride_VignetteIntensity = true;
 	CaptureComponent->PostProcessSettings.VignetteIntensity = 0.0f;
