Added 2 advanced actor actions : 1- Remove actor(s) from their External Data Layer, 2- Move actor(s) to another External Data Layer
These operations will first make sure that they won't generate invalid cross-domain references. If so, the operation will fail and be reported through logging and notification.

diff --git a/Engine/Source/Editor/DataLayerEditor/DataLayerEditor.Build.cs b/Engine/Source/Editor/DataLayerEditor/DataLayerEditor.Build.cs
--- a/Engine/Source/Editor/DataLayerEditor/DataLayerEditor.Build.cs
+++ b/Engine/Source/Editor/DataLayerEditor/DataLayerEditor.Build.cs
@@ -15,7 +15,8 @@ public class DataLayerEditor : ModuleRules
 				"EditorSubsystem",
 				"PropertyEditor",
 				"Engine",
-                "InputCore",
+				"LevelEditor",
+				"InputCore",
 				"Slate",
 				"SlateCore",
 				"UnrealEd",
diff --git a/Engine/Source/Editor/DataLayerEditor/Private/DataLayerEditorModule.cpp b/Engine/Source/Editor/DataLayerEditor/Private/DataLayerEditorModule.cpp
--- a/Engine/Source/Editor/DataLayerEditor/Private/DataLayerEditorModule.cpp
+++ b/Engine/Source/Editor/DataLayerEditor/Private/DataLayerEditorModule.cpp
@@ -6,8 +6,14 @@
 #include "DataLayer/DataLayerInstanceCustomization.h"
 #include "DataLayer/DataLayerNameEditSink.h"
 #include "DataLayer/DataLayerPropertyTypeCustomization.h"
+#include "DataLayer/DataLayerPropertyTypeCustomizationHelper.h"
 #include "DataLayer/SDataLayerBrowser.h"
+#include "WorldPartition/WorldPartitionLog.h"
 #include "WorldPartition/DataLayer/ExternalDataLayerAsset.h"
+#include "WorldPartition/DataLayer/ExternalDataLayerManager.h"
+#include "Framework/Notifications/NotificationManager.h"
+#include "Framework/Docking/TabManager.h"
+#include "Widgets/Notifications/SNotificationList.h"
 #include "EditorWidgetsModule.h"
 #include "HAL/Platform.h"
 #include "Modules/ModuleManager.h"
@@ -17,11 +23,21 @@
 #include "UObject/NameTypes.h"
 #include "Widgets/DeclarativeSyntaxSupport.h"
 #include "Widgets/SWidget.h"
+#include "LevelEditorMenuContext.h"
+#include "ToolMenus.h"
+#include "ToolMenu.h"
+#include "ToolMenuEntry.h"
+#include "ToolMenuSection.h"
+#include "Selection.h"
+#include "Algo/AnyOf.h"
+#include "Editor.h"
 
 class AActor;
 
 IMPLEMENT_MODULE(FDataLayerEditorModule, DataLayerEditor );
 
+#define LOCTEXT_NAMESPACE "DataLayerEditorModule"
+
 static const FName NAME_ActorDataLayer(TEXT("ActorDataLayer"));
 static const FName NAME_DataLayerInstance(TEXT("DataLayerInstance"));
 
@@ -33,10 +49,15 @@ void FDataLayerEditorModule::StartupModule()
 
 	FEditorWidgetsModule& EditorWidgetsModule = FModuleManager::LoadModuleChecked<FEditorWidgetsModule>("EditorWidgets");
 	EditorWidgetsModule.GetObjectNameEditSinkRegistry()->RegisterObjectNameEditSink(MakeShared<FDataLayerNameEditSink>());
+
+	UToolMenus::Get()->RegisterStartupCallback(FSimpleMulticastDelegate::FDelegate::CreateRaw(this, &FDataLayerEditorModule::RegisterMenus));
 }
 
 void FDataLayerEditorModule::ShutdownModule()
 {
+	UToolMenus::UnRegisterStartupCallback(this);
+	UToolMenus::UnregisterOwner(this);
+
 	if (FPropertyEditorModule* PropertyModule = FModuleManager::GetModulePtr<FPropertyEditorModule>("PropertyEditor"))
 	{
 		PropertyModule->UnregisterCustomPropertyTypeLayout(NAME_ActorDataLayer);
@@ -68,4 +89,148 @@ bool FDataLayerEditorModule::AddActorToDataLayers(AActor* Actor, const TArray<UD
 void FDataLayerEditorModule::SetActorEditorContextCurrentExternalDataLayer(const UExternalDataLayerAsset* InExternalDataLayerAsset)
 {
 	UDataLayerEditorSubsystem::Get()->SetActorEditorContextCurrentExternalDataLayer(InExternalDataLayerAsset);
-}
\ No newline at end of file
+}
+
+bool FDataLayerEditorModule::MoveActorsToExternalDataLayer(const TArray<AActor*>& InSelectedActors, const UExternalDataLayerInstance* InExternalDataLayerInstance, FText* OutReason)
+{
+	if (InSelectedActors.IsEmpty())
+	{
+		if (OutReason)
+		{
+			*OutReason = LOCTEXT("NoActorToProcess", "No actor to process");
+		}
+		return false;
+	}
+	return FExternalDataLayerHelper::MoveActorsToExternalDataLayer(InSelectedActors, InExternalDataLayerInstance, OutReason);
+}
+
+namespace UE::Private::DataLayerEditorModule
+{
+	void LogWarningAndNotify(const FText& InWarningMessage, const FText& InDetailedMessage)
+	{
+		UE_LOG(LogWorldPartition, Warning, TEXT("%s : %s"), *InWarningMessage.ToString(), *InDetailedMessage.ToString());
+		FNotificationInfo WarningInfo(InWarningMessage);
+		WarningInfo.SubText = InDetailedMessage;
+		WarningInfo.ExpireDuration = 5.0f;
+		WarningInfo.Hyperlink = FSimpleDelegate::CreateLambda([]() { FGlobalTabmanager::Get()->TryInvokeTab(FName("OutputLog")); });
+		WarningInfo.HyperlinkText = LOCTEXT("ShowMessageLogHyperlink", "Show Output Log");
+		FSlateNotificationManager::Get().AddNotification(WarningInfo);
+	}
+}
+
+void FDataLayerEditorModule::RegisterMenus()
+{
+	FToolMenuOwnerScoped OwnerScoped(this);
+	FToolUIAction RemoveExternalDataLayerAction;
+	RemoveExternalDataLayerAction.ExecuteAction = FToolMenuExecuteAction::CreateLambda([this](const FToolMenuContext& InContext)
+	{
+		FText FailureReason;
+		FScopedTransaction Transaction(LOCTEXT("RemoveActorsFromExternalDataLayer", "Remove Actor(s) From External Data Layer"));
+		TArray<AActor*> ActorsWithExternalDataLayer;
+		Algo::TransformIf(GetSelectedActors(), ActorsWithExternalDataLayer, [](const AActor* Actor) { return Actor && Actor->GetExternalDataLayerAsset(); }, [](AActor* Actor) { return Actor; });
+		if (!MoveActorsToExternalDataLayer(ActorsWithExternalDataLayer, nullptr, &FailureReason))
+		{
+			Transaction.Cancel();
+			const FText WarningMessage = LOCTEXT("RemoveActorsFromExternalDataLayerFailed", "Failed to remove actor(s) from External Data Layer");
+			UE::Private::DataLayerEditorModule::LogWarningAndNotify(WarningMessage, FailureReason);
+		}
+	});
+	RemoveExternalDataLayerAction.CanExecuteAction = FToolMenuCanExecuteAction::CreateLambda([this](const FToolMenuContext& InContext)
+	{
+		return Algo::AnyOf(GetSelectedActors(), [](AActor* Actor) { return Actor && Actor->GetExternalDataLayerAsset(); });
+	});
+
+	auto MoveToExternalDataLayerMenu = [this](UToolMenu* Menu)
+	{
+		TSet<const UExternalDataLayerAsset*> ExternalDataLayerAssets;
+		Algo::TransformIf(GetSelectedActors(), ExternalDataLayerAssets, [](AActor* Actor) { return Actor && Actor->GetExternalDataLayerAsset(); }, [](AActor* Actor) { return Actor->GetExternalDataLayerAsset(); });
+
+		FToolMenuSection& Section = Menu->AddSection("External Data Layer Picker");
+		TSharedRef< SWidget > MenuWidget =
+			SNew(SVerticalBox)
+			+ SVerticalBox::Slot()
+			.MaxHeight(400.0f)
+			[
+				FDataLayerPropertyTypeCustomizationHelper::CreateDataLayerMenu(
+					// OnDataLayerSelectedFunction
+					[this](const UDataLayerInstance* DataLayerInstance)
+					{
+						const UExternalDataLayerInstance* ExternalDataLayerInstance = Cast<UExternalDataLayerInstance>(DataLayerInstance);
+						if (ExternalDataLayerInstance && ExternalDataLayerInstance->GetExternalDataLayerAsset())
+						{
+							FText FailureReason;
+							FScopedTransaction Transaction(LOCTEXT("MoveActorsToExternalDataLayer", "Move Actor(s) To External Data Layer"));
+							if (!MoveActorsToExternalDataLayer(GetSelectedActors(), ExternalDataLayerInstance, &FailureReason))
+							{
+								Transaction.Cancel();
+								const FText WarningMessage = LOCTEXT("MoveActorsToExternalDataLayerFailed", "Failed to move actor(s) to External Data Layer");
+								UE::Private::DataLayerEditorModule::LogWarningAndNotify(WarningMessage, FailureReason);
+							}
+						}
+					},
+					// OnShouldFilterDataLayerInstanceFunction
+					[ExternalDataLayerAssets](const UDataLayerInstance* DataLayerInstance)
+					{
+						const UExternalDataLayerInstance* ExternalDataLayerInstance = Cast<UExternalDataLayerInstance>(DataLayerInstance);
+						const UExternalDataLayerAsset* ExternalDataLayerAsset = ExternalDataLayerInstance ? ExternalDataLayerInstance->GetExternalDataLayerAsset() : nullptr;
+						return (!ExternalDataLayerAsset || ExternalDataLayerAssets.Contains(ExternalDataLayerAsset));
+					})
+			];
+
+		Section.AddEntry(FToolMenuEntry::InitWidget("PickExternalDataLayer", MenuWidget, FText::GetEmpty(), false));
+	};
+
+	auto FillExternalDataLayerMenu = [this, RemoveExternalDataLayerAction, MoveToExternalDataLayerMenu](UToolMenu* SubMenu)
+	{
+		FToolMenuSection& ExternalDataLayerSection = SubMenu->AddSection("External Data Layer");
+		ExternalDataLayerSection.AddMenuEntry(
+			"ActorRemoveFromExternalDataLayerMenu",
+			LOCTEXT("RemoveActorsFromTheirExternalDataLayerMenuEntry", "Remove Actors(s)"),
+			LOCTEXT("RemoveActorsFromTheirExternalDataLayerMenu_ToolTip", "Remove Actor(s) from their External Data Layer"),
+			FSlateIcon(),
+			RemoveExternalDataLayerAction);
+
+		ExternalDataLayerSection.AddSubMenu(
+			"MoveToExternalDataLayerSubMenu",
+			LOCTEXT("MoveActorsToPickedExternalDataLayerSubMenu", "Move Actor(s) To"),
+			LOCTEXT("MoveActorsToPickedExternalDataLayerSubMenu_ToolTip", "Move Actor(s) to picked External Data Layer"),
+			FNewToolMenuDelegate::CreateLambda(MoveToExternalDataLayerMenu));
+	};
+
+	UToolMenu* Menu = UToolMenus::Get()->ExtendMenu("LevelEditor.ActorContextMenu");
+	FToolMenuSection& ActorSection = Menu->FindOrAddSection("ActorOptions");
+
+	FToolUIAction FillExternalDataLayerMenuAction;
+	FillExternalDataLayerMenuAction.IsActionVisibleDelegate = FToolMenuCanExecuteAction::CreateLambda([this](const FToolMenuContext& InContext)
+	{
+		UWorld* EditorWorld = GEditor->GetEditorWorldContext().World();
+		return EditorWorld && EditorWorld->IsPartitionedWorld();
+	});
+	FillExternalDataLayerMenuAction.CanExecuteAction = FToolMenuCanExecuteAction::CreateLambda([this](const FToolMenuContext& InContext)
+	{
+		TArray<AActor*> SelectedActors = GetSelectedActors();
+		UExternalDataLayerManager* ExternalDataLayerManager = SelectedActors.Num() ? UExternalDataLayerManager::GetExternalDataLayerManager(SelectedActors[0]) : nullptr;
+		return ExternalDataLayerManager && ExternalDataLayerManager->HasInjectedExternalDataLayerAssets();
+	});
+
+	ActorSection.AddSubMenu(
+		"ExternalDataLayerSubMenu",
+		LOCTEXT("ExternalDataLayerSubMenu", "External Data Layer"),
+		LOCTEXT("ExternalDataLayerSubMenu_ToolTip", "External Data Layer Utils"),
+		FNewToolMenuDelegate::CreateLambda(FillExternalDataLayerMenu),
+		FillExternalDataLayerMenuAction,
+		EUserInterfaceActionType::Button,
+		false,
+		FSlateIcon(FAppStyle::GetAppStyleSetName(), "DataLayer.External"));
+}
+
+TArray<AActor*> FDataLayerEditorModule::GetSelectedActors() const
+{
+	TArray<UObject*> SelectedObjects;
+	GEditor->GetSelectedActors()->GetSelectedObjects(AActor::StaticClass(), /*out*/ SelectedObjects);
+	TArray<AActor*> SelectedActors;
+	Algo::TransformIf(SelectedObjects, SelectedActors, [](UObject* Object) { return Object && Object->IsA<AActor>(); }, [](UObject* Object) { return Cast<AActor>(Object); });
+	return SelectedActors;
+}
+
+#undef LOCTEXT_NAMESPACE
\ No newline at end of file
diff --git a/Engine/Source/Editor/DataLayerEditor/Public/DataLayerEditorModule.h b/Engine/Source/Editor/DataLayerEditor/Public/DataLayerEditorModule.h
--- a/Engine/Source/Editor/DataLayerEditor/Public/DataLayerEditorModule.h
+++ b/Engine/Source/Editor/DataLayerEditor/Public/DataLayerEditorModule.h
@@ -14,8 +14,10 @@ class AActor;
 class FExtender;
 class FUICommandList;
 class SWidget;
+class UToolMenu;
 class UDataLayerInstance;
 class UExternalDataLayerAsset;
+class UExternalDataLayerInstance;
 
 /**
  * The module holding all of the UI related pieces for DataLayers
@@ -54,6 +56,11 @@ public:
 	virtual void SetActorEditorContextCurrentExternalDataLayer(const UExternalDataLayerAsset* InExternalDataLayerAsset) override;
 
 private:
+	void RegisterMenus();
+	void MoveToExternalDataLayerMenu(UToolMenu* Menu);
+	bool MoveActorsToExternalDataLayer(const TArray<AActor*>& InSelectedActors, const UExternalDataLayerInstance* InExternalDataLayerInstance, FText* OutReason = nullptr);
+	TArray<AActor*> GetSelectedActors() const;
+
 	TWeakPtr<SWidget> DataLayerBrowser;
 
 	/** All extender delegates for the DataLayers menus */
diff --git a/Engine/Source/Runtime/Engine/Classes/GameFramework/Actor.h b/Engine/Source/Runtime/Engine/Classes/GameFramework/Actor.h
--- a/Engine/Source/Runtime/Engine/Classes/GameFramework/Actor.h
+++ b/Engine/Source/Runtime/Engine/Classes/GameFramework/Actor.h
@@ -979,7 +979,7 @@ protected:
 
 	TArray<TSoftObjectPtr<UDataLayerAsset>> PreEditChangeDataLayers;
 
-	UPROPERTY(VisibleAnywhere, AdvancedDisplay, Category = DataLayers, TextExportTransient, NonTransactional)
+	UPROPERTY(VisibleAnywhere, AdvancedDisplay, Category = DataLayers, TextExportTransient)
 	TObjectPtr<const UExternalDataLayerAsset> ExternalDataLayerAsset;
 
 public:
@@ -4227,6 +4227,7 @@ private:
 	friend struct FSetActorInstanceGuid;
 	friend struct FSetActorContentBundleGuid;
 	friend struct FAssignActorDataLayer;
+	friend struct FScopeRawAssignActorExternalDataLayer;
 	friend struct FSetActorSelectable;
 	friend struct FSetActorFolderPath;
 #endif
@@ -4541,13 +4542,14 @@ private:
 struct FAssignActorDataLayer
 {
 private:
-	static bool AddDataLayerAsset(AActor* InActor, const UDataLayerAsset* InDataLayerAsset);
-	static bool RemoveDataLayerAsset(AActor* InActor, const UDataLayerAsset* InDataLayerAsset);
+	ENGINE_API static bool AddDataLayerAsset(AActor* InActor, const UDataLayerAsset* InDataLayerAsset);
+	ENGINE_API static bool RemoveDataLayerAsset(AActor* InActor, const UDataLayerAsset* InDataLayerAsset);
 
 	friend class UDataLayerInstanceWithAsset;
 	friend class UDataLayerInstancePrivate;
 	friend class UExternalDataLayerInstance;
 	friend class ULevelInstanceSubsystem;
+	friend class FExternalDataLayerHelper;
 };
 
 struct FSetActorIsInLevelInstance
diff --git a/Engine/Source/Runtime/Engine/Private/WorldPartition/DataLayer/ExternalDataLayerHelper.cpp b/Engine/Source/Runtime/Engine/Private/WorldPartition/DataLayer/ExternalDataLayerHelper.cpp
--- a/Engine/Source/Runtime/Engine/Private/WorldPartition/DataLayer/ExternalDataLayerHelper.cpp
+++ b/Engine/Source/Runtime/Engine/Private/WorldPartition/DataLayer/ExternalDataLayerHelper.cpp
@@ -10,8 +10,14 @@
 #include "ExternalPackageHelper.h"
 #include "UObject/Object.h"
 #include "UObject/AssetRegistryTagsContext.h"
+#include "Subsystems/EditorActorSubsystem.h"
+#include "WorldPartition/DataLayer/ExternalDataLayerManager.h"
+#include "ReferencedAssetsUtils.h"
+#include "Editor.h"
 #endif
 
+#define LOCTEXT_NAMESPACE "ExternalDataLayerHelper"
+
 FString FExternalDataLayerHelper::GetExternalStreamingObjectPackageName(const UExternalDataLayerAsset* InExternalDataLayerAsset)
 {
 	check(InExternalDataLayerAsset);
@@ -95,6 +101,230 @@ void FExternalDataLayerHelper::GetExternalDataLayerUIDs(const FAssetData& Asset,
 	}
 }
 
+namespace UE::Private::ExternalDataLayerHelper
+{
+	static bool ValidateAssetUsingAssetReferenceRestrictions(const UObject* InAsset, TSet<FString>& OutInvalidReferenceReasons)
+	{
+		if (!InAsset)
+		{
+			return false;
+		}
+
+		TArray<UClass*> IgnoreClasses;
+		TArray<UPackage*> IgnorePackages;
+		TSet<UObject*> ReferencedAssets;
+		FFindReferencedAssets::BuildAssetList(const_cast<UObject*>(InAsset), IgnoreClasses, IgnorePackages, ReferencedAssets);
+
+		uint32 ErrorCount = 0;
+		if (ReferencedAssets.Num() > 0)
+		{
+			FAssetData ReferencingAssetData = FAssetData(InAsset);
+			FAssetReferenceFilterContext AssetReferenceFilterContext;
+			AssetReferenceFilterContext.ReferencingAssets = { ReferencingAssetData };
+			TSharedPtr<IAssetReferenceFilter> AssetReferenceFilter = GEditor ? GEditor->MakeAssetReferenceFilter(AssetReferenceFilterContext) : nullptr;
+			if (ensure(AssetReferenceFilter.IsValid()))
+			{
+				for (UObject* ReferencedAsset : ReferencedAssets)
+				{
+					FText FailureReason;
+					if (!AssetReferenceFilter->PassesFilter(FAssetData(ReferencedAsset), &FailureReason))
+					{
+						OutInvalidReferenceReasons.Add(FailureReason.ToString());
+						++ErrorCount;
+					}
+				}
+			}
+		}
+		return !ErrorCount;
+	}
+}
+
+// Only meant to be used by FExternalDataLayerHelper on stack
+struct FScopeRawAssignActorExternalDataLayer
+{
+private:
+	FScopeRawAssignActorExternalDataLayer(AActor* InActor, const UExternalDataLayerAsset* InExternalDataLayerAsset)
+		: Actor(InActor)
+		, ExternalDataLayerAsset(InActor->ExternalDataLayerAsset)
+	{
+		check(Actor);
+		Actor->ExternalDataLayerAsset = InExternalDataLayerAsset;
+	}
+
+	~FScopeRawAssignActorExternalDataLayer()
+	{
+		Actor->ExternalDataLayerAsset = ExternalDataLayerAsset;
+	}
+
+	AActor* Actor;
+	const UExternalDataLayerAsset* ExternalDataLayerAsset;
+
+	friend class FExternalDataLayerHelper;
+};
+
+bool FExternalDataLayerHelper::CanMoveActorsToExternalDataLayer(const TArray<AActor*>& InActors, const UExternalDataLayerInstance* InExternalDataLayerInstance, FText* OutFailureReason)
+{
+	auto CanMoveActorToExternalDataLayer = [](AActor* InActor, const UExternalDataLayerInstance* InExternalDataLayerInstance, FText& OutFailureReason)
+	{
+		check(!InActor->IsTemplate());
+		check(InActor->GetLevel());
+
+		const UExternalDataLayerAsset* NewExternalDataLayerAsset = InExternalDataLayerInstance ? InExternalDataLayerInstance->GetExternalDataLayerAsset() : nullptr;
+		if (!InActor->IsPackageExternal())
+		{
+			OutFailureReason = FText::Format(LOCTEXT("CantMoveActorToEDL_NotPackageExternal", "Actor {0} is not using external package."), FText::FromString(InActor->GetName()));
+			return false;
+		}
+		
+		if (!InActor->IsUserManaged())
+		{
+			OutFailureReason = FText::Format(LOCTEXT("CantMoveActorToEDL_NotUserManaged", "Actor {0} cannot be manually modified."), FText::FromString(InActor->GetName()));
+			return false;
+		}
+
+		if (!InActor->GetExternalDataLayerAsset() && !NewExternalDataLayerAsset)
+		{
+			OutFailureReason = FText::Format(LOCTEXT("CantMoveActorToEDL_NoExternalDataLayer", "Actor {0} has already no External Data Layer."), FText::FromString(InActor->GetName()));
+			return false;
+		}
+
+		if (InActor->GetExternalDataLayerAsset() == NewExternalDataLayerAsset)
+		{
+			OutFailureReason = FText::Format(LOCTEXT("CantMoveActorToEDL_SameExternalDataLayer", "Actor {0} is already assigned to this External Data Layer."), FText::FromString(InActor->GetName()));
+			return false;
+		}
+
+		if (InExternalDataLayerInstance && InExternalDataLayerInstance->IsReadOnly())
+		{
+			OutFailureReason = FText::Format(LOCTEXT("CantMoveActorToEDL_ReadOnlyExternalDataLayer", "External Data Layer is read-only."), FText::FromString(InExternalDataLayerInstance->GetDataLayerShortName()));
+			return false;
+		}
+
+		UEditorActorSubsystem* EditorActorSubsystem = GEditor->GetEditorSubsystem<UEditorActorSubsystem>();
+		if (!EditorActorSubsystem)
+		{
+			OutFailureReason = LOCTEXT("CantMoveActorToEDL_MissingEditorActorSubsystem", "Missing EditorActorSubsystem.");
+			return false;
+		}
+
+		// Create a temporary transient copy of the actor
+		AActor* DuplicatedActorToValidate = nullptr;
+		{
+			// Avoid triggering any dirty package notifications as this is a transient object
+			TGuardValue<bool> IsEditorLoadingPackageGuard(GIsEditorLoadingPackage, true);
+			{
+				UWorld* World = InActor->GetWorld();
+				// Temporarily set actor to target value for the duplication to use this value
+				FScopeRawAssignActorExternalDataLayer ScopeSetActorExternalDataLayer(InActor, NewExternalDataLayerAsset);
+				// Flag the temporary actor transient to make sure IsAsset returns false (to avoid considering this temporary actor by any other systems)
+				FDelegateHandle PreSpawnDelegateHandle = World->AddOnActorPreSpawnInitialization(FOnActorSpawned::FDelegate::CreateLambda([](AActor* DuplicatedActorToValidate) { DuplicatedActorToValidate->SetFlags(RF_Transient); }));
+				DuplicatedActorToValidate = EditorActorSubsystem->DuplicateActor(InActor);
+				World->RemoveOnActorPreSpawnInitialization(PreSpawnDelegateHandle);
+			}
+		}
+		// Destroy the temporary actor
+		ON_SCOPE_EXIT{ DuplicatedActorToValidate->GetWorld()->EditorDestroyActor(DuplicatedActorToValidate, false); };
+
+		// Validate the temporary actor's asset references
+		TSet<FString> InvalidReferenceReasons;
+		if (!UE::Private::ExternalDataLayerHelper::ValidateAssetUsingAssetReferenceRestrictions(DuplicatedActorToValidate, InvalidReferenceReasons))
+		{
+			const FString JoinedReasons = FString::Join(InvalidReferenceReasons, TEXT("\n"));
+			if (NewExternalDataLayerAsset)
+			{
+				OutFailureReason = FText::Format(LOCTEXT("CantMoveActorToEDLReferenceRestrictions", "Can't move Actor {0} to External Data Layer {1}. Reason: {2}."), FText::FromString(InActor->GetName()), FText::FromString(NewExternalDataLayerAsset->GetName()), FText::FromString(JoinedReasons));
+			}
+			else
+			{
+				const UExternalDataLayerAsset* OldExternalDataLayerAsset = InActor->GetExternalDataLayerAsset();
+				check(OldExternalDataLayerAsset);
+				OutFailureReason = FText::Format(LOCTEXT("CantRemoveEDLFromActorReferenceRestrictions", "Can't remove External Data Layer {0} from Actor {1}. Reason: {2}."), FText::FromString(OldExternalDataLayerAsset->GetName()), FText::FromString(InActor->GetName()), FText::FromString(JoinedReasons));
+			}
+			return false;
+		}
+
+		return true;
+	};
+
+	// Validate that all actors can change their External Data Layer asset
+	for (AActor* Actor : InActors)
+	{
+		FText FailureReason;
+		if (!CanMoveActorToExternalDataLayer(Actor, InExternalDataLayerInstance, FailureReason))
+		{
+			if (OutFailureReason)
+			{
+				*OutFailureReason = FailureReason;
+			}
+			return false;
+		}
+	}
+
+	return true;
+}
+
+bool FExternalDataLayerHelper::MoveActorsToExternalDataLayer(const TArray<AActor*>& InActors, const UExternalDataLayerInstance* InExternalDataLayerInstance, FText* OutFailureReason)
+{
+	auto MoveActorToExternalDataLayer = [](AActor * InActor, const UExternalDataLayerInstance* InExternalDataLayerInstance)
+	{
+		const UExternalDataLayerAsset* NewExternalDataLayerAsset = InExternalDataLayerInstance ? InExternalDataLayerInstance->GetExternalDataLayerAsset() : nullptr;
+		const bool bShouldDirty = true;
+		bool bLevelPackageWasDirty = InActor->GetLevel()->GetPackage()->IsDirty();
+		InActor->SetPackageExternal(false, bShouldDirty);
+
+		// If set, remove EDL from actor
+		const UExternalDataLayerAsset* OldExternalDataLayerAsset = InActor->GetExternalDataLayerAsset();
+		if (OldExternalDataLayerAsset)
+		{
+			FAssignActorDataLayer::RemoveDataLayerAsset(InActor, OldExternalDataLayerAsset);
+			for (const UDataLayerInstance* DataLayerInstance : InActor->GetDataLayerInstances())
+			{
+				if (DataLayerInstance->GetRootExternalDataLayerInstance())
+				{
+					FAssignActorDataLayer::RemoveDataLayerAsset(InActor, DataLayerInstance->GetAsset());
+				}
+			}
+		}
+
+		// If set, add actor to new EDL
+		if (NewExternalDataLayerAsset)
+		{
+			FAssignActorDataLayer::AddDataLayerAsset(InActor, NewExternalDataLayerAsset);
+		}
+
+		InActor->SetPackageExternal(true, bShouldDirty);
+		if (!bLevelPackageWasDirty)
+		{
+			InActor->GetLevel()->GetPackage()->SetDirtyFlag(false);
+		}
+
+		return (InActor->GetExternalDataLayerAsset() == NewExternalDataLayerAsset);
+	};
+	
+	// First, validate that the whole operation can be done without any validation errors
+	if (!CanMoveActorsToExternalDataLayer(InActors, InExternalDataLayerInstance, OutFailureReason))
+	{
+		return false;
+	}
+
+	// Change all actors External Data Layer asset
+	for (AActor* Actor : InActors)
+	{
+		if (ensure(MoveActorToExternalDataLayer(Actor, InExternalDataLayerInstance)))
+		{
+			// Basic validation on the actor and its new External Data Layer asset
+			UExternalDataLayerManager* ExternalDataLayerManager = UExternalDataLayerManager::GetExternalDataLayerManager(Actor);
+			check(ExternalDataLayerManager->ValidateOnActorExternalDataLayerAssetChanged(Actor));
+
+			// Notify actor's External Data Layer asset changed
+			FProperty* ExternalDataLayerAssetChangeProperty = FindFProperty<FProperty>(Actor->GetClass(), "ExternalDataLayerAsset");
+			FPropertyChangedEvent PropertyChangedEvent(ExternalDataLayerAssetChangeProperty);
+			Actor->PostEditChangeProperty(PropertyChangedEvent);
+		}
+	}
+	return true;
+}
+
 void FExternalDataLayerHelper::ForEachExternalDataLayerLevelPackagePath(const FString& InLevelPackageName, TFunctionRef<void(const FString&)> Func)
 {
 	UClass* GameFeatureDataClass = FindObject<UClass>(nullptr, TEXT("/Script/GameFeatures.GameFeatureData"));
@@ -146,3 +376,5 @@ bool FExternalDataLayerHelper::IsExternalDataLayerPath(FStringView InExternalDat
 }
 
 #endif
+
+#undef LOCTEXT_NAMESPACE
\ No newline at end of file
diff --git a/Engine/Source/Runtime/Engine/Private/WorldPartition/DataLayer/ExternalDataLayerInstance.cpp b/Engine/Source/Runtime/Engine/Private/WorldPartition/DataLayer/ExternalDataLayerInstance.cpp
--- a/Engine/Source/Runtime/Engine/Private/WorldPartition/DataLayer/ExternalDataLayerInstance.cpp
+++ b/Engine/Source/Runtime/Engine/Private/WorldPartition/DataLayer/ExternalDataLayerInstance.cpp
@@ -98,15 +98,6 @@ bool UExternalDataLayerInstance::CanAddActor(AActor* InActor, FText* OutReason)
 		return false;
 	}
 
-	if (InActor->HasAllFlags(RF_Transient))
-	{
-		if (OutReason)
-		{
-			*OutReason = LOCTEXT("CantAddActorTransientActorNotSupported", "Can't assign transient actor to external data layer.");
-		}
-		return false;
-	}
-
 	if (!InActor->IsMainPackageActor())
 	{
 		if (OutReason)
@@ -196,7 +187,7 @@ bool UExternalDataLayerInstance::PerformAddActor(AActor* InActor) const
 	check(ExternalDataLayer);
 	if (FAssignActorDataLayer::AddDataLayerAsset(InActor, ExternalDataLayer))
 	{
-		if (UExternalDataLayerManager::GetExternalDataLayerManager(GetOuterWorld())->OnActorExternalDataLayerAssetChanged(InActor))
+		if (UExternalDataLayerManager::GetExternalDataLayerManager(GetOuterWorld())->ValidateOnActorExternalDataLayerAssetChanged(InActor))
 		{
 			return true;
 		}
diff --git a/Engine/Source/Runtime/Engine/Private/WorldPartition/DataLayer/ExternalDataLayerManager.cpp b/Engine/Source/Runtime/Engine/Private/WorldPartition/DataLayer/ExternalDataLayerManager.cpp
--- a/Engine/Source/Runtime/Engine/Private/WorldPartition/DataLayer/ExternalDataLayerManager.cpp
+++ b/Engine/Source/Runtime/Engine/Private/WorldPartition/DataLayer/ExternalDataLayerManager.cpp
@@ -871,24 +871,25 @@ bool UExternalDataLayerManager::OnActorPreSpawnInitialization(AActor* InActor, c
 	return SetupActorPackageForExternalDataLayerAsset(InActor, InExternalDataLayerAsset);
 }
 
-bool UExternalDataLayerManager::OnActorExternalDataLayerAssetChanged(AActor* InActor)
+bool UExternalDataLayerManager::ValidateOnActorExternalDataLayerAssetChanged(AActor* InActor)
 {
-	// Check that the actor package is valid (it is currently not supported to change the EDL of an actor)
 	check(InActor);
 	check(InActor->IsPackageExternal());
 	const UExternalDataLayerAsset* ExternalDataLayerAsset = InActor->GetExternalDataLayerAsset();
-	check(ExternalDataLayerAsset);
-
+	if (ExternalDataLayerAsset)
+	{
 #if DO_CHECK
-	// Validate that the container
-	FExternalDataLayerContainerMap::ValueType ActorDescContainer = EDLContainerMap.FindChecked(ExternalDataLayerAsset);
-	check(ActorDescContainer->GetExternalActorPath() == ULevel::GetExternalActorsPath(GetExternalDataLayerLevelRootPath(ExternalDataLayerAsset)));
+		// Validate the corresponding container for this External Data Layer asset
+		FExternalDataLayerContainerMap::ValueType ActorDescContainer = EDLContainerMap.FindChecked(ExternalDataLayerAsset);
+		check(ActorDescContainer->GetExternalActorPath() == ULevel::GetExternalActorsPath(GetExternalDataLayerLevelRootPath(ExternalDataLayerAsset)));
 #endif
-	
-	const FString ActorPackageName = InActor->GetExternalPackage()->GetName();
-	const FString NewActorPackageName = GetActorPackageName(ExternalDataLayerAsset, InActor->GetLevel(), InActor->GetPathName());
-	check(NewActorPackageName == ActorPackageName);
-	return (NewActorPackageName == ActorPackageName);
+		// Validate that the actor package is valid
+		const FString ActorPackageName = InActor->GetExternalPackage()->GetName();
+		const FString NewActorPackageName = GetActorPackageName(ExternalDataLayerAsset, InActor->GetLevel(), InActor->GetPathName());
+		check(NewActorPackageName == ActorPackageName);
+		return (NewActorPackageName == ActorPackageName);
+	}
+	return true;
 }
 
 AWorldDataLayers* UExternalDataLayerManager::GetWorldDataLayers(const UExternalDataLayerAsset* InExternalDataLayerAsset, bool bAllowCreate) const
diff --git a/Engine/Source/Runtime/Engine/Public/WorldPartition/DataLayer/ExternalDataLayerHelper.h b/Engine/Source/Runtime/Engine/Public/WorldPartition/DataLayer/ExternalDataLayerHelper.h
--- a/Engine/Source/Runtime/Engine/Public/WorldPartition/DataLayer/ExternalDataLayerHelper.h
+++ b/Engine/Source/Runtime/Engine/Public/WorldPartition/DataLayer/ExternalDataLayerHelper.h
@@ -5,8 +5,10 @@
 #include "Containers/StringView.h"
 #include "Templates/Function.h"
 
+class AActor;
 class FAssetRegistryTagsContext;
 class UExternalDataLayerAsset;
+class UExternalDataLayerInstance;
 struct FExternalDataLayerUID;
 struct FAssetData;
 
@@ -52,6 +54,21 @@ public:
 
 private:
 
+#if WITH_EDITOR
+	/** 
+	 * Validates that all actors can change their External Data Layer to the new provided value (supports passing null) 
+	 * Returns false if any actor fails and fills OutFailureReason with the reason (if non-null). 
+	 */
+	ENGINE_API static bool CanMoveActorsToExternalDataLayer(const TArray<AActor*>& InActors, const UExternalDataLayerInstance* InExternalDataLayerInstance, FText* OutFailureReason = nullptr);
+
+	/** 
+	 * Changes all actors External Data Layer to the new provided value (supports passing null)
+	 * Returns false if any actor fails and fills OutFailureReason with the reason (if non-null). 
+	 */
+	ENGINE_API static bool MoveActorsToExternalDataLayer(const TArray<AActor*>& InActors, const UExternalDataLayerInstance* InExternalDataLayerInstance, FText* OutFailureReason = nullptr);
+
+	friend class FDataLayerEditorModule;
+#endif
 	static constexpr FStringView GetExternalDataLayerFolder() { return ExternalDataLayerFolder; }
 	static constexpr FStringView ExternalDataLayerFolder = TEXTVIEW("/EDL/");
 };
\ No newline at end of file
diff --git a/Engine/Source/Runtime/Engine/Public/WorldPartition/DataLayer/ExternalDataLayerInstance.h b/Engine/Source/Runtime/Engine/Public/WorldPartition/DataLayer/ExternalDataLayerInstance.h
--- a/Engine/Source/Runtime/Engine/Public/WorldPartition/DataLayer/ExternalDataLayerInstance.h
+++ b/Engine/Source/Runtime/Engine/Public/WorldPartition/DataLayer/ExternalDataLayerInstance.h
@@ -50,4 +50,6 @@ protected:
 	ENGINE_API virtual bool PerformRemoveActor(AActor* Actor) const override;
 	//~ End UDataLayerInstanceWithAsset Interface
 #endif
+
+	friend class FExternalDataLayerHelper;
 };
\ No newline at end of file
diff --git a/Engine/Source/Runtime/Engine/Public/WorldPartition/DataLayer/ExternalDataLayerManager.h b/Engine/Source/Runtime/Engine/Public/WorldPartition/DataLayer/ExternalDataLayerManager.h
--- a/Engine/Source/Runtime/Engine/Public/WorldPartition/DataLayer/ExternalDataLayerManager.h
+++ b/Engine/Source/Runtime/Engine/Public/WorldPartition/DataLayer/ExternalDataLayerManager.h
@@ -75,7 +75,7 @@ private:
 	// Used in editor
 	UActorDescContainerInstance* RegisterExternalDataLayerActorDescContainer(const UExternalDataLayerAsset* InExternalDataLayerAsset);
 	bool UnregisterExternalDataLayerActorDescContainer(const UExternalDataLayerAsset* InExternalDataLayerAsset);
-	bool OnActorExternalDataLayerAssetChanged(AActor* InActor);
+	bool ValidateOnActorExternalDataLayerAssetChanged(AActor* InActor);
 	bool RegisterExternalDataLayerInstance(UExternalDataLayerInstance* InExternalDataLayerInstance);
 	bool UnregisterExternalDataLayerInstance(UExternalDataLayerInstance* InExternalDataLayerInstance);
 	const UExternalDataLayerAsset* GetMatchingExternalDataLayerAssetForObjectPath(const FSoftObjectPath& InObjectPath);
@@ -85,6 +85,7 @@ private:
 	FString GetActorPackageName(const UExternalDataLayerAsset* InExternalDataLayerAsset, const ULevel* InDestinationLevel, const FString& InActorPath) const;
 	bool SetupActorPackageForExternalDataLayerAsset(AActor* InActor, const UExternalDataLayerAsset* InExternalDataLayerAsset);
 	URuntimeHashExternalStreamingObjectBase* CreateExternalStreamingObjectUsingStreamingGeneration(const UExternalDataLayerAsset* InExternalDataLayerAsset);
+	bool HasInjectedExternalDataLayerAssets() const { return InjectedExternalDataLayerAssets.Num() > 0; }
 
 	// Used for PIE/-game
 	void PrepareEditorGameWorld();
@@ -119,6 +120,8 @@ private:
 #endif
 
 	friend class FDataLayerMode;
+	friend class FDataLayerEditorModule;
+	friend class FExternalDataLayerHelper;
 	friend class AWorldDataLayers;
 	friend class UWorldPartition;
 	friend class UExternalDataLayerInstance;
