World Partition Runtime HashSet: implemented local version of FMath::SphereAABBIntersection which is faster mainly due to better vectorization.
This makes static index queries twice as fast on PC.

diff --git a/Engine/Source/Runtime/Engine/Private/WorldPartition/RuntimeHashSet/StaticSpatialIndex.cpp b/Engine/Source/Runtime/Engine/Private/WorldPartition/RuntimeHashSet/StaticSpatialIndex.cpp
--- a/Engine/Source/Runtime/Engine/Private/WorldPartition/RuntimeHashSet/StaticSpatialIndex.cpp
+++ b/Engine/Source/Runtime/Engine/Private/WorldPartition/RuntimeHashSet/StaticSpatialIndex.cpp
@@ -3,17 +3,26 @@
 #include "WorldPartition/RuntimeHashSet/StaticSpatialIndex.h"
 #include "Algo/Transform.h"
 
+namespace FStaticSpatialIndex
+{
+	inline bool FastSphereAABBIntersection(const FVector& InSphereCenter, const FSphere::FReal InRadiusSquared, const FBox& InAABB)
+	{
+		const FVector ClosestPoint = FVector::Max(InAABB.Min, FVector::Min(InSphereCenter, InAABB.Max));
+		return (ClosestPoint - InSphereCenter).SizeSquared() <= InRadiusSquared;
+	}
+}
+
 void FStaticSpatialIndex::FListImpl::Init(const TArray<TPair<FBox, uint32>>& InElements)
 {
 	Elements.Reserve(InElements.Num());
 	Algo::Transform(InElements, Elements, [](const TPair<FBox, uint32>& Element) { return Element.Value; });
 }
 
-bool FStaticSpatialIndex::FListImpl::ForEachElement(TFunctionRef<bool(uint32 InValueIndex)> Func) const
+bool FStaticSpatialIndex::FListImpl::ForEachElement(TFunctionRef<bool(uint32 InValueIndex)> InFunc) const
 {
 	for (uint32 ValueIndex : Elements)
 	{
-		if (!Func(ValueIndex))
+		if (!InFunc(ValueIndex))
 		{
 			return false;
 		}
@@ -21,7 +30,7 @@ bool FStaticSpatialIndex::FListImpl::ForEachElement(TFunctionRef<bool(uint32 InV
 	return true;
 }
 
-bool FStaticSpatialIndex::FListImpl::ForEachIntersectingElement(const FBox& InBox, TFunctionRef<bool(uint32 InValueIndex)> Func) const
+bool FStaticSpatialIndex::FListImpl::ForEachIntersectingElement(const FBox& InBox, TFunctionRef<bool(uint32 InValueIndex)> InFunc) const
 {
 	for (uint32 ValueIndex : Elements)
 	{
@@ -29,7 +38,7 @@ bool FStaticSpatialIndex::FListImpl::ForEachIntersectingElement(const FBox& InBo
 
 		if (Box.Intersect(InBox))
 		{
-			if (!Func(ValueIndex))
+			if (!InFunc(ValueIndex))
 			{
 				return false;
 			}
@@ -38,15 +47,17 @@ bool FStaticSpatialIndex::FListImpl::ForEachIntersectingElement(const FBox& InBo
 	return true;
 }
 
-bool FStaticSpatialIndex::FListImpl::ForEachIntersectingElement(const FSphere& InSphere, TFunctionRef<bool(uint32 InValueIndex)> Func) const
+bool FStaticSpatialIndex::FListImpl::ForEachIntersectingElement(const FSphere& InSphere, TFunctionRef<bool(uint32 InValueIndex)> InFunc) const
 {
+	const FSphere::FReal RadiusSquared = FMath::Square(InSphere.W);
+
 	for (uint32 ValueIndex : Elements)
 	{
 		const FBox& Box = DataInterface.GetBox(ValueIndex);
 
-		if (FMath::SphereAABBIntersection(InSphere, Box))
+		if (FastSphereAABBIntersection(InSphere.Center, RadiusSquared, Box))
 		{
-			if (!Func(ValueIndex))
+			if (!InFunc(ValueIndex))
 			{
 				return false;
 			}
@@ -114,18 +125,18 @@ void FStaticSpatialIndex::FRTreeImpl::Init(const TArray<TPair<FBox, uint32>>& In
 	}
 }
 
-bool FStaticSpatialIndex::FRTreeImpl::ForEachElement(TFunctionRef<bool(uint32 InValueIndex)> Func) const
+bool FStaticSpatialIndex::FRTreeImpl::ForEachElement(TFunctionRef<bool(uint32 InValueIndex)> InFunc) const
 {
-	return ForEachElementRecursive(&RootNode, Func);
+	return ForEachElementRecursive(&RootNode, InFunc);
 }
 
-bool FStaticSpatialIndex::FRTreeImpl::ForEachElementRecursive(const FNode* Node, TFunctionRef<bool(uint32 InValueIndex)> Func) const
+bool FStaticSpatialIndex::FRTreeImpl::ForEachElementRecursive(const FNode* Node, TFunctionRef<bool(uint32 InValueIndex)> InFunc) const
 {
 	if (Node->Content.IsType<FNode::FNodeType>())
 	{
 		for (auto& ChildNode : Node->Content.Get<FNode::FNodeType>())
 		{
-			if (!ForEachElementRecursive(&ChildNode, Func))
+			if (!ForEachElementRecursive(&ChildNode, InFunc))
 			{
 				return false;
 			}
@@ -135,7 +146,7 @@ bool FStaticSpatialIndex::FRTreeImpl::ForEachElementRecursive(const FNode* Node,
 	{
 		for (uint32 ValueIndex : Node->Content.Get<FNode::FLeafType>())
 		{
-			if (!Func(ValueIndex))
+			if (!InFunc(ValueIndex))
 			{
 				return false;
 			}
@@ -144,20 +155,20 @@ bool FStaticSpatialIndex::FRTreeImpl::ForEachElementRecursive(const FNode* Node,
 	return true;
 }
 
-bool FStaticSpatialIndex::FRTreeImpl::ForEachIntersectingElement(const FBox& InBox, TFunctionRef<bool(uint32 InValueIndex)> Func) const
+bool FStaticSpatialIndex::FRTreeImpl::ForEachIntersectingElement(const FBox& InBox, TFunctionRef<bool(uint32 InValueIndex)> InFunc) const
 {
-	return ForEachIntersectingElementRecursive(&RootNode, InBox, Func);
+	return ForEachIntersectingElementRecursive(&RootNode, InBox, InFunc);
 }
 
-bool FStaticSpatialIndex::FRTreeImpl::ForEachIntersectingElementRecursive(const FNode* Node, const FBox& InBox, TFunctionRef<bool(uint32 InValueIndex)> Func) const
+bool FStaticSpatialIndex::FRTreeImpl::ForEachIntersectingElementRecursive(const FNode* InNode, const FBox& InBox, TFunctionRef<bool(uint32 InValueIndex)> InFunc) const
 {
-	if (Node->Content.IsType<FNode::FNodeType>())
+	if (InNode->Content.IsType<FNode::FNodeType>())
 	{
-		for (auto& ChildNode : Node->Content.Get<FNode::FNodeType>())
+		for (auto& ChildNode : InNode->Content.Get<FNode::FNodeType>())
 		{
 			if (ChildNode.Box.Intersect(InBox))
 			{
-				if (!ForEachIntersectingElementRecursive(&ChildNode, InBox, Func))
+				if (!ForEachIntersectingElementRecursive(&ChildNode, InBox, InFunc))
 				{
 					return false;
 				}
@@ -166,13 +177,13 @@ bool FStaticSpatialIndex::FRTreeImpl::ForEachIntersectingElementRecursive(const
 	}
 	else
 	{
-		for (uint32 ValueIndex : Node->Content.Get<FNode::FLeafType>())
+		for (uint32 ValueIndex : InNode->Content.Get<FNode::FLeafType>())
 		{
 			const FBox& Box = DataInterface.GetBox(ValueIndex);
 
 			if (Box.Intersect(InBox))
 			{
-				if (!Func(ValueIndex))
+				if (!InFunc(ValueIndex))
 				{
 					return false;
 				}
@@ -182,20 +193,21 @@ bool FStaticSpatialIndex::FRTreeImpl::ForEachIntersectingElementRecursive(const
 	return true;
 }
 
-bool FStaticSpatialIndex::FRTreeImpl::ForEachIntersectingElement(const FSphere& InSphere, TFunctionRef<bool(uint32 InValueIndex)> Func) const
+bool FStaticSpatialIndex::FRTreeImpl::ForEachIntersectingElement(const FSphere& InSphere, TFunctionRef<bool(uint32 InValueIndex)> InFunc) const
 {
-	return ForEachIntersectingElementRecursive(&RootNode, InSphere, Func);
+	const FSphere::FReal RadiusSquared = FMath::Square(InSphere.W);
+	return ForEachIntersectingElementRecursive(&RootNode, InSphere.Center, RadiusSquared, InFunc);
 }
 
-bool FStaticSpatialIndex::FRTreeImpl::ForEachIntersectingElementRecursive(const FNode* Node, const FSphere& InSphere, TFunctionRef<bool(uint32 InValueIndex)> Func) const
+bool FStaticSpatialIndex::FRTreeImpl::ForEachIntersectingElementRecursive(const FNode* InNode, const FVector& InSphereCenter, FSphere::FReal InRadiusSquared, TFunctionRef<bool(uint32 InValueIndex)> InFunc) const
 {
-	if (Node->Content.IsType<FNode::FNodeType>())
+	if (InNode->Content.IsType<FNode::FNodeType>())
 	{
-		for (auto& ChildNode : Node->Content.Get<FNode::FNodeType>())
+		for (auto& ChildNode : InNode->Content.Get<FNode::FNodeType>())
 		{
-			if (FMath::SphereAABBIntersection(InSphere, ChildNode.Box))
+			if (FastSphereAABBIntersection(InSphereCenter, InRadiusSquared, ChildNode.Box))
 			{
-				if (!ForEachIntersectingElementRecursive(&ChildNode, InSphere, Func))
+				if (!ForEachIntersectingElementRecursive(&ChildNode, InSphereCenter, InRadiusSquared, InFunc))
 				{
 					return false;
 				}
@@ -204,13 +216,13 @@ bool FStaticSpatialIndex::FRTreeImpl::ForEachIntersectingElementRecursive(const
 	}
 	else
 	{
-		for (uint32 ValueIndex : Node->Content.Get<FNode::FLeafType>())
+		for (uint32 ValueIndex : InNode->Content.Get<FNode::FLeafType>())
 		{
 			const FBox& Box = DataInterface.GetBox(ValueIndex);
 
-			if (FMath::SphereAABBIntersection(InSphere, Box))
+			if (FastSphereAABBIntersection(InSphereCenter, InRadiusSquared, Box))
 			{
-				if (!Func(ValueIndex))
+				if (!InFunc(ValueIndex))
 				{
 					return false;
 				}
diff --git a/Engine/Source/Runtime/Engine/Public/WorldPartition/RuntimeHashSet/StaticSpatialIndex.h b/Engine/Source/Runtime/Engine/Public/WorldPartition/RuntimeHashSet/StaticSpatialIndex.h
--- a/Engine/Source/Runtime/Engine/Public/WorldPartition/RuntimeHashSet/StaticSpatialIndex.h
+++ b/Engine/Source/Runtime/Engine/Public/WorldPartition/RuntimeHashSet/StaticSpatialIndex.h
@@ -124,9 +124,9 @@ namespace FStaticSpatialIndex
 
 		void Init(const TArray<TPair<FBox, uint32>>& InElements);
 
-		bool ForEachElement(TFunctionRef<bool(uint32 InValueIndex)> Func) const;
-		bool ForEachIntersectingElement(const FBox& InBox, TFunctionRef<bool(uint32 InValueIndex)> Func) const;
-		bool ForEachIntersectingElement(const FSphere& InSphere, TFunctionRef<bool(uint32 InValueIndex)> Func) const;
+		bool ForEachElement(TFunctionRef<bool(uint32 InValueIndex)> InFunc) const;
+		bool ForEachIntersectingElement(const FBox& InBox, TFunctionRef<bool(uint32 InValueIndex)> InFunc) const;
+		bool ForEachIntersectingElement(const FSphere& InSphere, TFunctionRef<bool(uint32 InValueIndex)> InFunc) const;
 
 		uint32 GetAllocatedSize() const;
 
@@ -143,9 +143,9 @@ namespace FStaticSpatialIndex
 
 		void Init(const TArray<TPair<FBox, uint32>>& InElements);
 
-		bool ForEachElement(TFunctionRef<bool(uint32 InValueIndex)> Func) const;
-		bool ForEachIntersectingElement(const FBox& InBox, TFunctionRef<bool(uint32 InValueIndex)> Func) const;
-		bool ForEachIntersectingElement(const FSphere& InSphere, TFunctionRef<bool(uint32 InValueIndex)> Func) const;
+		bool ForEachElement(TFunctionRef<bool(uint32 InValueIndex)> InFunc) const;
+		bool ForEachIntersectingElement(const FBox& InBox, TFunctionRef<bool(uint32 InValueIndex)> InFunc) const;
+		bool ForEachIntersectingElement(const FSphere& InSphere, TFunctionRef<bool(uint32 InValueIndex)> InFunc) const;
 
 		uint32 GetAllocatedSize() const;
 
@@ -159,9 +159,9 @@ namespace FStaticSpatialIndex
 			TVariant<FNodeType, FLeafType> Content;
 		};
 
-		bool ForEachElementRecursive(const FNode* Node, TFunctionRef<bool(uint32 InValueIndex)> Func) const;
-		bool ForEachIntersectingElementRecursive(const FNode* Node, const FBox& InBox, TFunctionRef<bool(uint32 InValueIndex)> Func) const;
-		bool ForEachIntersectingElementRecursive(const FNode* Node, const FSphere& InSphere, TFunctionRef<bool(uint32 InValueIndex)> Func) const;
+		bool ForEachElementRecursive(const FNode* InNode, TFunctionRef<bool(uint32 InValueIndex)> InFunc) const;
+		bool ForEachIntersectingElementRecursive(const FNode* InNode, const FBox& InBox, TFunctionRef<bool(uint32 InValueIndex)> InFunc) const;
+		bool ForEachIntersectingElementRecursive(const FNode* InNode, const FVector& InSphereCenter, FSphere::FReal InRadiusSquared, TFunctionRef<bool(uint32 InValueIndex)> InFunc) const;
 
 		FNode RootNode;
 	};
