Iris
* ReplicationBridge is now MarkedAsGarbage when ReplicationSystem is destroyed
* ActorReplicationBridge now unregisters its delegates in Deinit()

diff --git a/Engine/Source/Runtime/Engine/Private/Net/Iris/ReplicationSystem/ActorReplicationBridge.cpp b/Engine/Source/Runtime/Engine/Private/Net/Iris/ReplicationSystem/ActorReplicationBridge.cpp
--- a/Engine/Source/Runtime/Engine/Private/Net/Iris/ReplicationSystem/ActorReplicationBridge.cpp
+++ b/Engine/Source/Runtime/Engine/Private/Net/Iris/ReplicationSystem/ActorReplicationBridge.cpp
@@ -152,6 +152,15 @@ UActorReplicationBridge::~UActorReplicationBridge()
 	}
 }
 
+void UActorReplicationBridge::Deinitialize()
+{
+	if (NetDriver)
+	{
+		NetDriver->OnNetServerMaxTickRateChanged.RemoveAll(this);
+	}
+	Super::Deinitialize();
+}
+
 UE::Net::FNetRefHandle UActorReplicationBridge::BeginReplication(AActor* Actor, const FActorBeginReplicationParams& Params)
 {
 	using namespace UE::Net;
diff --git a/Engine/Source/Runtime/Engine/Public/Net/Iris/ReplicationSystem/ActorReplicationBridge.h b/Engine/Source/Runtime/Engine/Public/Net/Iris/ReplicationSystem/ActorReplicationBridge.h
--- a/Engine/Source/Runtime/Engine/Public/Net/Iris/ReplicationSystem/ActorReplicationBridge.h
+++ b/Engine/Source/Runtime/Engine/Public/Net/Iris/ReplicationSystem/ActorReplicationBridge.h
@@ -90,6 +90,7 @@ protected:
 
 	// UObjectReplicationBridge
 	virtual void Initialize(UReplicationSystem* ReplicationSystem) override;
+	virtual void Deinitialize() override;
 	virtual bool WriteCreationHeader(UE::Net::FNetSerializationContext& Context, FNetRefHandle Handle) override;
 	virtual FCreationHeader* ReadCreationHeader(UE::Net::FNetSerializationContext& Context) override;
 	virtual FObjectReplicationBridgeInstantiateResult BeginInstantiateFromRemote(FNetRefHandle RootObjectOfSubObject, const UE::Net::FNetObjectResolveContext& ResolveContext, const FCreationHeader* InHeader) override;
diff --git a/Engine/Source/Runtime/Experimental/Iris/Core/Private/Iris/ReplicationSystem/ReplicationSystem.cpp b/Engine/Source/Runtime/Experimental/Iris/Core/Private/Iris/ReplicationSystem/ReplicationSystem.cpp
--- a/Engine/Source/Runtime/Experimental/Iris/Core/Private/Iris/ReplicationSystem/ReplicationSystem.cpp
+++ b/Engine/Source/Runtime/Experimental/Iris/Core/Private/Iris/ReplicationSystem/ReplicationSystem.cpp
@@ -617,6 +617,8 @@ void UReplicationSystem::Shutdown()
 	Impl->Deinit();
 	Impl.Reset();
 
+	// Destroy bridge
+	ReplicationBridge->MarkAsGarbage();
 	ReplicationBridge = nullptr;
 }
 
