VirtualTextureDataBuilder fix 4G limit check should be on blocks, not virtual canvas size

diff --git a/Engine/Source/Runtime/Engine/Private/VT/VirtualTextureDataBuilder.cpp b/Engine/Source/Runtime/Engine/Private/VT/VirtualTextureDataBuilder.cpp
--- a/Engine/Source/Runtime/Engine/Private/VT/VirtualTextureDataBuilder.cpp
+++ b/Engine/Source/Runtime/Engine/Private/VT/VirtualTextureDataBuilder.cpp
@@ -362,25 +362,23 @@ bool FVirtualTextureBuilderDerivedInfo::InitializeFromBuildSettings(const FTextu
 
 	SizeX = (int32)FullSizeX;
 	SizeY = (int32)FullSizeY;
-
+	
 	// there is no strict limit on total pixel count
 	//	but output must fit in 4 GB
 	//	so as a sanity check, test if pixel count is over 4G
 	// see FImageCoreUtils::IsImageImportPossible
-	//	this is sort of the wrong check, it really depends on output pixel format
-	// @todo : I'm not sure this check is right; is there actually a limit on the virtual canvas size?
-	//	  or is it only on the actual output data size?
-	//		(note that SizeX is the VT canvas size, not a pixel count, when you have UDIM blocks where not all tiles are present)
-	//		(eg. see "bigoffsets" test case)
-	//if ( (int64)SizeX * SizeY > (1ULL<<32) )
-	//{
-	//	UE_LOG(LogVirtualTexturing,Warning,TEXT("InitializeFromBuildSettings failed : total pixel count over 4G "
-	//		"(%d x %d = %lld) [%s]"),
-	//		SizeX,SizeY,(int64)SizeX * SizeY,
-	//		*InSourceData.TextureFullName);
-
-	//	return false;
-	//}
+	//	(this is sort of the wrong check, it really depends on output pixel format)
+	int64 NumBlocks = InSourceData.Blocks.Num();
+	int64 TotalPixels = (int64) BlockSizeX * BlockSizeY * NumBlocks;
+	if ( TotalPixels > (1ULL<<32) )
+	{
+		UE_LOG(LogVirtualTexturing,Warning,TEXT("InitializeFromBuildSettings failed : total pixel count over 4G "
+			"(%d x %d = %lld) [%s]"),
+			SizeX,SizeY,(int64)SizeX * SizeY,
+			*InSourceData.TextureFullName);
+
+		return false;
+	}
 
 	const uint32 Size = FMath::Max(SizeX, SizeY);
 
