GDC prefab fix
Prevent the CPFUO to do an extra copy on object that has Overridable Serialization enabled

diff --git a/Engine/Source/Runtime/CoreUObject/Private/UObject/PropertyArray.cpp b/Engine/Source/Runtime/CoreUObject/Private/UObject/PropertyArray.cpp
--- a/Engine/Source/Runtime/CoreUObject/Private/UObject/PropertyArray.cpp
+++ b/Engine/Source/Runtime/CoreUObject/Private/UObject/PropertyArray.cpp
@@ -260,6 +260,10 @@ void FArrayProperty::SerializeItem(FStructuredArchive::FSlot Slot, void* Value,
 
 				FScriptArrayHelper DefaultsArrayHelper(this, Defaults);
 
+				// If the arrays are not the same size, the code will not be able to remove items correctly
+				// as it searches for the items to remove in the default array but removes them from this array.
+				checkf(DefaultsArrayHelper.Num() == ArrayHelper.Num(), TEXT("The array \'%s\' is expected to be the same size as it default for removal logic."), *GetName());
+
 				auto FindObject = [InnerObjectProperty](UObject* Object, UObject* Object2, FScriptArrayHelper& ArrayHelper) -> int32
 				{
 					if (Object)
diff --git a/Engine/Source/Runtime/Engine/Private/UnrealEngine.cpp b/Engine/Source/Runtime/Engine/Private/UnrealEngine.cpp
--- a/Engine/Source/Runtime/Engine/Private/UnrealEngine.cpp
+++ b/Engine/Source/Runtime/Engine/Private/UnrealEngine.cpp
@@ -245,6 +245,7 @@ UnrealEngine.cpp: Implements the UEngine class and helpers.
 #include "HDRHelper.h"
 #include "UObject/PropertyBagRepository.h"
 #include "UObject/UObjectThreadContext.h"
+#include "UObject/OverridableManager.h"
 
 #if WITH_DUMPGPU
 #include "RenderGraphBuilder.h"
@@ -17317,6 +17318,25 @@ void UEngine::CopyPropertiesForUnrelatedObjects(UObject* OldObject, UObject* New
 	// may be missing, but testing will determine that:
 	const auto CollectAllSubobjects = [&Params](UObject* Object, TArray<UObject*>& OutSubobjectArray)
 	{
+		auto ShouldExcludeSubobjectFromCopy = FOverridableManager::Get().IsEnabled(*Object) ?
+			[](UObject* SubObject)
+			{
+				// Here we are trying to prevent an extra copy of COD direct subobjects. The problem is that IsDefaultSubobject() 
+				// method is it not consistent when called on subobjects of an instance vs on subobjects of a CDO. The method 
+				// PreCreateSubObjectsForReinstantiation creates and return non default sub objects to be copied. But if those 
+				// subobject are then considered here by IsDefaultSubobject() as default sub object, it will trigger an extra copy. 
+				// This double copy might be handle ok by Delta Serialization but the Overridable Serialization is unable to support it. 
+				// The reason is that it implements container removal operations and might remove wrong items during the subsequent copies.
+				// So for that reason the object that enabled Overridable serialization, we are just checking if the archetype 
+				// of the subobject is a CDO which is partially what IsDefaultSubobject() does.
+				return SubObject->GetArchetype()->HasAnyFlags(RF_ClassDefaultObject);
+			}
+			:
+			[](UObject* SubObject)
+			{
+				return !SubObject->IsDefaultSubobject() && !SubObject->HasAnyFlags(RF_DefaultSubObject);
+			};
+			
 		const bool bIncludedNestedObjects = !Params.bOnlyHandleDirectSubObjects;
 		GetObjectsWithOuter(Object, OutSubobjectArray, bIncludedNestedObjects);
 
@@ -17324,7 +17344,7 @@ void UEngine::CopyPropertiesForUnrelatedObjects(UObject* OldObject, UObject* New
 		for ( int32 ComponentIndex = 0; ComponentIndex < OutSubobjectArray.Num(); ComponentIndex++ )
 		{
 			UObject* PotentialComponent = OutSubobjectArray[ComponentIndex];
-			if (!PotentialComponent->IsDefaultSubobject() && !PotentialComponent->HasAnyFlags(RF_DefaultSubObject))
+			if (ShouldExcludeSubobjectFromCopy(PotentialComponent))
 			{
 				OutSubobjectArray.RemoveAtSwap(ComponentIndex--);
 			}
