* Removed check for pre-existing URuntimeHashExternalStreamingObjectBase from UWorldPartitionRuntimeHash::CreateExternalStreamingObject and moved it to the specific derived class ::CreateExternalStreamingData() implementation that needed it.
* This is pertinent for runtime usages not editor ones and this check relies on the GC being run in between calls.

diff --git a/Engine/Source/Runtime/Engine/Private/WorldPartition/WorldPartitionRuntimeHash.cpp b/Engine/Source/Runtime/Engine/Private/WorldPartition/WorldPartitionRuntimeHash.cpp
--- a/Engine/Source/Runtime/Engine/Private/WorldPartition/WorldPartitionRuntimeHash.cpp
+++ b/Engine/Source/Runtime/Engine/Private/WorldPartition/WorldPartitionRuntimeHash.cpp
@@ -197,12 +197,6 @@ EWorldPartitionStreamingPerformance UWorldPartitionRuntimeHash::GetStreamingPerf
 
 URuntimeHashExternalStreamingObjectBase* UWorldPartitionRuntimeHash::CreateExternalStreamingObject(TSubclassOf<URuntimeHashExternalStreamingObjectBase> InClass, UObject* InOuter, FName InName, UWorld* InOuterWorld)
 {
-	if (FindObject<URuntimeHashExternalStreamingObjectBase>(InOuter, *InName.ToString()))
-	{
-		UE_LOG(LogWorldPartition, Warning, TEXT("UWorldPartitionRuntimeHash::CreateExternalStreamingObject can't create an already existing URuntimeHashExternalStreamingObjectBase object named %s"), *InName.ToString());
-		return nullptr;
-	}
-
 	URuntimeHashExternalStreamingObjectBase* StreamingObject = NewObject<URuntimeHashExternalStreamingObjectBase>(InOuter, InClass, InName, RF_Public);
 	StreamingObject->OuterWorld = InOuterWorld;	
 	return StreamingObject;
