Removed double-call to RemapTracksToNewSkeleton for 'this' sequence when remapping skeleton

The first RemapTracksToNewSkeleton call in UAnimationAsset::ReplaceSkeleton was correctly kicking off an async compression, but the additional RemapTracksToNewSkeleton call was causing the compression to always be cancelled (via UAnimSequence::OnSetSkeleton bing called from the super call to UAnimationAsset::RemapTracksToNewSkeleton). This also, depending on timing, sometimes happened part-way or after compression, causing it to seem to fail (via the if (bCompressionSuccessful && !Owner.IsCanceled()) branch in FAnimationSequenceAsyncCacheTask::BuildData).

Also guarded the logging for compression with a check for task cancellation so we dont double-log when cancelling/retracting tasks.

#tests Local content worker cook of sunstone, full cook & verify cooked game, PIE

diff --git a/Engine/Source/Runtime/Engine/Private/Animation/AnimationCompressionDerivedData.cpp b/Engine/Source/Runtime/Engine/Private/Animation/AnimationCompressionDerivedData.cpp
--- a/Engine/Source/Runtime/Engine/Private/Animation/AnimationCompressionDerivedData.cpp
+++ b/Engine/Source/Runtime/Engine/Private/Animation/AnimationCompressionDerivedData.cpp
@@ -174,6 +174,13 @@ bool FAnimationSequenceAsyncCacheTask::BuildData() const
 {	
 	// This is where we should do the compression parts
 	TRACE_CPUPROFILER_EVENT_SCOPE_TEXT(*(FString(TEXT("FAnimationSequenceAsyncCacheTask::BuildData ") + CompressibleAnimPtr->Name)));
+
+	// Early out before logging if we are canceled (could be retracting this task)
+	if (Owner.IsCanceled())
+	{
+		return false;
+	}
+
 	UE_LOG(LogAnimationCompression, Display, TEXT("Building compressed animation data for %s (Required Memory Estimate: %.2f MB)"),
 		*CompressibleAnimPtr->FullName, double(GetRequiredMemoryEstimate()) / (1024.0 * 1024.0));
 
@@ -201,11 +208,12 @@ bool FAnimationSequenceAsyncCacheTask::BuildData() const
 		return false;
 	}
 	const bool bCurveCompressionOk = FAnimationUtils::CompressAnimCurves(DataToCompress, OutData);
-				
+	const bool bIsCanceled = Owner.IsCanceled();
+
 	const bool bCompressionSuccessful = bBoneCompressionOk && bCurveCompressionOk;
 	const FString CompressionName = DataToCompress.BoneCompressionSettings->GetFullName();
-	
-	if (bCompressionSuccessful && !Owner.IsCanceled())
+
+	if (bCompressionSuccessful && !bIsCanceled)
 	{
 		OutData.CompressedByteStream = MoveTemp(CompressionResult.CompressedByteStream);
 		OutData.CompressedDataStructure = MoveTemp(CompressionResult.AnimData);
@@ -215,7 +223,7 @@ bool FAnimationSequenceAsyncCacheTask::BuildData() const
 		
 		return true;
 	}
-	else
+	else if(!bIsCanceled)
 	{
 		UE_LOG(LogAnimationCompression, Error, TEXT("Failed to generate compressed animation data for %s with compression scheme %s for target platform %s"), *CompressibleAnimPtr->FullName, *CompressionName, *TargetPlatform->DisplayName().ToString());
 	}
diff --git a/Engine/Source/Runtime/Engine/Private/AnimationAsset.cpp b/Engine/Source/Runtime/Engine/Private/AnimationAsset.cpp
--- a/Engine/Source/Runtime/Engine/Private/AnimationAsset.cpp
+++ b/Engine/Source/Runtime/Engine/Private/AnimationAsset.cpp
@@ -471,74 +471,56 @@ bool UAnimationAsset::ReplaceSkeleton(USkeleton* NewSkeleton, bool bConvertSpace
 		// get all sequences that need to change
 		TArray<UAnimationAsset*> AnimAssetsToReplace;
 
-		if (UAnimSequence* AnimSequence = Cast<UAnimSequence>(this))
-		{
-			AnimAssetsToReplace.AddUnique(AnimSequence);
-		}
-		if (GetAllAnimationSequencesReferred(AnimAssetsToReplace))
-		{
-			TArray<UAnimSequence*> Sequences;
-			
-			//Firstly need to remap
-			for (UAnimationAsset* AnimAsset : AnimAssetsToReplace)
-			{
-				//Make sure animation has finished loading before we start messing with it
-				if (FLinkerLoad* AnimLinker = AnimAsset->GetLinker())
-				{
-					AnimLinker->Preload(AnimAsset);
-				}
-				AnimAsset->ConditionalPostLoad();
+		// Always add 'this' asset
+		AnimAssetsToReplace.Add(this);
 
-				if (AnimAsset->GetSkeleton() != GetSkeleton())
-				{
-					UE_LOG(LogAnimation, Warning, TEXT("AnimationAsset referencing asset using different skeleton. This will generate undeterministic builds. Please Fix the Asset : AnimationAsset: [%s] - ReferencedAsset : [%s]"), *GetName(), *AnimAsset->GetName());
-				}
+		// We dont use the return value from this function because we always add 'this' above
+		GetAllAnimationSequencesReferred(AnimAssetsToReplace);
 
-				// This ensure that in subsequent behaviour the RawData GUID is never 'new-ed' but always calculated from the 
-				// raw animation data itself.
-				if (UAnimSequence* Sequence = Cast<UAnimSequence>(AnimAsset))
-				{
-					Sequences.Add(Sequence);				
-				}
-				else
-				{
-					// these two are different functions for now
-					// technically if you have implementation for Remap, it will also set skeleton 
-					AnimAsset->RemapTracksToNewSkeleton(NewSkeleton, bConvertSpaces);
-				}
-			}
+		TArray<UAnimSequence*> Sequences;
 
-			UE::Anim::FAnimSequenceCompilingManager::Get().FinishCompilation(Sequences);
-			for (UAnimSequence* Sequence : Sequences)
+		//Firstly need to remap
+		for (UAnimationAsset* AnimAsset : AnimAssetsToReplace)
+		{
+			//Make sure animation has finished loading before we start messing with it
+			if (FLinkerLoad* AnimLinker = AnimAsset->GetLinker())
 			{
-				Sequence->GetController().OpenBracket(LOCTEXT("ReplaceSkeleton_Bracket", "Replacing USkeleton"));
-				Sequence->RemapTracksToNewSkeleton(NewSkeleton, bConvertSpaces);
+				AnimLinker->Preload(AnimAsset);
 			}
+			AnimAsset->ConditionalPostLoad();
 
-			//Second need to process anim sequences themselves. This is done in two stages as additives can rely on other animations.
-			for (UAnimSequence* Sequence : Sequences)
+			if (AnimAsset->GetSkeleton() != GetSkeleton())
 			{
-				Sequence->GetController().CloseBracket();
+				UE_LOG(LogAnimation, Warning, TEXT("AnimationAsset referencing asset using different skeleton. This will generate undeterministic builds. Please Fix the Asset : AnimationAsset: [%s] - ReferencedAsset : [%s]"), *GetName(), *AnimAsset->GetName());
 			}
-		}
 
-		UAnimSequence* Seq = Cast<UAnimSequence>(this);
-		{			
 			// This ensure that in subsequent behaviour the RawData GUID is never 'new-ed' but always calculated from the 
 			// raw animation data itself.
-			if (Seq)
+			if (UAnimSequence* Sequence = Cast<UAnimSequence>(AnimAsset))
 			{
-				Seq->GetController().OpenBracket(LOCTEXT("ReplaceSkeleton_Bracket", "Replacing USkeleton"));
+				Sequences.Add(Sequence);
 			}
-  
-			RemapTracksToNewSkeleton(NewSkeleton, bConvertSpaces);
-
-			if (Seq)
+			else
 			{
-				Seq->GetController().CloseBracket();
+				// these two are different functions for now
+				// technically if you have implementation for Remap, it will also set skeleton 
+				AnimAsset->RemapTracksToNewSkeleton(NewSkeleton, bConvertSpaces);
 			}
 		}
 
+		UE::Anim::FAnimSequenceCompilingManager::Get().FinishCompilation(Sequences);
+		for (UAnimSequence* Sequence : Sequences)
+		{
+			Sequence->GetController().OpenBracket(LOCTEXT("ReplaceSkeleton_Bracket", "Replacing USkeleton"));
+			Sequence->RemapTracksToNewSkeleton(NewSkeleton, bConvertSpaces);
+		}
+
+		//Second need to process anim sequences themselves. This is done in two stages as additives can rely on other animations.
+		for (UAnimSequence* Sequence : Sequences)
+		{
+			Sequence->GetController().CloseBracket();
+		}
+
 		return true;
 	}
 
