[NetPhysics] Resim improvements

- Ensure resim particles are marked as not in sync if their state during resim doesn't match the cached history, so that PullData state gets updated.
--- Render interpolation bug fixed as a result of PullData being properly updated in all cases
- Reduce redundant inputs in each network sync to 3, same default value as the server input buffer.
- Implement resim runtime correction for velocity and angular velocity, guarded by CVar that is disabled by default for now, also implemented as an option int the NetworkPhysicsSettingsComponent
- Implemented the new particle correction call that recalculate friction points after an XR correction to make corrections ignore friction.
- Implemented CVar to hook into the option to apply particle correction to any connected particle. Also  implemented as an option int the NetworkPhysicsSettingsComponent.
- Removed unused extern defines for replication cvars (they are instead defined inside NetworkPhysicsSettingsComponent.h where they are used)

diff --git a/Engine/Source/Runtime/Engine/Private/PhysicsEngine/PhysicsReplication.cpp b/Engine/Source/Runtime/Engine/Private/PhysicsEngine/PhysicsReplication.cpp
--- a/Engine/Source/Runtime/Engine/Private/PhysicsEngine/PhysicsReplication.cpp
+++ b/Engine/Source/Runtime/Engine/Private/PhysicsEngine/PhysicsReplication.cpp
@@ -90,7 +90,13 @@ namespace PhysicsReplicationCVars
 	namespace ResimulationCVars
 	{
 		bool bRuntimeCorrectionEnabled = true;
-		static FAutoConsoleVariableRef CVarResimRuntimeCorrectionEnabled(TEXT("np2.Resim.RuntimeCorrectionEnabled"), bRuntimeCorrectionEnabled, TEXT("Apply runtime corrections while error is smalle enough not to trigger a resim."));
+		static FAutoConsoleVariableRef CVarResimRuntimeCorrectionEnabled(TEXT("np2.Resim.RuntimeCorrectionEnabled"), bRuntimeCorrectionEnabled, TEXT("Apply runtime corrections while error is small enough not to trigger a resim."));
+
+		bool bRuntimeVelocityCorrection = false;
+		static FAutoConsoleVariableRef CVarResimRuntimeVelocityCorrection(TEXT("np2.Resim.RuntimeVelocityCorrection"), bRuntimeVelocityCorrection, TEXT("Apply linear and angular velocity corrections in runtime while within resim trigger. Used if RuntimeCorrectionEnabled is true."));
+
+		bool bRuntimeCorrectConnectedBodies = true;
+		static FAutoConsoleVariableRef CVarResimRuntimeCorrectConnectedBodies(TEXT("np2.Resim.RuntimeCorrectConnectedBodies"), bRuntimeCorrectConnectedBodies, TEXT("If true runtime position and rotation correction will also shift transform of any connected physics objects. Used if RuntimeCorrectionEnabled is true."));
 
 		bool bDisableReplicationOnInteraction = true;
 		static FAutoConsoleVariableRef CVarResimDisableReplicationOnInteraction(TEXT("np2.Resim.DisableReplicationOnInteraction"), bDisableReplicationOnInteraction, TEXT("If a resim object interacts with another object not running resimulation, deactivate that objects replication until interaction stops."));
@@ -100,6 +106,15 @@ namespace PhysicsReplicationCVars
 
 		float RotStabilityMultiplier = 1.0f;
 		static FAutoConsoleVariableRef CVarResimRotStabilityMultiplier(TEXT("np2.Resim.RotStabilityMultiplier"), RotStabilityMultiplier, TEXT("Recommended range between 0.0-1.0. Lower value means more stable rotational corrections."));
+	
+		float VelStabilityMultiplier = 0.5f;
+		static FAutoConsoleVariableRef CVarResimVelStabilityMultiplier(TEXT("np2.Resim.VelStabilityMultiplier"), VelStabilityMultiplier, TEXT("Recommended range between 0.0-1.0. Lower value means more stable linear velocity corrections."));
+
+		float AngVelStabilityMultiplier = 0.5f;
+		static FAutoConsoleVariableRef CVarResimAngVelStabilityMultiplier(TEXT("np2.Resim.AngVelStabilityMultiplier"), AngVelStabilityMultiplier, TEXT("Recommended range between 0.0-1.0. Lower value means more stable angular velocity corrections."));
+
+		bool bDrawDebug = false;
+		static FAutoConsoleVariableRef CVarResimDrawDebug(TEXT("np2.Resim.DrawDebug"), bDrawDebug, TEXT("Resimulation debug draw-calls"));
 	}
 
 	namespace PredictiveInterpolationCVars
@@ -1794,17 +1809,18 @@ bool FPhysicsReplicationAsync::ResimulationReplication(Chaos::FPBDRigidParticleH
 		UE_LOG(LogTemp, Log, TEXT("Particle Target Velocity = %s | Current Velocity = %s"), *Target.TargetState.LinVel.ToString(), *PastState.GetV().ToString());
 		UE_LOG(LogTemp, Log, TEXT("Particle Target Quaternion = %s | Current Quaternion = %s"), *Target.TargetState.Quaternion.ToString(), *PastState.GetR().ToString());
 		UE_LOG(LogTemp, Log, TEXT("Particle Target Omega = %s | Current Omega= %s"), *Target.TargetState.AngVel.ToString(), *PastState.GetW().ToString());
+	}
 
-		{ // DrawDebug
-			static constexpr float BoxSize = 5.0f;
-			const float ColorLerp = ShouldTriggerResim ? 1.0f : 0.0f;
-			const FColor DebugColor = FLinearColor::LerpUsingHSV(FLinearColor::Green, FLinearColor::Red, ColorLerp).ToFColor(false);
+	if (PhysicsReplicationCVars::ResimulationCVars::bDrawDebug)
+	{ 
+		static constexpr float BoxSize = 5.0f;
+		const float ColorLerp = ShouldTriggerResim ? 1.0f : 0.0f;
+		const FColor DebugColor = FLinearColor::LerpUsingHSV(FLinearColor::Green, FLinearColor::Red, ColorLerp).ToFColor(false);
 
-			Chaos::FDebugDrawQueue::GetInstance().DrawDebugBox(Target.TargetState.Position, FVector(BoxSize, BoxSize, BoxSize), Target.TargetState.Quaternion, FColor::Orange, true, CharacterMovementCVars::NetCorrectionLifetime, 0, 1.0f);
-			Chaos::FDebugDrawQueue::GetInstance().DrawDebugBox(PastState.GetX(), FVector(6, 6, 6), PastState.GetR(), DebugColor, true, CharacterMovementCVars::NetCorrectionLifetime, 0, 1.0f);
+		Chaos::FDebugDrawQueue::GetInstance().DrawDebugBox(Target.TargetState.Position, FVector(BoxSize, BoxSize, BoxSize), Target.TargetState.Quaternion, FColor::Orange, true, CharacterMovementCVars::NetCorrectionLifetime, 0, 1.0f);
+		Chaos::FDebugDrawQueue::GetInstance().DrawDebugBox(PastState.GetX(), FVector(6, 6, 6), PastState.GetR(), DebugColor, true, CharacterMovementCVars::NetCorrectionLifetime, 0, 1.0f);
 
-			Chaos::FDebugDrawQueue::GetInstance().DrawDebugDirectionalArrow(PastState.GetX(), Target.TargetState.Position, 5.0f, FColor::MakeRandomSeededColor(LocalFrame), true, CharacterMovementCVars::NetCorrectionLifetime, 0, 0.5f);
-		}
+		Chaos::FDebugDrawQueue::GetInstance().DrawDebugDirectionalArrow(PastState.GetX(), Target.TargetState.Position, 5.0f, FColor::MakeRandomSeededColor(LocalFrame), true, CharacterMovementCVars::NetCorrectionLifetime, 0, 0.5f);
 	}
 #endif
 
@@ -1825,26 +1841,44 @@ bool FPhysicsReplicationAsync::ResimulationReplication(Chaos::FPBDRigidParticleH
 
 			if (Target.TickCount <= NumPredictedFrames && NumPredictedFrames > 0)
 			{
-				// Calculate correction to position
-				const float CorrectionAmountX = SettingsCurrent.ResimulationSettings.GetPosStabilityMultiplier() / NumPredictedFrames; // Same result as (ErrorOffset / NumPredictedFrames) * PosStabilityMultiplier
-				const FVector CorrectedX = Handle->GetX() + (ErrorOffset * CorrectionAmountX);
-
-				// Calculate correction to rotation
-				const float CorrectionAmountR = (1.f / NumPredictedFrames) * SettingsCurrent.ResimulationSettings.GetRotStabilityMultiplier();
-				const FQuat InvTargetQuat = Target.TargetState.Quaternion;
-				const FQuat DeltaQuat = PastState.GetR().Inverse() * InvTargetQuat;
+				// Positional Correction
+				const float CorrectionAmountX = SettingsCurrent.ResimulationSettings.GetPosStabilityMultiplier() / NumPredictedFrames;
+				const FVector PosDiffCorrection = ErrorOffset * CorrectionAmountX; // Same result as (ErrorOffset / NumPredictedFrames) * PosStabilityMultiplier
+				const FVector CorrectedX = Handle->GetX() + PosDiffCorrection;
+
+				// Rotational Correction
+				const float CorrectionAmountR = SettingsCurrent.ResimulationSettings.GetRotStabilityMultiplier() / NumPredictedFrames;
+				const FQuat DeltaQuat = PastState.GetR().Inverse() * Target.TargetState.Quaternion;
 				const FQuat TargetCorrectionR = Handle->GetR() * DeltaQuat;
 				const FQuat CorrectedR = FQuat::Slerp(Handle->GetR(), TargetCorrectionR, CorrectionAmountR);
 
-	#if !(UE_BUILD_SHIPPING || UE_BUILD_TEST)
-				if (Chaos::FPhysicsSolverBase::CanDebugNetworkPhysicsPrediction())
+				if (SettingsCurrent.ResimulationSettings.GetRuntimeVelocityCorrectionEnabled())
+				{
+					// Linear Velocity Correction
+					const FVector LinVelDiff = Target.TargetState.LinVel - PastState.GetV(); // Velocity vector that the server covers but the client doesn't
+					const float CorrectionAmountV = SettingsCurrent.ResimulationSettings.GetVelStabilityMultiplier() / NumPredictedFrames;
+					const FVector VelCorrection = LinVelDiff * CorrectionAmountV; // Same result as (LinVelDiff / NumPredictedFrames) * VelStabilityMultiplier
+					const FVector CorrectedV = Handle->GetV() + VelCorrection;
+
+					// Angular Velocity Correction
+					const FVector AngVelDiff = Target.TargetState.AngVel - PastState.GetW(); // Angular velocity vector that the server covers but the client doesn't
+					const float CorrectionAmountW = SettingsCurrent.ResimulationSettings.GetAngVelStabilityMultiplier() / NumPredictedFrames;
+					const FVector AngVelCorrection = AngVelDiff * CorrectionAmountW; // Same result as (AngVelDiff / NumPredictedFrames) * VelStabilityMultiplier
+					const FVector CorrectedW = Handle->GetW() + AngVelCorrection;
+					
+					// Apply correction to velocities
+					Handle->SetV(CorrectedV);
+					Handle->SetW(CorrectedW);
+				}
+
+#if !(UE_BUILD_SHIPPING || UE_BUILD_TEST)
+				if (PhysicsReplicationCVars::ResimulationCVars::bDrawDebug)
 				{
 					Chaos::FDebugDrawQueue::GetInstance().DrawDebugDirectionalArrow(Handle->GetX(), CorrectedX, 5.0f, FColor::MakeRandomSeededColor(LocalFrame), true, CharacterMovementCVars::NetCorrectionLifetime, 0, 0.5f);
 				}
-	#endif
-				// Apply correction
-				Handle->SetX(CorrectedX);
-				Handle->SetR(CorrectedR);
+#endif
+				// Apply correction to position and rotation
+				RigidsSolver->GetEvolution()->ApplyParticleTransformCorrection(Handle, CorrectedX, CorrectedR, PhysicsReplicationCVars::ResimulationCVars::bRuntimeCorrectConnectedBodies);
 			}
 
 			// Keep target for NumPredictedFrames time to perform runtime corrections with until a new target is received
diff --git a/Engine/Source/Runtime/Engine/Public/Physics/NetworkPhysicsComponent.h b/Engine/Source/Runtime/Engine/Public/Physics/NetworkPhysicsComponent.h
--- a/Engine/Source/Runtime/Engine/Public/Physics/NetworkPhysicsComponent.h
+++ b/Engine/Source/Runtime/Engine/Public/Physics/NetworkPhysicsComponent.h
@@ -619,7 +619,7 @@ private:
 	TUniquePtr<FNetworkPhysicsData> StateData;
 
 	// Send last N number of inputs each replication call to patch up holes due to packet loss
-	int8 InputRedundancy = 4;
+	int8 InputRedundancy = 3;
 
 	// Current index used in the inputs offsets
 	int8 InputIndex = 0;
diff --git a/Engine/Source/Runtime/Engine/Public/Physics/NetworkPhysicsSettingsComponent.h b/Engine/Source/Runtime/Engine/Public/Physics/NetworkPhysicsSettingsComponent.h
--- a/Engine/Source/Runtime/Engine/Public/Physics/NetworkPhysicsSettingsComponent.h
+++ b/Engine/Source/Runtime/Engine/Public/Physics/NetworkPhysicsSettingsComponent.h
@@ -46,8 +46,11 @@ namespace PhysicsReplicationCVars
 	namespace ResimulationCVars
 	{
 		extern bool bRuntimeCorrectionEnabled;
+		extern bool bRuntimeVelocityCorrection;
 		extern float PosStabilityMultiplier;
 		extern float RotStabilityMultiplier;
+		extern float VelStabilityMultiplier;
+		extern float AngVelStabilityMultiplier;
 	}
 }
 
@@ -178,7 +181,7 @@ struct FNetworkPhysicsSettingsResimulation
 	UPROPERTY(config, EditDefaultsOnly, Category = "Overrides", Meta = (InlineEditConditionToggle))
 	uint32 bOverrideRedundantInputs : 1;
 	UPROPERTY(config, EditDefaultsOnly, Category = "Overrides", Meta = (EditCondition = "bOverrideRedundantInputs"))
-	uint8 RedundantInputs = 4;
+	uint8 RedundantInputs = 3;
 
 	// Override how many states to synchronize each sync to cover packet loss.
 	UPROPERTY(config, EditDefaultsOnly, Category = "Overrides", Meta = (InlineEditConditionToggle))
@@ -193,6 +196,13 @@ struct FNetworkPhysicsSettingsResimulation
 	bool bRuntimeCorrectionEnabled = PhysicsReplicationCVars::ResimulationCVars::bRuntimeCorrectionEnabled;
 	bool GetRuntimeCorrectionEnabled() { return bOverrideRuntimeCorrectionEnabled ? bRuntimeCorrectionEnabled : PhysicsReplicationCVars::ResimulationCVars::bRuntimeCorrectionEnabled; }
 
+	// np2.Resim.RuntimeVelocityCorrection
+	UPROPERTY(config, EditDefaultsOnly, Category = "Overrides", Meta = (InlineEditConditionToggle))
+	uint32 bOverrideRuntimeVelocityCorrection : 1;
+	UPROPERTY(config, EditDefaultsOnly, Category = "Overrides", Meta = (EditCondition = "bOverrideRuntimeVelocityCorrection"))
+	bool bRuntimeVelocityCorrection = PhysicsReplicationCVars::ResimulationCVars::bRuntimeVelocityCorrection;
+	bool GetRuntimeVelocityCorrectionEnabled() { return bOverrideRuntimeVelocityCorrection ? bRuntimeVelocityCorrection : PhysicsReplicationCVars::ResimulationCVars::bRuntimeVelocityCorrection; }
+
 	// np2.Resim.PosStabilityMultiplier
 	UPROPERTY(config, EditDefaultsOnly, Category = "Overrides", Meta = (InlineEditConditionToggle))
 	uint32 bOverridePosStabilityMultiplier : 1;
@@ -207,6 +217,20 @@ struct FNetworkPhysicsSettingsResimulation
 	float RotStabilityMultiplier = PhysicsReplicationCVars::ResimulationCVars::RotStabilityMultiplier;
 	float GetRotStabilityMultiplier() { return bOverrideRotStabilityMultiplier ? RotStabilityMultiplier : PhysicsReplicationCVars::ResimulationCVars::RotStabilityMultiplier; }
 
+	// np2.Resim.VelStabilityMultiplier
+	UPROPERTY(config, EditDefaultsOnly, Category = "Overrides", Meta = (InlineEditConditionToggle))
+	uint32 bOverrideVelStabilityMultiplier : 1;
+	UPROPERTY(config, EditDefaultsOnly, Category = "Overrides", Meta = (EditCondition = "bOverrideVelStabilityMultiplier"))
+	float VelStabilityMultiplier = PhysicsReplicationCVars::ResimulationCVars::VelStabilityMultiplier;
+	float GetVelStabilityMultiplier() { return bOverrideVelStabilityMultiplier ? VelStabilityMultiplier : PhysicsReplicationCVars::ResimulationCVars::VelStabilityMultiplier; }
+
+	// np2.Resim.AngVelStabilityMultiplier
+	UPROPERTY(config, EditDefaultsOnly, Category = "Overrides", Meta = (InlineEditConditionToggle))
+	uint32 bOverrideAngVelStabilityMultiplier : 1;
+	UPROPERTY(config, EditDefaultsOnly, Category = "Overrides", Meta = (EditCondition = "bOverrideAngVelStabilityMultiplier"))
+	float AngVelStabilityMultiplier = PhysicsReplicationCVars::ResimulationCVars::AngVelStabilityMultiplier;
+	float GetAngVelStabilityMultiplier() { return bOverrideAngVelStabilityMultiplier ? AngVelStabilityMultiplier : PhysicsReplicationCVars::ResimulationCVars::AngVelStabilityMultiplier; }
+
 	// Project Settings -> Physics -> Replication -> Physics Prediction -> Resimulation Error Threshold
 	UPROPERTY(config, EditDefaultsOnly, Category = "Overrides", Meta = (InlineEditConditionToggle))
 	uint32 bOverrideResimulationErrorThreshold : 1;
diff --git a/Engine/Source/Runtime/Engine/Public/PhysicsReplication.h b/Engine/Source/Runtime/Engine/Public/PhysicsReplication.h
--- a/Engine/Source/Runtime/Engine/Public/PhysicsReplication.h
+++ b/Engine/Source/Runtime/Engine/Public/PhysicsReplication.h
@@ -42,31 +42,6 @@ namespace CharacterMovementCVars
 namespace PhysicsReplicationCVars
 {
 	extern ENGINE_API int32 LogPhysicsReplicationHardSnaps;
-	
-	namespace PredictiveInterpolationCVars
-	{
-		extern float PosCorrectionTimeBase;
-		extern float PosCorrectionTimeMin;
-		extern float PosCorrectionTimeMultiplier;
-		extern float RotCorrectionTimeBase;
-		extern float RotCorrectionTimeMin;
-		extern float RotCorrectionTimeMultiplier;
-		extern float PosInterpolationTimeMultiplier;
-		extern float RotInterpolationTimeMultiplier;
-		extern float SoftSnapPosStrength;
-		extern float SoftSnapRotStrength;
-		extern bool bSoftSnapToSource;
-		extern bool bSkipVelocityRepOnPosEarlyOut;
-		extern bool bPostResimWaitForUpdate;
-		extern bool bDisableSoftSnap;
-	}
-
-	namespace ResimulationCVars
-	{
-		extern bool bRuntimeCorrectionEnabled;
-		extern float PosStabilityMultiplier;
-		extern float RotStabilityMultiplier;
-	}
 }
 #endif
 
diff --git a/Engine/Source/Runtime/Experimental/Chaos/Private/RewindData.cpp b/Engine/Source/Runtime/Experimental/Chaos/Private/RewindData.cpp
--- a/Engine/Source/Runtime/Experimental/Chaos/Private/RewindData.cpp
+++ b/Engine/Source/Runtime/Experimental/Chaos/Private/RewindData.cpp
@@ -68,6 +68,7 @@ bool SimWritablePropsMayChange(const TGeometryParticleHandle<FReal,3>& Handle)
 	return ObjectState == EObjectStateType::Dynamic || ObjectState == EObjectStateType::Sleeping;
 }
 
+/* Deprecated UE5.4 */
 bool FGeometryParticleStateBase::IsResimFrameValid(const FGeometryParticleHandle& Handle, const FFrameAndPhase FrameAndPhase) const
 {
 	if (FPBDRigidsSolver* RigidSolver = Handle.PhysicsProxy()->GetSolver<FPBDRigidsSolver>())
@@ -87,68 +88,59 @@ bool FGeometryParticleStateBase::IsResimFrameValid(const FGeometryParticleHandle
 template <bool bSkipDynamics>
 bool FGeometryParticleStateBase::IsInSync(const FGeometryParticleHandle& Handle, const FFrameAndPhase FrameAndPhase, const FDirtyPropertiesPool& Pool) const
 {
-	if(Chaos::FPhysicsSolverBase::IsNetworkPhysicsPredictionEnabled() && !TargetPositions.IsEmpty() && !TargetVelocities.IsEmpty() && !TargetStates.IsEmpty())
+	if (!ParticlePositionRotation.IsInSync(Handle, FrameAndPhase, Pool))
 	{
-		if(!IsResimFrameValid(Handle, FrameAndPhase))
-		{
-			return false;
-		}
+		return false;
+	}
+
+	if (!NonFrequentData.IsInSync(Handle, FrameAndPhase, Pool))
+	{
+		return false;
 	}
-	else
-	{ 
-		if(!ParticlePositionRotation.IsInSync(Handle, FrameAndPhase, Pool))
+
+	//todo: deal with state change mismatch
+
+	if (auto Kinematic = Handle.CastToKinematicParticle())
+	{
+		if (!Velocities.IsInSync(*Kinematic, FrameAndPhase, Pool))
 		{
 			return false;
 		}
 
-		if(!NonFrequentData.IsInSync(Handle, FrameAndPhase, Pool))
+		if (!KinematicTarget.IsInSync(*Kinematic, FrameAndPhase, Pool))
 		{
 			return false;
 		}
+	}
 
-		//todo: deal with state change mismatch
-
-		if(auto Kinematic = Handle.CastToKinematicParticle())
+	if (auto Rigid = Handle.CastToRigidParticle())
+	{
+		if (!bSkipDynamics)
 		{
-			if(!Velocities.IsInSync(*Kinematic, FrameAndPhase, Pool))
-			{
-				return false;
-			}
-
-			if (!KinematicTarget.IsInSync(*Kinematic, FrameAndPhase, Pool))
+			if (!Dynamics.IsInSync(*Rigid, FrameAndPhase, Pool))
 			{
 				return false;
 			}
 		}
 
-		if(auto Rigid = Handle.CastToRigidParticle())
+		if (!DynamicsMisc.IsInSync(*Rigid, FrameAndPhase, Pool))
 		{
-			if(!bSkipDynamics)
-			{
-				if (!Dynamics.IsInSync(*Rigid, FrameAndPhase, Pool))
-				{
-					return false;
-				}
-			}
-
-			if(!DynamicsMisc.IsInSync(*Rigid, FrameAndPhase, Pool))
-			{
-				return false;
-			}
+			return false;
+		}
 
-			if(!MassProps.IsInSync(*Rigid, FrameAndPhase, Pool))
-			{
-				return false;
-			}
+		if (!MassProps.IsInSync(*Rigid, FrameAndPhase, Pool))
+		{
+			return false;
 		}
 	}
-
+	
 	//TODO: this assumes geometry is never modified. Geometry modification has various issues in higher up Chaos code. Need stable shape id
 	//For now iterate over all the shapes in latest and see if they have any mismatches
 	/*if(ShapesArrayState.PerShapeData.Num())
 	{
 		return false;	//if any shapes changed just resim, this is not efficient but at least it's correct
 	}*/
+
 	return true;
 }
 
diff --git a/Engine/Source/Runtime/Experimental/Chaos/Public/RewindData.h b/Engine/Source/Runtime/Experimental/Chaos/Public/RewindData.h
--- a/Engine/Source/Runtime/Experimental/Chaos/Public/RewindData.h
+++ b/Engine/Source/Runtime/Experimental/Chaos/Public/RewindData.h
@@ -1008,6 +1008,7 @@ struct FGeometryParticleStateBase
 	bool IsInSync(const FGeometryParticleHandle& Handle, const FFrameAndPhase FrameAndPhase, const FDirtyPropertiesPool& Pool) const;
 
 	/** Check if the handle resim frame is valid (before the current one) */
+	UE_DEPRECATED(5.4, "Not recommended for use, will not return a correct response if the particle is not in contact with anything else in the physics scene. Note also that naming is inversed, returns false when valid.")
 	bool IsResimFrameValid(const FGeometryParticleHandle& Handle, const FFrameAndPhase FrameAndPhase) const;
 	
 	template <typename TParticle>
