Level Instance: Allow Property Override ObjectPtrs to be properly fixed up (redirectors)
- Serialize Object pointer properties (ObjectPtrs, SoftObjectPaths) into a FSoftObjectTable array UPROPERTY (index table for the FWorldPartitionPropertyOverrideArchive)
- Keep a second set of ObjectPtrs for hard refs so that they can be loaded as imports (without having to explicitely load)

#rnx

diff --git a/Engine/Source/Runtime/Engine/Private/LevelInstance/LevelInstancePropertyOverrideAsset.cpp b/Engine/Source/Runtime/Engine/Private/LevelInstance/LevelInstancePropertyOverrideAsset.cpp
--- a/Engine/Source/Runtime/Engine/Private/LevelInstance/LevelInstancePropertyOverrideAsset.cpp
+++ b/Engine/Source/Runtime/Engine/Private/LevelInstance/LevelInstancePropertyOverrideAsset.cpp
@@ -34,7 +34,8 @@ struct FLevelInstancePropertyOverrideUtils
 				const UObject* DummyObject = GetDefault<UObject>();
 				TArray<uint8> Payload;
 				FWorldPartitionPropertyOverrideWriter Writer(Payload);
-				FWorldPartitionPropertyOverrideArchive Archive(Writer);
+				FPropertyOverrideReferenceTable ReferenceTable;
+				FWorldPartitionPropertyOverrideArchive Archive(Writer, ReferenceTable);
 				UClass* Class = DummyObject->GetClass();
 
 				// By serializing the component with itself as its defaults we guarantee that no properties will be written out
@@ -57,6 +58,9 @@ struct FLevelInstancePropertyOverrideUtils
 
 		Objects.Add(InActor);
 
+		// Reset table
+		OutActorPropertyOverrides.ReferenceTable = FPropertyOverrideReferenceTable();
+
 		TArray<uint8> Payload;
 		// Serialize SubObjects
 		const FString ActorName = InActor->GetName();
@@ -64,7 +68,7 @@ struct FLevelInstancePropertyOverrideUtils
 		{
 			Payload.Reset();
 			FWorldPartitionPropertyOverrideWriter Writer(Payload);
-			FWorldPartitionPropertyOverrideArchive Archive(Writer);
+			FWorldPartitionPropertyOverrideArchive Archive(Writer, OutActorPropertyOverrides.ReferenceTable);
 			UObject* Archetype = InLevelStreaming->GetArchetypeForObject(Object);
 			if (Archetype->GetTypedOuter<ULevel>() == InLevelStreaming->GetArchetypeLevel())
 			{
diff --git a/Engine/Source/Runtime/Engine/Private/WorldPartition/WorldPartitionPropertyOverride.cpp b/Engine/Source/Runtime/Engine/Private/WorldPartition/WorldPartitionPropertyOverride.cpp
--- a/Engine/Source/Runtime/Engine/Private/WorldPartition/WorldPartitionPropertyOverride.cpp
+++ b/Engine/Source/Runtime/Engine/Private/WorldPartition/WorldPartitionPropertyOverride.cpp
@@ -52,7 +52,7 @@ bool UWorldPartitionPropertyOverride::ApplyPropertyOverrides(const FActorPropert
 
 			Object->Modify(false);
 			FWorldPartitionPropertyOverrideReader Reader(SubObjectOverride->SerializedTaggedProperties);
-			FWorldPartitionPropertyOverrideArchive Archive(Reader);
+			FWorldPartitionPropertyOverrideArchive Archive(Reader, InPropertyOverride->ReferenceTable);
 			UClass* ObjectClass = Object->GetClass();
 			ObjectClass->SerializeTaggedProperties(Archive, (uint8*)Object, ObjectClass, nullptr);
 			bAppliedProperties = true;
diff --git a/Engine/Source/Runtime/Engine/Private/WorldPartition/WorldPartitionPropertyOverrideSerialization.cpp b/Engine/Source/Runtime/Engine/Private/WorldPartition/WorldPartitionPropertyOverrideSerialization.cpp
--- a/Engine/Source/Runtime/Engine/Private/WorldPartition/WorldPartitionPropertyOverrideSerialization.cpp
+++ b/Engine/Source/Runtime/Engine/Private/WorldPartition/WorldPartitionPropertyOverrideSerialization.cpp
@@ -7,8 +7,9 @@
 #include "WorldPartition/WorldPartitionSettings.h"
 #include "WorldPartition/WorldPartitionPropertyOverride.h"
 
-FWorldPartitionPropertyOverrideArchive::FWorldPartitionPropertyOverrideArchive(FArchive& InArchive)
-	: FObjectAndNameAsStringProxyArchive(InArchive, /*bLoadIfFindFails*/true)
+FWorldPartitionPropertyOverrideArchive::FWorldPartitionPropertyOverrideArchive(FArchive& InArchive, FPropertyOverrideReferenceTable& InReferenceTable)
+	: FNameAsStringProxyArchive(InArchive)
+	, ReferenceTable(InReferenceTable)
 {
 	check(InArchive.IsPersistent());
 	check(!InArchive.IsFilterEditorOnly());
@@ -35,11 +36,110 @@ bool FWorldPartitionPropertyOverrideArchive::ShouldSkipProperty(const FProperty*
 	return true;
 }
 
+FSoftObjectPath FWorldPartitionPropertyOverrideArchive::ReadSoftObjectPath()
+{
+	if (ReferenceTable.bIsValid)
+	{
+		int32 Index = INDEX_NONE;
+		*this << Index;
+		if (ReferenceTable.SoftObjectPathTable.IsValidIndex(Index))
+		{
+			return ReferenceTable.SoftObjectPathTable[Index];
+		}
+		else
+		{
+			// Only invalid Index we can expect is INDEX_NONE if the serialized SoftObjectPath was null
+			ensureMsgf(Index == INDEX_NONE, TEXT("Invalid Index (%d) was read from the SoftObjectPathTable"), Index);
+			return FSoftObjectPath();
+		}
+	}
+	else
+	{
+		// load the path name to the object
+		FString LoadedString;
+		InnerArchive << LoadedString;
+		return FSoftObjectPath(LoadedString);
+	}
+}
+
+void FWorldPartitionPropertyOverrideArchive::WriteSoftObjectPath(FSoftObjectPath SoftObjectPath)
+{
+	ReferenceTable.bIsValid = true;
+	int32 Index = INDEX_NONE;
+	if (SoftObjectPath.IsValid())
+	{
+		Index = ReferenceTable.SoftObjectPathTable.AddUnique(SoftObjectPath);
+	}
+	*this << Index;
+}
+
 FArchive& FWorldPartitionPropertyOverrideArchive::operator<<(FLazyObjectPtr& Value)
 { 
 	return FArchiveUObject::SerializeLazyObjectPtr(*this, Value); 
 }
 
+FArchive& FWorldPartitionPropertyOverrideArchive::operator<<(UObject*& Obj)
+{
+	if (IsLoading())
+	{
+		FSoftObjectPath Value = ReadSoftObjectPath();
+		Obj = Value.ResolveObject();
+
+		// Previous data didn't have hard references so make sure to load object if it isn't already
+		if (!Obj && !ReferenceTable.bIsValid)
+		{
+			Obj = Value.TryLoad();
+		}
+		return *this;
+	}
+	else
+	{
+		ReferenceTable.ObjectReferences.Add(Obj);
+		WriteSoftObjectPath(FSoftObjectPath(Obj));
+	}
+
+	return *this;
+}
+
+FArchive& FWorldPartitionPropertyOverrideArchive::operator<<(FWeakObjectPtr& Obj)
+{
+	return FArchiveUObject::SerializeWeakObjectPtr(*this, Obj);
+}
+
+FArchive& FWorldPartitionPropertyOverrideArchive::operator<<(FSoftObjectPtr& Value)
+{
+	if (IsLoading())
+	{
+		// Reset before serializing to clear the internal weak pointer. 
+		Value.ResetWeakPtr();
+		Value = ReadSoftObjectPath();
+	}
+	else
+	{
+		WriteSoftObjectPath(Value.GetUniqueID());
+	}
+	return *this;
+}
+
+FArchive& FWorldPartitionPropertyOverrideArchive::operator<<(FSoftObjectPath& Value)
+{
+	if (IsLoading())
+	{
+		Value = ReadSoftObjectPath();
+	}
+	else
+	{
+		WriteSoftObjectPath(Value);
+	}
+	return *this;
+}
+
+FArchive& FWorldPartitionPropertyOverrideArchive::operator<<(FObjectPtr& Obj)
+{
+	return FArchiveUObject::SerializeObjectPtr(*this, Obj);
+}
+
+
 FWorldPartitionPropertyOverrideWriter::FWorldPartitionPropertyOverrideWriter(TArray<uint8, TSizedDefaultAllocator<32>>& InBytes)
 	: FMemoryWriter(InBytes, true)
 {
diff --git a/Engine/Source/Runtime/Engine/Private/WorldPartition/WorldPartitionPropertyOverrideSerialization.h b/Engine/Source/Runtime/Engine/Private/WorldPartition/WorldPartitionPropertyOverrideSerialization.h
--- a/Engine/Source/Runtime/Engine/Private/WorldPartition/WorldPartitionPropertyOverrideSerialization.h
+++ b/Engine/Source/Runtime/Engine/Private/WorldPartition/WorldPartitionPropertyOverrideSerialization.h
@@ -8,23 +8,32 @@
 #include "UObject/ObjectMacros.h"
 #include "UObject/UObjectGlobals.h"
 #include "UObject/Object.h"
-#include "Serialization/ObjectAndNameAsStringProxyArchive.h"
+#include "Serialization/NameAsStringProxyArchive.h"
 #include "Serialization/MemoryReader.h"
 #include "Serialization/MemoryWriter.h"
+#include "WorldPartition/WorldPartitionPropertyOverride.h"
 
 class UWorldPartitionPropertyOverridePolicy;
 
-class FWorldPartitionPropertyOverrideArchive : public FObjectAndNameAsStringProxyArchive
+class FWorldPartitionPropertyOverrideArchive : public FNameAsStringProxyArchive
 {
 public:
-	FWorldPartitionPropertyOverrideArchive(FArchive& InArchive);
+	FWorldPartitionPropertyOverrideArchive(FArchive& InArchive, FPropertyOverrideReferenceTable& InReferenceTable);
 
 	virtual bool ShouldSkipProperty(const FProperty* InProperty) const override;
 	
 	virtual FArchive& operator<<(FLazyObjectPtr& Value) override;
-
+	virtual FArchive& operator<<(UObject*& Obj) override;
+	virtual FArchive& operator<<(FWeakObjectPtr& Obj) override;
+	virtual FArchive& operator<<(FSoftObjectPtr& Value) override;
+	virtual FArchive& operator<<(FSoftObjectPath& Value) override;
+	virtual FArchive& operator<<(FObjectPtr& Obj) override;
 private:
 	UWorldPartitionPropertyOverridePolicy* PropertyOverridePolicy = nullptr;
+	FPropertyOverrideReferenceTable& ReferenceTable;
+
+	FSoftObjectPath ReadSoftObjectPath();
+	void WriteSoftObjectPath(FSoftObjectPath SoftObjectPath);
 };
 
 class FWorldPartitionPropertyOverrideWriter : public FMemoryWriter
diff --git a/Engine/Source/Runtime/Engine/Public/WorldPartition/WorldPartitionPropertyOverride.h b/Engine/Source/Runtime/Engine/Public/WorldPartition/WorldPartitionPropertyOverride.h
--- a/Engine/Source/Runtime/Engine/Public/WorldPartition/WorldPartitionPropertyOverride.h
+++ b/Engine/Source/Runtime/Engine/Public/WorldPartition/WorldPartitionPropertyOverride.h
@@ -25,6 +25,27 @@ public:
 #endif
 };
 
+USTRUCT(NotBlueprintable)
+struct FPropertyOverrideReferenceTable
+{
+	GENERATED_BODY()
+
+#if WITH_EDITORONLY_DATA
+	// Contains SoftObjectPaths from the FSubObjectPropertyOverride serialization so that they can be properly fixed up (fixup redirectors)
+	// This table should not be changed outside of serialization of the SubObjectOverrides
+	UPROPERTY()
+	TArray<FSoftObjectPath> SoftObjectPathTable;
+
+	// Contains hard refs from the SoftObjectPathTable
+	UPROPERTY()
+	TSet<TObjectPtr<UObject>> ObjectReferences;
+
+	// Support previous data this will be false until this override is resaved
+	UPROPERTY()
+	bool bIsValid = false;
+#endif
+};
+
 // Per Actor overrides, includes a map of Sub-Object name to FSubObjectPropertyOverride data
 USTRUCT(NotBlueprintable)
 struct FActorPropertyOverride
@@ -39,6 +60,9 @@ public:
 
 	UPROPERTY()
 	TMap<FString, FSubObjectPropertyOverride> SubObjectOverrides;
+
+	UPROPERTY()
+	mutable FPropertyOverrideReferenceTable ReferenceTable;
 #endif
 };
 
