World Partition Runtime HahSet: don't resolve outer world and world data layers objects for every intersecting cells, do it once instead.

diff --git a/Engine/Source/Runtime/Engine/Private/WorldPartition/RuntimeHashSet/WorldPartitionRuntimeHashSet.cpp b/Engine/Source/Runtime/Engine/Private/WorldPartition/RuntimeHashSet/WorldPartitionRuntimeHashSet.cpp
--- a/Engine/Source/Runtime/Engine/Private/WorldPartition/RuntimeHashSet/WorldPartitionRuntimeHashSet.cpp
+++ b/Engine/Source/Runtime/Engine/Private/WorldPartition/RuntimeHashSet/WorldPartitionRuntimeHashSet.cpp
@@ -4,6 +4,7 @@
 #include "WorldPartition/RuntimeHashSet/RuntimePartition.h"
 #include "WorldPartition/RuntimeHashSet/RuntimePartitionLHGrid.h"
 #include "WorldPartition/RuntimeHashSet/RuntimePartitionPersistent.h"
+#include "WorldPartition/DataLayer/WorldDataLayers.h"
 #include "WorldPartition/HLOD/HLODLayer.h"
 #include "WorldPartition/ContentBundle/ContentBundleDescriptor.h"
 #include "WorldPartition/DataLayer/DataLayersID.h"
@@ -449,6 +450,10 @@ void UWorldPartitionRuntimeHashSet::ForEachStreamingCellsQuery(const FWorldParti
 
 void UWorldPartitionRuntimeHashSet::ForEachStreamingCellsSources(const TArray<FWorldPartitionStreamingSource>& Sources, TFunctionRef<bool(const UWorldPartitionRuntimeCell*, EStreamingSourceTargetState)> Func) const
 {
+	const UWorld* OuterWorld = GetTypedOuter<UWorld>();
+	const AWorldDataLayers* WorldDataLayers = OuterWorld->GetWorldDataLayers();
+	const int32 DataLayersStateEpoch = WorldDataLayers->GetDataLayersStateEpoch();
+
 	// Non-spatially loaded cells
 	for (const FRuntimePartitionStreamingData* StreamingData : RuntimeNonSpatiallyLoadedDataGridList)
 	{
@@ -462,7 +467,7 @@ void UWorldPartitionRuntimeHashSet::ForEachStreamingCellsSources(const TArray<FW
 #else
 			check(IsCellRelevantFor(Cell->GetClientOnlyVisible()));
 #endif
-			const EDataLayerRuntimeState CellEffectiveWantedState = Cell->GetCellEffectiveWantedState();
+			const EDataLayerRuntimeState CellEffectiveWantedState = Cell->GetCellEffectiveWantedState(DataLayersStateEpoch);
 			if (CellEffectiveWantedState != EDataLayerRuntimeState::Unloaded)
 			{
 				Func(Cell, (CellEffectiveWantedState == EDataLayerRuntimeState::Loaded) ? EStreamingSourceTargetState::Loaded : EStreamingSourceTargetState::Activated);
@@ -505,11 +510,11 @@ void UWorldPartitionRuntimeHashSet::ForEachStreamingCellsSources(const TArray<FW
 
 				for (const FRuntimePartitionStreamingData* StreamingData : *StreamingDataList)
 				{
-					Source.ForEachShape(StreamingData->LoadingRange, false, [this, &Source, StreamingData, &Func](const FSphericalSector& Shape)
+					Source.ForEachShape(StreamingData->LoadingRange, false, [this, &Source, StreamingData, DataLayersStateEpoch, &Func](const FSphericalSector& Shape)
 					{
 						const FSphere ShapeSphere(Shape.GetCenter(), Shape.GetRadius());
 
-						auto ForEachIntersectingElementFunc = [this, &Source, &Shape, &Func](UWorldPartitionRuntimeCell* Cell)
+						auto ForEachIntersectingElementFunc = [this, &Source, &Shape, DataLayersStateEpoch, &Func](UWorldPartitionRuntimeCell* Cell)
 						{
 #if WITH_EDITOR
 							if (!IsCellRelevantFor(Cell->GetClientOnlyVisible()))
@@ -519,7 +524,7 @@ void UWorldPartitionRuntimeHashSet::ForEachStreamingCellsSources(const TArray<FW
 #else
 							check(IsCellRelevantFor(Cell->GetClientOnlyVisible()));
 #endif
-							const EDataLayerRuntimeState CellEffectiveWantedState = Cell->GetCellEffectiveWantedState();
+							const EDataLayerRuntimeState CellEffectiveWantedState = Cell->GetCellEffectiveWantedState(DataLayersStateEpoch);
 							if (CellEffectiveWantedState != EDataLayerRuntimeState::Unloaded)
 							{
 								Cell->AppendStreamingSourceInfo(Source, Shape);
diff --git a/Engine/Source/Runtime/Engine/Public/WorldPartition/WorldPartitionRuntimeCell.h b/Engine/Source/Runtime/Engine/Public/WorldPartition/WorldPartitionRuntimeCell.h
--- a/Engine/Source/Runtime/Engine/Public/WorldPartition/WorldPartitionRuntimeCell.h
+++ b/Engine/Source/Runtime/Engine/Public/WorldPartition/WorldPartitionRuntimeCell.h
@@ -286,6 +286,14 @@ class UWorldPartitionRuntimeCell : public UObject, public IWorldPartitionCell, p
 
 	ENGINE_API UDataLayerManager* GetDataLayerManager() const;
 	ENGINE_API EDataLayerRuntimeState GetCellEffectiveWantedState() const;
+	FORCEINLINE EDataLayerRuntimeState GetCellEffectiveWantedState(int32 InDataLayersStateEpoch) const
+	{
+		if (EffectiveWantedStateEpoch != InDataLayersStateEpoch)
+		{
+			return GetCellEffectiveWantedState();
+		}
+		return EffectiveWantedState;
+	}
 	FORCEINLINE bool HasDataLayers() const { return !DataLayers.IsEmpty(); }
 
 	void SetBlockOnSlowLoading(bool bInBlockOnSlowLoading) { bBlockOnSlowLoading = bInBlockOnSlowLoading; }
