FStaticMeshComponentBulkReregisterContext::SanitizeMeshComponents : Fix crash with Unregistered PCG Components while moving an actor

New code in PCG (CL 32221745) now listens to OnObjectsReplaced which can cause PCG to unregister a component while ReRunConstructionScript is executed.
When the FStaticMeshComponentBulkReregisterContext exits scope we need to remove those unregistered components from the bulk register list as they can't be registered to the render anymore (they are unregistered from the world)

#rnx

diff --git a/Engine/Source/Runtime/Engine/Private/ActorEditor.cpp b/Engine/Source/Runtime/Engine/Private/ActorEditor.cpp
--- a/Engine/Source/Runtime/Engine/Private/ActorEditor.cpp
+++ b/Engine/Source/Runtime/Engine/Private/ActorEditor.cpp
@@ -257,8 +257,8 @@ void AActor::PostEditMove(bool bFinished)
 
 			FNavigationLockContext NavLock(GetWorld(), ENavigationLockReason::AllowUnregister);
 			RerunConstructionScripts();
-			// Construction scripts can have all manner of side effects, including creation of proxies..
-			// remove any static mesh components that have had their proxies created by RerunConstructionScripts:
+			// Construction scripts can have all manner of side effects, including creation of render proxies, unregistering / destruction of components
+			// Remove any static mesh components that are no longer valid for proxy registration after RerunConstructionScripts
 			ReregisterContext.SanitizeMeshComponents();
 		}
 	}
diff --git a/Engine/Source/Runtime/Engine/Private/StaticMeshResources.cpp b/Engine/Source/Runtime/Engine/Private/StaticMeshResources.cpp
--- a/Engine/Source/Runtime/Engine/Private/StaticMeshResources.cpp
+++ b/Engine/Source/Runtime/Engine/Private/StaticMeshResources.cpp
@@ -192,7 +192,7 @@ void FStaticMeshComponentBulkReregisterContext::AddSimpleConstructionScript(USim
 void FStaticMeshComponentBulkReregisterContext::SanitizeMeshComponents()
 {
 	// for contexts in which side effects are less predictable, e.g. in editor:
-	StaticMeshComponents.RemoveAllSwap([](const UPrimitiveComponent* Component) { return !IsValidChecked(Component) || Component->SceneProxy != nullptr; }, EAllowShrinking::No);
+	StaticMeshComponents.RemoveAllSwap([](const UPrimitiveComponent* Component) { return !IsValidChecked(Component) || !Component->IsRegistered() || Component->SceneProxy != nullptr; }, EAllowShrinking::No);
 }
 
 void FStaticMeshComponentBulkReregisterContext::AddConstructedComponent(USceneComponent* SceneComp)
