Reduce the sleeping done in AsyncCompilationHelper::FinishCompilation. It's very common for async tasks to be spawned and immediately waited on. Despite the tasks taking little time, we often won't complete _all_ before the end of the FinishCompilation wait loop, causing a 16ms stall on the game thread -- a stall which might be longer than the time necessary to complete the async work.

Since we don't have a means to steal async tasks we are waiting on, we instead only sleep when we haven't made progress, and when we do we simply yield the thread rather than wait an arbitrary amount of time.

PIE.Startup time spent in UAnimationAsset::ValidateSkeleton 21s -> 9s, warm PIE.Startup takes 3m22s > 3m4s

diff --git a/Engine/Source/Runtime/Engine/Private/AsyncCompilationHelpers.cpp b/Engine/Source/Runtime/Engine/Private/AsyncCompilationHelpers.cpp
--- a/Engine/Source/Runtime/Engine/Private/AsyncCompilationHelpers.cpp
+++ b/Engine/Source/Runtime/Engine/Private/AsyncCompilationHelpers.cpp
@@ -137,11 +137,14 @@ namespace AsyncCompilationHelpers
 			return FText::Format(LOCTEXT("WaitingOnFinishCompilationWithCount", "Waiting for {AssetType} to be ready {Done}/{Total} ({ObjectName}) ..."), Args);
 		};
 
+		const double MaxProcessingTimeSeconds = 0.016f;
+		double StartTimeSeconds = FPlatformTime::Seconds();
 		int32 NumDone = 0;
 		TBitArray<> JobsToFinish(true, Num);
 		TBitArray<> LoggedSlowTask(false, Num);
 		for(;;)
 		{
+			int32 OldNumDone = NumDone;
 			for (TBitArray<>::FWordIterator It(JobsToFinish); It; ++It)
 			{
 				const uint32_t BaseIndex = It.GetIndex();
@@ -175,41 +178,56 @@ namespace AsyncCompilationHelpers
 				It.SetWord(WordJobState);
 			}
 
-			// SN-DBS jobs (which make use of the http service) needs to be ticked
-			// from the game-thread to avoid starvation while we wait for other
-			// async tasks to finish.
-			if (GShaderCompilingManager && GShaderCompilingManager->IsCompiling())
-			{
-				const bool bLimitExecutionTime = true;
-				const bool bBlockOnGlobalShaderCompletion = false;
-				GShaderCompilingManager->ProcessAsyncResults(bLimitExecutionTime, bBlockOnGlobalShaderCompletion);
-			}
-
-			if(NumDone >= Num)
+			if (NumDone >= Num)
 			{
 				break;
 			}
 
-			if (SlowTask.IsSet())
+			// We haven't finished all tasks yet. Progress the slow task so the editor may 
+			// remain responsive if we hit our timeout, otherwise if we didn't make 
+			// progress yield the thread to avoid spinning needlessly
+			//
+			// Note, we want to instead call IncompleteJob.WaitCompletionWithTimeout(MaxProcessingTimeSeconds) here, 
+			// but until we fix the majority of tasks to signal instead of sleeping it's faster to poll and yield.
+			if ((FPlatformTime::Seconds() - StartTimeSeconds) > MaxProcessingTimeSeconds)
 			{
+				// We still have jobs inflight, so find one of the remaining jobs, update 
+				// progress on our slow task to keep the editor responsive and then wait 
+				// for the job to finish or our timeout, whichever comes first
 				int IncompleteJobIndex = JobsToFinish.Find(true);
 				check(IncompleteJobIndex != INDEX_NONE);
-
 				ICompilable& IncompleteJob = Getter(IncompleteJobIndex);
-				FText Progress = FormatProgress(NumDone, Num, IncompleteJob.GetName());
 
-				// Avoid spamming task progress while waiting
-				if (!LoggedSlowTask[IncompleteJobIndex])
+				if (SlowTask.IsSet())
+				{
+					FText Progress = FormatProgress(NumDone, Num, IncompleteJob.GetName());
+
+					// Avoid spamming task progress while waiting
+					if (!LoggedSlowTask[IncompleteJobIndex])
+					{
+						UE_LOG_REF(LogCategory, Display, TEXT("%s"), *Progress.ToString());
+						LoggedSlowTask[IncompleteJobIndex] = true;
+					}
+
+					SlowTask->EnterProgressFrame(0.0f, Progress);
+				}
+
+				// SN-DBS jobs (which make use of the http service) needs to be ticked
+				// from the game-thread to avoid starvation while we wait for other
+				// async tasks to finish.
+				if (GShaderCompilingManager && GShaderCompilingManager->IsCompiling())
 				{
-					UE_LOG_REF(LogCategory, Display, TEXT("%s"), *Progress.ToString());
-					LoggedSlowTask[IncompleteJobIndex] = true;
+					const bool bLimitExecutionTime = true;
+					const bool bBlockOnGlobalShaderCompletion = false;
+					GShaderCompilingManager->ProcessAsyncResults(bLimitExecutionTime, bBlockOnGlobalShaderCompletion);
 				}
 
-				SlowTask->EnterProgressFrame(0.0f, Progress);
+				StartTimeSeconds = FPlatformTime::Seconds();
+			}
+			else if (NumDone == OldNumDone)
+			{
+				FPlatformProcess::YieldThread();
 			}
-
-			// Jobs are still in flight so give them some time to complete
-			FPlatformProcess::Sleep(0.016);
 		}
 
 		SaveStallStack(FPlatformTime::Cycles64() - StartTime);
