Include retargeting transforms in the animation compression DDC key for additive sequences

diff --git a/Engine/Source/Runtime/Engine/Private/Animation/AnimCompressionTypes.cpp b/Engine/Source/Runtime/Engine/Private/Animation/AnimCompressionTypes.cpp
--- a/Engine/Source/Runtime/Engine/Private/Animation/AnimCompressionTypes.cpp
+++ b/Engine/Source/Runtime/Engine/Private/Animation/AnimCompressionTypes.cpp
@@ -244,8 +244,14 @@ void FCompressibleAnimData::BakeOutAdditiveIntoRawData(const FFrameRate& SampleR
 	}
 
 	FMemMark Mark(FMemStack::Get());
+
 	FByFramePoseEvalContext EvalContext(AnimSequence);
 
+	// Enable re-targeting as we wish for both the base and additive sequences to use the same frame of reference.
+	EvalContext.RequiredBones.SetDisableRetargeting(false);
+	EvalContext.RequiredBones.SetUseRAWData(true);
+	EvalContext.RequiredBones.SetUseSourceData(false);
+
 	TScriptInterface<IAnimationDataModel> DataModelInterface = AnimSequence->GetDataModelInterface();
 	// We actually need to resample bone transforms
 	const FFrameNumber ModelNumberOfFrames = DataModelInterface->GetNumberOfFrames();
@@ -445,7 +451,10 @@ void FCompressibleAnimData::ResampleAnimationTrackData(const FFrameRate& SampleR
             }
 
 			FMemMark Mark(FMemStack::Get());
+
 			FByFramePoseEvalContext EvalContext(AnimSequence);
+
+			// Disable re-targeting since we wish to compress pre-retargeting data. Retargeting is performed at runtime.
 			EvalContext.RequiredBones.SetDisableRetargeting(true);
 			EvalContext.RequiredBones.SetUseRAWData(true);
 			EvalContext.RequiredBones.SetUseSourceData(false);
diff --git a/Engine/Source/Runtime/Engine/Private/Animation/AnimSequence.cpp b/Engine/Source/Runtime/Engine/Private/Animation/AnimSequence.cpp
--- a/Engine/Source/Runtime/Engine/Private/Animation/AnimSequence.cpp
+++ b/Engine/Source/Runtime/Engine/Private/Animation/AnimSequence.cpp
@@ -4542,7 +4542,7 @@ FIoHash UAnimSequence::CreateDerivedDataKeyHash(const ITargetPlatform* TargetPla
 		bPerformFrameStripping = ShouldPerformStripping(bPerformFrameStripping, bPerformFrameStrippingOnOddNumberedFrames);
 		PRAGMA_ENABLE_DEPRECATION_WARNINGS
 	}
-		
+
 	FArcToHexString ArcToHexString;
 
 	ArcToHexString.Ar << CompressionErrorThresholdScale;
@@ -4550,7 +4550,17 @@ FIoHash UAnimSequence::CreateDerivedDataKeyHash(const ITargetPlatform* TargetPla
 	BoneCompressionSettings->PopulateDDCKey(UE::Anim::Compression::FAnimDDCKeyArgs(*this, TargetPlatform), ArcToHexString.Ar);
 	CurveCompressionSettings->PopulateDDCKey(ArcToHexString.Ar);
 	VariableFrameStrippingSettings->PopulateDDCKey(UE::Anim::Compression::FAnimDDCKeyArgs(*this, TargetPlatform), ArcToHexString.Ar);
-	
+
+	if (bIsValidAdditive)
+	{
+		// Additive sequences are compressed in re-targeted space, as such we need to include the re-targeting transforms in our key
+		const TArray<FTransform>& RetargetTransforms = GetRetargetTransforms();
+		for (FTransform RetargetTransform : RetargetTransforms)	// Copy the transform to allow us to use the non-const << operator below
+		{
+			ArcToHexString.Ar << RetargetTransform;
+		}
+	}
+
 	const FFrameRate FrameRate = PlatformTargetFrameRate.GetValueForPlatform(TargetPlatform->GetPlatformInfo().IniPlatformName);
 
 	FString Ret = FString::Printf(TEXT("%i_%s%s%s_%c%c%i_%s_%s_%i_%i_%s"),
