Fix stack use after scope in FLevelUtils::ApplyLevelTransform.

diff --git a/Engine/Source/Runtime/Engine/Private/WorldPartition/WorldPartitionLevelHelper.cpp b/Engine/Source/Runtime/Engine/Private/WorldPartition/WorldPartitionLevelHelper.cpp
--- a/Engine/Source/Runtime/Engine/Private/WorldPartition/WorldPartitionLevelHelper.cpp
+++ b/Engine/Source/Runtime/Engine/Private/WorldPartition/WorldPartitionLevelHelper.cpp
@@ -110,7 +110,7 @@ void FWorldPartitionLevelHelper::ApplyConstructionScriptPropertyOverridesFromAnn
 		{
 			if (InActor->GetRootComponent())
 			{
-				FTransform InverseTransform = Annotation.ContainerTransform.Inverse();
+				const FTransform InverseTransform = Annotation.ContainerTransform.Inverse();
 				FLevelUtils::FApplyLevelTransformParams TransformParams(InActor->GetLevel(), InverseTransform);
 				TransformParams.Actor = InActor;
 				TransformParams.bDoPostEditMove = false;
@@ -753,7 +753,8 @@ bool FWorldPartitionLevelHelper::LoadActorsInternal(FLoadActorsParams&& InParams
 						ActorPropertyOverridesAnnotation.AddAnnotation(Actor, FWorldPartitionLevelHelper::FActorPropertyOverridesAnnotation(MoveTemp(ActorPropertyOverrides), PackageObjectMapping->ContainerTransform));
 					}
 
-					FLevelUtils::FApplyLevelTransformParams TransformParams(nullptr, PackageObjectMapping->ContainerTransform * PackageObjectMapping->EditorOnlyParentTransform);
+					const FTransform TransformToApply = PackageObjectMapping->ContainerTransform * PackageObjectMapping->EditorOnlyParentTransform;
+					FLevelUtils::FApplyLevelTransformParams TransformParams(nullptr, TransformToApply);
 					TransformParams.Actor = Actor;
 					TransformParams.bDoPostEditMove = false;
 					FLevelUtils::ApplyLevelTransform(TransformParams);
