World Partition Optimizations:
- Implement streaming source structs move operators to avoid costly copies since they contains arrays.
- Convert FStreamingSourceShape::GetTYpeHash to HashCombineFast and avoid going through CRC computations when hashing Location.

diff --git a/Engine/Source/Runtime/Engine/Public/WorldPartition/WorldPartitionStreamingSource.h b/Engine/Source/Runtime/Engine/Public/WorldPartition/WorldPartitionStreamingSource.h
--- a/Engine/Source/Runtime/Engine/Public/WorldPartition/WorldPartitionStreamingSource.h
+++ b/Engine/Source/Runtime/Engine/Public/WorldPartition/WorldPartitionStreamingSource.h
@@ -130,14 +130,16 @@ struct FStreamingSourceShape
 	friend uint32 GetTypeHash(const FStreamingSourceShape& InShape)
 	{
 		uint32 Hash = GetTypeHash(InShape.bUseGridLoadingRange);
-		Hash = HashCombine(Hash, GetTypeHash(InShape.LoadingRangeScale));
-		Hash = HashCombine(Hash, GetTypeHash(InShape.Radius));
-		Hash = HashCombine(Hash, GetTypeHash(InShape.bIsSector));
-		Hash = HashCombine(Hash, GetTypeHash(InShape.SectorAngle));
-		Hash = HashCombine(Hash, GetTypeHash(InShape.Location));
-		Hash = HashCombine(Hash, GetTypeHash(InShape.Rotation.Pitch));
-		Hash = HashCombine(Hash, GetTypeHash(InShape.Rotation.Yaw));
-		Hash = HashCombine(Hash, GetTypeHash(InShape.Rotation.Roll));
+		Hash = HashCombineFast(Hash, GetTypeHash(InShape.LoadingRangeScale));
+		Hash = HashCombineFast(Hash, GetTypeHash(InShape.Radius));
+		Hash = HashCombineFast(Hash, GetTypeHash(InShape.bIsSector));
+		Hash = HashCombineFast(Hash, GetTypeHash(InShape.SectorAngle));
+		Hash = HashCombineFast(Hash, GetTypeHash(InShape.Location.X));
+		Hash = HashCombineFast(Hash, GetTypeHash(InShape.Location.Y));
+		Hash = HashCombineFast(Hash, GetTypeHash(InShape.Location.Z));
+		Hash = HashCombineFast(Hash, GetTypeHash(InShape.Rotation.Pitch));
+		Hash = HashCombineFast(Hash, GetTypeHash(InShape.Rotation.Yaw));
+		Hash = HashCombineFast(Hash, GetTypeHash(InShape.Rotation.Roll));
 		return Hash;
 	}
 };
@@ -246,26 +248,12 @@ struct FWorldPartitionStreamingQuerySource
 		, TargetBehavior(EStreamingSourceTargetBehavior::Include)
 	{}
 
-	// Define Copy Constructor to avoid deprecation warnings
-	FWorldPartitionStreamingQuerySource(const FWorldPartitionStreamingQuerySource& Other)
-	{
-		*this = Other;
-	}
-
-	FWorldPartitionStreamingQuerySource& operator=(const FWorldPartitionStreamingQuerySource& Other)
-	{
-		Location = Other.Location;
-		Radius = Other.Radius;
-		bUseGridLoadingRange = Other.bUseGridLoadingRange;
-		DataLayers = Other.DataLayers;
-		bDataLayersOnly = Other.bDataLayersOnly;
-		bSpatialQuery = Other.bSpatialQuery;
-		Rotation = Other.Rotation;
-		TargetBehavior = Other.TargetBehavior;
-		TargetGrids = Other.TargetGrids;
-		Shapes = Other.Shapes;
-		return *this;
-	}
+	PRAGMA_DISABLE_DEPRECATION_WARNINGS
+	FWorldPartitionStreamingQuerySource(const FWorldPartitionStreamingQuerySource& Other) = default;
+	FWorldPartitionStreamingQuerySource& operator=(const FWorldPartitionStreamingQuerySource& Other) = default;
+	FWorldPartitionStreamingQuerySource(FWorldPartitionStreamingQuerySource&& Other) = default;
+	FWorldPartitionStreamingQuerySource& operator=(FWorldPartitionStreamingQuerySource&& Other) = default;
+	PRAGMA_ENABLE_DEPRECATION_WARNINGS
 
 	/* Location to query. (not used if bSpatialQuery is false) */
 	UPROPERTY(BlueprintReadWrite, EditAnywhere, Category = "Query")
@@ -377,38 +365,13 @@ struct FWorldPartitionStreamingSource
 		, ExtraAngle(0)
 	{}
 
-	// Define Copy Constructor to avoid deprecation warnings
-	FWorldPartitionStreamingSource(const FWorldPartitionStreamingSource& Other)
-	{
-		*this = Other;
-	}
+	PRAGMA_DISABLE_DEPRECATION_WARNINGS
+	FWorldPartitionStreamingSource(const FWorldPartitionStreamingSource& Other) = default;
+	FWorldPartitionStreamingSource& operator=(const FWorldPartitionStreamingSource& Other) = default;
+	FWorldPartitionStreamingSource(FWorldPartitionStreamingSource&& Other) = default;
+	FWorldPartitionStreamingSource& operator=(FWorldPartitionStreamingSource&& Other) = default;
+	PRAGMA_ENABLE_DEPRECATION_WARNINGS
 
-	FWorldPartitionStreamingSource& operator=(const FWorldPartitionStreamingSource& Other)
-	{
-		Name = Other.Name;
-		Location = Other.Location;
-		Rotation = Other.Rotation;
-		TargetState = Other.TargetState;
-		bBlockOnSlowLoading = Other.bBlockOnSlowLoading;
-		Priority = Other.Priority;
-		bUseVelocityContributionToCellsSorting = Other.bUseVelocityContributionToCellsSorting;
-		Velocity = Other.Velocity;
-		DebugColor = Other.DebugColor;
-		TargetBehavior = Other.TargetBehavior;
-		TargetGrids = Other.TargetGrids;
-		Shapes = Other.Shapes;
-		bReplay = Other.bReplay;
-		bRemote = Other.bRemote;
-		bForce2D = Other.bForce2D;
-		Hash2D = Other.Hash2D;
-		Hash3D = Other.Hash3D;
-		OldLocation = Other.OldLocation;
-		OldRotation = Other.OldRotation;
-		ExtraRadius = Other.ExtraRadius;
-		ExtraAngle = Other.ExtraAngle;
-		return *this;
-	}
-	
 	FColor GetDebugColor() const
 	{
 		if (!DebugColor.ToPackedBGRA())
@@ -567,7 +530,7 @@ struct IWorldPartitionStreamingSourceProvider
 		FWorldPartitionStreamingSource StreamingSource;
 		if (GetStreamingSource(StreamingSource))
 		{
-			StreamingSources.Add(StreamingSource);
+			StreamingSources.Add(MoveTemp(StreamingSource));
 			return true;
 		}
 		return false;
