World Partition Optimization: avoid fetching world partition subsystem twice when computing maximum cells to load when updating streaming state.

diff --git a/Engine/Source/Runtime/Engine/Private/WorldPartition/WorldPartitionSubsystem.cpp b/Engine/Source/Runtime/Engine/Private/WorldPartition/WorldPartitionSubsystem.cpp
--- a/Engine/Source/Runtime/Engine/Private/WorldPartition/WorldPartitionSubsystem.cpp
+++ b/Engine/Source/Runtime/Engine/Private/WorldPartition/WorldPartitionSubsystem.cpp
@@ -579,35 +579,6 @@ void UWorldPartitionSubsystem::UpdateLoadingAndPendingLoadStreamingLevels(const
 	}
 }
 
-int32 UWorldPartitionSubsystem::GetMaxCellsToLoad(const UWorld* InWorld)
-{
-	TRACE_CPUPROFILER_EVENT_SCOPE(UWorldPartitionSubsystem::GetMaxCellsToLoad);
-	int32 MaxCellsToLoad = MAX_int32;
-	if (!IsServer(InWorld) && !InWorld->GetIsInBlockTillLevelStreamingCompleted() && (GMaxLoadingStreamingCells > 0))
-	{
-		if (UWorldPartitionSubsystem* WorldPartitionSubsystem = UWorld::GetSubsystem<UWorldPartitionSubsystem>(InWorld))
-		{
-			MaxCellsToLoad = GMaxLoadingStreamingCells;
-			for (auto It = WorldPartitionSubsystem->WorldPartitionLoadingAndPendingLoadStreamingLevels.CreateIterator(); It; ++It)
-			{
-				const ULevelStreaming* InStreamingLevel = It->Get();
-				if (IsLoadingOrPendingLoadStreamingLevel(InStreamingLevel))
-				{
-					if (!--MaxCellsToLoad)
-					{
-						break;
-					}
-				}
-				else
-				{
-					It.RemoveCurrent();
-				}
-			}
-		}
-	}
-	return MaxCellsToLoad;
-};
-
 void UWorldPartitionSubsystem::OnLevelStreamingStateChanged(UWorld* InWorld, const ULevelStreaming* InStreamingLevel, ULevel* LevelIfLoaded, ELevelStreamingState PreviousState, ELevelStreamingState NewState)
 {
 	if (InWorld != GetWorld())
@@ -1151,7 +1122,8 @@ void UWorldPartitionSubsystem::UpdateStreamingStateInternal(const UWorld* InWorl
 	const bool bServerStreamingEnabled = bIsServer && WorldPartitionSubsystem && WorldPartitionSubsystem->HasAnyWorldPartitionServerStreamingEnabled();
 	const int32 WorldPartitionUpdateCount = InWorldPartition ? 1 : WorldPartitionSubsystem->RegisteredWorldPartitions.Num();
 
-	const bool bForceDisableIncrementalUpdate = IsHighPriorityLoading(InWorld) || !InWorld->bMatchStarted || InWorld->IsInSeamlessTravel() || InWorld->GetIsInBlockTillLevelStreamingCompleted();
+	const bool bIsInBlockTillLevelStreamingCompleted = InWorld->GetIsInBlockTillLevelStreamingCompleted();
+	const bool bForceDisableIncrementalUpdate = IsHighPriorityLoading(InWorld) || !InWorld->bMatchStarted || InWorld->IsInSeamlessTravel() || bIsInBlockTillLevelStreamingCompleted;
 	const bool bIncrementalUpdate = (GUpdateStreamingStateTimeLimit > 0.f) &&
 									(WorldPartitionUpdateCount > 1) &&
 									!bForceDisableIncrementalUpdate &&
@@ -1196,7 +1168,33 @@ void UWorldPartitionSubsystem::UpdateStreamingStateInternal(const UWorld* InWorl
 	}
 
 	// Compute maximum number of cells to load
-	int32 MaxCellsToLoad = GetMaxCellsToLoad(World);
+	int32 MaxCellsToLoad = MAX_int32;
+	
+	if (WorldPartitionSubsystem)
+	{
+		TRACE_CPUPROFILER_EVENT_SCOPE(GetMaxCellsToLoad);
+
+		if (!bIsServer && !bIsInBlockTillLevelStreamingCompleted && (GMaxLoadingStreamingCells > 0))
+		{
+			MaxCellsToLoad = GMaxLoadingStreamingCells;
+
+			for (auto It = WorldPartitionSubsystem->WorldPartitionLoadingAndPendingLoadStreamingLevels.CreateIterator(); It; ++It)
+			{
+				const ULevelStreaming* InStreamingLevel = It->Get();
+				if (IsLoadingOrPendingLoadStreamingLevel(InStreamingLevel))
+				{
+					if (!--MaxCellsToLoad)
+					{
+						break;
+					}
+				}
+				else
+				{
+					It.RemoveCurrent();
+				}
+			}
+		}
+	}
 
 	// Process cells to activate
 	{
diff --git a/Engine/Source/Runtime/Engine/Public/WorldPartition/WorldPartitionSubsystem.h b/Engine/Source/Runtime/Engine/Public/WorldPartition/WorldPartitionSubsystem.h
--- a/Engine/Source/Runtime/Engine/Public/WorldPartition/WorldPartitionSubsystem.h
+++ b/Engine/Source/Runtime/Engine/Public/WorldPartition/WorldPartitionSubsystem.h
@@ -147,7 +147,6 @@ private:
 	void UpdateServerClientsVisibleLevelNames();
 
 	static ENGINE_API void UpdateStreamingStateInternal(const UWorld* InWorld, UWorldPartition* InWorldPartition = nullptr);
-	static int32 GetMaxCellsToLoad(const UWorld* InWorld);
 	static bool IsServer(const UWorld* InWorld);
 
 	ENGINE_API bool HasAnyWorldPartitionServerStreamingEnabled() const;
