World Partition Optimization: don't transform streaming sources to local if their owning world partition has no transform.

diff --git a/Engine/Source/Runtime/Engine/Private/WorldPartition/WorldPartitionSubsystem.cpp b/Engine/Source/Runtime/Engine/Private/WorldPartition/WorldPartitionSubsystem.cpp
--- a/Engine/Source/Runtime/Engine/Private/WorldPartition/WorldPartitionSubsystem.cpp
+++ b/Engine/Source/Runtime/Engine/Private/WorldPartition/WorldPartitionSubsystem.cpp
@@ -988,7 +988,7 @@ void UWorldPartitionSubsystem::GetStreamingSources(const UWorldPartition* InWorl
 #endif
 
 	// Transform to Local
-	if (OutStreamingSources.Num())
+	if (OutStreamingSources.Num() && InWorldPartition->HasInstanceTransform())
 	{
 		const FTransform WorldToLocal = InWorldPartition->GetInstanceTransform().Inverse();
 		for (FWorldPartitionStreamingSource& StreamingSource : OutStreamingSources)
diff --git a/Engine/Source/Runtime/Engine/Public/WorldPartition/WorldPartition.h b/Engine/Source/Runtime/Engine/Public/WorldPartition/WorldPartition.h
--- a/Engine/Source/Runtime/Engine/Public/WorldPartition/WorldPartition.h
+++ b/Engine/Source/Runtime/Engine/Public/WorldPartition/WorldPartition.h
@@ -181,6 +181,8 @@ public:
 	ENGINE_API const FTransform& GetInstanceTransform() const;
 	//~ End UActorDescContainer Interface
 
+	FORCEINLINE bool HasInstanceTransform() const { return InstanceTransform.IsSet(); }
+
 	//~ Begin UObject Interface
 #if WITH_EDITOR
 	ENGINE_API virtual bool CanEditChange(const FProperty* InProperty) const override;
