Removed deprecated StereoLayerComponent code.

diff --git a/Engine/Source/Runtime/Engine/Classes/Components/StereoLayerComponent.h b/Engine/Source/Runtime/Engine/Classes/Components/StereoLayerComponent.h
--- a/Engine/Source/Runtime/Engine/Classes/Components/StereoLayerComponent.h
+++ b/Engine/Source/Runtime/Engine/Classes/Components/StereoLayerComponent.h
@@ -200,7 +200,6 @@ public:
 	bool operator==(const class UStereoLayerShapeEquirect& Other) const;
 };
 
-
 UCLASS(meta = (DisplayName = "Equirect Layer"), MinimalAPI)
 class UStereoLayerShapeEquirect : public UStereoLayerShape
 {
@@ -276,7 +275,6 @@ public:
 	//~ Begin UObject Interface
 	ENGINE_API virtual void EndPlay(const EEndPlayReason::Type EndPlayReason) override;
 	ENGINE_API virtual void OnUnregister() override;
-	ENGINE_API virtual void Serialize(FArchive& Ar) override;
 	ENGINE_API virtual void PostLoad() override;
 	//~ End UObject Interface
 
@@ -396,30 +394,10 @@ protected:
 	UPROPERTY(EditAnywhere, BlueprintReadOnly, export, Category="StereoLayer")
 	FBox2D UVRect;
 
-	UE_DEPRECATED(4.25, "float CylinderRadius is deprecated. Use the corresponding property of the UStereoLayerShapeCylinder subobject instead.")
-	UPROPERTY()
-	float CylinderRadius_DEPRECATED;
-
-	UE_DEPRECATED(4.25, "float CylinderOverlayArc is deprecated. Use the corresponding property of the UStereoLayerShapeCylinder subobject instead.")
-	UPROPERTY()
-	float CylinderOverlayArc_DEPRECATED;
-
-	UE_DEPRECATED(4.25, "int CylinderHeight is deprecated. Use the corresponding property of the UStereoLayerShapeCylinder subobject instead.")
-	UPROPERTY()
-	int CylinderHeight_DEPRECATED;
-
-	UE_DEPRECATED(4.25, "FEquirectProps EquirectProps is deprecated. Use the corresponding properties of the UStereoLayerShapeEquirect subobject instead.")
-	UPROPERTY()
-	FEquirectProps EquirectProps_DEPRECATED;
-
 	/** Specifies how and where the quad is rendered to the screen **/
 	UPROPERTY(EditAnywhere, BlueprintReadOnly, export, Category="StereoLayer")
     TEnumAsByte<enum EStereoLayerType> StereoLayerType;
 
-	UE_DEPRECATED(4.25, "TEnumAsByte<enum EStereoLayerShape> StereoLayerShape is deprecated. Use Shape instead.")
-	UPROPERTY()
-	TEnumAsByte<enum EStereoLayerShape> StereoLayerShape_DEPRECATED;
-
 	/** Specifies which shape of layer it is.  Note that some shapes will be supported only on certain platforms! **/
 	UPROPERTY(EditAnywhere, BlueprintReadOnly, NoClear, Instanced, Category = "StereoLayer", DisplayName="Stereo Layer Shape")
 	TObjectPtr<UStereoLayerShape> Shape;
@@ -444,8 +422,5 @@ private:
 
 	/** Last frames visiblity state **/
 	bool bLastVisible;
-
-	/** Set if the component was loaded from an old version */
-	bool bNeedsPostLoadFixup;
 };
 
diff --git a/Engine/Source/Runtime/Engine/Private/Components/StereoLayerComponent.cpp b/Engine/Source/Runtime/Engine/Private/Components/StereoLayerComponent.cpp
--- a/Engine/Source/Runtime/Engine/Private/Components/StereoLayerComponent.cpp
+++ b/Engine/Source/Runtime/Engine/Private/Components/StereoLayerComponent.cpp
@@ -1,7 +1,6 @@
 // Copyright Epic Games, Inc. All Rights Reserved.
 
 #include "Components/StereoLayerComponent.h"
-#include "UObject/VRObjectVersion.h"
 #include "Engine/Engine.h"
 #include "TextureResource.h"
 #include "Engine/Texture.h"
@@ -30,7 +29,6 @@ UStereoLayerComponent::UStereoLayerComponent(const FObjectInitializer& ObjectIni
 	, bTextureNeedsUpdate(false)
 	, LastTransform(FTransform::Identity)
 	, bLastVisible(false)
-	, bNeedsPostLoadFixup(false)
 {
 	PrimaryComponentTick.bCanEverTick = true;
 	PrimaryComponentTick.TickGroup = TG_PrePhysics;
@@ -58,55 +56,9 @@ void UStereoLayerComponent::OnUnregister()
 	}
 }
 
-void UStereoLayerComponent::Serialize(FArchive& Ar)
-{
-	Ar.UsingCustomVersion(FVRObjectVersion::GUID);
-	Super::Serialize(Ar);
-	if (Ar.IsLoading() && Ar.CustomVer(FVRObjectVersion::GUID) < FVRObjectVersion::UseSubobjectForStereoLayerShapeProperties)
-	{
-		bNeedsPostLoadFixup = true; // Postponing fixups until after load, as we need to modify subobjects.
-	}
-}
-
 void UStereoLayerComponent::PostLoad()
 {
 	Super::PostLoad();
-	PRAGMA_DISABLE_DEPRECATION_WARNINGS
-	if (bNeedsPostLoadFixup)
-	{
-		switch (StereoLayerShape_DEPRECATED)
-		{
-			case SLSH_QuadLayer:
-			{
-				Shape = NewObject<UStereoLayerShapeQuad>(this, NAME_None, RF_Public);
-				break;
-			}
-			case SLSH_CubemapLayer:
-			{
-				Shape = NewObject<UStereoLayerShapeCubemap>(this, NAME_None, RF_Public);
-				break;
-			}
-			case SLSH_CylinderLayer:
-			{
-				auto Cylinder = NewObject<UStereoLayerShapeCylinder>(this, NAME_None, RF_Public);
-				Shape = Cylinder;
-				Cylinder->Height = CylinderHeight_DEPRECATED;
-				Cylinder->OverlayArc = CylinderOverlayArc_DEPRECATED;
-				Cylinder->Radius = CylinderRadius_DEPRECATED;
-				break;
-			}
-			case SLSH_EquirectLayer:
-			{
-				auto Equirect = NewObject<UStereoLayerShapeEquirect>(this, NAME_None, RF_Public);
-				Shape = Equirect;
-				Equirect->SetEquirectProps(EquirectProps_DEPRECATED);
-				break;
-			}
-		}
-		bNeedsPostLoadFixup = false;
-	}
-	PRAGMA_ENABLE_DEPRECATION_WARNINGS
-
 	if (Shape == nullptr)
 	{
 		Shape = NewObject<UStereoLayerShapeQuad>(this, NAME_None, RF_Public);
