Add code to UGameInstance::OnWorldChanged to repair the world context if it is null or set incorrectly.
Fix test worlds to set the owning game instance before setting the current world or Init as this is the order used for real worlds.
This fixes a crash where a test GameInstance would not have a valid world context.

diff --git a/Engine/Source/Runtime/Engine/Classes/Engine/GameInstance.h b/Engine/Source/Runtime/Engine/Classes/Engine/GameInstance.h
--- a/Engine/Source/Runtime/Engine/Classes/Engine/GameInstance.h
+++ b/Engine/Source/Runtime/Engine/Classes/Engine/GameInstance.h
@@ -302,8 +302,8 @@ public:
 	struct FWorldContext* GetWorldContext() const { return WorldContext; };
 	ENGINE_API class UGameViewportClient* GetGameViewportClient() const;
 
-	/** Callback from the world context when the world changes */
-	virtual void OnWorldChanged(UWorld* OldWorld, UWorld* NewWorld) {}
+	/** Callback from the world context when the world changes, this ensure the world context is set properly */
+	ENGINE_API virtual void OnWorldChanged(UWorld* OldWorld, UWorld* NewWorld);
 
 	/** Starts the GameInstance state machine running */
 	ENGINE_API virtual void StartGameInstance();
diff --git a/Engine/Source/Runtime/Engine/Private/GameInstance.cpp b/Engine/Source/Runtime/Engine/Private/GameInstance.cpp
--- a/Engine/Source/Runtime/Engine/Private/GameInstance.cpp
+++ b/Engine/Source/Runtime/Engine/Private/GameInstance.cpp
@@ -591,6 +591,23 @@ UGameViewportClient* UGameInstance::GetGameViewportClient() const
 	return WC ? WC->GameViewport : nullptr;
 }
 
+void UGameInstance::OnWorldChanged(UWorld* OldWorld, UWorld* NewWorld)
+{
+	// Fix up our world context if it is incorrect
+	FWorldContext* NewWorldContext = const_cast<FWorldContext*>(GEngine->GetWorldContextFromWorld(NewWorld));
+
+	if (WorldContext == nullptr)
+	{
+		// This may be a test/preview world that did not have a world context set before
+		WorldContext = NewWorldContext;
+	}
+	else if (WorldContext != NewWorldContext)
+	{
+		UE_LOG(LogLoad, Warning, TEXT("GameInstance %s changed from world %s to world %s with a different world context!"), *GetName(), *GetPathNameSafe(OldWorld), *GetPathNameSafe(NewWorld));
+		WorldContext = NewWorldContext;
+	}
+}
+
 // This can be defined in the target.cs file to allow map overrides in shipping builds
 #ifndef UE_ALLOW_MAP_OVERRIDE_IN_SHIPPING
 #define UE_ALLOW_MAP_OVERRIDE_IN_SHIPPING 0
diff --git a/Engine/Source/Runtime/Engine/Private/Tests/AutomationCommon.cpp b/Engine/Source/Runtime/Engine/Private/Tests/AutomationCommon.cpp
--- a/Engine/Source/Runtime/Engine/Private/Tests/AutomationCommon.cpp
+++ b/Engine/Source/Runtime/Engine/Private/Tests/AutomationCommon.cpp
@@ -70,7 +70,6 @@ bool FTestWorldWrapper::CreateTestWorld(EWorldType::Type WorldType)
 		case EWorldType::Game:
 			// Create something close to a game world that can be manually ticked
 			TestGameInstance = NewObject<UGameInstance>(GEngine);
-			TestGameInstance->Init();
 			break;
 
 		case EWorldType::Editor:
@@ -100,15 +99,17 @@ bool FTestWorldWrapper::CreateTestWorld(EWorldType::Type WorldType)
 	TestWorld->SetShouldTick(false);
 	TestWorld->AddToRoot();
 	
-	if (TestGameInstance)
-	{
-		TestWorld->SetGameInstance(TestGameInstance);
-	}
-
 	FWorldContext& WorldContext = GEngine->CreateNewWorldContext(WorldType);
-	WorldContext.SetCurrentWorld(TestWorld);
 	WorldContext.OwningGameInstance = TestGameInstance;
 
+	TestWorld->SetGameInstance(TestGameInstance);
+	WorldContext.SetCurrentWorld(TestWorld);
+
+	if (TestGameInstance)
+	{
+		TestGameInstance->Init();
+	}
+	
 	return true;
 }
 
