Add a NaN check to dirty area bounds.
Plus added early warning for invalid bounds when creating FNavigationDirtyArea.

diff --git a/Engine/Source/Runtime/Engine/Classes/AI/Navigation/NavigationTypes.h b/Engine/Source/Runtime/Engine/Classes/AI/Navigation/NavigationTypes.h
--- a/Engine/Source/Runtime/Engine/Classes/AI/Navigation/NavigationTypes.h
+++ b/Engine/Source/Runtime/Engine/Classes/AI/Navigation/NavigationTypes.h
@@ -101,7 +101,7 @@ struct FNavigationDirtyArea
 	TWeakObjectPtr<UObject> OptionalSourceObject;
 	
 	FNavigationDirtyArea() : Flags(0) {}
-	FNavigationDirtyArea(const FBox& InBounds, int32 InFlags, UObject* const InOptionalSourceObject = nullptr) : Bounds(InBounds), Flags(InFlags), OptionalSourceObject(InOptionalSourceObject) {}
+	ENGINE_API FNavigationDirtyArea(const FBox& InBounds, int32 InFlags, UObject* const InOptionalSourceObject = nullptr);
 	FORCEINLINE bool HasFlag(ENavigationDirtyFlag::Type Flag) const { return (Flags & Flag) != 0; }
 
 	bool operator==(const FNavigationDirtyArea& Other) const 
diff --git a/Engine/Source/Runtime/Engine/Private/AI/Navigation/NavigationTypes.cpp b/Engine/Source/Runtime/Engine/Private/AI/Navigation/NavigationTypes.cpp
--- a/Engine/Source/Runtime/Engine/Private/AI/Navigation/NavigationTypes.cpp
+++ b/Engine/Source/Runtime/Engine/Private/AI/Navigation/NavigationTypes.cpp
@@ -76,6 +76,20 @@ namespace FNavigationSystem
 	}	
 }
 
+FNavigationDirtyArea::FNavigationDirtyArea(const FBox& InBounds, int32 InFlags, UObject* const InOptionalSourceObject /*= nullptr*/)
+	: Bounds(InBounds)
+	, Flags(InFlags)
+	, OptionalSourceObject(InOptionalSourceObject)
+{
+#if !NO_LOGGING
+	if (!Bounds.IsValid || Bounds.ContainsNaN())
+	{
+		UE_LOG(LogNavigation, Warning, TEXT("Creation of FNavigationDirtyArea with invalid bounds%s. Bounds: %s, SourceObject: %s."),
+			Bounds.ContainsNaN() ? TEXT(" (contains NaN)") : TEXT(""), *Bounds.ToString(), *GetFullNameSafe(OptionalSourceObject.Get()));
+	}
+#endif //!NO_LOGGING	
+}
+
 //----------------------------------------------------------------------//
 // FNavigationQueryFilter
 //----------------------------------------------------------------------//
