Fix a bug where the substrate outputs get mixed up in Nanite CS materials when base pass velocity is disabled.

diff --git a/Engine/Shaders/Private/ComputeShaderOutputCommon.ush b/Engine/Shaders/Private/ComputeShaderOutputCommon.ush
--- a/Engine/Shaders/Private/ComputeShaderOutputCommon.ush
+++ b/Engine/Shaders/Private/ComputeShaderOutputCommon.ush
@@ -81,15 +81,19 @@ uint4 PassData;
 #if SUBSTRATE_OPAQUE_DEFERRED
 	RWTexture2DArray<uint> OutTargets;
 
-	#if PIXELSHADEROUTPUT_MRT3
-		RWTexture2D<SUBSTRATE_TOP_LAYER_TYPE> OutTarget4;
-	#elif PIXELSHADEROUTPUT_MRT2
-		RWTexture2D<SUBSTRATE_TOP_LAYER_TYPE> OutTarget3;
-	#elif PIXELSHADEROUTPUT_MRT1
-		RWTexture2D<SUBSTRATE_TOP_LAYER_TYPE> OutTarget2;
-	#else
-		RWTexture2D<SUBSTRATE_TOP_LAYER_TYPE> OutTarget1;
-	#endif
+#if SUBSTRATE_FIRST_MRT_INDEX > 4
+	#error Substrate doesn't currently handle more than 4 base GBuffer MRTs
+#endif
+#if SUBSTRATE_FIRST_MRT_INDEX == 4
+	#define SUBSTRATE_TOP_LAYER_NAME OutTarget4
+#elif SUBSTRATE_FIRST_MRT_INDEX == 3
+	#define SUBSTRATE_TOP_LAYER_NAME OutTarget3
+#elif SUBSTRATE_FIRST_MRT_INDEX == 2
+	#define SUBSTRATE_TOP_LAYER_NAME OutTarget2
+#else
+	#define SUBSTRATE_TOP_LAYER_NAME OutTarget1
+#endif
+	RWTexture2D<SUBSTRATE_TOP_LAYER_TYPE> SUBSTRATE_TOP_LAYER_NAME;
 
 #else // SUBSTRATE_OPAQUE_DEFERRED
 
@@ -219,16 +223,7 @@ void ExportPixel(const uint2 PixelPos, FPixelShaderOut ShadedPixel)
 	ConditionalExport(true,  PixelPos, ShadedPixel, 0);
 	ConditionalExport(false, PixelPos, ShadedPixel, 1);
 	ConditionalExport(false, PixelPos, ShadedPixel, 2);
-
-#if PIXELSHADEROUTPUT_MRT3
-	ConditionalExport(true, OutTarget4, PixelPos, ShadedPixel.SubstrateTopLayerData);
-#elif PIXELSHADEROUTPUT_MRT2
-	ConditionalExport(true, OutTarget3, PixelPos, ShadedPixel.SubstrateTopLayerData);
-#elif PIXELSHADEROUTPUT_MRT1
-	ConditionalExport(true, OutTarget2, PixelPos, ShadedPixel.SubstrateTopLayerData);
-#else
-	ConditionalExport(true, OutTarget1, PixelPos, ShadedPixel.SubstrateTopLayerData);
-#endif
+	ConditionalExport(true, SUBSTRATE_TOP_LAYER_NAME, PixelPos, ShadedPixel.SubstrateTopLayerData);
 
 #else // SUBSTRATE_OPAQUE_DEFERRED
 
diff --git a/Engine/Shaders/Private/PixelShaderOutputCommon.ush b/Engine/Shaders/Private/PixelShaderOutputCommon.ush
--- a/Engine/Shaders/Private/PixelShaderOutputCommon.ush
+++ b/Engine/Shaders/Private/PixelShaderOutputCommon.ush
@@ -56,20 +56,20 @@ void MainPS
 #if SUBSTRATE_BASE_PASS_MRT_OUTPUT_COUNT != 3
 #error Substrate SUBSTRATE_BASE_PASS_MRT_OUTPUT_COUNT has been updated but not the uint MRTs
 #endif
-#if PIXELSHADEROUTPUT_MRT4
+#if SUBSTRATE_FIRST_MRT_INDEX > 4
 #error Substrate cannot map to such a case
 #endif
-#if PIXELSHADEROUTPUT_MRT3
+#if SUBSTRATE_FIRST_MRT_INDEX == 4
 		, out uint SubstrateOutput4 : SV_Target4
 		, out uint SubstrateOutput5 : SV_Target5
 		, out uint SubstrateOutput6 : SV_Target6
 		, out SUBSTRATE_TOP_LAYER_TYPE SubstrateOutput7 : SV_Target7
-#elif PIXELSHADEROUTPUT_MRT2
+#elif SUBSTRATE_FIRST_MRT_INDEX == 3
 		, out uint SubstrateOutput3 : SV_Target3
 		, out uint SubstrateOutput4 : SV_Target4
 		, out uint SubstrateOutput5 : SV_Target5
 		, out SUBSTRATE_TOP_LAYER_TYPE SubstrateOutput6 : SV_Target6
-#elif PIXELSHADEROUTPUT_MRT1
+#elif SUBSTRATE_FIRST_MRT_INDEX == 2
 		, out uint SubstrateOutput2 : SV_Target2
 		, out uint SubstrateOutput3 : SV_Target3
 		, out uint SubstrateOutput4 : SV_Target4
@@ -171,17 +171,17 @@ void MainPS
 	// MRT2 is precomputed shadow if both velocity and prec shadow are enabled.
 	// After, Substrate output are appended (3 uint output)
 
-#if PIXELSHADEROUTPUT_MRT3
+#if SUBSTRATE_FIRST_MRT_INDEX == 4
 	SubstrateOutput4 = PixelShaderOut.SubstrateOutput[0];
 	SubstrateOutput5 = PixelShaderOut.SubstrateOutput[1];
 	SubstrateOutput6 = PixelShaderOut.SubstrateOutput[2];
 	SubstrateOutput7 = PixelShaderOut.SubstrateTopLayerData;
-#elif PIXELSHADEROUTPUT_MRT2
+#elif SUBSTRATE_FIRST_MRT_INDEX == 3
 	SubstrateOutput3 = PixelShaderOut.SubstrateOutput[0];
 	SubstrateOutput4 = PixelShaderOut.SubstrateOutput[1];
 	SubstrateOutput5 = PixelShaderOut.SubstrateOutput[2];
 	SubstrateOutput6 = PixelShaderOut.SubstrateTopLayerData;
-#elif PIXELSHADEROUTPUT_MRT1
+#elif SUBSTRATE_FIRST_MRT_INDEX == 2
 	SubstrateOutput2 = PixelShaderOut.SubstrateOutput[0];
 	SubstrateOutput3 = PixelShaderOut.SubstrateOutput[1];
 	SubstrateOutput4 = PixelShaderOut.SubstrateOutput[2];
diff --git a/Engine/Source/Runtime/Engine/Private/ShaderCompiler/ShaderGenerationUtil.cpp b/Engine/Source/Runtime/Engine/Private/ShaderCompiler/ShaderGenerationUtil.cpp
--- a/Engine/Source/Runtime/Engine/Private/ShaderCompiler/ShaderGenerationUtil.cpp
+++ b/Engine/Source/Runtime/Engine/Private/ShaderCompiler/ShaderGenerationUtil.cpp
@@ -1738,30 +1738,59 @@ static void SetSlotsForShadingModelType(bool Slots[], EMaterialShadingModel Shad
 	}
 }
 
+enum class EGBufferSlotUsage : uint8
+{
+	Unused,
+	Used,
+	Written
+};
 
+static bool operator> (EGBufferSlotUsage A, EGBufferSlotUsage B) { return uint8(A) >  uint8(B); }
+static bool operator>=(EGBufferSlotUsage A, EGBufferSlotUsage B) { return uint8(A) >= uint8(B); }
+static bool operator< (EGBufferSlotUsage A, EGBufferSlotUsage B) { return uint8(A) <  uint8(B); }
+static bool operator<=(EGBufferSlotUsage A, EGBufferSlotUsage B) { return uint8(A) <= uint8(B); }
 
-static void SetStandardGBufferSlots(bool Slots[], bool bWriteEmissive, bool bHasTangent, bool bHasVelocity, bool bHasStaticLighting, bool bIsSubstrateMaterial)
+static EGBufferSlotUsage GetGBufferSlotUsage(bool bWritten, bool bForceUsed = false)
 {
-	Slots[GBS_SceneColor] = bWriteEmissive;
-	Slots[GBS_Velocity] = bHasVelocity;
-	Slots[GBS_PrecomputedShadowFactor] = bHasStaticLighting;
-
-	Slots[GBS_WorldNormal] =			bIsSubstrateMaterial ? false : true;
-	Slots[GBS_PerObjectGBufferData] =	bIsSubstrateMaterial ? false : true;
-	Slots[GBS_Metallic] =				bIsSubstrateMaterial ? false : true;
-	Slots[GBS_Specular] =				bIsSubstrateMaterial ? false : true;
-	Slots[GBS_Roughness] =				bIsSubstrateMaterial ? false : true;
-	Slots[GBS_ShadingModelId] =			bIsSubstrateMaterial ? false : true;
-	Slots[GBS_SelectiveOutputMask] =	bIsSubstrateMaterial ? false : true;
-	Slots[GBS_BaseColor] =				bIsSubstrateMaterial ? false : true;
-	Slots[GBS_GenericAO] =				bIsSubstrateMaterial ? false : true;
-	Slots[GBS_AO] =						false;//bIsSubstrateMaterial ? false : false;// true;		// Why only false?
-	Slots[GBS_WorldTangent] =			bIsSubstrateMaterial ? false : bHasTangent;
-	Slots[GBS_Anisotropy] =				bIsSubstrateMaterial ? false : bHasTangent;
+	return bWritten ? EGBufferSlotUsage::Written : (bForceUsed ? EGBufferSlotUsage::Used : EGBufferSlotUsage::Unused);
+}
+
+static void SetStandardGBufferSlots(
+	EGBufferSlotUsage Slots[GBS_Num],
+	bool bWriteEmissive,
+	bool bHasTangent,
+	bool bHasVelocity,
+	bool bWritesVelocity,
+	bool bHasStaticLighting,
+	bool bIsSubstrateMaterial)
+{
+	auto GetUsageIfNotSubstrate = [bIsSubstrateMaterial](bool bWritten, bool bForceUsed = false)
+	{
+		return GetGBufferSlotUsage(bWritten && !bIsSubstrateMaterial, bForceUsed && !bIsSubstrateMaterial);
+	};
+
+	Slots[GBS_SceneColor]				= GetGBufferSlotUsage(bWriteEmissive);
+	Slots[GBS_Velocity]					= GetGBufferSlotUsage(bWritesVelocity, bHasVelocity);
+	Slots[GBS_PrecomputedShadowFactor]	= GetGBufferSlotUsage(bHasStaticLighting);
+
+	Slots[GBS_WorldNormal] 				= GetUsageIfNotSubstrate(true);
+	Slots[GBS_PerObjectGBufferData]		= GetUsageIfNotSubstrate(true);
+	Slots[GBS_Metallic] 				= GetUsageIfNotSubstrate(true);
+	Slots[GBS_Specular] 				= GetUsageIfNotSubstrate(true);
+	Slots[GBS_Roughness] 				= GetUsageIfNotSubstrate(true);
+	Slots[GBS_ShadingModelId] 			= GetUsageIfNotSubstrate(true);
+	Slots[GBS_SelectiveOutputMask] 		= GetUsageIfNotSubstrate(true);
+	Slots[GBS_BaseColor] 				= GetUsageIfNotSubstrate(true);
+	Slots[GBS_GenericAO] 				= GetUsageIfNotSubstrate(true);
+
+	Slots[GBS_WorldTangent] 			= GetUsageIfNotSubstrate(bHasTangent);
+	Slots[GBS_Anisotropy]				= GetUsageIfNotSubstrate(bHasTangent);
+
+	Slots[GBS_AO] 						= EGBufferSlotUsage::Unused;
 }
 
 static void DetermineUsedMaterialSlots(
-	bool Slots[],
+	EGBufferSlotUsage Slots[GBS_Num],
 	const FShaderMaterialDerivedDefines& Dst,
 	const FShaderMaterialPropertyDefines& Mat,
 	const FShaderLightmapPropertyDefines& Lightmap,
@@ -1771,7 +1800,8 @@ static void DetermineUsedMaterialSlots(
 {
 	bool bWriteEmissive = Dst.NEEDS_BASEPASS_VERTEX_FOGGING || Mat.USES_EMISSIVE_COLOR || SrcGlobal.ALLOW_STATIC_LIGHTING || Mat.MATERIAL_SHADINGMODEL_SINGLELAYERWATER;
 	bool bHasTangent = SrcGlobal.GBUFFER_HAS_TANGENT;
-	bool bHasVelocity = Dst.WRITES_VELOCITY_TO_GBUFFER;
+	bool bHasVelocity = SrcGlobal.GBUFFER_HAS_VELOCITY;
+	bool bWritesVelocity = Dst.WRITES_VELOCITY_TO_GBUFFER;
 	bool bHasStaticLighting = Dst.GBUFFER_HAS_PRECSHADOWFACTOR || Dst.WRITES_PRECSHADOWFACTOR_TO_GBUFFER;
 	bool bIsSubstrateMaterial = Mat.SUBSTRATE_ENABLED; // Similarly to FetchFullGBufferInfo, we do not check for MATERIAL_IS_SUBSTRATE as this is decided per project.
 
@@ -1781,70 +1811,69 @@ static void DetermineUsedMaterialSlots(
 	// we have to use if statements, not switch or if/else statements because we can have multiple shader model ids.
 	if (Mat.MATERIAL_SHADINGMODEL_UNLIT)
 	{
-		Slots[GBS_SceneColor] = true;
-		Slots[GBS_Velocity] = bHasVelocity;
+		SetStandardGBufferSlots(Slots, true, false, bHasVelocity, bWritesVelocity, false, bIsSubstrateMaterial);
 	}
 
 	if (Mat.MATERIAL_SHADINGMODEL_DEFAULT_LIT)
 	{
-		SetStandardGBufferSlots(Slots, bWriteEmissive, bHasTangent, bHasVelocity, bHasStaticLighting, bIsSubstrateMaterial);
+		SetStandardGBufferSlots(Slots, bWriteEmissive, bHasTangent, bHasVelocity, bWritesVelocity, bHasStaticLighting, bIsSubstrateMaterial);
 	}
 
 	if (Mat.MATERIAL_SHADINGMODEL_SUBSURFACE)
 	{
-		SetStandardGBufferSlots(Slots, bWriteEmissive, bHasTangent, bHasVelocity, bHasStaticLighting, bIsSubstrateMaterial);
-		Slots[GBS_CustomData] = bUseCustomData;
+		SetStandardGBufferSlots(Slots, bWriteEmissive, bHasTangent, bHasVelocity, bWritesVelocity, bHasStaticLighting, bIsSubstrateMaterial);
+		Slots[GBS_CustomData] = GetGBufferSlotUsage(bUseCustomData);
 	}
 
 	if (Mat.MATERIAL_SHADINGMODEL_PREINTEGRATED_SKIN)
 	{
-		SetStandardGBufferSlots(Slots, bWriteEmissive, bHasTangent, bHasVelocity, bHasStaticLighting, bIsSubstrateMaterial);
-		Slots[GBS_CustomData] = bUseCustomData;
+		SetStandardGBufferSlots(Slots, bWriteEmissive, bHasTangent, bHasVelocity, bWritesVelocity, bHasStaticLighting, bIsSubstrateMaterial);
+		Slots[GBS_CustomData] = GetGBufferSlotUsage(bUseCustomData);
 	}
 
 	if (Mat.MATERIAL_SHADINGMODEL_SUBSURFACE_PROFILE)
 	{
-		SetStandardGBufferSlots(Slots, bWriteEmissive, bHasTangent, bHasVelocity, bHasStaticLighting, bIsSubstrateMaterial);
-		Slots[GBS_CustomData] = bUseCustomData;
+		SetStandardGBufferSlots(Slots, bWriteEmissive, bHasTangent, bHasVelocity, bWritesVelocity, bHasStaticLighting, bIsSubstrateMaterial);
+		Slots[GBS_CustomData] = GetGBufferSlotUsage(bUseCustomData);
 	}
 
 	if (Mat.MATERIAL_SHADINGMODEL_CLEAR_COAT)
 	{
-		SetStandardGBufferSlots(Slots, bWriteEmissive, bHasTangent, bHasVelocity, bHasStaticLighting, bIsSubstrateMaterial);
-		Slots[GBS_CustomData] = bUseCustomData;
+		SetStandardGBufferSlots(Slots, bWriteEmissive, bHasTangent, bHasVelocity, bWritesVelocity, bHasStaticLighting, bIsSubstrateMaterial);
+		Slots[GBS_CustomData] = GetGBufferSlotUsage(bUseCustomData);
 	}
 
 	if (Mat.MATERIAL_SHADINGMODEL_TWOSIDED_FOLIAGE)
 	{
-		SetStandardGBufferSlots(Slots, bWriteEmissive, bHasTangent, bHasVelocity, bHasStaticLighting, bIsSubstrateMaterial);
-		Slots[GBS_CustomData] = bUseCustomData;
+		SetStandardGBufferSlots(Slots, bWriteEmissive, bHasTangent, bHasVelocity, bWritesVelocity, bHasStaticLighting, bIsSubstrateMaterial);
+		Slots[GBS_CustomData] = GetGBufferSlotUsage(bUseCustomData);
 	}
 
 	if (Mat.MATERIAL_SHADINGMODEL_HAIR)
 	{
-		SetStandardGBufferSlots(Slots, bWriteEmissive, bHasTangent, bHasVelocity, bHasStaticLighting, bIsSubstrateMaterial);
-		Slots[GBS_CustomData] = bUseCustomData;
+		SetStandardGBufferSlots(Slots, bWriteEmissive, bHasTangent, bHasVelocity, bWritesVelocity, bHasStaticLighting, bIsSubstrateMaterial);
+		Slots[GBS_CustomData] = GetGBufferSlotUsage(bUseCustomData);
 	}
 
 	if (Mat.MATERIAL_SHADINGMODEL_CLOTH)
 	{
-		SetStandardGBufferSlots(Slots, bWriteEmissive, bHasTangent, bHasVelocity, bHasStaticLighting, bIsSubstrateMaterial);
-		Slots[GBS_CustomData] = bUseCustomData;
+		SetStandardGBufferSlots(Slots, bWriteEmissive, bHasTangent, bHasVelocity, bWritesVelocity, bHasStaticLighting, bIsSubstrateMaterial);
+		Slots[GBS_CustomData] = GetGBufferSlotUsage(bUseCustomData);
 	}
 
 	if (Mat.MATERIAL_SHADINGMODEL_EYE)
 	{
-		SetStandardGBufferSlots(Slots, bWriteEmissive, bHasTangent, bHasVelocity, bHasStaticLighting, bIsSubstrateMaterial);
-		Slots[GBS_CustomData] = bUseCustomData;
+		SetStandardGBufferSlots(Slots, bWriteEmissive, bHasTangent, bHasVelocity, bWritesVelocity, bHasStaticLighting, bIsSubstrateMaterial);
+		Slots[GBS_CustomData] = GetGBufferSlotUsage(bUseCustomData);
 	}
 
 	if (Mat.MATERIAL_SHADINGMODEL_SINGLELAYERWATER)
 	{
 		// single layer water uses standard slots
-		SetStandardGBufferSlots(Slots, bWriteEmissive, bHasTangent, bHasVelocity, bHasStaticLighting, bIsSubstrateMaterial);
+		SetStandardGBufferSlots(Slots, bWriteEmissive, bHasTangent, bHasVelocity, bWritesVelocity, bHasStaticLighting, bIsSubstrateMaterial);
 		if (Mat.SINGLE_LAYER_WATER_SEPARATED_MAIN_LIGHT)
 		{
-			Slots[GBS_SeparatedMainDirLight] = true;
+			Slots[GBS_SeparatedMainDirLight] = EGBufferSlotUsage::Written;
 		}
 	}
 
@@ -1852,7 +1881,6 @@ static void DetermineUsedMaterialSlots(
 	if (Mat.MATERIAL_SHADINGMODEL_THIN_TRANSLUCENT)
 	{
 	}
-
 }
 
 void FShaderCompileUtilities::ApplyDerivedDefines(FShaderCompilerEnvironment& OutEnvironment, FShaderCompilerEnvironment * SharedEnvironment, const EShaderPlatform Platform)
@@ -1894,16 +1922,16 @@ void FShaderCompileUtilities::ApplyDerivedDefines(FShaderCompilerEnvironment& Ou
 
 	EGBufferLayout Layout = (EGBufferLayout)MaterialDefines.GBUFFER_LAYOUT;
 	FGBufferParams Params = FShaderCompileUtilities::FetchGBufferParamsRuntime(Platform, Layout);
-	FGBufferInfo BufferInfo = FetchFullGBufferInfo(Params);
+	FGBufferInfo BufferInfo = FetchFullGBufferInfo(Params);	
 
-	bool bTargetUsage[FGBufferInfo::MaxTargets] = {};
+	EGBufferSlotUsage TargetUsage[FGBufferInfo::MaxTargets] = {};
 	if (MaterialDefines.IS_BASE_PASS)
 	{
 		// if we are using a gbuffer, and this is the base pass that writes a gbuffer, search the gbuffer for each slot
 		if (DerivedDefines.USES_GBUFFER)
 		{
-			bTargetUsage[0] = true;
-			bool Slots[GBS_Num] = {};
+			TargetUsage[0] = EGBufferSlotUsage::Written;
+			EGBufferSlotUsage Slots[GBS_Num] = {};
 
 			DetermineUsedMaterialSlots(Slots, DerivedDefines, MaterialDefines, LightmapDefines, GlobalDefines, CompilerDefines, FeatureLevel);
 
@@ -1917,7 +1945,7 @@ void FShaderCompileUtilities::ApplyDerivedDefines(FShaderCompilerEnvironment& Ou
 			for (int32 Index = 0; Index < GBS_Num; Index++)
 			{
 				// if we are using this slot
-				if (Slots[Index])
+				if (Slots[Index] != EGBufferSlotUsage::Unused)
 				{
 					// if we are using this slot, it must have a valid spot in our gbuffer
 					const FGBufferItem& Item = BufferInfo.Slots[Index];
@@ -1929,88 +1957,84 @@ void FShaderCompileUtilities::ApplyDerivedDefines(FShaderCompilerEnvironment& Ou
 						if (Packing.bIsValid)
 						{
 							check(Packing.TargetIndex >= 0);
-							bTargetUsage[Packing.TargetIndex] = true;
+							TargetUsage[Packing.TargetIndex] = FMath::Max(TargetUsage[Packing.TargetIndex], Slots[Index]);
 						}
 					}
 				}
 			}
-
 		}
 		else
 		{
-			bTargetUsage[0] = true;
+			TargetUsage[0] = EGBufferSlotUsage::Written;
 			// we also need MRT for thin translucency due to dual blending if we are not on the fallback path
-			bTargetUsage[1] = (DerivedDefines.WRITES_VELOCITY_TO_GBUFFER || (MaterialDefines.DUAL_SOURCE_COLOR_BLENDING_ENABLED && DerivedDefines.MATERIAL_WORKS_WITH_DUAL_SOURCE_COLOR_BLENDING));
+			const bool bDualSourceBlending = MaterialDefines.DUAL_SOURCE_COLOR_BLENDING_ENABLED && DerivedDefines.MATERIAL_WORKS_WITH_DUAL_SOURCE_COLOR_BLENDING;
+			const bool bHasVelocity = GlobalDefines.GBUFFER_HAS_VELOCITY;
+			const bool bWritesVelocity = DerivedDefines.WRITES_VELOCITY_TO_GBUFFER;
+			TargetUsage[1] = GetGBufferSlotUsage(bWritesVelocity || bDualSourceBlending, bHasVelocity);
 		}
 	}
 	else if (MaterialDefines.IS_VIRTUAL_TEXTURE_MATERIAL)
 	{
-		// these whill change, of course
+		// these will change, of course
 		if (MaterialDefines.OUT_BASECOLOR)
 		{
-			bTargetUsage[0] = 1;
+			TargetUsage[0] = EGBufferSlotUsage::Written;
 		}
 		else if (MaterialDefines.OUT_BASECOLOR_NORMAL_ROUGHNESS)
 		{
-			bTargetUsage[0] = 1;
-			bTargetUsage[1] = 1;
+			TargetUsage[0] = EGBufferSlotUsage::Written;
+			TargetUsage[1] = EGBufferSlotUsage::Written;
 		}
 		else if (MaterialDefines.OUT_BASECOLOR_NORMAL_SPECULAR)
 		{
-			bTargetUsage[0] = 1;
-			bTargetUsage[1] = 1;
-			bTargetUsage[2] = 1;
+			TargetUsage[0] = EGBufferSlotUsage::Written;
+			TargetUsage[1] = EGBufferSlotUsage::Written;
+			TargetUsage[2] = EGBufferSlotUsage::Written;
 		}
 		else if (MaterialDefines.OUT_WORLDHEIGHT)
 		{
-			bTargetUsage[0] = 1;
+			TargetUsage[0] = EGBufferSlotUsage::Written;
 		}
 		else if (MaterialDefines.OUT_DISPLACEMENT)
 		{
-			bTargetUsage[0] = 1;
+			TargetUsage[0] = EGBufferSlotUsage::Written;
 		}
 	}
 	else if (MaterialDefines.IS_DECAL)
 	{
 		// these will have to change too
-		bTargetUsage[0] = MaterialDefines.DECAL_RENDERTARGET_COUNT > 0;
-		bTargetUsage[1] = MaterialDefines.DECAL_RENDERTARGET_COUNT > 1;
-		bTargetUsage[2] = MaterialDefines.DECAL_RENDERTARGET_COUNT > 2;
-		bTargetUsage[3] = MaterialDefines.DECAL_RENDERTARGET_COUNT > 3;
-		bTargetUsage[4] = MaterialDefines.DECAL_RENDERTARGET_COUNT > 4;
+		TargetUsage[0] = GetGBufferSlotUsage(MaterialDefines.DECAL_RENDERTARGET_COUNT > 0);
+		TargetUsage[1] = GetGBufferSlotUsage(MaterialDefines.DECAL_RENDERTARGET_COUNT > 1);
+		TargetUsage[2] = GetGBufferSlotUsage(MaterialDefines.DECAL_RENDERTARGET_COUNT > 2);
+		TargetUsage[3] = GetGBufferSlotUsage(MaterialDefines.DECAL_RENDERTARGET_COUNT > 3);
+		TargetUsage[4] = GetGBufferSlotUsage(MaterialDefines.DECAL_RENDERTARGET_COUNT > 4);
 	}
 	else
 	{
 		// something else, so no op
 	}
-#if 1
-	static bool bTestNewVersion = true;
-	if (bTestNewVersion)
+
+	// Decide which pixel shader outputs are enabled based on which targets are written and what the first substrate
+	// target is based on which target slots are in use
+	int32 SubstrateFirstMRT = 0;
+	for (int32 Iter = 0; Iter < FGBufferInfo::MaxTargets; Iter++)
 	{
-		//if (DerivedDefines.USES_GBUFFER)
+		if (TargetUsage[Iter] >= EGBufferSlotUsage::Written)
 		{
-			for (int32 Iter = 0; Iter < FGBufferInfo::MaxTargets; Iter++)
-			{
-				if (bTargetUsage[Iter])
-				{
-					FString TargetName = FString::Printf(TEXT("PIXELSHADEROUTPUT_MRT%d"), Iter);
-					OutEnvironment.SetDefine(TargetName.GetCharArray().GetData(), TEXT("1"));
-				}
-			}
+			FString TargetName = FString::Printf(TEXT("PIXELSHADEROUTPUT_MRT%d"), Iter);
+			OutEnvironment.SetDefine(TargetName.GetCharArray().GetData(), TEXT("1"));
+		}
+
+		if (TargetUsage[Iter] >= EGBufferSlotUsage::Used)
+		{
+			SubstrateFirstMRT = Iter + 1;
 		}
 	}
-	else
+
+	if (MaterialDefines.MATERIAL_IS_SUBSTRATE && SubstrateFirstMRT > 0)
 	{
-		// This uses the legacy logic from CalculateDerivedMaterialParameters(); Just keeping it around momentarily for testing during the transition.
-		SET_COMPILE_BOOL_IF_TRUE(PIXELSHADEROUTPUT_MRT0)
-		SET_COMPILE_BOOL_IF_TRUE(PIXELSHADEROUTPUT_MRT1)
-		SET_COMPILE_BOOL_IF_TRUE(PIXELSHADEROUTPUT_MRT2)
-		SET_COMPILE_BOOL_IF_TRUE(PIXELSHADEROUTPUT_MRT3)
-		SET_COMPILE_BOOL_IF_TRUE(PIXELSHADEROUTPUT_MRT4)
-		SET_COMPILE_BOOL_IF_TRUE(PIXELSHADEROUTPUT_MRT5)
-		SET_COMPILE_BOOL_IF_TRUE(PIXELSHADEROUTPUT_MRT6)
+		OutEnvironment.SetDefine(TEXT("SUBSTRATE_FIRST_MRT_INDEX"), SubstrateFirstMRT);
 	}
-#endif
 }
 
 void FShaderCompileUtilities::AppendGBufferDDCKeyString(const EShaderPlatform Platform, FString& KeyString)
