[Engine]
* Added PSOPrecacheFwd.h and removed includes in PrimitiveComponent.h/PrimitiveSceneInfoData.h

diff --git a/Engine/Plugins/Experimental/VirtualHeightfieldMesh/Source/VirtualHeightfieldMesh/Private/VirtualHeightfieldMeshEnable.cpp b/Engine/Plugins/Experimental/VirtualHeightfieldMesh/Source/VirtualHeightfieldMesh/Private/VirtualHeightfieldMeshEnable.cpp
--- a/Engine/Plugins/Experimental/VirtualHeightfieldMesh/Source/VirtualHeightfieldMesh/Private/VirtualHeightfieldMeshEnable.cpp
+++ b/Engine/Plugins/Experimental/VirtualHeightfieldMesh/Source/VirtualHeightfieldMesh/Private/VirtualHeightfieldMeshEnable.cpp
@@ -3,6 +3,7 @@
 #include "VirtualHeightfieldMeshEnable.h"
 
 #include "Components/RuntimeVirtualTextureComponent.h"
+#include "RHIGlobals.h"
 #include "UObject/UObjectIterator.h"
 #include "VirtualHeightfieldMeshComponent.h"
 #include "VT/RuntimeVirtualTextureVolume.h"
diff --git a/Engine/Source/Runtime/Engine/Classes/Components/PrimitiveComponent.h b/Engine/Source/Runtime/Engine/Classes/Components/PrimitiveComponent.h
--- a/Engine/Source/Runtime/Engine/Classes/Components/PrimitiveComponent.h
+++ b/Engine/Source/Runtime/Engine/Classes/Components/PrimitiveComponent.h
@@ -9,7 +9,6 @@
 #include "UObject/ObjectMacros.h"
 #include "UObject/UObjectGlobals.h"
 #include "Misc/Guid.h"
-#include "InputCoreTypes.h"
 #include "Interfaces/IPhysicsComponent.h"
 #include "Templates/SubclassOf.h"
 #include "Engine/EngineTypes.h"
@@ -30,7 +29,7 @@
 #include "HLOD/HLODBatchingPolicy.h"
 #include "HLOD/HLODLevelExclusion.h"
 #include "Stats/Stats2.h"
-#include "PSOPrecache.h"
+#include "PSOPrecacheFwd.h"
 #include "MeshDrawCommandStatsDefines.h"
 #include "PrimitiveSceneInfoData.h"
 #include "PrimitiveComponent.generated.h"
@@ -3224,3 +3223,8 @@ bool UPrimitiveComponent::GetOverlapsWithActor_Template(const AActor* Actor, TAr
 
 	return InitialCount != OutOverlaps.Num();
 }
+
+#if UE_ENABLE_INCLUDE_ORDER_DEPRECATED_IN_5_5
+#include "InputCoreTypes.h"
+#include "PSOPrecache.h"
+#endif
\ No newline at end of file
diff --git a/Engine/Source/Runtime/Engine/Classes/Components/StaticMeshComponent.h b/Engine/Source/Runtime/Engine/Classes/Components/StaticMeshComponent.h
--- a/Engine/Source/Runtime/Engine/Classes/Components/StaticMeshComponent.h
+++ b/Engine/Source/Runtime/Engine/Classes/Components/StaticMeshComponent.h
@@ -3,6 +3,7 @@
 #pragma once
 
 #include "CoreMinimal.h"
+#include "RHIResources.h"
 #include "UObject/ObjectMacros.h"
 #include "UObject/Object.h"
 #include "Misc/Guid.h"
@@ -27,6 +28,7 @@ class FLightingBuildOptions;
 class FMeshMapBuildData;
 class FPrimitiveSceneProxy;
 class FStaticMeshStaticLightingMesh;
+class FVertexFactoryType;
 class ULightComponent;
 class UStaticMesh;
 class UStaticMeshComponent;
diff --git a/Engine/Source/Runtime/Engine/Private/LevelActor.cpp b/Engine/Source/Runtime/Engine/Private/LevelActor.cpp
--- a/Engine/Source/Runtime/Engine/Private/LevelActor.cpp
+++ b/Engine/Source/Runtime/Engine/Private/LevelActor.cpp
@@ -6,6 +6,7 @@
 #include "EngineStats.h"
 #include "Engine/Level.h"
 #include "AI/NavigationSystemBase.h"
+#include "RenderUtils.h"
 #include "UObject/LinkerLoad.h"
 #include "GameFramework/OnlineReplStructs.h"
 #include "Engine/Engine.h"
diff --git a/Engine/Source/Runtime/Engine/Public/PSOPrecache.h b/Engine/Source/Runtime/Engine/Public/PSOPrecache.h
--- a/Engine/Source/Runtime/Engine/Public/PSOPrecache.h
+++ b/Engine/Source/Runtime/Engine/Public/PSOPrecache.h
@@ -11,16 +11,10 @@
 #include "RHIResources.h"
 #include "Engine/EngineTypes.h"
 #include "PipelineStateCache.h"
+#include "PSOPrecacheFwd.h"
 
 class FVertexFactoryType;
 
-// General switch that decides whether to compile out some PSO precaching code (most importantly, reduce sizeofs of common classes)
-#ifndef UE_WITH_PSO_PRECACHING
-	#define UE_WITH_PSO_PRECACHING		(PLATFORM_SUPPORTS_PSO_PRECACHING)
-#endif // UE_WITH_PSO_PRECACHING
-
-#define PSO_PRECACHING_VALIDATE !WITH_EDITOR && UE_WITH_PSO_PRECACHING
-
 /**
  * Parameters which are needed to collect all possible PSOs used by the PSO collectors
  */
@@ -192,8 +186,6 @@ struct FMaterialInterfacePSOPrecacheParams
 	FPSOPrecacheVertexFactoryDataList VertexFactoryDataList;
 };
 
-typedef TArray<FMaterialInterfacePSOPrecacheParams, TInlineAllocator<4> > FMaterialInterfacePSOPrecacheParamsList;
-
 extern ENGINE_API void AddMaterialInterfacePSOPrecacheParamsToList(const FMaterialInterfacePSOPrecacheParams& EntryToAdd, FMaterialInterfacePSOPrecacheParamsList& List);
 
 /**
@@ -296,6 +288,3 @@ extern ENGINE_API EPSOPrecacheProxyCreationStrategy GetPSOPrecacheProxyCreationS
  * Delay component proxy creation when it's requested PSOs are still precaching
  */
 extern ENGINE_API bool ProxyCreationWhenPSOReady();
-
-// Unique request ID of MaterialPSOPrecache which can be used to boost the priority of a PSO precache requests if it's needed for rendering
-using FMaterialPSOPrecacheRequestID = uint32;
diff --git a/Engine/Source/Runtime/Engine/Public/PSOPrecacheFwd.h b/Engine/Source/Runtime/Engine/Public/PSOPrecacheFwd.h
new file mode 100644
--- /dev/null
+++ b/Engine/Source/Runtime/Engine/Public/PSOPrecacheFwd.h
@@ -0,0 +1,25 @@
+// Copyright Epic Games, Inc. All Rights Reserved.
+
+/*=============================================================================
+	PSOPrecacheFwd.h
+=============================================================================*/
+
+#pragma once
+
+#include "Containers/Array.h"
+#include "HAL/Platform.h"
+
+// General switch that decides whether to compile out some PSO precaching code (most importantly, reduce sizeofs of common classes)
+#ifndef UE_WITH_PSO_PRECACHING
+	#define UE_WITH_PSO_PRECACHING		(PLATFORM_SUPPORTS_PSO_PRECACHING)
+#endif // UE_WITH_PSO_PRECACHING
+
+#define PSO_PRECACHING_VALIDATE !WITH_EDITOR && UE_WITH_PSO_PRECACHING
+
+struct FMaterialInterfacePSOPrecacheParams;
+struct FPSOPrecacheParams;
+
+typedef TArray<FMaterialInterfacePSOPrecacheParams, TInlineAllocator<4> > FMaterialInterfacePSOPrecacheParamsList;
+
+// Unique request ID of MaterialPSOPrecache which can be used to boost the priority of a PSO precache requests if it's needed for rendering
+using FMaterialPSOPrecacheRequestID = uint32;
\ No newline at end of file
diff --git a/Engine/Source/Runtime/Engine/Public/PrimitiveSceneInfoData.h b/Engine/Source/Runtime/Engine/Public/PrimitiveSceneInfoData.h
--- a/Engine/Source/Runtime/Engine/Public/PrimitiveSceneInfoData.h
+++ b/Engine/Source/Runtime/Engine/Public/PrimitiveSceneInfoData.h
@@ -2,7 +2,6 @@
 
 #pragma once 
 
-#include "RendererInterface.h"
 #include "HAL/ThreadSafeCounter.h"
 #include "PrimitiveComponentId.h"
 
@@ -107,4 +106,8 @@ public:
 		PrimitiveSceneId.PrimIDValue = NextPrimitiveId.Increment();
 	}
 
-};
\ No newline at end of file
+};
+
+#if UE_ENABLE_INCLUDE_ORDER_DEPRECATED_IN_5_5
+#include "RendererInterface.h"
+#endif
\ No newline at end of file
