Fixed crash caused by attempted immediate async pipe transition of certain SVT resources when using async compute. This seems to only happen when manually updating the streamer by calling UE::SVT::FStreamingManager::Update_GameThread(), which is only done by Niagara.
Fixed by removing the ForceImmediateFirstBarrier flag on the involved resources. The flag forces RDG to only transition the resource in the first pass it's used in and to not use a split barrier with the first half at the start of the graph. Initially this flag was set to avoid issues where the resources might be used outside of RDG but still within the graph prior to the RDG pass accessing them as RDG resources.
However, it seems like all non-RDG accesses of the SVT resources within the same graph always happen after the streaming system update, so removing this flag should be fine.

#rnx

diff --git a/Engine/Source/Runtime/Engine/Private/SparseVolumeTexture/SparseVolumeTextureStreamingManager.cpp b/Engine/Source/Runtime/Engine/Private/SparseVolumeTexture/SparseVolumeTextureStreamingManager.cpp
--- a/Engine/Source/Runtime/Engine/Private/SparseVolumeTexture/SparseVolumeTextureStreamingManager.cpp
+++ b/Engine/Source/Runtime/Engine/Private/SparseVolumeTexture/SparseVolumeTextureStreamingManager.cpp
@@ -730,6 +730,8 @@ void FStreamingManager::AddInternal(FRDGBuilder& GraphBuilder, FNewSparseVolumeT
 			// No need to write to TileDataPtrA and TileDataPtrB because we zeroed out all the occupancy bits.
 		}
 
+		const bool bAsyncCompute = UseAsyncComputeForStreaming();
+
 		// Process frames
 		for (int32 FrameIdx = 0; FrameIdx < NumFrames; ++FrameIdx)
 		{
@@ -751,7 +753,7 @@ void FStreamingManager::AddInternal(FRDGBuilder& GraphBuilder, FNewSparseVolumeT
 				for (int32 MipLevelIndex = 0; MipLevelIndex < NumMipLevels; ++MipLevelIndex)
 				{
 					FRDGTextureUAV* UAV = GraphBuilder.CreateUAV(FRDGTextureUAVDesc(PageTableRDG, MipLevelIndex, PageEntryFormat));
-					AddClearUAVPass(GraphBuilder, UAV, FUintVector4(ForceInitToZero));
+					AddClearUAVPass(GraphBuilder, UAV, FUintVector4(ForceInitToZero), bAsyncCompute ? ERDGPassFlags::AsyncCompute : ERDGPassFlags::Compute);
 				}
 
 				FrameInfo.PageTableTexture = GraphBuilder.ConvertToExternalTexture(PageTableRDG);
diff --git a/Engine/Source/Runtime/Engine/Private/SparseVolumeTexture/SparseVolumeTextureUpload.cpp b/Engine/Source/Runtime/Engine/Private/SparseVolumeTexture/SparseVolumeTextureUpload.cpp
--- a/Engine/Source/Runtime/Engine/Private/SparseVolumeTexture/SparseVolumeTextureUpload.cpp
+++ b/Engine/Source/Runtime/Engine/Private/SparseVolumeTexture/SparseVolumeTextureUpload.cpp
@@ -87,6 +87,13 @@ namespace UE
 namespace SVT
 {
 
+bool UseAsyncComputeForStreaming()
+{
+	// Disable async compute for streaming systems when MGPU is active, to work around GPU hangs
+	const bool bAsyncCompute = GSupportsEfficientAsyncCompute && (GSVTStreamingAsyncCompute != 0) && (GNumExplicitGPUsForRendering == 1);
+	return bAsyncCompute;
+}
+
 FTileUploader::FTileUploader()
 {
 	ResetState();
@@ -246,12 +253,12 @@ void FTileUploader::ResourceUploadTo(FRDGBuilder& GraphBuilder, const TRefCountP
 			FRDGTexture* DstTextureBRDG = nullptr;
 			if (InDstTextureA)
 			{
-				DstTextureARDG = GraphBuilder.RegisterExternalTexture(InDstTextureA, ERDGTextureFlags::ForceImmediateFirstBarrier);
+				DstTextureARDG = GraphBuilder.RegisterExternalTexture(InDstTextureA);
 				GraphBuilder.UseInternalAccessMode(DstTextureARDG);
 			}
 			if (InDstTextureB)
 			{
-				DstTextureBRDG = GraphBuilder.RegisterExternalTexture(InDstTextureB, ERDGTextureFlags::ForceImmediateFirstBarrier);
+				DstTextureBRDG = GraphBuilder.RegisterExternalTexture(InDstTextureB);
 				GraphBuilder.UseInternalAccessMode(DstTextureBRDG);
 			}
 
@@ -351,8 +358,7 @@ void FTileUploader::ResourceUploadTo(FRDGBuilder& GraphBuilder, const TRefCountP
 				PassParameters->NumDispatchedGroups = FMath::Min(NumTilesInThisBatch, GRHIMaxDispatchThreadGroupsPerDimension.X);
 				PassParameters->PaddedTileSize = SPARSE_VOLUME_TILE_RES_PADDED;
 	
-				// Disable async compute for streaming systems when MGPU is active, to work around GPU hangs
-				const bool bAsyncCompute = GSupportsEfficientAsyncCompute && (GSVTStreamingAsyncCompute != 0) && (GNumExplicitGPUsForRendering == 1);
+				const bool bAsyncCompute = UseAsyncComputeForStreaming();
 	
 				FComputeShaderUtils::AddPass(
 					GraphBuilder,
@@ -466,8 +472,7 @@ void FPageTableUpdater::Apply(FRDGBuilder& GraphBuilder)
 	
 		if (NumWrittenUpdates > 0)
 		{
-			// Disable async compute for streaming systems when MGPU is active, to work around GPU hangs
-			const bool bAsyncCompute = GSupportsEfficientAsyncCompute && (GSVTStreamingAsyncCompute != 0) && (GNumExplicitGPUsForRendering == 1);
+			const bool bAsyncCompute = UseAsyncComputeForStreaming();
 			auto ComputeShader = GetGlobalShaderMap(GMaxRHIFeatureLevel)->GetShader<FSparseVolumeTextureUpdatePageTableCS>();
 
 			// Register page table textures with RDG and set it to internal access mode
@@ -476,7 +481,7 @@ void FPageTableUpdater::Apply(FRDGBuilder& GraphBuilder)
 				FRDGTexture* PageTableRDG = GraphBuilder.FindExternalTexture(Batch.PageTable);
 				if (!PageTableRDG)
 				{
-					PageTableRDG = GraphBuilder.RegisterExternalTexture(Batch.PageTable, ERDGTextureFlags::ForceImmediateFirstBarrier);
+					PageTableRDG = GraphBuilder.RegisterExternalTexture(Batch.PageTable);
 				}
 				GraphBuilder.UseInternalAccessMode(PageTableRDG); // Make sure to use graph tracking in case we previously called UseExternalAccessMode() within the current graph
 			}
diff --git a/Engine/Source/Runtime/Engine/Private/SparseVolumeTexture/SparseVolumeTextureUpload.h b/Engine/Source/Runtime/Engine/Private/SparseVolumeTexture/SparseVolumeTextureUpload.h
--- a/Engine/Source/Runtime/Engine/Private/SparseVolumeTexture/SparseVolumeTextureUpload.h
+++ b/Engine/Source/Runtime/Engine/Private/SparseVolumeTexture/SparseVolumeTextureUpload.h
@@ -15,6 +15,9 @@ namespace UE
 namespace SVT
 {
 
+// Whether to use async compute for all SVT streaming GPU work
+bool UseAsyncComputeForStreaming();
+
 // Utility class for uploading tiles to a physical tile data texture
 class FTileUploader
 {
