Removed now unused public function ULevel::CreateActorPackage receiving a basedir (introduced after 5.3)

diff --git a/Engine/Source/Runtime/Engine/Classes/Engine/Level.h b/Engine/Source/Runtime/Engine/Classes/Engine/Level.h
--- a/Engine/Source/Runtime/Engine/Classes/Engine/Level.h
+++ b/Engine/Source/Runtime/Engine/Classes/Engine/Level.h
@@ -1268,15 +1268,6 @@ public:
 	 */
 	static ENGINE_API UPackage* CreateActorPackage(UPackage* InLevelPackage, EActorPackagingScheme InActorPackagingScheme, const FString& InActorPath, const UObject* InMountPointContext = nullptr);
 
-	/**
-	 * Create an package for this actor
-	 * @param InBaseDir the base directory used when building the actor package name
-	 * @param InActorPackagingScheme the packaging scheme to use
-	 * @param InActorPath the fully qualified actor path, in the format: 'Outermost.Outer.Name'
-	 * @return the created package
-	 */
-	static ENGINE_API UPackage* CreateActorPackage(const FString& InBaseDir, EActorPackagingScheme InActorPackagingScheme, const FString& InActorPath);
-
 	/**
 	 * Detach or reattach all level actors to from/to their external package
 	 * @param bReattach if false will detach actors from their external package until reattach is called, passing true will reattach actors, no-op for non external actors
diff --git a/Engine/Source/Runtime/Engine/Private/Level.cpp b/Engine/Source/Runtime/Engine/Private/Level.cpp
--- a/Engine/Source/Runtime/Engine/Private/Level.cpp
+++ b/Engine/Source/Runtime/Engine/Private/Level.cpp
@@ -4234,12 +4234,6 @@ UPackage* ULevel::CreateActorPackage(UPackage* InLevelPackage, EActorPackagingSc
 	return ActorPackage;
 }
 
-UPackage* ULevel::CreateActorPackage(const FString& InBaseDir, EActorPackagingScheme InActorPackagingScheme, const FString& InActorPath)
-{
-	const FString PackageName = ULevel::GetActorPackageName(InBaseDir, InActorPackagingScheme, InActorPath);
-	return CreateActorPackageInternal(PackageName, InActorPath);
-}
-
 void ULevel::DetachAttachAllActorsPackages(bool bReattach)
 {
 	if (bReattach)
