Fixed crash when running SVTs with the nullrhi. The crash was caused by trying to register a SVT with the streamer. Due to FApp::CanEverRender(), the SVT didn't create any render resources. The streamer then fails in a check() for those render resources.
Fixed by introducing additional FApp::CanEverRender() and !IsTemplate() checks on the USVT side to make sure SVTs are only added to the streamer if the application can actually render. These checks also replace the old AllowGlobalShaderLoad() way of checking if the streamer should run. By preventing SVTs from being added to the streamer, all other public functions will early exit because there are no SVTs registered.

#rnx

diff --git a/Engine/Source/Runtime/Engine/Private/SparseVolumeTexture/SparseVolumeTexture.cpp b/Engine/Source/Runtime/Engine/Private/SparseVolumeTexture/SparseVolumeTexture.cpp
--- a/Engine/Source/Runtime/Engine/Private/SparseVolumeTexture/SparseVolumeTexture.cpp
+++ b/Engine/Source/Runtime/Engine/Private/SparseVolumeTexture/SparseVolumeTexture.cpp
@@ -971,7 +971,7 @@ bool USparseVolumeTextureFrame::Initialize(USparseVolumeTexture* InOwner, int32
 
 bool USparseVolumeTextureFrame::CreateTextureRenderResources()
 {
-	if (!TextureRenderResources && FApp::CanEverRender())
+	if (!TextureRenderResources && !IsTemplate() && FApp::CanEverRender())
 	{
 		TextureRenderResources = new UE::SVT::FTextureRenderResources();
 		TextureRenderResources->SetGlobalVolumeResolution_GameThread(Owner->GetVolumeResolution());
@@ -998,30 +998,33 @@ void USparseVolumeTextureFrame::BeginDestroy()
 {
 	// Ensure that the streamable SVT has been removed from the streaming manager
 	
-	if (IsValid(Owner))
-	{		
-		UStreamableSparseVolumeTexture* SVTOwner = CastChecked<UStreamableSparseVolumeTexture>(Owner);
-		for (int i = 0; i < SVTOwner->GetNumFrames(); ++i)
+	if (!IsTemplate())
+	{
+		if (IsValid(Owner) && FApp::CanEverRender())
 		{
-			// if the owner contains the current frame being deleted, remove the owner from the streaming manager
-			// SVT_TODO: This is a temporary fix for a GC problem.  In the future this will be replaced with a more robust solution.
-			if (SVTOwner->GetFrame(i) == this)
+			UStreamableSparseVolumeTexture* SVTOwner = CastChecked<UStreamableSparseVolumeTexture>(Owner);
+			for (int i = 0; i < SVTOwner->GetNumFrames(); ++i)
 			{
-				UE::SVT::GetStreamingManager().Remove_GameThread(SVTOwner);
-				break;
+				// if the owner contains the current frame being deleted, remove the owner from the streaming manager
+				// SVT_TODO: This is a temporary fix for a GC problem.  In the future this will be replaced with a more robust solution.
+				if (SVTOwner->GetFrame(i) == this)
+				{
+					UE::SVT::GetStreamingManager().Remove_GameThread(SVTOwner);
+					break;
+				}
 			}
-		}		
-	}
-	
-	if (TextureRenderResources)
-	{
-		ENQUEUE_RENDER_COMMAND(USparseVolumeTextureFrame_DeleteTextureRenderResources)(
-			[Resources = TextureRenderResources](FRHICommandListImmediate& RHICmdList)
-			{
-				Resources->ReleaseResource();
-				delete Resources;
-			});
-		TextureRenderResources = nullptr;
+		}
+
+		if (TextureRenderResources)
+		{
+			ENQUEUE_RENDER_COMMAND(USparseVolumeTextureFrame_DeleteTextureRenderResources)(
+				[Resources = TextureRenderResources](FRHICommandListImmediate& RHICmdList)
+				{
+					Resources->ReleaseResource();
+					delete Resources;
+				});
+			TextureRenderResources = nullptr;
+		}
 	}
 
 	Super::BeginDestroy();
@@ -1342,16 +1345,22 @@ void UStreamableSparseVolumeTexture::PostLoad()
 	// Ensure that NumFrames always corresponds to the actual number of frames
 	NumFrames = GetNumFrames();
 
+	if (!IsTemplate())
+	{
 #if WITH_EDITORONLY_DATA
-	RecacheFrames();
+		RecacheFrames();
 #else
-	for (USparseVolumeTextureFrame* Frame : Frames)
-	{
-		check(Frame); // Elements in Frames should only ever be null when the SVT is being deleted
-		Frame->CreateTextureRenderResources();
-	}
-	UE::SVT::GetStreamingManager().Add_GameThread(this); // RecacheFrames() handles this in editor builds
+		if (FApp::CanEverRender())
+		{
+			for (USparseVolumeTextureFrame* Frame : Frames)
+			{
+				check(Frame); // Elements in Frames should only ever be null when the SVT is being deleted
+				Frame->CreateTextureRenderResources();
+			}
+			UE::SVT::GetStreamingManager().Add_GameThread(this); // RecacheFrames() handles this in editor builds
+		}
 #endif
+	}
 }
 
 void UStreamableSparseVolumeTexture::FinishDestroy()
@@ -1362,14 +1371,18 @@ void UStreamableSparseVolumeTexture::FinishDestroy()
 void UStreamableSparseVolumeTexture::BeginDestroy()
 {
 	Super::BeginDestroy();
-	UE::SVT::GetStreamingManager().Remove_GameThread(this);
-
-#if WITH_EDITOR
+	
 	if (!IsTemplate())
 	{
+		if (FApp::CanEverRender())
+		{
+			UE::SVT::GetStreamingManager().Remove_GameThread(this);
+		}
+
+#if WITH_EDITOR
 		UnregisterEditorDelegates();
-	}
 #endif
+	}
 }
 
 void UStreamableSparseVolumeTexture::Serialize(FArchive& Ar)
@@ -1422,8 +1435,11 @@ void UStreamableSparseVolumeTexture::PostEditChangeProperty(FPropertyChangedEven
 		|| PropertyChangedEvent.GetPropertyName() == GET_MEMBER_NAME_CHECKED(UStreamableSparseVolumeTexture, PrefetchPercentageBias))
 	{
 		// Re-register the SVT with the streamer so it picks up on the changed streaming parameters.
-		UE::SVT::GetStreamingManager().Remove_GameThread(this);
-		UE::SVT::GetStreamingManager().Add_GameThread(this);
+		if (FApp::CanEverRender() && !IsTemplate())
+		{
+			UE::SVT::GetStreamingManager().Remove_GameThread(this);
+			UE::SVT::GetStreamingManager().Add_GameThread(this);
+		}
 		
 		bRecacheFrames = false;
 	}
@@ -1534,6 +1550,11 @@ void UStreamableSparseVolumeTexture::OnAssetsAddExtraObjectsToDelete(TArray<UObj
 #if WITH_EDITORONLY_DATA
 void UStreamableSparseVolumeTexture::RecacheFrames()
 {
+	if (IsTemplate())
+	{
+		return;
+	}
+
 	if (InitState != EInitState_Done)
 	{
 		UE_LOG(LogSparseVolumeTexture, Warning, TEXT("Tried to cache derived data of an uninitialized SVT: %s"), *GetName());
@@ -1543,7 +1564,12 @@ void UStreamableSparseVolumeTexture::RecacheFrames()
 	FScopedSlowTask RecacheTask(static_cast<float>(Frames.Num() + 2), LOCTEXT("SparseVolumeTextureCacheFrames", "Caching SparseVolumeTexture frames in Derived Data Cache"));
 	RecacheTask.MakeDialog(true);
 
-	UE::SVT::GetStreamingManager().Remove_GameThread(this);
+	const bool bCanEverRender = FApp::CanEverRender();
+	if (bCanEverRender)
+	{
+		UE::SVT::GetStreamingManager().Remove_GameThread(this);
+	}
+	
 	RecacheTask.EnterProgressFrame(1.0f);
 
 	bool bCanceled = false;
@@ -1571,7 +1597,10 @@ void UStreamableSparseVolumeTexture::RecacheFrames()
 		RecacheTask.EnterProgressFrame(1.0f);
 	}
 	
-	UE::SVT::GetStreamingManager().Add_GameThread(this);
+	if (bCanEverRender)
+	{
+		UE::SVT::GetStreamingManager().Add_GameThread(this);
+	}
 }
 #endif
 
diff --git a/Engine/Source/Runtime/Engine/Private/SparseVolumeTexture/SparseVolumeTextureStreamingManager.cpp b/Engine/Source/Runtime/Engine/Private/SparseVolumeTexture/SparseVolumeTextureStreamingManager.cpp
--- a/Engine/Source/Runtime/Engine/Private/SparseVolumeTexture/SparseVolumeTextureStreamingManager.cpp
+++ b/Engine/Source/Runtime/Engine/Private/SparseVolumeTexture/SparseVolumeTextureStreamingManager.cpp
@@ -119,9 +119,8 @@ IStreamingManager& GetStreamingManager()
 
 static bool DoesPlatformSupportSparseVolumeTexture(EShaderPlatform Platform)
 {
-	// SVT_TODO: This is a bit of a hack: FStreamingManager::Add_GameThread() issues a rendering thread lambda for creating the RHI resources and uploading root tile data.
-	// Uploading root tile data involves access to the global shader map, which is empty under certain circumstances. By checking AllowGlobalShaderLoad(), we disallow streaming completely.
-	return AllowGlobalShaderLoad();
+	// There are currently no hard platform restrictions for SVT support
+	return true;
 }
 
 struct FStreamingUpdateParameters
