Only precache component PSOs from GameThread because material PSO request take a TStrongObjectPtr which can only be constructed on the GameThread.

diff --git a/Engine/Source/Runtime/Engine/Private/Components/PrimitiveComponent.cpp b/Engine/Source/Runtime/Engine/Private/Components/PrimitiveComponent.cpp
--- a/Engine/Source/Runtime/Engine/Private/Components/PrimitiveComponent.cpp
+++ b/Engine/Source/Runtime/Engine/Private/Components/PrimitiveComponent.cpp
@@ -4683,7 +4683,11 @@ void UPrimitiveComponent::SetupPrecachePSOParams(FPSOPrecacheParams& Params)
 void UPrimitiveComponent::PrecachePSOs()
 {
 #if UE_WITH_PSO_PRECACHING
-	if (!FApp::CanEverRender() || !IsComponentPSOPrecachingEnabled())
+	// Only request PSO precaching if app is rendering and per component PSO precaching is enabled
+	// Also only request PSOs from game thread because TStrongObjectPtr is used on the material to make
+	// it's not deleted via garbage collection when PSO precaching is still busy. TStrongObjectPtr can only
+	// be constructed on the GameThread
+	if (!FApp::CanEverRender() || !IsComponentPSOPrecachingEnabled() || !IsInGameThread())
 	{
 		return;
 	}
diff --git a/Engine/Source/Runtime/Engine/Private/Components/StaticMeshComponent.cpp b/Engine/Source/Runtime/Engine/Private/Components/StaticMeshComponent.cpp
--- a/Engine/Source/Runtime/Engine/Private/Components/StaticMeshComponent.cpp
+++ b/Engine/Source/Runtime/Engine/Private/Components/StaticMeshComponent.cpp
@@ -2275,11 +2275,7 @@ bool UStaticMeshComponent::SetStaticMesh(UStaticMesh* NewMesh)
 	}
 
 #if UE_WITH_PSO_PRECACHING
-	// Precache the PSOs if already precache otherwise will happen later on during PostLoad
-	if (bPSOPrecacheCalled)
-	{ 
-		PrecachePSOs();
-	}
+	PrecachePSOs();
 #endif // UE_WITH_PSO_PRECACHING
 
 	// Need to send this to render thread at some point
