Fixed memory stomping in the SVT streaming manager caused by improperly cancelled IO requests. This issue only affected cooked/non-editor builds. The problem was that calling Reset() on the request doesn't necessarily cancel the request if there are additional references being held to it. Fixed by explicitly cancelling the request before resetting the handle.

#rnx

diff --git a/Engine/Source/Runtime/Engine/Private/SparseVolumeTexture/SparseVolumeTextureStreamingManager.h b/Engine/Source/Runtime/Engine/Private/SparseVolumeTexture/SparseVolumeTextureStreamingManager.h
--- a/Engine/Source/Runtime/Engine/Private/SparseVolumeTexture/SparseVolumeTextureStreamingManager.h
+++ b/Engine/Source/Runtime/Engine/Private/SparseVolumeTexture/SparseVolumeTextureStreamingManager.h
@@ -176,7 +176,7 @@ private:
 		uint32 RequestVersion = 0;
 #endif
 		FIoBuffer RequestBuffer;
-		FBulkDataBatchReadRequest Request;
+		FBulkDataBatchRequest Request;
 		uint32 IssuedInFrame = 0;
 		bool bBlocking = false;
 
@@ -194,6 +194,11 @@ private:
 			RetryCount = 0;
 			++RequestVersion;
 #endif
+			if (Request.IsPending())
+			{
+				Request.Cancel();
+				Request.Wait(); // Even after calling Cancel(), we still need to Wait() before we can touch the RequestBuffer.
+			}
 			Request.Reset();
 			RequestBuffer = {};
 			IssuedInFrame = 0;
