FActorInstanceHandle - GetTypeHash() : Return value of FCrc::StrCrc32() was ignored

diff --git a/Engine/Source/Runtime/Engine/Private/Engine/ActorInstanceHandle.cpp b/Engine/Source/Runtime/Engine/Private/Engine/ActorInstanceHandle.cpp
--- a/Engine/Source/Runtime/Engine/Private/Engine/ActorInstanceHandle.cpp
+++ b/Engine/Source/Runtime/Engine/Private/Engine/ActorInstanceHandle.cpp
@@ -451,7 +451,7 @@ uint32 GetTypeHash(const FActorInstanceHandle& Handle)
 	uint32 Hash = 0;
 	if (const AActor* Actor = Handle.GetCachedActor())
 	{
-		FCrc::StrCrc32(*(Actor->GetPathName()), Hash);
+		Hash = FCrc::StrCrc32(*(Actor->GetPathName()), Hash);
 	}
 	if (UObject* ManagerInterfaceObject = Handle.ManagerInterface.GetObject())
 	{
