Take a local copy of our AsyncRequestIDs array for flushing async loads. Levelstreaming relies on callbacks that may empty the AsyncRequestIDs array as loads complete, which will invalidate the array view passed into FlushAsyncLoading.

diff --git a/Engine/Source/Runtime/Engine/Private/LevelStreaming.cpp b/Engine/Source/Runtime/Engine/Private/LevelStreaming.cpp
--- a/Engine/Source/Runtime/Engine/Private/LevelStreaming.cpp
+++ b/Engine/Source/Runtime/Engine/Private/LevelStreaming.cpp
@@ -1709,8 +1709,10 @@ bool ULevelStreaming::RequestLevel(UWorld* PersistentWorld, bool bAllowLevelLoad
 					UE_LOG(LogStreaming, Display, TEXT("ULevelStreaming::RequestLevel(%s) is flushing async loading"), *DesiredPackageName.ToString());
 				}
 
-				// Finish all async loading.
-				FlushAsyncLoading(AsyncRequestIDs);
+				// Finish all async loading. Since we will clear our requests upon completion of all loads, 
+				// we take a copy so FlushAsyncLoading won't touch an invalidated array view
+				TArray<int32> LocalRequestIDs = AsyncRequestIDs;
+				FlushAsyncLoading(LocalRequestIDs);
 			}
 		}
 		else
diff --git a/Engine/Source/Runtime/Engine/Private/WorldPartition/WorldPartitionLevelStreamingDynamic.cpp b/Engine/Source/Runtime/Engine/Private/WorldPartition/WorldPartitionLevelStreamingDynamic.cpp
--- a/Engine/Source/Runtime/Engine/Private/WorldPartition/WorldPartitionLevelStreamingDynamic.cpp
+++ b/Engine/Source/Runtime/Engine/Private/WorldPartition/WorldPartitionLevelStreamingDynamic.cpp
@@ -332,8 +332,10 @@ bool UWorldPartitionLevelStreamingDynamic::RequestLevel(UWorld* InPersistentWorl
 						UE_LOG(LogStreaming, Display, TEXT("UWorldPartitionLevelStreamingDynamic::RequestLevel(%s) is flushing async loading"), *WorldAssetPackageName);
 					}
 
-					// Finish all async loading.
-					FlushAsyncLoading(AsyncRequestIDs);
+					// Finish all async loading. Since we will clear our requests upon completion of all loads, 
+					// we take a copy so FlushAsyncLoading won't touch an invalidated array view
+					TArray<int32> LocalRequestIDs = AsyncRequestIDs;
+					FlushAsyncLoading(LocalRequestIDs);
 				}
 				else
 				{
