Fixing up Visual Logger automation tests

diff --git a/Engine/Source/Runtime/Engine/Private/VisualLogger/VisualLoggerAutomationTests.cpp b/Engine/Source/Runtime/Engine/Private/VisualLogger/VisualLoggerAutomationTests.cpp
--- a/Engine/Source/Runtime/Engine/Private/VisualLogger/VisualLoggerAutomationTests.cpp
+++ b/Engine/Source/Runtime/Engine/Private/VisualLogger/VisualLoggerAutomationTests.cpp
@@ -131,7 +131,7 @@ bool FVisualLogTest::RunTest(const FString& Parameters)
 
 	{
 		const FString TextToLog = TEXT("Simple text line to test if UE_VLOG_UELOG works fine");
-		double CurrentTimestamp = World->TimeSeconds;
+		double CurrentTimestamp = FVisualLogger::Get().GetTimeStampForObject(World);
 		UE_VLOG_UELOG(World, LogVisual, Log, TEXT("%s"), *TextToLog);
 		CHECK_SUCCESS(Context.Device.LastObject != World);
 		CHECK_SUCCESS(Context.Device.LastEntry.TimeStamp == -1);
@@ -191,7 +191,7 @@ bool FVisualLogSegmentsTest::RunTest(const FString& Parameters)
 		CHECK_SUCCESS(Context.Device.LastEntry.TimeStamp == -1);
 		FVisualLogEntry* CurrentEntry = FVisualLogger::Get().GetEntryToWrite(World, World->TimeSeconds, ECreateIfNeeded::DontCreate);
 
-		double CurrentTimestamp = World->TimeSeconds;
+		double CurrentTimestamp = FVisualLogger::Get().GetTimeStampForObject(World);
 		{
 			CHECK_NOT_NULL(CurrentEntry);
 			CHECK_SUCCESS(CurrentEntry->TimeStamp == CurrentTimestamp);
@@ -245,7 +245,7 @@ bool FVisualLogEventsTest::RunTest(const FString& Parameters)
 	CHECK_SUCCESS(EventTest3.Name == TEXT("EventTest3"));
 	CHECK_SUCCESS(EventTest3.FriendlyDesc == TEXT("Third simple event for vlog tests"));
 
-	double CurrentTimestamp = World->TimeSeconds;
+	double CurrentTimestamp = FVisualLogger::Get().GetTimeStampForObject(World);
 	FVisualLogEntry* CurrentEntry = FVisualLogger::Get().GetEntryToWrite(World, CurrentTimestamp, ECreateIfNeeded::DontCreate);
 	CHECK_SUCCESS(CurrentEntry == nullptr);
 
@@ -265,7 +265,7 @@ bool FVisualLogEventsTest::RunTest(const FString& Parameters)
 	CHECK_SUCCESS(CurrentEntry->Events[1].Counter == 1);
 	CHECK_SUCCESS(CurrentEntry->Events[1].Name == TEXT("EventTest2"));
 
-	CurrentTimestamp = World->TimeSeconds;
+	CurrentTimestamp = FVisualLogger::Get().GetTimeStampForObject(World);
 	UE_VLOG_EVENTS(World, NAME_None, EventTest, EventTest2, EventTest3);
 
 	{
@@ -283,6 +283,9 @@ bool FVisualLogEventsTest::RunTest(const FString& Parameters)
 		CHECK_SUCCESS(CurrentEntry->Events[1].UserFriendlyDesc == TEXT("Second simple event for vlog tests"));
 		CHECK_SUCCESS(CurrentEntry->Events[2].UserFriendlyDesc == TEXT("Third simple event for vlog tests"));
 
+		// Create a NewEntry that has no data in it.  In this case, our entry won't coalesce with the previous data and we should end up with NewTimestamp.
+		// This functionality (requesting an explicit TimeStamp) is deprecated and will be removed, since we should always be using GetTimeStampForObject()
+		// as that functionality is what's used when actually logging.
 		const double NewTimestamp = CurrentTimestamp + 0.1;
 		FVisualLogEntry* NewEntry = FVisualLogger::Get().GetEntryToWrite(World, NewTimestamp); //generate new entry and serialize old one
 		FVisualLogger::Get().FlushThreadsEntries();
