Support basic dynamic actor renaming/level changes in Iris.
* Iris now respects net.SerializeNewActorOverrideLevel.
* UActorReplicationBridge will now update level groups of actors that change levels. Use the cvar net.Iris.EnableActorLevelChanges to disable this functionality.
* Removed FActorBeginReplicationParams::bIncludeInLevelGroupFilter. PlayerController was the only class using it, now the UActorReplicationBridge handles PlayerControllers explicitly.
* Minor update to FRenameReplicatedActor_ServerOnlyDynamicOuterChangeTest to verify that net.SerializeNewActorOverrideLevel is working as well.

diff --git a/Engine/Source/Runtime/Engine/Classes/Engine/NetDriver.h b/Engine/Source/Runtime/Engine/Classes/Engine/NetDriver.h
--- a/Engine/Source/Runtime/Engine/Classes/Engine/NetDriver.h
+++ b/Engine/Source/Runtime/Engine/Classes/Engine/NetDriver.h
@@ -390,6 +390,12 @@ extern ENGINE_API double GServerReplicateActorTimeSeconds;
 extern ENGINE_API int32 GNumClientConnections;
 extern ENGINE_API int32 GNumClientUpdateLevelVisibility;
 
+namespace UE::Net::Private
+{
+	/** Allow other internal systems to check this cvar */
+	extern int32 SerializeNewActorOverrideLevel;
+}
+
 // Delegates
 
 #if !UE_BUILD_SHIPPING
diff --git a/Engine/Source/Runtime/Engine/Classes/Engine/PackageMapClient.h b/Engine/Source/Runtime/Engine/Classes/Engine/PackageMapClient.h
--- a/Engine/Source/Runtime/Engine/Classes/Engine/PackageMapClient.h
+++ b/Engine/Source/Runtime/Engine/Classes/Engine/PackageMapClient.h
@@ -148,9 +148,6 @@ private:
 
 namespace UE::Net::Private
 {
-	/** Allow other internal systems to check this cvar */
-	extern int32 SerializeNewActorOverrideLevel;
-
 	/** Simple class to manage a reference-counted array of FNetworkGUIDs */
 	class FRefCountedNetGUIDArray
 	{
diff --git a/Engine/Source/Runtime/Engine/Private/Net/Iris/ReplicationSystem/ActorReplicationBridge.cpp b/Engine/Source/Runtime/Engine/Private/Net/Iris/ReplicationSystem/ActorReplicationBridge.cpp
--- a/Engine/Source/Runtime/Engine/Private/Net/Iris/ReplicationSystem/ActorReplicationBridge.cpp
+++ b/Engine/Source/Runtime/Engine/Private/Net/Iris/ReplicationSystem/ActorReplicationBridge.cpp
@@ -30,6 +30,7 @@
 #include "Engine/Level.h"
 #include "Engine/World.h"
 #include "GameFramework/Actor.h"
+#include "GameFramework/PlayerController.h"
 #include "HAL/LowLevelMemStats.h"
 #include "Net/DataBunch.h"
 #include "Net/DataChannel.h"
@@ -49,6 +50,13 @@ extern bool GDefaultUseSubObjectReplicationList;
 namespace UE::Net::Private
 {
 
+static bool bEnableActorLevelChanges = true;
+static FAutoConsoleVariableRef CVarEnableActorLevelChanges(
+	TEXT("net.Iris.EnableActorLevelChanges"),
+	bEnableActorLevelChanges,
+	TEXT("When true the ActorReplicationBridge will process actors that change levels by updating the actor's level groups.")
+);
+
 bool IsActorValidForIrisReplication(const AActor* Actor)
 {
 	return IsValid(Actor) && !Actor->IsActorBeingDestroyed() && !Actor->IsUnreachable();
@@ -77,6 +85,13 @@ void ActorReplicationBridgeGetActorWorldObjectInfo(FNetRefHandle /*Handle*/, con
 	}
 }
 
+bool ShouldIncludeActorInLevelGroups(const AActor* Actor)
+{
+	// Never filter out PlayerControllers based on level as they are required for travel.
+	// Preserves the special case for PlayerControllers from UNetDriver::IsLevelInitializedForActor.
+	return !Actor->IsA<APlayerController>();
+}
+
 } // end namespace UE::Net::Private
 
 UActorReplicationBridge::UActorReplicationBridge()
@@ -292,30 +307,8 @@ UE::Net::FNetRefHandle UActorReplicationBridge::BeginReplication(AActor* Actor,
 		}
 	}
 
-	// Setup Level filtering unless this actor belongs to the persistent level
-	{
-		ULevel* Level = Actor->GetLevel();
-		if (Params.bIncludeInLevelGroupFilter && (!Level->IsPersistentLevel() || (Level != NetDriver->GetWorld()->PersistentLevel)))
-		{
-			const UPackage* const LevelPackage = Level->GetOutermost();
-			const FName PackageName = LevelPackage->GetFName();
-
-			FNetObjectGroupHandle LevelGroup = GetLevelGroup(Level);
-			if (!LevelGroup.IsValid())
-			{
-				LevelGroup = CreateLevelGroup(Level);
-
-				UE_LOG_ACTORREPLICATIONBRIDGE(Log, TEXT("Created new GroupIndex: %u for Level: %s"), LevelGroup.GetGroupIndex(), ToCStr(PackageName.ToString()));
-
-				// Update the filtering status of the group based on current level visibility for all connections
-				NetDriver->UpdateGroupFilterStatusForLevel(Level, LevelGroup);
-			}
-
-			// Add object to group
-			UE_LOG_ACTORREPLICATIONBRIDGE(Verbose, TEXT("Added %s to GroupIndex: %u Level: %s"), *ActorRefHandle.ToString(), LevelGroup.GetGroupIndex(), ToCStr(PackageName.ToString()));
-			GetReplicationSystem()->AddToGroup(LevelGroup, ActorRefHandle);	
-		}
-	}
+	// Setup Level filtering
+	AddActorToLevelGroup(Actor);
 
 	// If we have registered sub objects we replicate them as well
 	const FSubObjectRegistry& ActorSubObjects = FSubObjectRegistryGetter::GetSubObjects(Actor);
@@ -981,7 +974,7 @@ void UActorReplicationBridge::GetActorCreationHeader(const AActor* Actor, UE::Ne
 
 		// Fill in Header
 		Header.ArchetypeReference = GetOrCreateObjectReference(Archetype);
-		Header.bUsePersistentLevel = NetDriver->GetWorld()->PersistentLevel == ActorLevel;
+		Header.bUsePersistentLevel = (UE::Net::Private::SerializeNewActorOverrideLevel == 0) || (NetDriver->GetWorld()->PersistentLevel == ActorLevel);
 		if (!Header.bUsePersistentLevel)
 		{
 			Header.LevelReference = GetOrCreateObjectReference(ActorLevel);
@@ -1232,6 +1225,78 @@ void UActorReplicationBridge::ReportErrorWithNetRefHandle(uint32 ErrorType, FNet
 	}
 }
 
+void UActorReplicationBridge::ActorChangedLevel(const AActor* Actor, const ULevel* PreviousLevel)
+{
+	if (!UE::Net::Private::bEnableActorLevelChanges)
+	{
+		return;
+	}
+
+	UE_LOG_ACTORREPLICATIONBRIDGE(Verbose, TEXT("ActorChangedLevel: Actor %s from PreviousLevel %s"), *GetFullNameSafe(Actor), *GetFullNameSafe(PreviousLevel));
+
+	if (!UE::Net::Private::ShouldIncludeActorInLevelGroups(Actor))
+	{
+		return;
+	}
+
+	// Remove from previous level group
+	const bool bPreviousLevelIsPersistentLevel = PreviousLevel && (PreviousLevel->IsPersistentLevel() || (PreviousLevel == NetDriver->GetWorld()->PersistentLevel));
+	
+	if (PreviousLevel && !bPreviousLevelIsPersistentLevel)
+	{
+		const UPackage* const PreviousLevelPackage = PreviousLevel->GetOutermost();
+		const FName PreviousLevelPackageName = PreviousLevelPackage->GetFName();
+
+		const UE::Net::FNetRefHandle ActorRefHandle = GetReplicatedRefHandle(Actor);
+		UE::Net::FNetObjectGroupHandle PreviousLevelGroup = GetLevelGroup(PreviousLevel);
+		if (GetReplicationSystem()->IsInGroup(PreviousLevelGroup, ActorRefHandle))
+		{
+			UE_LOG_ACTORREPLICATIONBRIDGE(Verbose, TEXT("ActorChangedLevel: removing %s from GroupIndex: %u PreviousLevel: %s"), *ActorRefHandle.ToString(), PreviousLevelGroup.GetGroupIndex(), ToCStr(PreviousLevelPackageName.ToString()));
+			GetReplicationSystem()->RemoveFromGroup(PreviousLevelGroup, ActorRefHandle);
+		}
+		else
+		{
+			UE_LOG_ACTORREPLICATIONBRIDGE(Warning, TEXT("ActorChangedLevel: %s not found in GroupIndex: %u PreviousLevel: %s"), *ActorRefHandle.ToString(), PreviousLevelGroup.GetGroupIndex(), ToCStr(PreviousLevelPackageName.ToString()));
+		}
+	}
+
+	AddActorToLevelGroup(Actor);
+}
+
+void UActorReplicationBridge::AddActorToLevelGroup(const AActor* Actor)
+{
+	if (!UE::Net::Private::ShouldIncludeActorInLevelGroups(Actor))
+	{
+		return;
+	}
+	
+	const ULevel* Level = Actor->GetLevel();
+	
+	// Don't filter out actors in the persistent level
+	if (Level && (!Level->IsPersistentLevel() || (Level != NetDriver->GetWorld()->PersistentLevel)))
+	{
+		const UPackage* const LevelPackage = Level->GetOutermost();
+		const FName PackageName = LevelPackage->GetFName();
+
+		UE::Net::FNetObjectGroupHandle LevelGroup = GetLevelGroup(Level);
+		if (!LevelGroup.IsValid())
+		{
+			LevelGroup = CreateLevelGroup(Level);
+
+			UE_LOG_ACTORREPLICATIONBRIDGE(Log, TEXT("Created new GroupIndex: %u for Level: %s"), LevelGroup.GetGroupIndex(), ToCStr(PackageName.ToString()));
+
+			// Update the filtering status of the group based on current level visibility for all connections
+			NetDriver->UpdateGroupFilterStatusForLevel(Level, LevelGroup);
+		}
+
+		const UE::Net::FNetRefHandle ActorRefHandle = GetReplicatedRefHandle(Actor);
+
+		// Add object to group
+		UE_LOG_ACTORREPLICATIONBRIDGE(Verbose, TEXT("Added %s to GroupIndex: %u Level: %s"), *ActorRefHandle.ToString(), LevelGroup.GetGroupIndex(), ToCStr(PackageName.ToString()));
+		GetReplicationSystem()->AddToGroup(LevelGroup, ActorRefHandle);
+	}
+}
+
 FString UActorReplicationBridge::PrintConnectionInfo(uint32 ConnectionId) const
 {
 	if (NetDriver)
diff --git a/Engine/Source/Runtime/Engine/Private/NetDriver.cpp b/Engine/Source/Runtime/Engine/Private/NetDriver.cpp
--- a/Engine/Source/Runtime/Engine/Private/NetDriver.cpp
+++ b/Engine/Source/Runtime/Engine/Private/NetDriver.cpp
@@ -311,6 +311,11 @@ namespace UE::Net::Private
 		TEXT("Time in milliseconds to limit client incoming bunch processing to. If 0, no limit. As long as we're below the limit, will start processing another bunch. A single bunch that takes a while to process can overshoot the limit. ")
 		TEXT("After the limit is hit, remaining bunches in a packet are queued, and the IpNetDriver will not process any more packets in the current frame."));
 
+	int32 SerializeNewActorOverrideLevel = 1;
+	static FAutoConsoleVariableRef CVarNetSerializeNewActorOverrideLevel(
+		TEXT("net.SerializeNewActorOverrideLevel"),
+		SerializeNewActorOverrideLevel,
+		TEXT("If true, servers will serialize a spawned, replicated actor's level so the client attempts to spawn it into that level too. If false, clients will spawn all these actors into the persistent level."));
 
 } //namespace UE::Net::Private
 
@@ -3991,15 +3996,23 @@ void UNetDriver::NotifyActorRenamed(AActor* ThisActor, UObject* PreviousOuter, F
 			{
 				UE_LOG(LogNet, Log, TEXT("NotifyActorRenamed on server, dynamic actor changed outer: %s PreviousOuter: %s PreviousName: %s"), *GetFullNameSafe(ThisActor), *GetFullNameSafe(PreviousOuter), *PreviousName.ToString());
 
-#if UE_WITH_IRIS
-				ensureMsgf(!ReplicationSystem, TEXT("Dynamic actor renaming not supported in Iris. Actor: %s PreviousOuter: %s"), *GetFullNameSafe(ThisActor), *GetFullNameSafe(PreviousOuter));
-#endif
 				// Forward change to ReplicationDriver so it can update its state
 				if (ReplicationDriver)
 				{
 					ReplicationDriver->NotifyActorRenamed(ThisActor, PreviousOuter, PreviousName);
 				}
 
+#if UE_WITH_IRIS
+				if (ReplicationSystem)
+				{
+					if (UActorReplicationBridge* Bridge = ReplicationSystem->GetReplicationBridgeAs<UActorReplicationBridge>())
+					{
+						Bridge->ActorChangedLevel(ThisActor, Cast<ULevel>(PreviousOuter));
+					}
+				}
+				else
+#endif // UE_WITH_IRIS
+
 				if (UE::Net::Private::CleanUpRenamedDynamicActors)
 				{
 					// Close the channel or send destruction info for connections that don't have the new level visible.
@@ -4132,7 +4145,7 @@ void UNetDriver::NotifyActorTearOff(AActor* Actor)
 	{
 		if (UObjectReplicationBridge* Bridge = ReplicationSystem->GetReplicationBridgeAs<UObjectReplicationBridge>())
 		{
-			// Set the actor to be torn-off during the next update of the replication systeem
+			// Set the actor to be torn-off during the next update of the replication system
 			const UE::Net::FNetRefHandle ActorRefHandle = Bridge->GetReplicatedRefHandle(Actor);
 			ReplicationSystem->TearOffNextUpdate(ActorRefHandle);
 		}
@@ -5180,6 +5193,8 @@ int32 UNetDriver::ServerReplicateActors_PrioritizeActors( UNetConnection* Connec
 				// See of actor wants to try and go dormant
 				if ( ShouldActorGoDormant( Actor, ConnectionViewers, Channel, ElapsedTime, bLowNetBandwidth ) )
 				{
+					CA_ASSUME(Channel); // ShouldActorGoDormant returns false if Channel is null, but analyzers don't seem to see that
+
 					// Channel is marked to go dormant now once all properties have been replicated (but is not dormant yet)
 					Channel->StartBecomingDormant();
 				}
@@ -5287,7 +5302,7 @@ int32 UNetDriver::ServerReplicateActors_ProcessPrioritizedActorsRange( UNetConne
 
 		// Normal actor replication
 		UActorChannel* Channel = PriorityActors[j]->Channel;
-		UE_LOG( LogNetTraffic, Log, TEXT( " Maybe Replicate %s" ), *ActorInfo->Actor->GetName() );
+		UE_LOG( LogNetTraffic, Log, TEXT( " Maybe Replicate %s" ), ActorInfo ? *ActorInfo->Actor->GetName() : TEXT("None") );
 		if ( !Channel || Channel->Actor ) //make sure didn't just close this channel
 		{
 			AActor* Actor = ActorInfo->Actor;
diff --git a/Engine/Source/Runtime/Engine/Private/PackageMapClient.cpp b/Engine/Source/Runtime/Engine/Private/PackageMapClient.cpp
--- a/Engine/Source/Runtime/Engine/Private/PackageMapClient.cpp
+++ b/Engine/Source/Runtime/Engine/Private/PackageMapClient.cpp
@@ -120,12 +120,6 @@ namespace UE::Net::Private
 			}
 		}
 	}
-
-	int32 SerializeNewActorOverrideLevel = 1;
-	static FAutoConsoleVariableRef CVarNetSerializeNewActorOverrideLevel(
-		TEXT("net.SerializeNewActorOverrideLevel"),
-		SerializeNewActorOverrideLevel,
-		TEXT("If true, servers will serialize a spawned, replicated actor's level so the client attempts to spawn it into that level too. If false, clients will spawn all these actors into the persistent level."));
 }
 
 static TAutoConsoleVariable<int32> CVarAllowAsyncLoading(
@@ -327,7 +321,7 @@ bool UPackageMapClient::SerializeObject( FArchive& Ar, UClass* Class, UObject*&
 		InternalWriteObject( Ar, NetGUID, Object, TEXT( "" ), NULL );
 
 		// If we need to export this GUID (its new or hasnt been ACKd, do so here)
-		if (!NetGUID.IsDefault() && ShouldSendFullPath(Object, NetGUID))
+		if (!NetGUID.IsDefault() && Object && ShouldSendFullPath(Object, NetGUID))
 		{
 			check(IsNetGUIDAuthority());
 			if ( !ExportNetGUID( NetGUID, Object, TEXT(""), NULL ) )
diff --git a/Engine/Source/Runtime/Engine/Private/PlayerController.cpp b/Engine/Source/Runtime/Engine/Private/PlayerController.cpp
--- a/Engine/Source/Runtime/Engine/Private/PlayerController.cpp
+++ b/Engine/Source/Runtime/Engine/Private/PlayerController.cpp
@@ -6112,10 +6112,7 @@ void APlayerController::BeginReplication()
 {
 	using namespace UE::Net;
 
-	// Always allow the PlayerController to be replicated as it is required for travel.
-	FActorBeginReplicationParams Params;
-	Params.bIncludeInLevelGroupFilter = false;
-	Super::BeginReplication(Params);
+	Super::BeginReplication();
 
 	// Bump prio of playercontroller in order to make sure it replicates really early
 	static constexpr float PlayerControllerStaticPriority = 100.f;
diff --git a/Engine/Source/Runtime/Engine/Public/Net/Iris/ReplicationSystem/ActorReplicationBridge.h b/Engine/Source/Runtime/Engine/Public/Net/Iris/ReplicationSystem/ActorReplicationBridge.h
--- a/Engine/Source/Runtime/Engine/Public/Net/Iris/ReplicationSystem/ActorReplicationBridge.h
+++ b/Engine/Source/Runtime/Engine/Public/Net/Iris/ReplicationSystem/ActorReplicationBridge.h
@@ -32,9 +32,6 @@ ENGINE_API bool ShouldUseIrisReplication(const UObject* Actor);
 /** Parameters passed to UActorReplicationBridge::BeginReplication. */
 struct FActorBeginReplicationParams
 {
-	/** Will the actor be part to the level visibility filters. This makes it that only clients who loaded the dynamic level of the actor can have the actor replicated to them.*/
-	bool bIncludeInLevelGroupFilter = true;
-
 	/** When true we ignore the configured dynamic filter for this actor type and use the explicit filter instead */
 	bool bOverrideDynamicFilterConfig = false;
 
@@ -84,6 +81,9 @@ public:
 	/** Tell the remote connection that we detected a reading error with a specific replicated object */
 	ENGINE_API virtual void ReportErrorWithNetRefHandle(uint32 ErrorType, FNetRefHandle RefHandle, uint32 ConnectionId) override;
 	
+	/** Updates the level group for an actor that changed levels */
+	void ActorChangedLevel(const AActor* Actor, const ULevel* PreviousLevel);
+
 	using UObjectReplicationBridge::EndReplication;
 
 protected:
@@ -116,6 +116,7 @@ private:
 	void OnMaxTickRateChanged(UNetDriver* InNetDriver, int32 NewMaxTickRate, int32 OldMaxTickRate);
 
 	void WakeUpObjectInstantiatedFromRemote(AActor* Actor) const;
+	void AddActorToLevelGroup(const AActor* Actor);
 
 private:
 	UNetDriver* NetDriver;
