[StaticMeshComponent] removed ensure for GetNavigationData/GetNavigationBounds in Editor builds  since there are cases where mesh compilation is still active of a registered component.
Also fixed a printf format and removed some unused code.

diff --git a/Engine/Source/Runtime/Engine/Private/Components/StaticMeshComponent.cpp b/Engine/Source/Runtime/Engine/Private/Components/StaticMeshComponent.cpp
--- a/Engine/Source/Runtime/Engine/Private/Components/StaticMeshComponent.cpp
+++ b/Engine/Source/Runtime/Engine/Private/Components/StaticMeshComponent.cpp
@@ -483,16 +483,7 @@ void UStaticMeshComponent::PreSave(FObjectPreSaveContext ObjectSaveContext)
 void UStaticMeshComponent::CheckForErrors()
 {
 	Super::CheckForErrors();
-
-	const FCoreTexts& CoreTexts = FCoreTexts::Get();
-
-	// Get the mesh owner's name.
 	AActor* Owner = GetOwner();
-	FString OwnerName(*(CoreTexts.None.ToString()));
-	if ( Owner )
-	{
-		OwnerName = Owner->GetName();
-	}
 
 	if (GetStaticMesh() != nullptr && GetStaticMesh()->IsNaniteEnabled() != 0)
 	{
@@ -808,17 +799,6 @@ void UStaticMeshComponent::OnDestroyPhysicsState()
 
 #if WITH_EDITORONLY_DATA
 
-/** Return the total number of LOD sections in the LOD resources */
-static int32 GetNumberOfElements(const TIndirectArray<FStaticMeshLODResources>& LODResources)
-{
-	int32 Count = 0;
-	for (int32 LODIndex = 0; LODIndex < LODResources.Num(); ++LODIndex)
-	{
-		Count += LODResources[LODIndex].Sections.Num();
-	}
-	return Count;
-}
-
 /**
  *	Pack the texture into data ready for saving. Also ensures a single entry per texture.
  *
@@ -1081,14 +1061,6 @@ void UStaticMeshComponent::GetStreamingRenderAssetInfo(FStreamingTextureLevelCon
 
 	// Since GetTextureStreamingTransformScale can be slow for certain component types, only call it if necessary
 	TOptional<float> LazyTransformScale;
-	auto GetTransformScale = [this, &LazyTransformScale]()
-	{
-		if (!LazyTransformScale.IsSet())
-		{
-			LazyTransformScale = GetTextureStreamingTransformScale();
-		}
-		return *LazyTransformScale;
-	};
 	
 	if (!CanSkipGetTextureStreamingRenderAssetInfo())
 	{
@@ -1239,7 +1211,7 @@ FTransform UStaticMeshComponent::GetSocketTransform(FName InSocketName, ERelativ
 					{
 						if( const AActor* Actor = GetOwner() )
 						{
-							return SocketWorldTransform.GetRelativeTransform( GetOwner()->GetTransform() );
+							return SocketWorldTransform.GetRelativeTransform( Actor->GetTransform() );
 						}
 						break;
 					}
@@ -1368,12 +1340,10 @@ void UStaticMeshComponent::CopyInstanceVertexColorsIfCompatible( const UStaticMe
 		// Copy vertex colors from Source to Target (this)
 		for ( int32 CurrentLOD = 0; CurrentLOD != NumSourceLODs; CurrentLOD++ )
 		{
-			FStaticMeshLODResources& SourceLODModel = SourceComponent->GetStaticMesh()->GetRenderData()->LODResources[CurrentLOD];
 			if (SourceComponent->LODData.IsValidIndex(CurrentLOD))
 			{
 				const FStaticMeshComponentLODInfo& SourceLODInfo = SourceComponent->LODData[CurrentLOD];
 
-				FStaticMeshLODResources& TargetLODModel = GetStaticMesh()->GetRenderData()->LODResources[CurrentLOD];
 				FStaticMeshComponentLODInfo& TargetLODInfo = LODData[CurrentLOD];
 
 				if ( SourceLODInfo.OverrideVertexColors != nullptr )
@@ -1791,7 +1761,6 @@ void UStaticMeshComponent::UpdatePreCulledData(int32 LODIndex, const TArray<uint
 
 		for (int32 SectionIndex = 0; SectionIndex < StaticMeshLODResources.Sections.Num(); SectionIndex++)
 		{
-			const FStaticMeshSection& Section = StaticMeshLODResources.Sections[SectionIndex];
 			FPreCulledStaticMeshSection PreCulledSection;
 			PreCulledSection.FirstIndex = FirstIndex;
 			PreCulledSection.NumTriangles = NumTrianglesPerSection[SectionIndex];
@@ -2446,7 +2415,7 @@ void UStaticMeshComponent::SetEvaluateWorldPositionOffsetInRayTracing(bool NewVa
 	{
 		// Update render thread data
 		ENQUEUE_RENDER_COMMAND(UpdateEvaluateWPORTCmd)
-		([NewValue, Scene = GetScene(), PrimitiveSceneProxy = static_cast<FStaticMeshSceneProxy*>(SceneProxy)](FRHICommandList& RHICmdList)
+		([NewValue, PrimitiveSceneProxy = static_cast<FStaticMeshSceneProxy*>(SceneProxy)](FRHICommandList& RHICmdList)
 		{
 			PrimitiveSceneProxy->SetEvaluateWorldPositionOffsetInRayTracing(RHICmdList, NewValue);
 		});
@@ -3138,7 +3107,10 @@ FBox UStaticMeshComponent::GetNavigationBounds() const
 {
 	if (const UStaticMesh* Mesh = GetStaticMesh())
 	{
-		if (ensureMsgf(!Mesh->IsCompiling(), TEXT("%s is not considered relevant to navigation until associated mesh is compiled."), *GetFullName()))
+#if WITH_EDITOR
+		// @see GetNavigationData
+		if (!Mesh->IsCompiling())
+#endif // WITH_EDITOR
 		{
 			return Mesh->GetNavigationBounds(GetComponentTransform());
 		}
@@ -3156,7 +3128,14 @@ void UStaticMeshComponent::GetNavigationData(FNavigationRelevantData& Data) cons
 	{
 		if (const UStaticMesh* Mesh = GetStaticMesh())
 		{
-			if (ensureMsgf(!Mesh->IsCompiling(), TEXT("%s is not considered relevant to navigation until associated mesh is compiled."), *GetFullName()))
+#if WITH_EDITOR
+			// In Editor it's possible that compilation of a StaticMesh gets triggered on a newly registered component for
+			// which a pending update is queued for the Navigation system.
+			// Then GetNavigationData is called when the pending update is processed but we don't consider the current component
+			// relevant to navigation until associated mesh is compiled.
+			// On mesh post compilation the component will reregister with the right mesh.
+			if (!Mesh->IsCompiling())
+#endif // WITH_EDITOR
 			{
 				if (UNavCollisionBase* NavCollision = Mesh->GetNavCollision())
 				{
@@ -3558,7 +3537,7 @@ void FStaticMeshComponentLODInfo::ExportText(FString& ValueStr)
 		FPaintedVertex& Vert = PaintedVertices[i];
 
 		ValueStr += FString::Printf(TEXT("((Position=(X=%.6f,Y=%.6f,Z=%.6f),"), Vert.Position.X, Vert.Position.Y, Vert.Position.Z);
-		ValueStr += FString::Printf(TEXT("(Normal=(X=%d,Y=%d,Z=%d,W=%d),"), Vert.Normal.X, Vert.Normal.Y, Vert.Normal.Z, Vert.Normal.W);
+		ValueStr += FString::Printf(TEXT("(Normal=(X=%f,Y=%f,Z=%f,W=%f),"), Vert.Normal.X, Vert.Normal.Y, Vert.Normal.Z, Vert.Normal.W);
 		ValueStr += FString::Printf(TEXT("(Color=(B=%d,G=%d,R=%d,A=%d))"), Vert.Color.B, Vert.Color.G, Vert.Color.R, Vert.Color.A);
 
 		// Seperate each vertex entry with a comma
