Add in some Texture information to the Content Browser as columns that can be searched and sorted. Max In Game X, Y, Z (where appropriate) as well as calculated resource size for memory.
#test Loaded assets in Editor and viewed the columns.

diff --git a/Engine/Source/Runtime/Engine/Classes/Engine/Texture2DArray.cpp b/Engine/Source/Runtime/Engine/Classes/Engine/Texture2DArray.cpp
--- a/Engine/Source/Runtime/Engine/Classes/Engine/Texture2DArray.cpp
+++ b/Engine/Source/Runtime/Engine/Classes/Engine/Texture2DArray.cpp
@@ -568,6 +568,17 @@ void UTexture2DArray::GetAssetRegistryTags(FAssetRegistryTagsContext Context) co
 	Context.AddTag(FAssetRegistryTag("Dimensions", Dimensions, FAssetRegistryTag::TT_Dimensional));
 	Context.AddTag(FAssetRegistryTag("Format", GPixelFormats[GetPixelFormat()].Name, FAssetRegistryTag::TT_Alphabetical));
 
+#if WITH_EDITORONLY_DATA
+	const int32 SurfaceWidth = static_cast<int32>(GetSurfaceWidth());
+	const int32 SurfaceHeight = static_cast<int32>(GetSurfaceHeight());
+	const int32 MaxResMipBias = GetCachedLODBias();
+	const int32 MaxInGameWidth = SurfaceWidth ? FMath::Max(SurfaceWidth >> MaxResMipBias, 1) : 0;
+	const int32 MaxInGameHeight = SurfaceHeight ? FMath::Max(SurfaceHeight >> MaxResMipBias, 1) : 0;
+
+	Context.AddTag(FAssetRegistryTag("MaxInGameX", FString::FromInt(MaxInGameWidth), FAssetRegistryTag::TT_Numerical));
+	Context.AddTag(FAssetRegistryTag("MaxInGameY", FString::FromInt(MaxInGameHeight), FAssetRegistryTag::TT_Numerical));
+#endif
+
 	Super::GetAssetRegistryTags(Context);
 }
 
diff --git a/Engine/Source/Runtime/Engine/Private/Texture.cpp b/Engine/Source/Runtime/Engine/Private/Texture.cpp
--- a/Engine/Source/Runtime/Engine/Private/Texture.cpp
+++ b/Engine/Source/Runtime/Engine/Private/Texture.cpp
@@ -1388,6 +1388,15 @@ void UTexture::GetAssetRegistryTags(FAssetRegistryTagsContext Context) const
 	
 	Context.AddTag(FAssetRegistryTag("IsSourceValid", Source.IsValid() ? TEXT("True") : TEXT("False"), FAssetRegistryTag::TT_Alphabetical));
 
+#if WITH_EDITORONLY_DATA
+	Context.AddTag(FAssetRegistryTag("MaxTextureSize", FString::FromInt(MaxTextureSize), FAssetRegistryTag::TT_Numerical));
+
+	FTexturePlatformData** PlatformDataPtr = const_cast<UTexture*>(this)->GetRunningPlatformData();
+	const int32 ActualMipBias = GetCachedLODBias();
+	const int64 ResourceSize = PlatformDataPtr && *PlatformDataPtr ? (*PlatformDataPtr)->GetPayloadSize(ActualMipBias) : const_cast<UTexture*>(this)->GetResourceSizeBytes(EResourceSizeMode::Exclusive);
+	Context.AddTag(FAssetRegistryTag("ResourceSizeKB", FString::SanitizeFloat(ResourceSize / 1024.f), FAssetRegistryTag::TT_Numerical));
+#endif
+
 	Super::GetAssetRegistryTags(Context);
 }
 #endif
diff --git a/Engine/Source/Runtime/Engine/Private/Texture2D.cpp b/Engine/Source/Runtime/Engine/Private/Texture2D.cpp
--- a/Engine/Source/Runtime/Engine/Private/Texture2D.cpp
+++ b/Engine/Source/Runtime/Engine/Private/Texture2D.cpp
@@ -719,6 +719,18 @@ void UTexture2D::GetAssetRegistryTags(FAssetRegistryTagsContext Context) const
 	Context.AddTag(FAssetRegistryTag("MaxTextureSize", FString::FromInt(MaxTextureSize), FAssetRegistryTag::TT_Numerical));
 #endif
 
+#if WITH_EDITORONLY_DATA
+	// In game max bias and dimensions
+	const int32 SurfaceWidth = static_cast<int32>(GetSurfaceWidth());
+	const int32 SurfaceHeight = static_cast<int32>(GetSurfaceHeight());
+	const int32 MaxResMipBias = GetNumMips() - GetNumMipsAllowed(true);
+	const int32 MaxInGameWidth = SurfaceWidth ? FMath::Max(SurfaceWidth >> MaxResMipBias, 1) : 0;
+	const int32 MaxInGameHeight = SurfaceHeight ? FMath::Max(SurfaceHeight >> MaxResMipBias, 1) : 0;
+
+	Context.AddTag(FAssetRegistryTag("MaxInGameX", FString::FromInt(MaxInGameWidth), FAssetRegistryTag::TT_Numerical));
+	Context.AddTag(FAssetRegistryTag("MaxInGameY", FString::FromInt(MaxInGameHeight), FAssetRegistryTag::TT_Numerical));
+#endif
+
 	Super::GetAssetRegistryTags(Context);
 }
 
diff --git a/Engine/Source/Runtime/Engine/Private/TextureCube.cpp b/Engine/Source/Runtime/Engine/Private/TextureCube.cpp
--- a/Engine/Source/Runtime/Engine/Private/TextureCube.cpp
+++ b/Engine/Source/Runtime/Engine/Private/TextureCube.cpp
@@ -197,6 +197,17 @@ void UTextureCube::GetAssetRegistryTags(FAssetRegistryTagsContext Context) const
 	Context.AddTag( FAssetRegistryTag("Dimensions", Dimensions, FAssetRegistryTag::TT_Dimensional) );
 	Context.AddTag( FAssetRegistryTag("Format", GPixelFormats[GetPixelFormat()].Name, FAssetRegistryTag::TT_Alphabetical) );
 
+#if WITH_EDITORONLY_DATA
+	const int32 SurfaceWidth = static_cast<int32>(GetSurfaceWidth());
+	const int32 SurfaceHeight = static_cast<int32>(GetSurfaceHeight());
+	const int32 MaxResMipBias = GetCachedLODBias();
+	const int32 MaxInGameWidth = SurfaceWidth ? FMath::Max(SurfaceWidth >> MaxResMipBias, 1) : 0;
+	const int32 MaxInGameHeight = SurfaceHeight ? FMath::Max(SurfaceHeight >> MaxResMipBias, 1) : 0;
+
+	Context.AddTag(FAssetRegistryTag("MaxInGameX", FString::FromInt(MaxInGameWidth), FAssetRegistryTag::TT_Numerical));
+	Context.AddTag(FAssetRegistryTag("MaxInGameY", FString::FromInt(MaxInGameHeight), FAssetRegistryTag::TT_Numerical));
+#endif
+
 	Super::GetAssetRegistryTags(Context);
 }
 
diff --git a/Engine/Source/Runtime/Engine/Private/TextureCubeArray.cpp b/Engine/Source/Runtime/Engine/Private/TextureCubeArray.cpp
--- a/Engine/Source/Runtime/Engine/Private/TextureCubeArray.cpp
+++ b/Engine/Source/Runtime/Engine/Private/TextureCubeArray.cpp
@@ -825,6 +825,18 @@ void UTextureCubeArray::GetAssetRegistryTags(FAssetRegistryTagsContext Context)
 	Context.AddTag(FAssetRegistryTag("Dimensions", Dimensions, FAssetRegistryTag::TT_Dimensional));
 	Context.AddTag(FAssetRegistryTag("Format", GPixelFormats[GetPixelFormat()].Name, FAssetRegistryTag::TT_Alphabetical));
 
+#if WITH_EDITORONLY_DATA
+	const int32 SurfaceWidth = static_cast<int32>(GetSurfaceWidth());
+	const int32 SurfaceHeight = static_cast<int32>(GetSurfaceHeight());
+	const int32 SurfaceDepth = static_cast<int32>(GetSurfaceDepth());
+	const int32 MaxResMipBias = GetCachedLODBias();
+	const int32 MaxInGameWidth = SurfaceWidth ? FMath::Max(SurfaceWidth >> MaxResMipBias, 1) : 0;
+	const int32 MaxInGameHeight = SurfaceHeight ? FMath::Max(SurfaceHeight >> MaxResMipBias, 1) : 0;
+
+	Context.AddTag(FAssetRegistryTag("MaxInGameX", FString::FromInt(MaxInGameWidth), FAssetRegistryTag::TT_Numerical));
+	Context.AddTag(FAssetRegistryTag("MaxInGameY", FString::FromInt(MaxInGameHeight), FAssetRegistryTag::TT_Numerical));
+#endif
+
 	Super::GetAssetRegistryTags(Context);
 }
 
diff --git a/Engine/Source/Runtime/Engine/Private/VolumeTexture.cpp b/Engine/Source/Runtime/Engine/Private/VolumeTexture.cpp
--- a/Engine/Source/Runtime/Engine/Private/VolumeTexture.cpp
+++ b/Engine/Source/Runtime/Engine/Private/VolumeTexture.cpp
@@ -357,6 +357,20 @@ void UVolumeTexture::GetAssetRegistryTags(FAssetRegistryTagsContext Context) con
 	Context.AddTag( FAssetRegistryTag("Dimensions", Dimensions, FAssetRegistryTag::TT_Dimensional) );
 	Context.AddTag( FAssetRegistryTag("Format", GPixelFormats[GetPixelFormat()].Name, FAssetRegistryTag::TT_Alphabetical) );
 
+#if WITH_EDITORONLY_DATA
+	const int32 SurfaceWidth = static_cast<int32>(GetSurfaceWidth());
+	const int32 SurfaceHeight = static_cast<int32>(GetSurfaceHeight());
+	const int32 SurfaceDepth = static_cast<int32>(GetSurfaceDepth());
+	const int32 MaxResMipBias = GetCachedLODBias();
+	const int32 MaxInGameWidth = SurfaceWidth ? FMath::Max(SurfaceWidth >> MaxResMipBias, 1) : 0;
+	const int32 MaxInGameHeight = SurfaceHeight ? FMath::Max(SurfaceHeight >> MaxResMipBias, 1) : 0;
+	const int32 MaxInGameDepth = SurfaceDepth ? FMath::Max(SurfaceDepth >> MaxResMipBias, 1) : 0;
+
+	Context.AddTag(FAssetRegistryTag("MaxInGameX", FString::FromInt(MaxInGameWidth), FAssetRegistryTag::TT_Numerical));
+	Context.AddTag(FAssetRegistryTag("MaxInGameY", FString::FromInt(MaxInGameHeight), FAssetRegistryTag::TT_Numerical));
+	Context.AddTag(FAssetRegistryTag("MaxInGameZ", FString::FromInt(MaxInGameDepth), FAssetRegistryTag::TT_Numerical));
+#endif
+
 	Super::GetAssetRegistryTags(Context);
 }
 
