fix format typo in stat dump

diff --git a/Engine/Source/Runtime/Engine/Private/ProfilingDebugging/HealthSnapshot.cpp b/Engine/Source/Runtime/Engine/Private/ProfilingDebugging/HealthSnapshot.cpp
--- a/Engine/Source/Runtime/Engine/Private/ProfilingDebugging/HealthSnapshot.cpp
+++ b/Engine/Source/Runtime/Engine/Private/ProfilingDebugging/HealthSnapshot.cpp
@@ -200,8 +200,8 @@ void FHealthSnapshot::DumpStats(FOutputDevice& Ar, FName CategoryName)
 		Ar.CategorizedLogf(CategoryName, ELogVerbosity::Log, TEXT("GPU: Avg %.02fms, Hitches/Min: %.02f, Bound Frames: %.02f%%"), GPU.AvgTime * 1000, GPU.HitchesPerMinute, GPU.PercentFramesBound);
 		Ar.CategorizedLogf(CategoryName, ELogVerbosity::Log, TEXT("DrawCalls: Avg: %d, Max: %d, Min: %d"), DrawCalls.Avg, DrawCalls.Max, DrawCalls.Min);
 		Ar.CategorizedLogf(CategoryName, ELogVerbosity::Log, TEXT("DrawnPrims: Avg: %d, Max: %d, Min: %d"), PrimitivesDrawn.Avg, PrimitivesDrawn.Max, PrimitivesDrawn.Min);
-		Ar.CategorizedLogf(CategoryName, ELogVerbosity::Log, TEXT("PhysicalMemoryMB: Avg: %d, Max: %d, Min: %d"), PhysicalMemory.Avg, PhysicalMemory.Max, PhysicalMemory.Min);
-		Ar.CategorizedLogf(CategoryName, ELogVerbosity::Log, TEXT("VirtualMemoryMB: Avg: %d, Max: %d, Min: %d"), VirtualMemory.Avg, VirtualMemory.Max, VirtualMemory.Min);
+		Ar.CategorizedLogf(CategoryName, ELogVerbosity::Log, TEXT("PhysicalMemoryMB: Avg: %llu, Max: %llu, Min: %llu"), PhysicalMemory.Avg, PhysicalMemory.Max, PhysicalMemory.Min);
+		Ar.CategorizedLogf(CategoryName, ELogVerbosity::Log, TEXT("VirtualMemoryMB: Avg: %llu, Max: %llu, Min: %llu"), VirtualMemory.Avg, VirtualMemory.Max, VirtualMemory.Min);
 	}
 
 	Ar.CategorizedLogf(CategoryName, ELogVerbosity::Log, TEXT("CPU Memory: Used %.2fMB, Peak %.2fMB"), CPUMemoryMB.Used, CPUMemoryMB.Peak);
