SharedMaterial: Change CachedMaterialResourcesForCooking to hold EShaderFormat in addition to the FMaterial*, so that the upcoming CookDependency code can recalculate the FShaderMapId from the FMaterial when it is constructing the dependencies.
Added FMaterialShaderParameters parameterless constructor so we can load it from compact binary.
#rnx

diff --git a/Engine/Source/Runtime/Engine/Classes/Materials/Material.h b/Engine/Source/Runtime/Engine/Classes/Materials/Material.h
--- a/Engine/Source/Runtime/Engine/Classes/Materials/Material.h
+++ b/Engine/Source/Runtime/Engine/Classes/Materials/Material.h
@@ -1090,7 +1090,7 @@ private:
 	TArray<FMaterialResource*> MaterialResources;
 #if WITH_EDITOR
 	/** Material resources being cached for cooking. */
-	TMap<const class ITargetPlatform*, TArray<FMaterialResource*>> CachedMaterialResourcesForCooking;
+	TMap<const ITargetPlatform*, TArray<FMaterialResourceForCooking>> CachedMaterialResourcesForCooking;
 #endif
 	/** Flag used to guarantee that the RT is finished using various resources in this UMaterial before cleanup. */
 	FThreadSafeBool ReleasedByRT;
@@ -1537,6 +1537,7 @@ private:
 	 */
 	void CacheResourceShadersForRendering(bool bRegenerateId, EMaterialShaderPrecompileMode PrecompileMode = EMaterialShaderPrecompileMode::Default);
 
+#if WITH_EDITOR
 	/**
 	 * Cache resource shaders for cooking on the given shader platform.
 	 * If a matching shader map is not found in memory or the DDC, a new one will be compiled.
@@ -1544,7 +1545,9 @@ private:
 	 * Caller is responsible for deleting OutCachedMaterialResources.
 	 * Note: This modifies material variables used for rendering and is assumed to be called within a FMaterialUpdateContext!
 	 */
-	void CacheResourceShadersForCooking(EShaderPlatform Platform, TArray<FMaterialResource*>& OutCachedMaterialResources, const ITargetPlatform* TargetPlatform = nullptr, bool bBlocking = false);
+	void CacheResourceShadersForCooking(EShaderPlatform Platform, TArray<FMaterialResourceForCooking>& OutCachedMaterialResources, 
+		const ITargetPlatform* TargetPlatform = nullptr, bool bBlocking = false);
+#endif
 
 	void GetNewResources(EShaderPlatform ShaderPlatform, TArray<FMaterialResource*>& NewResourcesToCache);
 
diff --git a/Engine/Source/Runtime/Engine/Classes/Materials/MaterialInstance.h b/Engine/Source/Runtime/Engine/Classes/Materials/MaterialInstance.h
--- a/Engine/Source/Runtime/Engine/Classes/Materials/MaterialInstance.h
+++ b/Engine/Source/Runtime/Engine/Classes/Materials/MaterialInstance.h
@@ -710,7 +710,7 @@ private:
 	TArray<FMaterialResource*> StaticPermutationMaterialResources;
 #if WITH_EDITOR
 	/** Material resources being cached for cooking. */
-	TMap<const class ITargetPlatform*, TArray<FMaterialResource*>> CachedMaterialResourcesForCooking;
+	TMap<const ITargetPlatform*, TArray<FMaterialResourceForCooking>> CachedMaterialResourcesForCooking;
 #endif
 	/** Thread-safe flags used to track whether this instance is still in use by the render thread (EMaterialInstanceUsedByRTFlag) */
 	mutable std::atomic<uint32> UsedByRT;
@@ -879,6 +879,7 @@ public:
 
 	ENGINE_API void UpdateOverridableBaseProperties();
 
+#if WITH_EDITOR
 	/** 
 	 * Cache resource shaders for rendering on the given shader platform. 
 	 * If a matching shader map is not found in memory or the DDC, a new one will be compiled.
@@ -887,11 +888,12 @@ public:
 	 */
 	void CacheResourceShadersForCooking(
 		EShaderPlatform ShaderPlatform,
-		TArray<FMaterialResource*>& OutCachedMaterialResources,
+		TArray<FMaterialResourceForCooking>& OutCachedMaterialResources,
 		EMaterialShaderPrecompileMode PrecompileMode = EMaterialShaderPrecompileMode::Default,
 		const ITargetPlatform* TargetPlatform = nullptr,
 		bool bBlocking = false
 	);
+#endif
 
 	/** 
 	 * Gathers actively used shader maps from all material resources used by this material instance
diff --git a/Engine/Source/Runtime/Engine/Classes/Materials/MaterialInterface.h b/Engine/Source/Runtime/Engine/Classes/Materials/MaterialInterface.h
--- a/Engine/Source/Runtime/Engine/Classes/Materials/MaterialInterface.h
+++ b/Engine/Source/Runtime/Engine/Classes/Materials/MaterialInterface.h
@@ -63,6 +63,7 @@ struct FMaterialCachedExpressionData;
 struct FMaterialCachedExpressionEditorOnlyData;
 #if WITH_EDITOR
 class FMaterialCachedHLSLTree;
+struct FMaterialResourceForCooking;
 #endif
 #if WITH_EDITORONLY_DATA
 struct FParameterChannelNames;
@@ -1216,12 +1217,21 @@ private:
 #endif
 };
 
+namespace UE::MaterialInterface::Private
+{
+
 /** Helper function to serialize inline shader maps for the given material resources. */
 extern void SerializeInlineShaderMaps(
-	const TMap<const class ITargetPlatform*, TArray<FMaterialResource*>>* PlatformMaterialResourcesToSave,
 	FArchive& Ar,
 	TArray<FMaterialResource>& OutLoadedResources,
-	const FName& SerializingAsset = NAME_None);
+	const FName& SerializingAsset = NAME_None
+#if WITH_EDITOR
+	, const TMap<const ITargetPlatform*, TArray<FMaterialResourceForCooking>>* PlatformMaterialResourcesToSave = nullptr
+#endif
+);
+
+}
+
 /** Helper function to process (register) serialized inline shader maps for the given material resources. */
 extern void ProcessSerializedInlineShaderMaps(UMaterialInterface* Owner, TArray<FMaterialResource>& LoadedResources, TArray<FMaterialResource*>& OutMaterialResourcesLoaded);
 
diff --git a/Engine/Source/Runtime/Engine/Private/Materials/Material.cpp b/Engine/Source/Runtime/Engine/Private/Materials/Material.cpp
--- a/Engine/Source/Runtime/Engine/Private/Materials/Material.cpp
+++ b/Engine/Source/Runtime/Engine/Private/Materials/Material.cpp
@@ -730,11 +730,17 @@ static TAutoConsoleVariable<int32> CVarDiscardUnusedQualityLevels(
 	TEXT("1: Discard unused quality levels on load."),
 	ECVF_ReadOnly);
 
+namespace UE::MaterialInterface::Private
+{
+
 void SerializeInlineShaderMaps(
-	const TMap<const ITargetPlatform*, TArray<FMaterialResource*>>* PlatformMaterialResourcesToSavePtr,
 	FArchive& Ar,
 	TArray<FMaterialResource>& OutLoadedResources,
-	const FName& SerializingAsset)
+	const FName& SerializingAsset
+#if WITH_EDITOR
+	, const TMap<const ITargetPlatform*, TArray<FMaterialResourceForCooking>>* PlatformMaterialResourcesToSavePtr
+#endif
+)
 {
 	LLM_SCOPE(ELLMTag::Shaders);
 	SCOPED_LOADTIMER(SerializeInlineShaderMaps);
@@ -744,16 +750,17 @@ void SerializeInlineShaderMaps(
 
 	if (Ar.IsSaving())
 	{
+#if WITH_EDITOR
 		int32 NumResourcesToSave = 0;
-		const TArray<FMaterialResource*> *MaterialResourcesToSavePtr = NULL;
+		const TArray<FMaterialResourceForCooking> *MaterialResourcesToSavePtr = nullptr;
 		if (Ar.IsCooking() && Ar.IsPersistent() && !Ar.IsObjectReferenceCollector() && !Ar.ShouldSkipBulkData())
 		{
 			check( PlatformMaterialResourcesToSavePtr );
 			auto& PlatformMaterialResourcesToSave = *PlatformMaterialResourcesToSavePtr;
 
 			MaterialResourcesToSavePtr = PlatformMaterialResourcesToSave.Find( Ar.CookingTarget() );
-			check( MaterialResourcesToSavePtr != NULL || (Ar.GetLinker()==NULL) );
-			if (MaterialResourcesToSavePtr!= NULL )
+			check(MaterialResourcesToSavePtr != nullptr || Ar.GetLinker() == nullptr);
+			if (MaterialResourcesToSavePtr!= nullptr )
 			{
 				NumResourcesToSave = MaterialResourcesToSavePtr->Num();
 			}
@@ -765,13 +772,17 @@ void SerializeInlineShaderMaps(
 			&& NumResourcesToSave > 0)
 		{
 			FMaterialResourceMemoryWriter ResourceAr(Ar);
-			const TArray<FMaterialResource*> &MaterialResourcesToSave = *MaterialResourcesToSavePtr;
-			for (int32 ResourceIndex = 0; ResourceIndex < NumResourcesToSave; ResourceIndex++)
+			for (const FMaterialResourceForCooking& ResourceForCooking : *MaterialResourcesToSavePtr)
 			{
-				FMaterialResourceWriteScope Scope(&ResourceAr, *MaterialResourcesToSave[ResourceIndex]);
-				MaterialResourcesToSave[ResourceIndex]->SerializeInlineShaderMap(ResourceAr);
+				FMaterialResource& Resource = *ResourceForCooking.Resource;
+				FMaterialResourceWriteScope Scope(&ResourceAr, Resource);
+				Resource.SerializeInlineShaderMap(ResourceAr);
 			}
 		}
+#else
+		int32 NumResourcesToSave = 0;
+		Ar << NumResourcesToSave;
+#endif
 	}
 	else if (Ar.IsLoading())
 	{
@@ -793,6 +804,8 @@ void SerializeInlineShaderMaps(
 	}
 }
 
+} // namespace UE::MaterialInterface::Private
+
 void ProcessSerializedInlineShaderMaps(UMaterialInterface* Owner, TArray<FMaterialResource>& LoadedResources, TArray<FMaterialResource*>& OutMaterialResourcesLoaded)
 {
 	LLM_SCOPE(ELLMTag::Shaders);
@@ -2365,30 +2378,31 @@ void UMaterial::CacheResourceShadersForRendering(bool bRegenerateId, EMaterialSh
 			CacheShadersForResources(ShaderPlatform, ResourcesToCache, PrecompileMode);
 		}
 
+#if WITH_EDITOR
 		FString AdditionalFormatToCache = GCompileMaterialsForShaderFormatCVar->GetString();
 		if (!AdditionalFormatToCache.IsEmpty())
 		{
 			EShaderPlatform AdditionalPlatform = ShaderFormatToLegacyShaderPlatform(FName(*AdditionalFormatToCache));
 			if (AdditionalPlatform != SP_NumPlatforms)
 			{
-				ResourcesToCache.Reset();
-				CacheResourceShadersForCooking(AdditionalPlatform, ResourcesToCache);
-				for (int32 i = 0; i < ResourcesToCache.Num(); ++i)
-				{
-					FMaterialResource* Resource = ResourcesToCache[i];
-					delete Resource;
-				}
-				ResourcesToCache.Reset();
+				TArray<FMaterialResourceForCooking> CookResourcesToCache;
+				CacheResourceShadersForCooking(AdditionalPlatform, CookResourcesToCache);
 			}
 		}
+#endif
 
 		RecacheUniformExpressions(true);
 	}
 }
 
-void UMaterial::CacheResourceShadersForCooking(EShaderPlatform ShaderPlatform, TArray<FMaterialResource*>& OutCachedMaterialResources, const ITargetPlatform* TargetPlatform, bool bBlocking)
+#if WITH_EDITOR
+void UMaterial::CacheResourceShadersForCooking(EShaderPlatform ShaderPlatform,
+	TArray<FMaterialResourceForCooking>& OutCachedMaterialResources, const ITargetPlatform* TargetPlatform,
+	bool bBlocking)
 {
-	TArray<FMaterialResource*> NewResourcesToCache;	// only new resources need to have CacheShaders() called on them, whereas OutCachedMaterialResources may already contain resources for another shader platform
+	// Only new resources need to have CacheShaders() called on them, whereas OutCachedMaterialResources
+	// may already contain resources for another shader platform
+	TArray<FMaterialResource*> NewResourcesToCache;
 	GetNewResources(ShaderPlatform, NewResourcesToCache);
 
 #if WITH_EDITOR
@@ -2407,8 +2421,13 @@ void UMaterial::CacheResourceShadersForCooking(EShaderPlatform ShaderPlatform, T
 	}
 #endif
 
-	OutCachedMaterialResources.Append(NewResourcesToCache);
+	OutCachedMaterialResources.Reserve(OutCachedMaterialResources.Num() + NewResourcesToCache.Num());
+	for (FMaterialResource* Resource : NewResourcesToCache)
+	{
+		OutCachedMaterialResources.Add({ Resource, ShaderPlatform });
+	}
 }
+#endif // WITH_EDITOR
 
 void UMaterial::GetNewResources(EShaderPlatform ShaderPlatform, TArray<FMaterialResource*>& NewResourcesToCache)
 {
@@ -2718,16 +2737,15 @@ void UMaterial::Serialize(FArchive& Ar)
 
 	if (Ar.UEVer() >= VER_UE4_PURGED_FMATERIAL_COMPILE_OUTPUTS)
 	{
+		UE::MaterialInterface::Private::SerializeInlineShaderMaps(
+			Ar, LoadedMaterialResources
 #if WITH_EDITOR
-		SerializeInlineShaderMaps(&CachedMaterialResourcesForCooking, Ar, LoadedMaterialResources);
+			, NAME_None
+			, &CachedMaterialResourcesForCooking
 #else
-		SerializeInlineShaderMaps(
-			NULL,
-			Ar,
-			LoadedMaterialResources,
-			GetFName()
-		);
+			, GetFName()
 #endif
+			);
 	}
 	else
 	{
@@ -4173,11 +4191,11 @@ void UMaterial::DumpDebugInfo(FOutputDevice& OutputDevice) const
 	}
 
 #if WITH_EDITOR
-	for (auto& It : CachedMaterialResourcesForCooking)
+	for (const TPair<const ITargetPlatform*, TArray<FMaterialResourceForCooking>>& It : CachedMaterialResourcesForCooking)
 	{
-		for (FMaterialResource* Resource : It.Value)
+		for (const FMaterialResourceForCooking& ResourceForCooking : It.Value)
 		{
-			Resource->DumpDebugInfo(OutputDevice);
+			ResourceForCooking.Resource->DumpDebugInfo(OutputDevice);
 		}
 	}
 #endif
@@ -4197,14 +4215,14 @@ void UMaterial::SaveShaderStableKeysInner(const class ITargetPlatform* TP, const
 #if WITH_EDITOR
 	FStableShaderKeyAndValue SaveKeyVal(InSaveKeyVal);
 	SaveKeyVal.MaterialDomain = FName(*MaterialDomainString(MaterialDomain));
-	TArray<FMaterialResource*>* MatRes = CachedMaterialResourcesForCooking.Find(TP);
+	TArray<FMaterialResourceForCooking>* MatRes = CachedMaterialResourcesForCooking.Find(TP);
 	if (MatRes)
 	{
-		for (FMaterialResource* Mat : *MatRes)
+		for (const FMaterialResourceForCooking& Mat : *MatRes)
 		{
-			if (Mat)
+			if (Mat.Resource)
 			{
-				Mat->SaveShaderStableKeys(EShaderPlatform::SP_NumPlatforms, SaveKeyVal);
+				Mat.Resource->SaveShaderStableKeys(EShaderPlatform::SP_NumPlatforms, SaveKeyVal);
 			}
 		}
 	}
@@ -4309,20 +4327,15 @@ void UMaterial::PropagateDataToMaterialProxy()
 void UMaterial::BeginCacheForCookedPlatformData( const ITargetPlatform *TargetPlatform )
 {
 	LLM_SCOPE(ELLMTag::Materials);
-	TArray<FName> DesiredShaderFormats;
-	TargetPlatform->GetAllTargetedShaderFormats(DesiredShaderFormats);
-
-	GetCmdLineFilterShaderFormats(DesiredShaderFormats);
-
-	TArray<FMaterialResource*> *CachedMaterialResourcesForPlatform = CachedMaterialResourcesForCooking.Find( TargetPlatform );
+	TArray<FMaterialResourceForCooking>* CachedMaterialResourcesForPlatform = CachedMaterialResourcesForCooking.Find(TargetPlatform);
 
 	if (CachedMaterialResourcesForPlatform == nullptr)
 	{
-		CachedMaterialResourcesForCooking.Add( TargetPlatform );
-		CachedMaterialResourcesForPlatform = CachedMaterialResourcesForCooking.Find( TargetPlatform );
-
-		check(CachedMaterialResourcesForPlatform != nullptr);
+		TArray<FName> DesiredShaderFormats;
+		TargetPlatform->GetAllTargetedShaderFormats(DesiredShaderFormats);
+		GetCmdLineFilterShaderFormats(DesiredShaderFormats);
 
+		CachedMaterialResourcesForPlatform = &CachedMaterialResourcesForCooking.FindOrAdd(TargetPlatform);
 		if (DesiredShaderFormats.Num())
 		{
 			// Cache for all the shader formats that the cooking target requires
@@ -4340,13 +4353,15 @@ void UMaterial::BeginCacheForCookedPlatformData( const ITargetPlatform *TargetPl
 bool UMaterial::IsCachedCookedPlatformDataLoaded( const ITargetPlatform* TargetPlatform ) 
 {
 	LLM_SCOPE(ELLMTag::Materials);
-	const TArray<FMaterialResource*>* CachedMaterialResourcesForPlatform = CachedMaterialResourcesForCooking.Find( TargetPlatform );
+	const TArray<FMaterialResourceForCooking>* CachedMaterialResourcesForPlatform =
+		CachedMaterialResourcesForCooking.Find(TargetPlatform);
 
-	if (CachedMaterialResourcesForPlatform != nullptr) // this should always succeed if BeginCacheForCookedPlatformData is called first
+	// this should always succeed if BeginCacheForCookedPlatformData is called first
+	if (CachedMaterialResourcesForPlatform != nullptr)
 	{
-		for ( const auto& MaterialResource : *CachedMaterialResourcesForPlatform )
+		for (const FMaterialResourceForCooking& MaterialResource : *CachedMaterialResourcesForPlatform)
 		{
-			if ( MaterialResource->IsCompilationFinished() == false )
+			if (MaterialResource.Resource->IsCompilationFinished() == false)
 			{
 				return false;
 			}
@@ -4358,7 +4373,6 @@ bool UMaterial::IsCachedCookedPlatformDataLoaded( const ITargetPlatform* TargetP
 
 void UMaterial::ClearCachedCookedPlatformData( const ITargetPlatform *TargetPlatform )
 {
-#if WITH_EDITOR
 	if (GIsBuildMachine)
 	{
 		// Dump debug info for the DefaultMaterial.
@@ -4380,24 +4394,33 @@ void UMaterial::ClearCachedCookedPlatformData( const ITargetPlatform *TargetPlat
 			}
 		}
 	}
-#endif
 
-	TArray<FMaterialResource*>* CachedMaterialResourcesForPlatform = CachedMaterialResourcesForCooking.Find( TargetPlatform );
-	if ( CachedMaterialResourcesForPlatform != nullptr)
+	TArray<TRefCountPtr<FMaterialResource>> MaterialsToDelete;
 	{
-		FMaterial::DeferredDeleteArray(*CachedMaterialResourcesForPlatform);
+		TArray<FMaterialResourceForCooking> CachedMaterialResourcesForPlatform;
+		CachedMaterialResourcesForCooking.RemoveAndCopyValue(TargetPlatform, CachedMaterialResourcesForPlatform);
+		MaterialsToDelete.Reserve(CachedMaterialResourcesForPlatform.Num());
+		for (FMaterialResourceForCooking& MaterialToDelete : CachedMaterialResourcesForPlatform)
+		{
+			MaterialsToDelete.Add(MoveTemp(MaterialToDelete.Resource));
+		}
 	}
-	CachedMaterialResourcesForCooking.Remove( TargetPlatform );
+	FMaterial::DeferredDeleteArray(MaterialsToDelete);
 }
 
 void UMaterial::ClearAllCachedCookedPlatformData()
 {
-	for ( auto& It : CachedMaterialResourcesForCooking )
+	TArray<TRefCountPtr<FMaterialResource>> MaterialsToDelete;
+	for (TPair<const ITargetPlatform*, TArray<FMaterialResourceForCooking>>& It : CachedMaterialResourcesForCooking)
 	{
-		TArray<FMaterialResource*>& CachedMaterialResourcesForPlatform = It.Value;
-		FMaterial::DeferredDeleteArray(CachedMaterialResourcesForPlatform);
+		MaterialsToDelete.Reserve(MaterialsToDelete.Num() + It.Value.Num());
+		for (FMaterialResourceForCooking& MaterialToDelete : It.Value)
+		{
+			MaterialsToDelete.Add(MoveTemp(MaterialToDelete.Resource));
+		}
 	}
 	CachedMaterialResourcesForCooking.Empty();
+	FMaterial::DeferredDeleteArray(MaterialsToDelete);
 }
 #endif // WITH_EDITOR
 
@@ -5364,14 +5387,13 @@ void UMaterial::AddReferencedObjects(UObject* InThis, FReferenceCollector& Colle
 	}
 
 #if WITH_EDITOR
-	for (auto& It : This->CachedMaterialResourcesForCooking)
+	for (TPair<const ITargetPlatform*, TArray<FMaterialResourceForCooking>>& It : This->CachedMaterialResourcesForCooking)
 	{
-		TArray<FMaterialResource*>& CachedMaterialResourcesForPlatform = It.Value;
-		for (FMaterialResource* CurrentResource : CachedMaterialResourcesForPlatform)
+		for (FMaterialResourceForCooking& CurrentResource : It.Value)
 		{
-			if (CurrentResource)
+			if (CurrentResource.Resource)
 			{
-				CurrentResource->AddReferencedObjects(Collector);
+				CurrentResource.Resource->AddReferencedObjects(Collector);
 			}
 		}
 	}
@@ -5486,7 +5508,7 @@ void UMaterial::CompileMaterialsForRemoteRecompile(
 	TRACE_CPUPROFILER_EVENT_SCOPE(UMaterial::CompileMaterialsForRemoteRecompile);
 
 	// Build a map from UMaterial / UMaterialInstance to the resources which are being compiled
-	TMap<FString, TArray<FMaterialResource*> > CompilingResources;
+	TMap<FString, TArray<TRefCountPtr<FMaterialResource>> > CompilingResources;
 
 	// compile the requested materials
 	for (int32 Index = 0; Index < MaterialsToCompile.Num(); Index++)
@@ -5496,15 +5518,27 @@ void UMaterial::CompileMaterialsForRemoteRecompile(
 		UMaterialInstance* MaterialInstance = Cast<UMaterialInstance>(Material);
 		UMaterial* BaseMaterial = Cast<UMaterial>(Material);
 
+		TArray<FMaterialResourceForCooking> ResourcesForCooking;
 		if (MaterialInstance && MaterialInstance->bHasStaticPermutationResource)
 		{
-			TArray<FMaterialResource*>& ResourceArray = CompilingResources.Add(Material->GetPathName(), TArray<FMaterialResource*>());
-			MaterialInstance->CacheResourceShadersForCooking(ShaderPlatform, ResourceArray, EMaterialShaderPrecompileMode::Default, TargetPlatform, true /* Blocking */);
+			MaterialInstance->CacheResourceShadersForCooking(ShaderPlatform, ResourcesForCooking,
+				EMaterialShaderPrecompileMode::Default, TargetPlatform, true /* Blocking */);
 		}
 		else if (BaseMaterial)
 		{
-			TArray<FMaterialResource*>& ResourceArray = CompilingResources.Add(Material->GetPathName(), TArray<FMaterialResource*>());
-			BaseMaterial->CacheResourceShadersForCooking(ShaderPlatform, ResourceArray, TargetPlatform, true /* Blocking */);
+			BaseMaterial->CacheResourceShadersForCooking(ShaderPlatform, ResourcesForCooking,
+				TargetPlatform, true /* Blocking */);
+		}
+
+		if (!ResourcesForCooking.IsEmpty())
+		{
+			TArray<TRefCountPtr<FMaterialResource>>& ResourceArray = CompilingResources.Add(Material->GetPathName(), 
+				TArray<TRefCountPtr<FMaterialResource>>());
+			ResourceArray.Reserve(ResourcesForCooking.Num());
+			for (FMaterialResourceForCooking& ResourceForCooking : ResourcesForCooking)
+			{
+				ResourceArray.Add(MoveTemp(ResourceForCooking.Resource));
+			}
 		}
 	}
 
@@ -5515,16 +5549,15 @@ void UMaterial::CompileMaterialsForRemoteRecompile(
 	// This is code that should be run on the cooker so shouldn't be a big deal.
 	FlushRenderingCommands();
 
-	for(TMap<FString, TArray<FMaterialResource*> >::TIterator It(CompilingResources); It; ++It)
+	for(TMap<FString, TArray<TRefCountPtr<FMaterialResource>> >::TIterator It(CompilingResources); It; ++It)
 	{
-		TArray<FMaterialResource*>& ResourceArray = It.Value();
-		TArray<TRefCountPtr<FMaterialShaderMap> >& OutShaderMapArray = OutShaderMaps.Add(It.Key(), TArray<TRefCountPtr<FMaterialShaderMap> >());
+		TArray<TRefCountPtr<FMaterialResource>>& ResourceArray = It.Value();
+		TArray<TRefCountPtr<FMaterialShaderMap>>& OutShaderMapArray = OutShaderMaps.Add(It.Key(), TArray<TRefCountPtr<FMaterialShaderMap>>());
 
 		for (int32 Index = 0; Index < ResourceArray.Num(); Index++)
 		{
-			FMaterialResource* CurrentResource = ResourceArray[Index];
+			TRefCountPtr<FMaterialResource> CurrentResource = MoveTemp(ResourceArray[Index]);
 			OutShaderMapArray.Add(CurrentResource->GetGameThreadShaderMap());
-			delete CurrentResource;
 		}
 	}
 }
diff --git a/Engine/Source/Runtime/Engine/Private/Materials/MaterialInstance.cpp b/Engine/Source/Runtime/Engine/Private/Materials/MaterialInstance.cpp
--- a/Engine/Source/Runtime/Engine/Private/Materials/MaterialInstance.cpp
+++ b/Engine/Source/Runtime/Engine/Private/Materials/MaterialInstance.cpp
@@ -2336,9 +2336,10 @@ void UMaterialInstance::CacheResourceShadersForRendering(EMaterialShaderPrecompi
 	FMaterial::DeferredDeleteArray(ResourcesToFree);
 }
 
+#if WITH_EDITOR
 void UMaterialInstance::CacheResourceShadersForCooking(
 	EShaderPlatform ShaderPlatform,
-	TArray<FMaterialResource*>& OutCachedMaterialResources,
+	TArray<FMaterialResourceForCooking>& OutCachedMaterialResources,
 	EMaterialShaderPrecompileMode PrecompileMode,
 	const ITargetPlatform* TargetPlatform,
 	bool bBlocking
@@ -2356,7 +2357,9 @@ void UMaterialInstance::CacheResourceShadersForCooking(
 
 		ERHIFeatureLevel::Type TargetFeatureLevel = GetMaxSupportedFeatureLevel(ShaderPlatform);
 
-		TArray<FMaterialResource*> NewResourcesToCache;	// only new resources need to have CacheShaders() called on them, whereas OutCachedMaterialResources may already contain resources for another shader platform
+		// only new resources need to have CacheShaders() called on them, whereas OutCachedMaterialResources
+		// may already contain resources for another shader platform
+		TArray<FMaterialResource*> NewResourcesToCache;
 		for (int32 QualityLevelIndex = 0; QualityLevelIndex < EMaterialQualityLevel::Num; QualityLevelIndex++)
 		{
 			// Cache all quality levels actually used
@@ -2384,25 +2387,26 @@ void UMaterialInstance::CacheResourceShadersForCooking(
 			NewResourcesToCache.Add(NewResource);
 		}
 
-	#if WITH_EDITOR
 		// The editor needs to block if the caching call comes from cook on the fly, where the polling mechanisms are not active.
 		// This is important so that the jobs finish and the CacheShadersCompletion() callback is triggered via FinishCacheShaders()!
 		if (bBlocking)
-	#endif
 		{
 			CacheShadersForResources(ShaderPlatform, NewResourcesToCache, PrecompileMode, TargetPlatform);
 		}
-	#if WITH_EDITOR
 		else
 		{
 			// For cooking, we can call the begin function and it will be completed as part of the polling mechanism.
 			BeginCacheShadersForResources(ShaderPlatform, NewResourcesToCache, PrecompileMode, TargetPlatform);
 		}
-	#endif
 
-		OutCachedMaterialResources.Append(NewResourcesToCache);
+		OutCachedMaterialResources.Reserve(NewResourcesToCache.Num());
+		for (FMaterialResource* NewResource : NewResourcesToCache)
+		{
+			OutCachedMaterialResources.Add({ NewResource , ShaderPlatform });
+		}
 	}
 }
+#endif // WITH_EDITOR
 
 namespace MaterialInstanceImpl
 {
@@ -2692,16 +2696,11 @@ void TrimToOverriddenOnly(TArray<ParameterType>& Parameters)
 void UMaterialInstance::BeginCacheForCookedPlatformData( const ITargetPlatform *TargetPlatform )
 {
 	LLM_SCOPE(ELLMTag::Materials);
-	TArray<FMaterialResource*> *CachedMaterialResourcesForPlatform = CachedMaterialResourcesForCooking.Find( TargetPlatform );
+	TArray<FMaterialResourceForCooking> *CachedMaterialResourcesForPlatform = CachedMaterialResourcesForCooking.Find( TargetPlatform );
 
 	if ( CachedMaterialResourcesForPlatform == NULL )
 	{
-		check( CachedMaterialResourcesForPlatform == NULL );
-
-		CachedMaterialResourcesForCooking.Add( TargetPlatform );
-		CachedMaterialResourcesForPlatform = CachedMaterialResourcesForCooking.Find( TargetPlatform );
-
-		check( CachedMaterialResourcesForPlatform != NULL );
+		CachedMaterialResourcesForPlatform = &CachedMaterialResourcesForCooking.FindOrAdd( TargetPlatform );
 
 		TArray<FName> DesiredShaderFormats;
 		TargetPlatform->GetAllTargetedShaderFormats(DesiredShaderFormats);
@@ -2713,7 +2712,8 @@ void UMaterialInstance::BeginCacheForCookedPlatformData( const ITargetPlatform *
 		{
 			const EShaderPlatform TargetShaderPlatform = ShaderFormatToLegacyShaderPlatform(DesiredShaderFormats[FormatIndex]);
 
-			CacheResourceShadersForCooking(TargetShaderPlatform, *CachedMaterialResourcesForPlatform, EMaterialShaderPrecompileMode::Background, TargetPlatform);
+			CacheResourceShadersForCooking(TargetShaderPlatform, *CachedMaterialResourcesForPlatform,
+				EMaterialShaderPrecompileMode::Background, TargetPlatform);
 		}
 	}
 }
@@ -2721,12 +2721,13 @@ void UMaterialInstance::BeginCacheForCookedPlatformData( const ITargetPlatform *
 bool UMaterialInstance::IsCachedCookedPlatformDataLoaded( const ITargetPlatform* TargetPlatform ) 
 {
 	LLM_SCOPE(ELLMTag::Materials);
-	const TArray<FMaterialResource*> *CachedMaterialResourcesForPlatform = CachedMaterialResourcesForCooking.Find( TargetPlatform );
+	const TArray<FMaterialResourceForCooking>* CachedMaterialResourcesForPlatform =
+		CachedMaterialResourcesForCooking.Find( TargetPlatform );
 	if ( CachedMaterialResourcesForPlatform != NULL )
 	{
-		for ( const auto& MaterialResource : *CachedMaterialResourcesForPlatform )
+		for (const FMaterialResourceForCooking& MaterialResource : *CachedMaterialResourcesForPlatform)
 		{
-			if ( MaterialResource->IsCompilationFinished() == false )
+			if (MaterialResource.Resource->IsCompilationFinished() == false)
 				return false;
 		}
 
@@ -2736,26 +2737,35 @@ bool UMaterialInstance::IsCachedCookedPlatformDataLoaded( const ITargetPlatform*
 }
 
 
-void UMaterialInstance::ClearCachedCookedPlatformData( const ITargetPlatform *TargetPlatform )
+void UMaterialInstance::ClearCachedCookedPlatformData(const ITargetPlatform* TargetPlatform)
 {
-	TArray<FMaterialResource*>* CachedMaterialResourcesForPlatform = CachedMaterialResourcesForCooking.Find( TargetPlatform );
-	if ( CachedMaterialResourcesForPlatform != nullptr )
+	TArray<TRefCountPtr<FMaterialResource>> MaterialsToDelete;
 	{
-		FMaterial::DeferredDeleteArray(*CachedMaterialResourcesForPlatform);
+		TArray<FMaterialResourceForCooking> CachedMaterialResourcesForPlatform;
+		CachedMaterialResourcesForCooking.RemoveAndCopyValue(TargetPlatform, CachedMaterialResourcesForPlatform);
+		MaterialsToDelete.Reserve(CachedMaterialResourcesForPlatform.Num());
+		for (FMaterialResourceForCooking& MaterialToDelete : CachedMaterialResourcesForPlatform)
+		{
+			MaterialsToDelete.Add(MoveTemp(MaterialToDelete.Resource));
+		}
 	}
-	CachedMaterialResourcesForCooking.Remove( TargetPlatform );
+	FMaterial::DeferredDeleteArray(MaterialsToDelete);
 }
 
 
 void UMaterialInstance::ClearAllCachedCookedPlatformData()
 {
-	for ( auto& It : CachedMaterialResourcesForCooking )
+	TArray<TRefCountPtr<FMaterialResource>> MaterialsToDelete;
+	for (TPair<const ITargetPlatform*, TArray<FMaterialResourceForCooking>>& It : CachedMaterialResourcesForCooking)
 	{
-		TArray<FMaterialResource*>& CachedMaterialResourcesForPlatform = It.Value;
-		FMaterial::DeferredDeleteArray(CachedMaterialResourcesForPlatform);
+		MaterialsToDelete.Reserve(MaterialsToDelete.Num() + It.Value.Num());
+		for (FMaterialResourceForCooking& MaterialToDelete : It.Value)
+		{
+			MaterialsToDelete.Add(MoveTemp(MaterialToDelete.Resource));
+		}
 	}
-
 	CachedMaterialResourcesForCooking.Empty();
+	FMaterial::DeferredDeleteArray(MaterialsToDelete);
 }
 
 #endif
@@ -2892,16 +2902,17 @@ void UMaterialInstance::Serialize(FArchive& Ar)
 				StaticParametersRuntime = MoveTemp(StaticParameters_DEPRECATED.GetRuntime());
 				GetEditorOnlyData()->StaticParameters = MoveTemp(StaticParameters_DEPRECATED.EditorOnly);
 			}
+#endif
 
-			SerializeInlineShaderMaps(&CachedMaterialResourcesForCooking, Ar, LoadedMaterialResources);
+			UE::MaterialInterface::Private::SerializeInlineShaderMaps(
+				Ar, LoadedMaterialResources
+#if WITH_EDITOR
+				, NAME_None
+				, &CachedMaterialResourcesForCooking
 #else
-			SerializeInlineShaderMaps(
-				NULL,
-				Ar,
-				LoadedMaterialResources,
-				GetFName()
-			);
+				, GetFName()
 #endif
+			);
 		}
 #if WITH_EDITOR
 		else
@@ -4829,11 +4840,11 @@ void UMaterialInstance::DumpDebugInfo(FOutputDevice& OutputDevice) const
 			}
 
 #if WITH_EDITOR
-			for (auto& It : CachedMaterialResourcesForCooking)
+			for (const TPair<const ITargetPlatform*, TArray<FMaterialResourceForCooking>>& It : CachedMaterialResourcesForCooking)
 			{
-				for (FMaterialResource* CurrentResource : It.Value)
+				for (const FMaterialResourceForCooking& CurrentResource : It.Value)
 				{
-					CurrentResource->DumpDebugInfo(OutputDevice);
+					CurrentResource.Resource->DumpDebugInfo(OutputDevice);
 				}
 			}
 #endif // WITH_EDITOR
@@ -4865,14 +4876,14 @@ void UMaterialInstance::SaveShaderStableKeysInner(const class ITargetPlatform* T
 	if (bHasStaticPermutationResource)
 	{
 		FStableShaderKeyAndValue SaveKeyVal(InSaveKeyVal);
-		TArray<FMaterialResource*>* MatRes = CachedMaterialResourcesForCooking.Find(TP);
+		TArray<FMaterialResourceForCooking>* MatRes = CachedMaterialResourcesForCooking.Find(TP);
 		if (MatRes)
 		{
-			for (FMaterialResource* Mat : *MatRes)
+			for (FMaterialResourceForCooking& Mat : *MatRes)
 			{
-				if (Mat)
+				if (Mat.Resource)
 				{
-					Mat->SaveShaderStableKeys(EShaderPlatform::SP_NumPlatforms, SaveKeyVal);
+					Mat.Resource->SaveShaderStableKeys(EShaderPlatform::SP_NumPlatforms, SaveKeyVal);
 				}
 			}
 		}
diff --git a/Engine/Source/Runtime/Engine/Private/Materials/MaterialInterface.cpp b/Engine/Source/Runtime/Engine/Private/Materials/MaterialInterface.cpp
--- a/Engine/Source/Runtime/Engine/Private/Materials/MaterialInterface.cpp
+++ b/Engine/Source/Runtime/Engine/Private/Materials/MaterialInterface.cpp
@@ -135,7 +135,7 @@ void FMaterialRelevance::SetPrimitiveViewRelevance(FPrimitiveViewRelevance& OutV
 
 #if WITH_EDITORONLY_DATA
 
-namespace MaterialInterface
+namespace UE::MaterialInterface
 {
 	FString GetEditorOnlyDataName(const TCHAR* InMaterialName)
 	{
@@ -735,7 +735,7 @@ void UMaterialInterface::PostDuplicate(bool bDuplicateForPIE)
 #if WITH_EDITORONLY_DATA
 	// Duplication will create a new editor only data but this MI already links a populated data. This makes sure that the
 	// duplicate instance EOD has the correct name (necessary for running the editor on cooked data).
-	FString EditorOnlyDataName = MaterialInterface::GetEditorOnlyDataName(*GetName());
+	FString EditorOnlyDataName = UE::MaterialInterface::GetEditorOnlyDataName(*GetName());
 	DisposeInvalidEditorOnlyData(this, EditorOnlyDataName);
 	if (GetEditorOnlyData())
 	{
@@ -1691,7 +1691,7 @@ void UMaterialInterface::PreSave(FObjectPreSaveContext ObjectSaveContext)
 	// Make sure the EditorOnlyData is named correctly. Some MI assets have a differently named editor only data that can cause problems in
 	// the editor running on cooked data.
 	UMaterialInterfaceEditorOnlyData* EditorOnly = GetEditorOnlyData();
-	FString EditorOnlyDataName = MaterialInterface::GetEditorOnlyDataName(*GetName());
+	FString EditorOnlyDataName = UE::MaterialInterface::GetEditorOnlyDataName(*GetName());
 	if (!ObjectSaveContext.IsProceduralSave() && EditorOnly && EditorOnly->GetName() != EditorOnlyDataName)
 	{
 		UE_LOG(LogMaterial, Display, TEXT("MaterialInterface %s has a incorrectly name EditorOnlyData '%s'. This may cause issues when running the editor on cooked data. Trying to rename it to the correct name '%s'."), *GetName(), *EditorOnly->GetName(), *EditorOnlyDataName);
@@ -1815,7 +1815,7 @@ UMaterialInterfaceEditorOnlyData* UMaterialInterface::CreateEditorOnlyData()
 	check(EditorOnlyClass);
 	check(EditorOnlyClass->HasAllClassFlags(CLASS_Optional));
 
-	const FString EditorOnlyName = MaterialInterface::GetEditorOnlyDataName(*GetName());
+	const FString EditorOnlyName = UE::MaterialInterface::GetEditorOnlyDataName(*GetName());
 	const EObjectFlags EditorOnlyFlags = GetMaskedFlags(RF_PropagateToSubObjects);
 	return NewObject<UMaterialInterfaceEditorOnlyData>(this, EditorOnlyClass, *EditorOnlyName, EditorOnlyFlags);
 }
@@ -1828,7 +1828,7 @@ bool UMaterialInterface::Rename(const TCHAR* NewName, UObject* NewOuter, ERename
 	// if we have EditorOnlyData, also rename it if we are changing the material's name
 	if (bRenamed && NewName && EditorOnlyData)
 	{
-		FString EditorOnlyDataName = MaterialInterface::GetEditorOnlyDataName(NewName);
+		FString EditorOnlyDataName = UE::MaterialInterface::GetEditorOnlyDataName(NewName);
 		DisposeInvalidEditorOnlyData(this, EditorOnlyDataName);
 		bRenamed = EditorOnlyData->Rename(*EditorOnlyDataName, nullptr, Flags);
 	}
diff --git a/Engine/Source/Runtime/Engine/Private/Materials/MaterialShared.cpp b/Engine/Source/Runtime/Engine/Private/Materials/MaterialShared.cpp
--- a/Engine/Source/Runtime/Engine/Private/Materials/MaterialShared.cpp
+++ b/Engine/Source/Runtime/Engine/Private/Materials/MaterialShared.cpp
@@ -5378,8 +5378,13 @@ bool FMemoryImageMaterialParameterInfo::RemapLayerIndex(TArrayView<const int32>
 
 FMaterialShaderParameters::FMaterialShaderParameters(const FMaterial* InMaterial)
 {
-	// Make sure to zero-initialize so we get consistent hashes
+	// zero-initialize is required even when InMaterial!=null so that all bytes in sizeof(*this) are initialized
+	// and we get consistent hashes.
 	FMemory::Memzero(*this);
+	if (!InMaterial)
+	{
+		return;
+	}
 
 	MaterialDomain = InMaterial->GetMaterialDomain();
 	ShadingModels = InMaterial->GetShadingModels();
diff --git a/Engine/Source/Runtime/Engine/Public/MaterialShared.h b/Engine/Source/Runtime/Engine/Public/MaterialShared.h
--- a/Engine/Source/Runtime/Engine/Public/MaterialShared.h
+++ b/Engine/Source/Runtime/Engine/Public/MaterialShared.h
@@ -973,7 +973,7 @@ struct FMaterialShaderParameters
 		};
 	};
 
-	FMaterialShaderParameters(const FMaterial* InMaterial);
+	FMaterialShaderParameters(const FMaterial* InMaterial=nullptr);
 };
 
 /** Contains all the information needed to uniquely identify a FMaterialShaderMap. */
@@ -3107,6 +3107,15 @@ ENGINE_API uint8 GetRayTracingMaskFromMaterial(const EBlendMode BlendMode);
 
 //
 
+#if WITH_EDITOR
+/** An FMaterialResource and other information needed for cooking it. */
+struct FMaterialResourceForCooking
+{
+	TRefCountPtr<FMaterialResource> Resource;
+	EShaderPlatform Platform;
+};
+#endif
+
 inline bool ShouldIncludeMaterialInDefaultOpaquePass(const FMaterial& Material)
 {
 	return !Material.IsSky()
