Small optim AWorldDataLayers::GetDataLayerInstance early out if Data Layer Instance name is None.

diff --git a/Engine/Source/Runtime/Engine/Private/WorldPartition/DataLayer/WorldDataLayers.cpp b/Engine/Source/Runtime/Engine/Private/WorldPartition/DataLayer/WorldDataLayers.cpp
--- a/Engine/Source/Runtime/Engine/Private/WorldPartition/DataLayer/WorldDataLayers.cpp
+++ b/Engine/Source/Runtime/Engine/Private/WorldPartition/DataLayer/WorldDataLayers.cpp
@@ -1021,6 +1021,11 @@ bool AWorldDataLayers::ContainsDataLayer(const UDataLayerInstance* InDataLayerIn
 
 const UDataLayerInstance* AWorldDataLayers::GetDataLayerInstance(const FName& InDataLayerInstanceName) const
 {
+	if (InDataLayerInstanceName.IsNone())
+	{
+		return nullptr;
+	}
+
 #if WITH_EDITOR	
 	{
 		const UDataLayerInstance* FoundDataLayerInstance = nullptr;
