Iris - Reduce log spam for the case where  FlushNetDormancy is called too early.

#rnx

diff --git a/Engine/Source/Runtime/Engine/Private/Net/Iris/ReplicationSystem/ReplicationSystemUtil.cpp b/Engine/Source/Runtime/Engine/Private/Net/Iris/ReplicationSystem/ReplicationSystemUtil.cpp
--- a/Engine/Source/Runtime/Engine/Private/Net/Iris/ReplicationSystem/ReplicationSystemUtil.cpp
+++ b/Engine/Source/Runtime/Engine/Private/Net/Iris/ReplicationSystem/ReplicationSystemUtil.cpp
@@ -467,6 +467,12 @@ void FReplicationSystemUtil::FlushNetDormancy(UReplicationSystem* ReplicationSys
 		return;
 	}
 
+	if (!Actor->IsActorInitialized())
+	{
+		UE_LOG(LogIris, Verbose, TEXT("FReplicationSystemUtil::FlushNetDormancy called on %s that isn't fully initialized yet. Ingoring."), ToCStr(GetFullNameSafe(Actor)));
+		return;
+	}
+
 	FNetHandle ActorHandle = GetNetHandle(Actor);
 	if (ActorHandle.IsValid())
 	{
