SVT streamer can now accept requests for fractional mip levels. Smoothly changing the requested fractional mip level can currently lead to certain performance issues with DDC because additional small requests are made to stream in a few missing tiles at a time. To mitigate this, there is now a new cvar to discretize the computed streaming size.
Also fixed a crash when unregistering a SVT from the streamer while it is currently playing back. Fixed by properly removing all FStreamingInstances belonging to the SVT to be removed.

#rnx

diff --git a/Engine/Source/Runtime/Engine/Classes/SparseVolumeTexture/SparseVolumeTexture.h b/Engine/Source/Runtime/Engine/Classes/SparseVolumeTexture/SparseVolumeTexture.h
--- a/Engine/Source/Runtime/Engine/Classes/SparseVolumeTexture/SparseVolumeTexture.h
+++ b/Engine/Source/Runtime/Engine/Classes/SparseVolumeTexture/SparseVolumeTexture.h
@@ -372,7 +372,7 @@ public:
 	// If bBlocking is true, DDC streaming requests will block on completion, guaranteeing that the requested frame will have been streamed in after the next streaming system update.
 	// If streaming cooked data from disk, the highest priority will be used, but no guarantee is given.
 	// if bHasValidFrameRate is true, the FrameRate argument will be use to predict the required streaming IO bandwidth.
-	static ENGINE_API USparseVolumeTextureFrame* GetFrameAndIssueStreamingRequest(USparseVolumeTexture* SparseVolumeTexture, uint32 StreamingInstanceKey, float FrameRate, float FrameIndex, int32 MipLevel, bool bBlocking, bool bHasValidFrameRate);
+	static ENGINE_API USparseVolumeTextureFrame* GetFrameAndIssueStreamingRequest(USparseVolumeTexture* SparseVolumeTexture, uint32 StreamingInstanceKey, float FrameRate, float FrameIndex, float MipLevel, bool bBlocking, bool bHasValidFrameRate);
 
 	ENGINE_API bool Initialize(USparseVolumeTexture* InOwner, int32 InFrameIndex, const FTransform& InFrameTransform, UE::SVT::FTextureData& UncookedFrame);
 	int32 GetFrameIndex() const { return FrameIndex; }
diff --git a/Engine/Source/Runtime/Engine/Private/SparseVolumeTexture/SparseVolumeTexture.cpp b/Engine/Source/Runtime/Engine/Private/SparseVolumeTexture/SparseVolumeTexture.cpp
--- a/Engine/Source/Runtime/Engine/Private/SparseVolumeTexture/SparseVolumeTexture.cpp
+++ b/Engine/Source/Runtime/Engine/Private/SparseVolumeTexture/SparseVolumeTexture.cpp
@@ -939,7 +939,7 @@ USparseVolumeTextureFrame::USparseVolumeTextureFrame(const FObjectInitializer& O
 {
 }
 
-USparseVolumeTextureFrame* USparseVolumeTextureFrame::GetFrameAndIssueStreamingRequest(USparseVolumeTexture* SparseVolumeTexture, uint32 StreamingInstanceKey, float FrameRate, float FrameIndex, int32 MipLevel, bool bBlocking, bool bHasValidFrameRate)
+USparseVolumeTextureFrame* USparseVolumeTextureFrame::GetFrameAndIssueStreamingRequest(USparseVolumeTexture* SparseVolumeTexture, uint32 StreamingInstanceKey, float FrameRate, float FrameIndex, float MipLevel, bool bBlocking, bool bHasValidFrameRate)
 {
 	if (UStreamableSparseVolumeTexture* StreamableSVT = Cast<UStreamableSparseVolumeTexture>(SparseVolumeTexture))
 	{
diff --git a/Engine/Source/Runtime/Engine/Private/SparseVolumeTexture/SparseVolumeTextureStreamingInstance.cpp b/Engine/Source/Runtime/Engine/Private/SparseVolumeTexture/SparseVolumeTextureStreamingInstance.cpp
--- a/Engine/Source/Runtime/Engine/Private/SparseVolumeTexture/SparseVolumeTextureStreamingInstance.cpp
+++ b/Engine/Source/Runtime/Engine/Private/SparseVolumeTexture/SparseVolumeTextureStreamingInstance.cpp
@@ -20,6 +20,14 @@ static FAutoConsoleVariableRef CVarSVTStreamingForceEstimateFrameRate(
 	ECVF_RenderThreadSafe
 );
 
+static int32 GSVTStreamingRequestSizeGranularity = 128 * 1024;
+static FAutoConsoleVariableRef CVarSVTStreamingRequestSizeGranularity(
+	TEXT("r.SparseVolumeTexture.Streaming.RequestSizeGranularity"),
+	GSVTStreamingRequestSizeGranularity,
+	TEXT("Rounds up calculated streaming request sizes to a multiple of this value (in bytes). This avoids performance issues where the system might issue multiple very small requests."),
+	ECVF_RenderThreadSafe
+);
+
 namespace UE
 {
 namespace SVT
@@ -144,9 +152,9 @@ int64 FStreamingInstance::GetRequestedBandwidth(bool bZeroIfBlocking) const
 	{
 		return 0;
 	}
-	const int32 MipLevelIndex = FMath::Clamp(FMath::FloorToInt32(LowestRequestedMip), 0, MipLevelStreamingSizes.Num() - 1);
-	check(MipLevelStreamingSizes.IsValidIndex(MipLevelIndex));
-	const int64 RequestedBandwidth = FMath::CeilToInt64(MipLevelStreamingSizes[MipLevelIndex] * FMath::Max(1.0f, GetEstimatedFrameRate()));
+	const int64 StreamingSize = GetStreamingSize(LowestRequestedMip);
+	const int64 RequestedBandwidthRaw = FMath::CeilToInt64(StreamingSize * FMath::Max(1.0f, GetEstimatedFrameRate()));
+	const int64 RequestedBandwidth = ApplyDiscretization(RequestedBandwidthRaw);
 	return RequestedBandwidth;
 }
 
@@ -163,10 +171,9 @@ bool FStreamingInstance::IsFrameInWindow(float FrameIndex) const
 
 float FStreamingInstance::GetPrefetchMipLevel(float RequestedMipLevel, float Percentage) const
 {
-	const int32 RequestedMipLevelIndex = FMath::Clamp(FMath::FloorToInt32(RequestedMipLevel), 0, MipLevelStreamingSizes.Num() - 1);
-	check(MipLevelStreamingSizes.IsValidIndex(RequestedMipLevelIndex));
-	const uint32 RequestedMipLevelStreamingSize = MipLevelStreamingSizes[RequestedMipLevelIndex];
-	const int64 PrefetchStreamingSize = FMath::CeilToInt64(RequestedMipLevelStreamingSize * FMath::Clamp(Percentage, 0.0f, 1.0f));
+	const int64 RequestedMipLevelStreamingSize = GetStreamingSize(RequestedMipLevel);
+	const int64 PrefetchStreamingSizeRaw = FMath::CeilToInt64(RequestedMipLevelStreamingSize * FMath::Clamp(Percentage, 0.0f, 1.0f));
+	const int64 PrefetchStreamingSize = ApplyDiscretization(PrefetchStreamingSizeRaw);
 
 	// Walk from the highest to lowest mip level and find the fractional mip level which corresponds to StreamingMemorySizeOf(RequestedMipLevel) * Percentage.
 	int64 PrevMipStreamingSize = 0;
@@ -227,5 +234,22 @@ float FStreamingInstance::GetWrappedWeightedAverage(float ValueA, float WeightA,
 	return Result;
 }
 
+int64 FStreamingInstance::ApplyDiscretization(int64 Value)
+{
+	const int64 DiscretizationStep = FMath::Max(GSVTStreamingRequestSizeGranularity, 1);
+	Value = FMath::DivideAndRoundUp(Value, DiscretizationStep) * DiscretizationStep;
+	return Value;
+}
+
+int64 FStreamingInstance::GetStreamingSize(float MipLevel) const
+{
+	const int32 LowerMipLevelIndex = FMath::Clamp(FMath::FloorToInt32(MipLevel), 0, MipLevelStreamingSizes.Num() - 1);
+	const int32 UpperMipLevelIndex = FMath::Clamp(FMath::CeilToInt32(MipLevel), 0, MipLevelStreamingSizes.Num() - 1);
+	check(MipLevelStreamingSizes.IsValidIndex(LowerMipLevelIndex));
+	check(MipLevelStreamingSizes.IsValidIndex(UpperMipLevelIndex));
+	const double InterpolatedStreamingSize = FMath::Lerp((double)MipLevelStreamingSizes[LowerMipLevelIndex], (double)MipLevelStreamingSizes[UpperMipLevelIndex], FMath::Frac(MipLevel));
+	return static_cast<int64>(InterpolatedStreamingSize);
+}
+
 }
 }
diff --git a/Engine/Source/Runtime/Engine/Private/SparseVolumeTexture/SparseVolumeTextureStreamingInstance.h b/Engine/Source/Runtime/Engine/Private/SparseVolumeTexture/SparseVolumeTextureStreamingInstance.h
--- a/Engine/Source/Runtime/Engine/Private/SparseVolumeTexture/SparseVolumeTextureStreamingInstance.h
+++ b/Engine/Source/Runtime/Engine/Private/SparseVolumeTexture/SparseVolumeTextureStreamingInstance.h
@@ -75,6 +75,10 @@ private:
 	static float GetShortestWrappedDistance(float To, float From, int32 NumFrames);
 	// Computes a weighted average between ValueA and ValueB, applying wrapping logic to both the calculation and the result.
 	static float GetWrappedWeightedAverage(float ValueA, float WeightA, float ValueB, float WeightB, int32 NumFrames);
+	// Rounds up to the next multiple of r.SparseVolumeTexture.Streaming.RequestSizeGranularity
+	static int64 ApplyDiscretization(int64 Value);
+	// Get the streaming size of the fractional mip level. GetStreamingSize(MipLevel) * FrameRate is the required peak bandwidth in bytes/s.
+	int64 GetStreamingSize(float MipLevel) const;
 };
 
 }
diff --git a/Engine/Source/Runtime/Engine/Private/SparseVolumeTexture/SparseVolumeTextureStreamingManager.cpp b/Engine/Source/Runtime/Engine/Private/SparseVolumeTexture/SparseVolumeTextureStreamingManager.cpp
--- a/Engine/Source/Runtime/Engine/Private/SparseVolumeTexture/SparseVolumeTextureStreamingManager.cpp
+++ b/Engine/Source/Runtime/Engine/Private/SparseVolumeTexture/SparseVolumeTextureStreamingManager.cpp
@@ -233,7 +233,7 @@ void FStreamingManager::Remove_GameThread(UStreamableSparseVolumeTexture* Sparse
 		});
 }
 
-void FStreamingManager::Request_GameThread(UStreamableSparseVolumeTexture* SparseVolumeTexture, uint32 StreamingInstanceKey, float FrameRate, float FrameIndex, int32 MipLevel, EStreamingRequestFlags Flags)
+void FStreamingManager::Request_GameThread(UStreamableSparseVolumeTexture* SparseVolumeTexture, uint32 StreamingInstanceKey, float FrameRate, float FrameIndex, float MipLevel, EStreamingRequestFlags Flags)
 {
 	if (!DoesPlatformSupportSparseVolumeTexture(GMaxRHIShaderPlatform) || !SparseVolumeTexture)
 	{
@@ -263,7 +263,7 @@ void FStreamingManager::Update_GameThread()
 		});
 }
 
-void FStreamingManager::Request(UStreamableSparseVolumeTexture* SparseVolumeTexture, uint32 StreamingInstanceKey, float FrameRate, float FrameIndex, int32 MipLevel, EStreamingRequestFlags Flags)
+void FStreamingManager::Request(UStreamableSparseVolumeTexture* SparseVolumeTexture, uint32 StreamingInstanceKey, float FrameRate, float FrameIndex, float MipLevel, EStreamingRequestFlags Flags)
 {
 	check(IsInRenderingThread());
 	if (!DoesPlatformSupportSparseVolumeTexture(GMaxRHIShaderPlatform) || !SparseVolumeTexture)
@@ -856,6 +856,12 @@ void FStreamingManager::RemoveInternal(UStreamableSparseVolumeTexture* SparseVol
 			}
 		}
 
+		// Make sure ActiveStreamingInstances doesn't have any references to this SVT anymore
+		for (TUniquePtr<FStreamingInstance>& StreamingInstance : SVTInfo->StreamingInstances)
+		{
+			ActiveStreamingInstances.Remove(StreamingInstance.Get());
+		}
+
 		// Release resources
 		for (FFrameInfo& FrameInfo : SVTInfo->PerFrameInfo)
 		{
@@ -881,7 +887,11 @@ void FStreamingManager::AddRequest(const FStreamingRequest& Request)
 		ExistingRequestPayload->MipLevelMask |= Request.Payload.MipLevelMask;
 		const int32 ExistingLowestMip = FMath::CountTrailingZeros(ExistingRequestPayload->MipLevelMask);
 		const int32 IncomingLowestMip = FMath::CountTrailingZeros(Request.Payload.MipLevelMask);
-		if (IncomingLowestMip < ExistingLowestMip)
+		if (IncomingLowestMip == ExistingLowestMip)
+		{
+			ExistingRequestPayload->LowestMipFraction = FMath::Max(ExistingRequestPayload->LowestMipFraction, Request.Payload.LowestMipFraction);
+		}
+		else if (IncomingLowestMip < ExistingLowestMip)
 		{
 			ExistingRequestPayload->LowestMipFraction = Request.Payload.LowestMipFraction;
 		}
diff --git a/Engine/Source/Runtime/Engine/Private/SparseVolumeTexture/SparseVolumeTextureStreamingManager.h b/Engine/Source/Runtime/Engine/Private/SparseVolumeTexture/SparseVolumeTextureStreamingManager.h
--- a/Engine/Source/Runtime/Engine/Private/SparseVolumeTexture/SparseVolumeTextureStreamingManager.h
+++ b/Engine/Source/Runtime/Engine/Private/SparseVolumeTexture/SparseVolumeTextureStreamingManager.h
@@ -101,10 +101,10 @@ public:
 	//~ Begin IStreamingManager Interface.
 	virtual void Add_GameThread(UStreamableSparseVolumeTexture* SparseVolumeTexture) override;
 	virtual void Remove_GameThread(UStreamableSparseVolumeTexture* SparseVolumeTexture) override;
-	virtual void Request_GameThread(UStreamableSparseVolumeTexture* SparseVolumeTexture, uint32 StreamingInstanceKey, float FrameRate, float FrameIndex, int32 MipLevel, EStreamingRequestFlags Flags) override;
+	virtual void Request_GameThread(UStreamableSparseVolumeTexture* SparseVolumeTexture, uint32 StreamingInstanceKey, float FrameRate, float FrameIndex, float MipLevel, EStreamingRequestFlags Flags) override;
 	virtual void Update_GameThread() override;
 	
-	virtual void Request(UStreamableSparseVolumeTexture* SparseVolumeTexture, uint32 StreamingInstanceKey, float FrameRate, float FrameIndex, int32 MipLevel, EStreamingRequestFlags Flags) override;
+	virtual void Request(UStreamableSparseVolumeTexture* SparseVolumeTexture, uint32 StreamingInstanceKey, float FrameRate, float FrameIndex, float MipLevel, EStreamingRequestFlags Flags) override;
 	virtual void BeginAsyncUpdate(FRDGBuilder& GraphBuilder, bool bUseAsyncThread) override;
 	virtual void EndAsyncUpdate(FRDGBuilder& GraphBuilder) override;
 	virtual const FStreamingDebugInfo* GetStreamingDebugInfo(FRDGBuilder& GraphBuilder) const override;
diff --git a/Engine/Source/Runtime/Engine/Public/SparseVolumeTexture/ISparseVolumeTextureStreamingManager.h b/Engine/Source/Runtime/Engine/Public/SparseVolumeTexture/ISparseVolumeTextureStreamingManager.h
--- a/Engine/Source/Runtime/Engine/Public/SparseVolumeTexture/ISparseVolumeTextureStreamingManager.h
+++ b/Engine/Source/Runtime/Engine/Public/SparseVolumeTexture/ISparseVolumeTextureStreamingManager.h
@@ -64,7 +64,7 @@ public:
 	// FrameRate is an optional argument which helps to more accurately predict the required bandwidth when using non-blocking requests. EStreamingRequestFlags::HasFrameRate must be set for this.
 	// FrameIndex is of float type so that the fractional part can be used to better track the playback speed/direction.
 	// This function automatically also requests all higher mip levels and adds prefetch requests for upcoming frames.
-	virtual void Request_GameThread(UStreamableSparseVolumeTexture* SparseVolumeTexture, uint32 StreamingInstanceKey, float FrameRate, float FrameIndex, int32 MipLevel, EStreamingRequestFlags Flags) = 0;
+	virtual void Request_GameThread(UStreamableSparseVolumeTexture* SparseVolumeTexture, uint32 StreamingInstanceKey, float FrameRate, float FrameIndex, float MipLevel, EStreamingRequestFlags Flags) = 0;
 	// Issues a rendering command for updating the streaming manager. This is not normally needed, but is necessary in cases where blocking requests are required
 	// and the SVT is used in a rendering command (which is executed before the streaming manager would normally update). A call to this function is not needed when the SVT
 	// is used in a regular pass called from FDeferredShadingSceneRenderer::Render().
@@ -72,7 +72,7 @@ public:
 	//~ End game thread functions.
 	
 	//~ Begin rendering thread functions.
-	virtual void Request(UStreamableSparseVolumeTexture* SparseVolumeTexture, uint32 StreamingInstanceKey, float FrameRate, float FrameIndex, int32 MipLevel, EStreamingRequestFlags Flags) = 0;
+	virtual void Request(UStreamableSparseVolumeTexture* SparseVolumeTexture, uint32 StreamingInstanceKey, float FrameRate, float FrameIndex, float MipLevel, EStreamingRequestFlags Flags) = 0;
 	// Begins updating the streaming manager. If r.SparseVolumeTexture.Streaming.AsyncThread is 1 and bUseAsyncThread is true, most of the updating work is done in another thread.
 	virtual void BeginAsyncUpdate(FRDGBuilder& GraphBuilder, bool bUseAsyncThread = true) = 0;
 	// Waits for the job started in BeginAsyncUpdate() to complete, issues GPU work and does some cleanup.
