* Fix ActorDesc delta serialization for the first instance of a Blueprint Actor in a World
* It would previously fail to obtain an ActorDesc for the base class and the delta serialization would save all values as if they were overriding the base class and this could cause incoherents states between the ActorDesc & the Actor down the line
* Now we alway prefetch the classdesc for Blueprints when saving a Blueprint Actor. (this early outs if it's already loaded)

diff --git a/Engine/Source/Runtime/Engine/Private/WorldPartition/WorldPartitionClassDescRegistry.cpp b/Engine/Source/Runtime/Engine/Private/WorldPartition/WorldPartitionClassDescRegistry.cpp
--- a/Engine/Source/Runtime/Engine/Private/WorldPartition/WorldPartitionClassDescRegistry.cpp
+++ b/Engine/Source/Runtime/Engine/Private/WorldPartition/WorldPartitionClassDescRegistry.cpp
@@ -44,6 +44,10 @@
  *
  *	[Saving a existing Blueprint]
  *	  - Create or update the class descriptor.
+ * 
+ *	[Saving a Blueprint Actor]
+ *	  - Prefetch the class descriptor for the Blueprint Actor Class
+ * 
  */
 
 static FAutoConsoleCommand DumpClassDescs(
@@ -526,6 +530,14 @@ void FWorldPartitionClassDescRegistry::OnObjectPreSave(UObject* InObject, FObjec
 			}
 		}
 	}
+
+	if (InObject->IsA<AActor>())
+	{
+		if (UBlueprintGeneratedClass* BlueprintGeneratedClass = Cast<UBlueprintGeneratedClass>(InObject->GetClass()))
+		{			
+			PrefetchClassDesc(BlueprintGeneratedClass);
+		}
+	}
 }
 
 void FWorldPartitionClassDescRegistry::OnObjectPropertyChanged(UObject* InObject, FPropertyChangedEvent& InPropertyChangedEvent)
