The HV component now uses the newly added USparseVolumeTexture::GetOptimalStreamingMipLevel() function to compute the mip level at which it issues streaming requests. Replaced the old MipLevel property on the HV component with a new StreamingMipBias property to allow users to bias the computed streaming mip level.

#rnx

diff --git a/Engine/Source/Runtime/Engine/Classes/Components/HeterogeneousVolumeComponent.h b/Engine/Source/Runtime/Engine/Classes/Components/HeterogeneousVolumeComponent.h
--- a/Engine/Source/Runtime/Engine/Classes/Components/HeterogeneousVolumeComponent.h
+++ b/Engine/Source/Runtime/Engine/Classes/Components/HeterogeneousVolumeComponent.h
@@ -24,7 +24,6 @@ class UHeterogeneousVolumeComponent : public UMeshComponent
 	UPROPERTY(EditAnywhere, BlueprintReadOnly, Category = Volume)
 	FIntVector VolumeResolution;
 
-	//UPROPERTY(EditAnywhere, BlueprintReadWrite, Category = Animation, AdvancedDisplay, meta = (EditCondition = "false"))
 	UPROPERTY(VisibleAnywhere, BlueprintReadWrite, Category = Animation, AdvancedDisplay)
 	FTransform FrameTransform;
 
@@ -46,8 +45,10 @@ class UHeterogeneousVolumeComponent : public UMeshComponent
 	UPROPERTY(EditAnywhere, Category = Animation)
 	uint32 bLooping : 1;
 
+	// Bias to apply to the calculated mip level to stream at. This property essentially influences the distance from the
+	// volume at which certain mip levels are no longer requested. Higher values result in earlier mip level transitions.
 	UPROPERTY(EditAnywhere, Category = SparseVolumeTextureStreaming, AdvancedDisplay)
-	uint32 MipLevel = 0;
+	float StreamingMipBias;
 
 	UPROPERTY(EditAnywhere, Category = SparseVolumeTextureStreaming)
 	uint32 bIssueBlockingRequests : 1;
@@ -91,6 +92,9 @@ class UHeterogeneousVolumeComponent : public UMeshComponent
 	UFUNCTION(BlueprintCallable, Category = "Animation")
 	ENGINE_API void Play();
 
+	UFUNCTION(BlueprintCallable, Category = "SparseVolumeTextureStreaming")
+	ENGINE_API void SetStreamingMipBias(int32 NewValue);
+
 	~UHeterogeneousVolumeComponent() {}
 
 public:
diff --git a/Engine/Source/Runtime/Engine/Classes/SparseVolumeTexture/SparseVolumeTexture.h b/Engine/Source/Runtime/Engine/Classes/SparseVolumeTexture/SparseVolumeTexture.h
--- a/Engine/Source/Runtime/Engine/Classes/SparseVolumeTexture/SparseVolumeTexture.h
+++ b/Engine/Source/Runtime/Engine/Classes/SparseVolumeTexture/SparseVolumeTexture.h
@@ -330,6 +330,8 @@ public:
 	virtual TextureAddress GetTextureAddressY() const { return TA_Wrap; }
 	virtual TextureAddress GetTextureAddressZ() const { return TA_Wrap; }
 	virtual const UE::SVT::FTextureRenderResources* GetTextureRenderResources() const { return nullptr; }
+	// Computes the optimal mip level to stream the SVT at, based on projected screen space size and voxel resolution.
+	float GetOptimalStreamingMipLevel(const FBoxSphereBounds& Bounds, float MipBias) const;
 
 	/** Getter for the shader uniform parameters with index as ESparseVolumeTextureShaderUniform. */
 	FVector4 GetUniformParameter(int32 Index) const { return FVector4(ForceInitToZero); } // SVT_TODO: This mechanism is no longer needed and can be removed
diff --git a/Engine/Source/Runtime/Engine/Private/Components/HeterogeneousVolumeComponent.cpp b/Engine/Source/Runtime/Engine/Private/Components/HeterogeneousVolumeComponent.cpp
--- a/Engine/Source/Runtime/Engine/Private/Components/HeterogeneousVolumeComponent.cpp
+++ b/Engine/Source/Runtime/Engine/Private/Components/HeterogeneousVolumeComponent.cpp
@@ -256,12 +256,22 @@ UHeterogeneousVolumeComponent::UHeterogeneousVolumeComponent(const FObjectInitia
 	ShadowStepFactor = 2.0f;
 	ShadowBiasFactor = 0.5f;
 	LightingDownsampleFactor = 2.0f;
-	MipLevel = 0;
+	StreamingMipBias = 0.0f;
 	bIssueBlockingRequests = false;
 	bPivotAtCentroid = false;
 	PreviousSVT = nullptr;
 }
 
+void UHeterogeneousVolumeComponent::SetStreamingMipBias(int32 NewValue)
+{
+	if (AreDynamicDataChangesAllowed()
+		&& StreamingMipBias != NewValue)
+	{
+		StreamingMipBias = NewValue;
+		MarkRenderStateDirty();
+	}
+}
+
 void UHeterogeneousVolumeComponent::SetFrame(float NewValue)
 {
 	if (AreDynamicDataChangesAllowed()
@@ -584,6 +594,7 @@ void UHeterogeneousVolumeComponent::TickComponent(float DeltaTime, ELevelTick Ti
 
 			const bool bIsBlocking = bIssueBlockingRequests != 0;
 			const bool bHasValidFrameRate = bPlaying != 0;
+			const float MipLevel = SparseVolumeTexture->GetOptimalStreamingMipLevel(Bounds, StreamingMipBias);
 			USparseVolumeTextureFrame* SparseVolumeTextureFrame = USparseVolumeTextureFrame::GetFrameAndIssueStreamingRequest(SparseVolumeTexture, GetTypeHash(this), FrameRate, Frame, MipLevel, bIsBlocking, bHasValidFrameRate);
 			if (SparseVolumeTextureFrame)
 			{
diff --git a/Engine/Source/Runtime/Engine/Private/SparseVolumeTexture/SparseVolumeTexture.cpp b/Engine/Source/Runtime/Engine/Private/SparseVolumeTexture/SparseVolumeTexture.cpp
--- a/Engine/Source/Runtime/Engine/Private/SparseVolumeTexture/SparseVolumeTexture.cpp
+++ b/Engine/Source/Runtime/Engine/Private/SparseVolumeTexture/SparseVolumeTexture.cpp
@@ -19,6 +19,7 @@
 #include "SparseVolumeTexture/SparseVolumeTextureData.h"
 #include "SparseVolumeTexture/SparseVolumeTextureUtility.h"
 #include "SparseVolumeTexture/ISparseVolumeTextureStreamingManager.h"
+#include "ContentStreaming.h"
 
 #if WITH_EDITORONLY_DATA
 #include "Misc/ScopedSlowTask.h"
@@ -54,6 +55,15 @@ static FAutoConsoleVariableRef CVarSVTRemoteDDCBehavior(
 	ECVF_Default
 );
 
+static float GSVTStreamingRequestMipBias = 1.0f;
+static FAutoConsoleVariableRef CVarSVTStreamingRequestMipBias(
+	TEXT("r.SparseVolumeTexture.Streaming.RequestMipBias"),
+	GSVTStreamingRequestMipBias,
+	TEXT("Bias to apply the calculated mip level to stream at. This is used to account for the mip estimation based on projected screen space size being very conservative. ")
+	TEXT("The default value of 1.0 was found empirically to roughly result in a 1:1 voxel to pixel ratio."),
+	ECVF_Default
+);
+
 ////////////////////////////////////////////////////////////////////////////////////////////////
 
 FArchive& operator<<(FArchive& Ar, UE::SVT::FMipLevelStreamingInfo& MipLevelStreamingInfo)
@@ -873,6 +883,40 @@ USparseVolumeTexture::USparseVolumeTexture(const FObjectInitializer& ObjectIniti
 {
 }
 
+float USparseVolumeTexture::GetOptimalStreamingMipLevel(const FBoxSphereBounds& Bounds, float MipBias) const
+{
+	check(IsInGameThread());
+	float ResultMipLevel = 0.0f;
+	if (IStreamingManager* StreamingManager = IStreamingManager::Get_Concurrent())
+	{
+		ResultMipLevel = FLT_MAX;
+		const int32 NumViews = StreamingManager->GetNumViews();
+		for (int32 ViewIndex = 0; ViewIndex < NumViews; ++ViewIndex)
+		{
+			const FStreamingViewInfo& ViewInfo = StreamingManager->GetViewInformation(ViewIndex);
+
+			// Determine the pixel-width at the near-plane.
+			const float PixelWidth = 1.0f / (ViewInfo.FOVScreenSize * 0.5f); // FOVScreenSize = ViewRect.Width / Tan(FOV * 0.5)
+
+			// Project to nearest distance of volume bounds.
+			const float Distance = FMath::Max<float>(1.0f, ((ViewInfo.ViewOrigin - Bounds.Origin).GetAbs() - Bounds.BoxExtent).Length());
+			const float VoxelWidth = Distance * PixelWidth;
+
+			// MIP is defined as the log of the ratio of native voxel resolution to pixel-coverage of volume bounds.
+			// We want to be conservative here (use potentially lower mip), so try to minimize the term we pass into Log2() by using
+			// the maximum dimension of the bounds and the minimum extent of the volume resolution. The bounds are axis aligned, so
+			// we can't assume that a given dimension in SVT UV space aligns with any particular dimension of the axis aligned bounds.
+			const float PixelWidthCoverage = (2.0f * Bounds.BoxExtent.GetMax()) / VoxelWidth;
+			const float VoxelResolution = GetVolumeResolution().GetMin();
+			float ViewMipLevel = FMath::Log2(VoxelResolution / PixelWidthCoverage) + MipBias + GSVTStreamingRequestMipBias;
+			ViewMipLevel = FMath::Clamp(ViewMipLevel, 0.0f, GetNumMipLevels() - 1.0f);
+
+			ResultMipLevel = FMath::Min(ViewMipLevel, ResultMipLevel);
+		}
+	}
+	return ResultMipLevel;
+}
+
 UE::Shader::EValueType USparseVolumeTexture::GetUniformParameterType(int32 Index)
 {
 	switch (Index)
