HISM Navmesh Update

Removed all the Navigation function overrides of ISMComponent in HISMComponent, and cleaned up the new navigation function/structures based on HISM's rendering structures.

Our navigation stuff was relying on the rendering structures used by HISM and that has issues:
- when we do a in-place update on instance transform, the cluster tree and BuiltInstanceBounds will be out of sync. Navigation then uses Cluster Tree to compute its bounds which would be incorrect.
- the cluster tree represents bounds for rendering. If we have a custom navigation collider then using that as navigation bounds will be incorrect.

diff --git a/Engine/Source/Runtime/Engine/Classes/Components/HierarchicalInstancedStaticMeshComponent.h b/Engine/Source/Runtime/Engine/Classes/Components/HierarchicalInstancedStaticMeshComponent.h
--- a/Engine/Source/Runtime/Engine/Classes/Components/HierarchicalInstancedStaticMeshComponent.h
+++ b/Engine/Source/Runtime/Engine/Classes/Components/HierarchicalInstancedStaticMeshComponent.h
@@ -382,14 +382,19 @@ protected:
 
 	ENGINE_API static FBox GetClusterTreeBounds(TArray<FClusterNode> const& InClusterTree, const FVector& InOffset);
 
+	UE_DEPRECATED(5.4, "UHierarchicalInstancedStaticMeshComponent is no longer overriding UInstancedStaticMeshComponent::GetNavigationPerInstanceTransforms.")
 	ENGINE_API virtual void GetNavigationPerInstanceTransforms(const FBox& AreaBox, TArray<FTransform>& InstanceData) const override;
+	UE_DEPRECATED(5.4, "UHierarchicalInstancedStaticMeshComponent is no longer overriding UInstancedStaticMeshComponent::PartialNavigationUpdate.")
 	ENGINE_API virtual void PartialNavigationUpdate(int32 InstanceIdx) override;
-	virtual bool SupportsPartialNavigationUpdate() const override { return true; }
+	UE_DEPRECATED(5.4, "UHierarchicalInstancedStaticMeshComponent is no longer overriding UInstancedStaticMeshComponent::SupportsPartialNavigationUpdate.")
+	virtual bool SupportsPartialNavigationUpdate() const override { return Super::SupportsPartialNavigationUpdate(); }
+	UE_DEPRECATED(5.4, "FlushAccumulatedNavigationUpdates has been deprecated. Use UInstancedStaticMeshComponent::PartialNavigationUpdates if you want to manually update the navigation data.")
 	ENGINE_API void FlushAccumulatedNavigationUpdates();
 
-	UE_DEPRECATED(5.4, "Use AccumulatedNavigationDirtyAreas instead.")
+	UE_DEPRECATED(5.4, "AccumulatedNavigationDirtyArea has been deprecated.")
     mutable FBox AccumulatedNavigationDirtyArea;
-	mutable TArray<FBox> AccumulatedNavigationDirtyAreas;
+	UE_DEPRECATED(5.4, "AccumulatedNavigationDirtyAreas has been deprecated.")
+	mutable TArray<FBox> AccumulatedNavigationDirtyAreas;	//Remember to also remove the disable deprecation macro wrapper around ~UHierarchicalInstancedStaticMeshComponent() when removing the variable
 
 	FGraphEventArray BuildTreeAsyncTasks;
 	EHISMViewRelevanceType ViewRelevanceType = EHISMViewRelevanceType::HISM;
diff --git a/Engine/Source/Runtime/Engine/Private/HierarchicalInstancedStaticMesh.cpp b/Engine/Source/Runtime/Engine/Private/HierarchicalInstancedStaticMesh.cpp
--- a/Engine/Source/Runtime/Engine/Private/HierarchicalInstancedStaticMesh.cpp
+++ b/Engine/Source/Runtime/Engine/Private/HierarchicalInstancedStaticMesh.cpp
@@ -2046,6 +2046,8 @@ UHierarchicalInstancedStaticMeshComponent::UHierarchicalInstancedStaticMeshCompo
 	bUseAsOccluder = false;
 }
 
+//We deprecated a TArray and that's being referenced by the dtor as TArray has a non-trivial dtor which will trigger the deprecation warning.
+PRAGMA_DISABLE_DEPRECATION_WARNINGS
 UHierarchicalInstancedStaticMeshComponent::~UHierarchicalInstancedStaticMeshComponent()
 {
 	if (ProxySize)
@@ -2054,6 +2056,7 @@ UHierarchicalInstancedStaticMeshComponent::~UHierarchicalInstancedStaticMeshComp
 	}
 	ProxySize = 0;
 }
+PRAGMA_ENABLE_DEPRECATION_WARNINGS
 
 #if WITH_EDITOR
 
@@ -2730,7 +2733,6 @@ void UHierarchicalInstancedStaticMeshComponent::ApplyBuildTree(FClusterBuilder&
 	// BUT: the hitproxy data was rebuilt right here & the per instance random is handled differently inside the tree builder so, nope.
 	PrimitiveInstanceDataManager.MarkForRebuildFromLegacy(MoveTemp(BuiltInstanceData), InstanceReorderTable, HitProxies);
 
-	FlushAccumulatedNavigationUpdates();
 	PostBuildStats();
 	MarkRenderStateDirty();
 
@@ -3229,62 +3231,16 @@ FVector UHierarchicalInstancedStaticMeshComponent::GetAverageScale() const
 
 void UHierarchicalInstancedStaticMeshComponent::GetNavigationPerInstanceTransforms(const FBox& AreaBox, TArray<FTransform>& InstanceData) const
 {
-	if (IsTreeFullyBuilt())
-	{
-		const TArray<FClusterNode>& ClusterTree = *ClusterTreePtr;
-		if (ClusterTree.Num())
-		{
-			GetOverlappingBoxTransforms(AreaBox, InstanceData);
-		}
-	}
-	else
-	{
-		// This area should be processed again by navigation system when cluster tree is available
-		// Store smaller tile box in accumulated dirty area, so we will not unintentionally mark as dirty neighbor tiles 
-		AccumulatedNavigationDirtyAreas.Emplace(AreaBox.ExpandBy(-AreaBox.GetExtent()/2.f));
-	}
+	Super::GetNavigationPerInstanceTransforms(AreaBox, InstanceData);
 }
 
 void UHierarchicalInstancedStaticMeshComponent::PartialNavigationUpdate(const int32 InstanceIdx)
 {
-	if (InstanceIdx == INDEX_NONE)
-	{
-		AccumulatedNavigationDirtyAreas.Reset();
-		FNavigationSystem::UpdateComponentData(*this);
-	}
-	else if (GetStaticMesh() && FNavigationSystem::HasComponentData(*this))
-	{
-		// Accumulate dirty areas and send them to navigation system once cluster tree is rebuilt
-		FTransform InstanceTransformInWorldSpace;
-		if (GetInstanceTransform(InstanceIdx, InstanceTransformInWorldSpace, /*bWorldSpace*/true))
-		{
-			const FBox InstanceBounds = GetInstanceNavigationBounds();
-			if (InstanceBounds.IsValid)
-			{
-				AccumulatedNavigationDirtyAreas.Emplace(InstanceBounds.TransformBy(InstanceTransformInWorldSpace));
-			}
-		}
-	}
-
-	CalcAndCacheNavigationBounds();
+	Super::PartialNavigationUpdate(InstanceIdx);
 }
 
 void UHierarchicalInstancedStaticMeshComponent::FlushAccumulatedNavigationUpdates()
-{
-	if (AccumulatedNavigationDirtyAreas.Num() > 0)
-	{
-		QUICK_SCOPE_CYCLE_COUNTER(STAT_UHierarchicalInstancedStaticMeshComponent_FlushAccumulatedNavigationUpdates);
-
-		const TArray<FClusterNode>& ClusterTree = *ClusterTreePtr;
-		if (ClusterTree.Num())
-		{
-			const FBox NewBounds = GetClusterTreeBounds(ClusterTree, TranslatedInstanceSpaceOrigin).TransformBy(GetComponentTransform());
-			FNavigationSystem::OnObjectBoundsChanged(*this, NewBounds, AccumulatedNavigationDirtyAreas);
-		}
-
-		AccumulatedNavigationDirtyAreas.Reset();
-	}
-}
+{}
 
 // recursive helper to gather all instances with locations inside the specified area. Supply a Filter to exclude leaf nodes based on the instance transform.
 static void GatherInstancesOverlappingArea(const UHierarchicalInstancedStaticMeshComponent& Component, const FBox& AreaBox, int32 Child, TFunctionRef<bool(const FMatrix&)> Filter, TArray<int32>& OutInstanceIndices)
