Fix for scene capture orthographic scenes being skewed by incorrect rewrite of the projection matrix. Also moved associated CVar for this functionality to ensure this can be force disabled in all use cases, as previously the SceneCapture path was skipped.

diff --git a/Engine/Source/Runtime/Engine/Classes/Camera/CameraTypes.h b/Engine/Source/Runtime/Engine/Classes/Camera/CameraTypes.h
--- a/Engine/Source/Runtime/Engine/Classes/Camera/CameraTypes.h
+++ b/Engine/Source/Runtime/Engine/Classes/Camera/CameraTypes.h
@@ -147,7 +147,7 @@ public:
 	ENGINE_API void AddWeightedViewInfo(const FMinimalViewInfo& OtherView, const float& Weight);
 
 	/** Calculates the projection matrix using this view info's aspect ratio (regardless of bConstrainAspectRatio) */
-	ENGINE_API FMatrix CalculateProjectionMatrix(float ForceOrthoNegativeNearPlane = 0.0f) const;
+	ENGINE_API FMatrix CalculateProjectionMatrix() const;
 
 	/** Calculates the projection matrix (and potentially a constrained view rectangle) given a FMinimalViewInfo and partially configured projection data (must have the view rect already set) */
 	ENGINE_API static void CalculateProjectionMatrixGivenView(FMinimalViewInfo& ViewInfo, TEnumAsByte<enum EAspectRatioAxisConstraint> AspectRatioAxisConstraint, class FViewport* Viewport, struct FSceneViewProjectionData& InOutProjectionData);
diff --git a/Engine/Source/Runtime/Engine/Private/Camera/CameraStackTypes.cpp b/Engine/Source/Runtime/Engine/Private/Camera/CameraStackTypes.cpp
--- a/Engine/Source/Runtime/Engine/Private/Camera/CameraStackTypes.cpp
+++ b/Engine/Source/Runtime/Engine/Private/Camera/CameraStackTypes.cpp
@@ -17,16 +17,6 @@ static TAutoConsoleVariable<bool> CVarUseLegacyMaintainYFOV(
 	ECVF_Default
 );
 
-static TAutoConsoleVariable<bool> CVarAllowOrthoNearPlaneCorrection(
-	TEXT("r.Ortho.AllowNearPlaneCorrection"),
-	true,
-	TEXT("Orthographic near planes can be behind the camera position, which causes some issues with Unreal resolving lighting behind the camera position ")
-	TEXT("This CVar enables the Orthographic cameras globally to automatically update the camera location to match the NearPlane location, ")
-	TEXT("and force the pseudo camera position to be the replaced near plane location for the projection matrix calculation. ")
-	TEXT("This means Unreal can resolve lighting behind the camera correctly."),
-	ECVF_Default
-);
-
 static TAutoConsoleVariable<bool> CVarOrthoAllowAutoPlanes(
 	TEXT("r.Ortho.AutoPlanes"),
 	true,
@@ -179,7 +169,7 @@ void FMinimalViewInfo::AddWeightedViewInfo(const FMinimalViewInfo& OtherView, co
 	bUseFieldOfViewForLOD |= OtherViewWeighted.bUseFieldOfViewForLOD;
 }
 
-FMatrix FMinimalViewInfo::CalculateProjectionMatrix(float NegativeNearPlaneOverride) const
+FMatrix FMinimalViewInfo::CalculateProjectionMatrix() const
 {
 	FMatrix ProjectionMatrix;
 
@@ -190,7 +180,7 @@ FMatrix FMinimalViewInfo::CalculateProjectionMatrix(float NegativeNearPlaneOverr
 		const float HalfOrthoWidth = OrthoWidth / 2.0f;
 		const float ScaledOrthoHeight = OrthoWidth / 2.0f * YScale;
 
-		const float NearPlane = CVarAllowOrthoNearPlaneCorrection.GetValueOnAnyThread() && NegativeNearPlaneOverride < 0.0f ? NegativeNearPlaneOverride : OrthoNearClipPlane;
+		const float NearPlane = OrthoNearClipPlane;
 		const float FarPlane = OrthoFarClipPlane;
 
 		const float ZScale = 1.0f / (FarPlane - NearPlane);
@@ -240,10 +230,11 @@ void FMinimalViewInfo::CalculateProjectionMatrixGivenViewRectangle(FMinimalViewI
 	}
 #endif
 
-	bool bOrthographicNearPlaneCorrection = ViewInfo.ProjectionMode == ECameraProjectionMode::Orthographic
-											&& CVarAllowOrthoNearPlaneCorrection.GetValueOnAnyThread();
-
-	ViewInfo.AutoCalculateOrthoPlanes(InOutProjectionData);	
+	bool bOrthographic = ViewInfo.ProjectionMode == ECameraProjectionMode::Orthographic;
+	if(bOrthographic)
+	{
+		ViewInfo.AutoCalculateOrthoPlanes(InOutProjectionData);	
+	}
 
 	// Create the projection matrix (and possibly constrain the view rectangle)
 	if (ViewInfo.bConstrainAspectRatio)
@@ -251,7 +242,7 @@ void FMinimalViewInfo::CalculateProjectionMatrixGivenViewRectangle(FMinimalViewI
 		// Enforce a particular aspect ratio for the render of the scene. 
 		// Results in black bars at top/bottom etc.
 		InOutProjectionData.SetConstrainedViewRectangle(ConstrainedViewRectangle);
-		if(bOrthographicNearPlaneCorrection)
+		if(bOrthographic)
 		{
 			InOutProjectionData.UpdateOrthoPlanes(ViewInfo);
 		}
@@ -303,17 +294,15 @@ void FMinimalViewInfo::CalculateProjectionMatrixGivenViewRectangle(FMinimalViewI
 			MatrixHalfFOV = FMath::Max(0.001f, ViewInfo.FOV) * (float)UE_PI / 360.0f;
 		}
 	
-		if (ViewInfo.ProjectionMode == ECameraProjectionMode::Orthographic)
+		if (bOrthographic)
 		{
 			const float OrthoWidth = ViewInfo.OrthoWidth / 2.0f * XAxisMultiplier;
 			const float OrthoHeight = (ViewInfo.OrthoWidth / 2.0f) / YAxisMultiplier;
 
 			float FarPlane = ViewInfo.OrthoFarClipPlane;
 			float NearPlane = ViewInfo.OrthoNearClipPlane;
-			if (bOrthographicNearPlaneCorrection)
-			{
-				InOutProjectionData.UpdateOrthoPlanes(NearPlane, FarPlane);
-			}
+
+			InOutProjectionData.UpdateOrthoPlanes(NearPlane, FarPlane);
 
 			const float ZScale = 1.0f / (FarPlane - NearPlane);
 			const float ZOffset = -NearPlane;
diff --git a/Engine/Source/Runtime/Engine/Private/SceneView.cpp b/Engine/Source/Runtime/Engine/Private/SceneView.cpp
--- a/Engine/Source/Runtime/Engine/Private/SceneView.cpp
+++ b/Engine/Source/Runtime/Engine/Private/SceneView.cpp
@@ -341,6 +341,16 @@ static FAutoConsoleVariableRef CVarDefaultUpdateOrthoNearPlane(
 	ECVF_RenderThreadSafe
 );
 
+static TAutoConsoleVariable<bool> CVarAllowOrthoNearPlaneCorrection(
+	TEXT("r.Ortho.AllowNearPlaneCorrection"),
+	true,
+	TEXT("Orthographic near planes can be behind the camera position, which causes some issues with Unreal resolving lighting behind the camera position ")
+	TEXT("This CVar enables the Orthographic cameras globally to automatically update the camera location to match the NearPlane location, ")
+	TEXT("and force the pseudo camera position to be the replaced near plane location for the projection matrix calculation. ")
+	TEXT("This means Unreal can resolve lighting behind the camera correctly."),
+	ECVF_Default
+);
+
 static TAutoConsoleVariable<bool> CVarOrthoCameraHeightAsViewTarget(
 	TEXT("r.Ortho.CameraHeightAsViewTarget"),
 	true,
@@ -562,7 +572,7 @@ FVector4f CreateInvDeviceZToWorldZTransform(const FMatrix& ProjMatrix)
 
 bool FSceneViewProjectionData::UpdateOrthoPlanes(FSceneViewProjectionData* InOutProjectionData, float& NearPlane, float& FarPlane)
 {
-	if (!InOutProjectionData)
+	if (!InOutProjectionData || !CVarAllowOrthoNearPlaneCorrection.GetValueOnAnyThread())
 	{
 		return false;
 	}
@@ -608,6 +618,7 @@ bool FSceneViewProjectionData::UpdateOrthoPlanes()
 	* Separating this step from the near plane calculation logic itself helps avoid applying the NearPlane correctiontwice /makes it easier to read.
 	*/
 	if(IsPerspectiveProjection() 
+		|| !CVarAllowOrthoNearPlaneCorrection.GetValueOnAnyThread()
 		|| ProjectionMatrix.M[2][2] == 0.0 
 		|| ProjectionMatrix.M[2][2] == ProjectionMatrix.M[2][3])
 	{ 
@@ -630,7 +641,7 @@ bool FSceneViewProjectionData::UpdateOrthoPlanes()
 
 	//Only the Near/Far plane elements need correcting, OrthoWidth/Height remains the same
 	ProjectionMatrix.M[2][2] = -ZScale;
-	ProjectionMatrix.M[2][3] = 1.0f - (ZOffset * ZScale);
+	ProjectionMatrix.M[3][2] = 1.0f - (ZOffset * ZScale);
 
 	return true;
 }
