Fix potential non-zero ISM bounds when there are no instances in the ISM.
This can happen in some Super::UpdateBounds() paths.

diff --git a/Engine/Source/Runtime/Engine/Private/InstancedStaticMesh.cpp b/Engine/Source/Runtime/Engine/Private/InstancedStaticMesh.cpp
--- a/Engine/Source/Runtime/Engine/Private/InstancedStaticMesh.cpp
+++ b/Engine/Source/Runtime/Engine/Private/InstancedStaticMesh.cpp
@@ -2701,7 +2701,11 @@ FBoxSphereBounds UInstancedStaticMeshComponent::CalcBoundsImpl(const FTransform&
 
 void UInstancedStaticMeshComponent::UpdateBounds()
 {
-	if (bUseConservativeBounds && CachedConservativeInstanceBounds.IsValid)
+	if (PerInstanceSMData.Num() == 0)
+	{
+		Bounds = FBoxSphereBounds(GetComponentTransform().GetLocation(), FVector::ZeroVector, 0.f);
+	}
+	else if (bUseConservativeBounds && CachedConservativeInstanceBounds.IsValid)
 	{
 		Bounds = CachedConservativeInstanceBounds.TransformBy(GetComponentTransform());
 	}
