Fixed shutdown crash in FNetGuidCache and FObjectReferenceCache

diff --git a/Engine/Source/Runtime/Engine/Private/PackageMapClient.cpp b/Engine/Source/Runtime/Engine/Private/PackageMapClient.cpp
--- a/Engine/Source/Runtime/Engine/Private/PackageMapClient.cpp
+++ b/Engine/Source/Runtime/Engine/Private/PackageMapClient.cpp
@@ -3319,10 +3319,18 @@ PRAGMA_ENABLE_DEPRECATION_WARNINGS
 
 	DelinquentAsyncLoads.MaxConcurrentAsyncLoads = FMath::Max<uint32>(DelinquentAsyncLoads.MaxConcurrentAsyncLoads, PendingAsyncLoadRequests.Num());
 
-	LoadPackageAsync(CacheObject.PathName.ToString(), FLoadPackageAsyncDelegate::CreateRaw(this, &FNetGUIDCache::AsyncPackageCallback));
+	FLoadPackageAsyncDelegate LoadPackageCompleteDelegate = FLoadPackageAsyncDelegate::CreateWeakLambda(Driver, [NetDriver = Driver](const FName& PackageName, UPackage* Package, EAsyncLoadingResult::Type Result)
+	{
+		if (NetDriver->GuidCache.IsValid())
+		{
+			NetDriver->GuidCache->AsyncPackageCallback(PackageName, Package, Result);
+		}
+	});
+
+	LoadPackageAsync(CacheObject.PathName.ToString(), LoadPackageCompleteDelegate);
 }
 
-void FNetGUIDCache::AsyncPackageCallback(const FName& PackageName, UPackage * Package, EAsyncLoadingResult::Type Result)
+void FNetGUIDCache::AsyncPackageCallback(const FName& PackageName, UPackage* Package, EAsyncLoadingResult::Type Result)
 {
 	LLM_SCOPE_BYTAG(GuidCache);
 
diff --git a/Engine/Source/Runtime/Experimental/Iris/Core/Private/Iris/ReplicationSystem/ObjectReferenceCache.cpp b/Engine/Source/Runtime/Experimental/Iris/Core/Private/Iris/ReplicationSystem/ObjectReferenceCache.cpp
--- a/Engine/Source/Runtime/Experimental/Iris/Core/Private/Iris/ReplicationSystem/ObjectReferenceCache.cpp
+++ b/Engine/Source/Runtime/Experimental/Iris/Core/Private/Iris/ReplicationSystem/ObjectReferenceCache.cpp
@@ -1888,7 +1888,12 @@ void FObjectReferenceCache::StartAsyncLoadingPackage(FCachedNetObjectReference&
 
 	//DelinquentAsyncLoads.MaxConcurrentAsyncLoads = FMath::Max<uint32>(DelinquentAsyncLoads.MaxConcurrentAsyncLoads, PendingAsyncLoadRequests.Num());
 
-	LoadPackageAsync(PackagePath.ToString(), FLoadPackageAsyncDelegate::CreateRaw(this, &FObjectReferenceCache::AsyncPackageCallback));
+	LoadPackageAsync(PackagePath.ToString(), FLoadPackageAsyncDelegate::CreateWeakLambda(ReplicationSystem, 
+		[this](const FName& PackageName, UPackage* Package, EAsyncLoadingResult::Type Result)
+		{
+			AsyncPackageCallback(PackageName, Package, Result);
+		}
+	));
 }
 
 void FObjectReferenceCache::AsyncPackageCallback(const FName& PackageName, UPackage* Package, EAsyncLoadingResult::Type Result)
