Added CSV stats for external code called from within UpdateStreamingState to be able to identify if perf regressions comes from engine or game code.
New stat goes into WorldPartition/GameThread/GetStreamingSourceProviders which is part of Exclusive/GameThread/UpdateStreamingState.

diff --git a/Engine/Source/Runtime/Engine/Private/WorldPartition/WorldPartitionSubsystem.cpp b/Engine/Source/Runtime/Engine/Private/WorldPartition/WorldPartitionSubsystem.cpp
--- a/Engine/Source/Runtime/Engine/Private/WorldPartition/WorldPartitionSubsystem.cpp
+++ b/Engine/Source/Runtime/Engine/Private/WorldPartition/WorldPartitionSubsystem.cpp
@@ -29,9 +29,11 @@
 #include "Debug/DebugDrawService.h"
 #include "GameFramework/PlayerController.h"
 #include "GameFramework/WorldSettings.h"
+#include "ProfilingDebugging/CsvProfiler.h"
 #include "Async/ParallelFor.h"
 #include "Algo/ForEach.h"
 #include "Misc/HashBuilder.h"
+#include "Stats/Stats.h"
 
 #if WITH_EDITOR
 #include "Editor.h"
@@ -42,6 +44,10 @@
 
 #include UE_INLINE_GENERATED_CPP_BY_NAME(WorldPartitionSubsystem)
 
+DECLARE_STATS_GROUP(TEXT("World Partition"), STATGROUP_WorldPartition, STATCAT_Advanced);
+DECLARE_CYCLE_STAT(TEXT("UpdateStreamingStateExternal"), STAT_UpdateStreamingStateExternal, STATGROUP_WorldPartition);
+CSV_DEFINE_CATEGORY(WorldPartition, (!UE_BUILD_SHIPPING));
+
 extern int32 GBlockOnSlowStreaming;
 static const FName NAME_WorldPartitionRuntimeHash("WorldPartitionRuntimeHash");
 
@@ -897,17 +903,24 @@ void UWorldPartitionSubsystem::UpdateStreamingSources()
 				bAllowPlayerControllerStreamingSources = false;
 			}
 #endif
-			TArray<FWorldPartitionStreamingSource> ProviderStreamingSources;
-			for (IWorldPartitionStreamingSourceProvider* StreamingSourceProvider : GetStreamingSourceProviders())
+
 			{
-				if (bAllowPlayerControllerStreamingSources || !Cast<APlayerController>(StreamingSourceProvider->GetStreamingSourceOwner()))
+				// This will include game-sepcific code called from GetStreamingSourceProviders if IsStreamingSourceProviderFiltered is bound, and also GetStreamingSources calls.
+				SCOPE_CYCLE_COUNTER(STAT_UpdateStreamingStateExternal);
+				CSV_SCOPED_TIMING_STAT(WorldPartition, GetStreamingSourceProviders);
+
+				TArray<FWorldPartitionStreamingSource> ProviderStreamingSources;
+				for (IWorldPartitionStreamingSourceProvider* StreamingSourceProvider : GetStreamingSourceProviders())
 				{
-					ProviderStreamingSources.Reset();
-					if (StreamingSourceProvider->GetStreamingSources(ProviderStreamingSources))
+					if (bAllowPlayerControllerStreamingSources || !Cast<APlayerController>(StreamingSourceProvider->GetStreamingSourceOwner()))
 					{
-						for (FWorldPartitionStreamingSource& ProviderStreamingSource : ProviderStreamingSources)
+						ProviderStreamingSources.Reset();
+						if (StreamingSourceProvider->GetStreamingSources(ProviderStreamingSources))
 						{
-							StreamingSources.Add(MoveTemp(ProviderStreamingSource));
+							for (FWorldPartitionStreamingSource& ProviderStreamingSource : ProviderStreamingSources)
+							{
+								StreamingSources.Add(MoveTemp(ProviderStreamingSource));
+							}
 						}
 					}
 				}
