Heterogeneous Volumes: Stereo rendering support

Minor: Updated NiagaraVolumeRenderer's ShadowStepFactor default value to 2.0

diff --git a/Engine/Plugins/FX/Niagara/Source/Niagara/Private/NiagaraRendererVolumes.cpp b/Engine/Plugins/FX/Niagara/Source/Niagara/Private/NiagaraRendererVolumes.cpp
--- a/Engine/Plugins/FX/Niagara/Source/Niagara/Private/NiagaraRendererVolumes.cpp
+++ b/Engine/Plugins/FX/Niagara/Source/Niagara/Private/NiagaraRendererVolumes.cpp
@@ -362,7 +362,7 @@ void FNiagaraRendererVolumes::GetDynamicMeshElements(const TArray<const FSceneVi
 			HeterogeneousVolumeData->LightingDownsampleFactor = VolumeDynamicData->LightingDownsampleFactor;
 			BatchElement.UserData = HeterogeneousVolumeData;
 
-			Collector.AddMesh(0, Mesh);
+			Collector.AddMesh(ViewIndex, Mesh);
 		}
 	}
 }
diff --git a/Engine/Plugins/FX/Niagara/Source/Niagara/Public/NiagaraVolumeRendererProperties.h b/Engine/Plugins/FX/Niagara/Source/Niagara/Public/NiagaraVolumeRendererProperties.h
--- a/Engine/Plugins/FX/Niagara/Source/Niagara/Public/NiagaraVolumeRendererProperties.h
+++ b/Engine/Plugins/FX/Niagara/Source/Niagara/Public/NiagaraVolumeRendererProperties.h
@@ -78,7 +78,7 @@ public:
 	static int32 GetDefaultVolumeResolutionMaxAxis() { return 0; }
 	static FVector3f GetDefaultVolumeWorldSpaceSize() { return FVector3f::ZeroVector; }
 	static float GetDefaultStepFactor() { return 1.0f; }
-	static float GetDefaultShadowStepFactor() { return 8.0f; }
+	static float GetDefaultShadowStepFactor() { return 2.0f; }
 	static float GetDefaultShadowBiasFactor() { return 0.5f; }
 	static float GetDefaultLightingDownsampleFactor() { return 2.0f; }
 
@@ -111,7 +111,7 @@ public:
 	float LightingDownsampleFactor = 2.0f;
 
 	UPROPERTY(EditAnywhere, Category = "Volume Rendering")
-	float ShadowStepFactor = 8.0f;
+	float ShadowStepFactor = 2.0f;
 
 	UPROPERTY(EditAnywhere, Category = "Volume Rendering")
 	float ShadowBiasFactor = 0.5f;
diff --git a/Engine/Shaders/Private/HeterogeneousVolumes/HeterogeneousVolumesComposite.usf b/Engine/Shaders/Private/HeterogeneousVolumes/HeterogeneousVolumesComposite.usf
--- a/Engine/Shaders/Private/HeterogeneousVolumes/HeterogeneousVolumesComposite.usf
+++ b/Engine/Shaders/Private/HeterogeneousVolumes/HeterogeneousVolumesComposite.usf
@@ -11,11 +11,11 @@ void HeterogeneousVolumesCompositeCS(
 	uint2 DispatchThreadId : SV_DispatchThreadID
 )
 {
-	uint2 PixelCoord = DispatchThreadId.xy;
-	if (any(PixelCoord >= View.ViewSizeAndInvSize.xy))
+	if (any(DispatchThreadId.xy >= View.ViewSizeAndInvSize.xy))
 	{
 		return;
 	}
+	uint2 PixelCoord = DispatchThreadId.xy + View.ViewRectMin.xy;
 
 	float4 Result;
 	float HeterogeneousVolumeTransmittance = saturate(1.0 - HeterogeneousVolumeRadiance[PixelCoord].a);
diff --git a/Engine/Source/Runtime/Engine/Private/Components/HeterogeneousVolumeComponent.cpp b/Engine/Source/Runtime/Engine/Private/Components/HeterogeneousVolumeComponent.cpp
--- a/Engine/Source/Runtime/Engine/Private/Components/HeterogeneousVolumeComponent.cpp
+++ b/Engine/Source/Runtime/Engine/Private/Components/HeterogeneousVolumeComponent.cpp
@@ -180,48 +180,45 @@ void FHeterogeneousVolumeSceneProxy::GetDynamicMeshElements(
 		return;
 	}
 
-	// Create a dummy MeshBatch to make the system happy..
-	if (MaterialInterface)
-	{
-		// Set up MeshBatch
-		FMeshBatch& Mesh = Collector.AllocateMesh();
-
-		Mesh.VertexFactory = &VertexFactory;
-		Mesh.MaterialRenderProxy = MaterialInterface->GetRenderProxy();
-		Mesh.LCI = NULL;
-		Mesh.ReverseCulling = IsLocalToWorldDeterminantNegative() ? true : false;
-		Mesh.CastShadow = CastsDynamicShadow();
-		//Mesh.DepthPriorityGroup = (ESceneDepthPriorityGroup)GetDepthPriorityGroup(View);
-		Mesh.Type = PT_TriangleStrip;
-		Mesh.bDisableBackfaceCulling = true;
-
-		// Set up the FMeshBatchElement.
-		FMeshBatchElement& BatchElement = Mesh.Elements[0];
-		BatchElement.IndexBuffer = NULL;
-		BatchElement.FirstIndex = 0;
-		BatchElement.MinVertexIndex = 0;
-		BatchElement.MaxVertexIndex = 3;
-		BatchElement.NumPrimitives = 2;
-		BatchElement.BaseVertexIndex = 0;
-
-		//FHeterogeneousVolumeData* HeterogeneousVolumeData = &Collector.AllocateOneFrameResource<FHeterogeneousVolumeData>(this);
-		BatchElement.UserData = &HeterogeneousVolumeData;
-
-		Mesh.bCanApplyViewModeOverrides = true;
-		Mesh.bUseWireframeSelectionColoring = IsSelected();
-		Mesh.bUseSelectionOutline = false;
-		Mesh.bSelectable = false;
-
-		Collector.AddMesh(0, Mesh);
-	}
-
-	// Draw bounds 
 	if (ViewFamily.EngineShowFlags.HeterogeneousVolumes)
 	{
 		for (int32 ViewIndex = 0; ViewIndex < Views.Num(); ViewIndex++)
 		{
 			if (VisibilityMap & (1 << ViewIndex))
 			{
+				if (MaterialInterface)
+				{
+					// Set up MeshBatch
+					FMeshBatch& Mesh = Collector.AllocateMesh();
+
+					Mesh.VertexFactory = &VertexFactory;
+					Mesh.MaterialRenderProxy = MaterialInterface->GetRenderProxy();
+					Mesh.LCI = NULL;
+					Mesh.ReverseCulling = IsLocalToWorldDeterminantNegative() ? true : false;
+					Mesh.CastShadow = CastsDynamicShadow();
+					Mesh.Type = PT_TriangleStrip;
+					Mesh.bDisableBackfaceCulling = true;
+
+					// Set up the FMeshBatchElement.
+					FMeshBatchElement& BatchElement = Mesh.Elements[0];
+					BatchElement.IndexBuffer = NULL;
+					BatchElement.FirstIndex = 0;
+					BatchElement.MinVertexIndex = 0;
+					BatchElement.MaxVertexIndex = 3;
+					BatchElement.NumPrimitives = 2;
+					BatchElement.BaseVertexIndex = 0;
+
+					// Heterogeneous Volume Interface is passed through UserData.
+					BatchElement.UserData = &HeterogeneousVolumeData;
+
+					Mesh.bCanApplyViewModeOverrides = true;
+					Mesh.bUseWireframeSelectionColoring = IsSelected();
+					Mesh.bUseSelectionOutline = false;
+					Mesh.bSelectable = false;
+
+					Collector.AddMesh(ViewIndex, Mesh);
+				}
+
 				const FSceneView* View = Views[ViewIndex];
 				RenderBounds(Collector.GetPDI(ViewIndex), View->Family->EngineShowFlags, GetBounds(), IsSelected());
 			}
diff --git a/Engine/Source/Runtime/Renderer/Private/HeterogeneousVolumes/HeterogeneousVolumes.cpp b/Engine/Source/Runtime/Renderer/Private/HeterogeneousVolumes/HeterogeneousVolumes.cpp
--- a/Engine/Source/Runtime/Renderer/Private/HeterogeneousVolumes/HeterogeneousVolumes.cpp
+++ b/Engine/Source/Runtime/Renderer/Private/HeterogeneousVolumes/HeterogeneousVolumes.cpp
@@ -684,17 +684,21 @@ void FDeferredShadingSceneRenderer::RenderHeterogeneousVolumes(
 	TRDGUniformBufferRef<FOrthoVoxelGridUniformBufferParameters> OrthoGridUniformBuffer = HeterogeneousVolumes::GetOrthoVoxelGridUniformBuffer(GraphBuilder, Views[0].ViewState);
 	TRDGUniformBufferRef<FFrustumVoxelGridUniformBufferParameters> FrustumGridUniformBuffer = HeterogeneousVolumes::GetFrustumVoxelGridUniformBuffer(GraphBuilder, Views[0].ViewState);
 
+	FRDGTextureRef HeterogeneousVolumeRadiance = nullptr;
+	if (ShouldRenderHeterogeneousVolumesForAnyView(Views))
+	{
+		FRDGTextureDesc Desc = SceneTextures.Color.Target->Desc;
+		Desc.Format = PF_FloatRGBA;
+		Desc.Flags &= ~(TexCreate_FastVRAM);
+		HeterogeneousVolumeRadiance = GraphBuilder.CreateTexture(Desc, TEXT("HeterogeneousVolumes"));
+		AddClearUAVPass(GraphBuilder, GraphBuilder.CreateUAV(HeterogeneousVolumeRadiance), FLinearColor::Transparent);
+	}
+
 	for (int32 ViewIndex = 0; ViewIndex < Views.Num(); ViewIndex++)
 	{
 		FViewInfo& View = Views[ViewIndex];
 		if (ShouldRenderHeterogeneousVolumesForView(View))
 		{
-			FRDGTextureDesc Desc = SceneTextures.Color.Target->Desc;
-			Desc.Format = PF_FloatRGBA;
-			Desc.Flags &= ~(TexCreate_FastVRAM);
-			FRDGTextureRef HeterogeneousVolumeRadiance = GraphBuilder.CreateTexture(Desc, TEXT("HeterogeneousVolumes"));
-			AddClearUAVPass(GraphBuilder, GraphBuilder.CreateUAV(HeterogeneousVolumeRadiance), FLinearColor::Transparent);
-
 			if (HeterogeneousVolumes::GetDebugMode() != 0)
 			{
 				// TODO: Replace with single-scattering voxel grid implementation.
@@ -961,4 +965,4 @@ namespace HeterogeneousVolumes
 			}
 		}
 	}
-}
\ No newline at end of file
+}
