Add the ability to ignore the dirty flag when querying for external objects to save

diff --git a/Engine/Source/Runtime/Engine/Private/ExternalPackageHelper.cpp b/Engine/Source/Runtime/Engine/Private/ExternalPackageHelper.cpp
--- a/Engine/Source/Runtime/Engine/Private/ExternalPackageHelper.cpp
+++ b/Engine/Source/Runtime/Engine/Private/ExternalPackageHelper.cpp
@@ -136,7 +136,7 @@ FString FExternalPackageHelper::GetExternalObjectPackageInstanceName(const FStri
 	return FLinkerInstancingContext::GetInstancedPackageName(OuterPackageName, ObjectPackageName);
 }
 
-void FExternalPackageHelper::GetExternalSaveableObjects(UObject* InOuter, TArray<UObject*>& OutObjects)
+void FExternalPackageHelper::GetExternalSaveableObjects(UObject* InOuter, TArray<UObject*>& OutObjects, EGetExternalSaveableObjectsFlags InFlags)
 {
 	// Get external packages
 	TSet<UPackage*> ExternalObjectPackages;
@@ -145,7 +145,8 @@ void FExternalPackageHelper::GetExternalSaveableObjects(UObject* InOuter, TArray
 	// Find assets for external packages
 	for (UPackage* ExternalPackage : ExternalObjectPackages)
 	{
-		if(FPackageName::IsValidLongPackageName(ExternalPackage->GetName()) && ExternalPackage->IsDirty())
+		const bool bPassesDirtyCheck = !EnumHasAnyFlags(InFlags, EGetExternalSaveableObjectsFlags::CheckDirty) || ExternalPackage->IsDirty();
+		if(bPassesDirtyCheck && FPackageName::IsValidLongPackageName(ExternalPackage->GetName()))
 		{
 			if(UObject* Asset = ExternalPackage->FindAssetInPackage())
 			{
diff --git a/Engine/Source/Runtime/Engine/Public/ExternalPackageHelper.h b/Engine/Source/Runtime/Engine/Public/ExternalPackageHelper.h
--- a/Engine/Source/Runtime/Engine/Public/ExternalPackageHelper.h
+++ b/Engine/Source/Runtime/Engine/Public/ExternalPackageHelper.h
@@ -79,12 +79,24 @@ public:
 	template<typename T>
 	static void LoadObjectsFromExternalPackages(UObject* InOuter, TFunctionRef<void(T*)> Operation);
 
+	enum class EGetExternalSaveableObjectsFlags : uint32
+	{
+		// No flags
+		None                   = 0,
+
+		// Whether to check the object's package dirty flag. Controls whether only dirty or all external objects are returned.
+		CheckDirty             = (1 << 0) 
+	};
+
+	FRIEND_ENUM_CLASS_FLAGS(EGetExternalSaveableObjectsFlags);
+
 	/**
 	 * Get the saveable external objects that should be saved alongside this outer's package
 	 * @param InOuter		The external object's outer
 	 * @param OutObjects	The objects that should be saved
+	 * @param InFlags		Flags controlling behavior @see EGetExternalSaveableObjectsFlags
 	 */
-	static ENGINE_API void GetExternalSaveableObjects(UObject* InOuter, TArray<UObject*>& OutObjects);
+	static ENGINE_API void GetExternalSaveableObjects(UObject* InOuter, TArray<UObject*>& OutObjects, EGetExternalSaveableObjectsFlags InFlags = EGetExternalSaveableObjectsFlags::CheckDirty);
 
 	/**
 	 * Returns an array of external package file paths for the provided objects
