Make scene capture ortho plane correction opt-in rather than automatically enabled to avoid breaking older use case assets where the capture is used for determining location data (e.g. water/terrain).

diff --git a/Engine/Source/Runtime/Engine/Classes/Components/SceneCaptureComponent2D.h b/Engine/Source/Runtime/Engine/Classes/Components/SceneCaptureComponent2D.h
--- a/Engine/Source/Runtime/Engine/Classes/Components/SceneCaptureComponent2D.h
+++ b/Engine/Source/Runtime/Engine/Classes/Components/SceneCaptureComponent2D.h
@@ -34,10 +34,14 @@ public:
 	UPROPERTY(EditAnywhere, BlueprintReadWrite, Category=Projection)
 	float OrthoWidth;
 
-	/** Automatically determine a min/max Near/Far clip plane position depending on OrthoWidth value*/
+	/** Automatically determine a min/max Near/Far clip plane position depending on OrthoWidth value */
 	UPROPERTY(EditAnywhere, BlueprintReadWrite, Category=Projection)
 	bool bAutoCalculateOrthoPlanes;
 
+	/** Adjusts the view origin to the near plane location, and forces near plane to 0. Helps with artefacting when auto-adjusting the plane range. */
+	UPROPERTY(EditAnywhere, BlueprintReadWrite, Category=Projection)
+	bool bUpdateOrthoPlanes;
+
 	/** Output render target of the scene capture that can be read in materials. */
 	UPROPERTY(EditAnywhere, BlueprintReadWrite, Category=SceneCapture)
 	TObjectPtr<class UTextureRenderTarget2D> TextureTarget;
diff --git a/Engine/Source/Runtime/Engine/Private/Components/SceneCaptureComponent.cpp b/Engine/Source/Runtime/Engine/Private/Components/SceneCaptureComponent.cpp
--- a/Engine/Source/Runtime/Engine/Private/Components/SceneCaptureComponent.cpp
+++ b/Engine/Source/Runtime/Engine/Private/Components/SceneCaptureComponent.cpp
@@ -569,6 +569,7 @@ USceneCaptureComponent2D::USceneCaptureComponent2D(const FObjectInitializer& Obj
 
 	OrthoWidth = DEFAULT_ORTHOWIDTH;
 	bAutoCalculateOrthoPlanes = true;
+	bUpdateOrthoPlanes = false;
 
 	bUseCustomProjectionMatrix = false;
 	bAutoActivate = true;
diff --git a/Engine/Source/Runtime/Renderer/Private/SceneCaptureRendering.cpp b/Engine/Source/Runtime/Renderer/Private/SceneCaptureRendering.cpp
--- a/Engine/Source/Runtime/Renderer/Private/SceneCaptureRendering.cpp
+++ b/Engine/Source/Runtime/Renderer/Private/SceneCaptureRendering.cpp
@@ -845,7 +845,7 @@ static FSceneRenderer* CreateSceneRendererForSceneCapture(
 	// Use camera position correction for ortho scene captures
 	if(USceneCaptureComponent2D * SceneCaptureComponent2D = Cast<USceneCaptureComponent2D>(SceneCaptureComponent))
 	{
-		if (!SceneCaptureViewInfo.IsPerspectiveProjection() && !SceneCaptureComponent2D->bUseCustomProjectionMatrix)
+		if (!SceneCaptureViewInfo.IsPerspectiveProjection() && SceneCaptureComponent2D->bUpdateOrthoPlanes)
 		{
 			SceneCaptureViewInfo.UpdateOrthoPlanes();
 		}
