World Partition - When running the "wp.Editor.DumpStreamingGenerationLog" command to obtain a streaming generation log, output the path to the file in the log

diff --git a/Engine/Source/Runtime/Engine/Private/WorldPartition/WorldPartitionStreamingGeneration.cpp b/Engine/Source/Runtime/Engine/Private/WorldPartition/WorldPartitionStreamingGeneration.cpp
--- a/Engine/Source/Runtime/Engine/Private/WorldPartition/WorldPartitionStreamingGeneration.cpp
+++ b/Engine/Source/Runtime/Engine/Private/WorldPartition/WorldPartitionStreamingGeneration.cpp
@@ -54,6 +54,11 @@ static FAutoConsoleCommand DumpStreamingGenerationLog(
 				UWorldPartition::FGenerateStreamingContext Context;
 				WorldPartition->GenerateStreaming(Params, Context);
 				WorldPartition->FlushStreaming();
+
+				if (Context.OutputLogFilename.IsSet())
+				{
+					UE_LOG(LogWorldPartition, Display, TEXT("Streaming generation details logged to '%s'"), *FPaths::ConvertRelativePathToFull(Context.OutputLogFilename.GetValue()));
+				}
 			}
 		}
 	})
