Add GameAnalyticsVar to EngineNetAnalytics

#rnx

diff --git a/Engine/Source/Runtime/Engine/Private/Analytics/EngineNetAnalytics.cpp b/Engine/Source/Runtime/Engine/Private/Analytics/EngineNetAnalytics.cpp
--- a/Engine/Source/Runtime/Engine/Private/Analytics/EngineNetAnalytics.cpp
+++ b/Engine/Source/Runtime/Engine/Private/Analytics/EngineNetAnalytics.cpp
@@ -555,6 +555,8 @@ void FNetConnAnalyticsData::SendAnalytics()
 			AppendAnalyticsEventAttributeArray(Attrs, EZAttrib_FailedPingUDP, FJsonFragment(MoveTemp(FailedPingAddressesUDPJsonStr)));
 		}
 
+		Aggregator->AppendGameInstanceAttributes(Attrs);
+
 		AnalyticsProvider->RecordEvent(EZEventName, MoveTemp(Attrs));
 	}
 }
