ISM: Fix bounds logic so that it works for all ISM types.
Existing logic tested PerInstanceSMData array size but HISM doesn't populate that array. Use the virtual GetNumRenderInstances() instead.
Fixes grass getting incorrect bounds and being culled.

diff --git a/Engine/Source/Runtime/Engine/Private/InstancedStaticMesh.cpp b/Engine/Source/Runtime/Engine/Private/InstancedStaticMesh.cpp
--- a/Engine/Source/Runtime/Engine/Private/InstancedStaticMesh.cpp
+++ b/Engine/Source/Runtime/Engine/Private/InstancedStaticMesh.cpp
@@ -2701,7 +2701,7 @@ FBoxSphereBounds UInstancedStaticMeshComponent::CalcBoundsImpl(const FTransform&
 
 void UInstancedStaticMeshComponent::UpdateBounds()
 {
-	if (PerInstanceSMData.Num() == 0)
+	if (GetNumRenderInstances() == 0)
 	{
 		Bounds = FBoxSphereBounds(GetComponentTransform().GetLocation(), FVector::ZeroVector, 0.f);
 	}
