Fix deadlock in FAudioChunkCache::AddForceInlineSoundWave. Remove unmatched mutex Lock call from FAudioChunkCache::RemoveForceInlineSoundWave

#rnx

#localization none
#tests PIE

diff --git a/Engine/Source/Runtime/Engine/Private/AudioStreamingCache.cpp b/Engine/Source/Runtime/Engine/Private/AudioStreamingCache.cpp
--- a/Engine/Source/Runtime/Engine/Private/AudioStreamingCache.cpp
+++ b/Engine/Source/Runtime/Engine/Private/AudioStreamingCache.cpp
@@ -1087,7 +1087,6 @@ void FAudioChunkCache::RemoveForceInlineSoundWave(const FSoundWaveProxyPtr& Soun
 	// scope lock
 	{
 		FScopeLock Lock(&SoundWaveMemoryTrackerCritSec);
-		SoundWaveMemoryTrackerCritSec.Lock();
 
 		FSoundWaveMemoryTracker* Tracker = SoundWaveTracker.Find(SoundWave);
 		checkf(Tracker != nullptr, TEXT("AudioStreamCache::RemoveForceInlineSoundWave: Attempted to remove SoundWave that was never added, or has already been removed: %s"), *SoundWave->GetFName().ToString());
