World Partition - Level Instance: Disable support of partial loading in non world partitionned worlds.

diff --git a/Engine/Source/Runtime/Engine/Private/LevelInstance/LevelInstanceInterface.cpp b/Engine/Source/Runtime/Engine/Private/LevelInstance/LevelInstanceInterface.cpp
--- a/Engine/Source/Runtime/Engine/Private/LevelInstance/LevelInstanceInterface.cpp
+++ b/Engine/Source/Runtime/Engine/Private/LevelInstance/LevelInstanceInterface.cpp
@@ -45,10 +45,18 @@ bool ILevelInstanceInterface::SupportsPartialEditorLoading() const
 				return false;
 			}
 
-			// If the level is loaded, check that it has editor streaming enabled
+			// If the level is loaded
 			if (ULevel* Level = GetLoadedLevel())
 			{
-				if (UWorldPartition* WorldPartition = Level->GetWorldPartition(); WorldPartition && WorldPartition->IsInitialized() && !WorldPartition->IsStreamingEnabledInEditor())
+				UWorldPartition* WorldPartition = Level->GetWorldPartition();
+
+				// Disable partial loading for non WP levels
+				if (!WorldPartition)
+				{
+					return false;
+				}
+				// Check that it has editor streaming enabled
+				else if (WorldPartition->IsInitialized() && !WorldPartition->IsStreamingEnabledInEditor())
 				{
 					return false;
 				}
