Fixes for TObjectPtr usage enforcement that were missed due to a bug that I'm fixing in another changelist.

diff --git a/Engine/Source/Runtime/Engine/Classes/PhysicsEngine/ClusterUnionComponent.h b/Engine/Source/Runtime/Engine/Classes/PhysicsEngine/ClusterUnionComponent.h
--- a/Engine/Source/Runtime/Engine/Classes/PhysicsEngine/ClusterUnionComponent.h
+++ b/Engine/Source/Runtime/Engine/Classes/PhysicsEngine/ClusterUnionComponent.h
@@ -170,7 +170,7 @@ struct FClusterUnionInitializationData
 	TObjectPtr<UClusterUnionComponent> ClusterUnionComponent;
 
 	UPROPERTY()
-	TArray<UPrimitiveComponent*> ProcessedComponents;
+	TArray<TObjectPtr<UPrimitiveComponent>> ProcessedComponents;
 };
 
 DECLARE_DYNAMIC_MULTICAST_DELEGATE_ThreeParams(FOnClusterUnionAddedComponent, UPrimitiveComponent*, Component, const TSet<int32>&, BoneIds, bool, bIsNew);
diff --git a/Engine/Source/Runtime/Engine/Private/PhysicsEngine/ClusterUnionComponent.cpp b/Engine/Source/Runtime/Engine/Private/PhysicsEngine/ClusterUnionComponent.cpp
--- a/Engine/Source/Runtime/Engine/Private/PhysicsEngine/ClusterUnionComponent.cpp
+++ b/Engine/Source/Runtime/Engine/Private/PhysicsEngine/ClusterUnionComponent.cpp
@@ -1061,7 +1061,14 @@ void UClusterUnionComponent::SyncClusterUnionFromProxy(const FTransform& NewTran
 
 	if (OnClusterUnionPostSyncBodiesEvent.IsBound())
 	{
-		OnClusterUnionPostSyncBodiesEvent.Broadcast({ this,GetPrimitiveComponents() });
+		TArray<TObjectPtr<UPrimitiveComponent>> PrimitiveComponents;
+		PrimitiveComponents.Reserve(PerComponentData.Num());
+
+		for (auto Iter = PerComponentData.CreateIterator(); Iter; ++Iter)
+		{
+			PrimitiveComponents.Add(Iter.Key().ResolveObjectPtr());
+		}
+		OnClusterUnionPostSyncBodiesEvent.Broadcast({ this,PrimitiveComponents });
 	}
 }
 
diff --git a/Engine/Source/Runtime/Engine/Private/WorldPartition/WorldPartitionRuntimeHash.cpp b/Engine/Source/Runtime/Engine/Private/WorldPartition/WorldPartitionRuntimeHash.cpp
--- a/Engine/Source/Runtime/Engine/Private/WorldPartition/WorldPartitionRuntimeHash.cpp
+++ b/Engine/Source/Runtime/Engine/Private/WorldPartition/WorldPartitionRuntimeHash.cpp
@@ -94,11 +94,11 @@ void URuntimeHashExternalStreamingObjectBase::OnStreamingObjectLoaded(UWorld* In
 #if WITH_EDITOR
 UWorldPartitionRuntimeCell* URuntimeHashExternalStreamingObjectBase::GetCellForCookPackage(const FString& InCookPackageName) const
 {
-	if (UWorldPartitionRuntimeCell* const* MatchingCell = PackagesToGenerateForCook.Find(InCookPackageName))
+	if (const TObjectPtr<UWorldPartitionRuntimeCell>* MatchingCell = PackagesToGenerateForCook.Find(InCookPackageName))
 	{
 		if (ensure(*MatchingCell))
 		{
-			return const_cast<UWorldPartitionRuntimeCell*>(*MatchingCell);
+			return const_cast<UWorldPartitionRuntimeCell*>(ToRawPtr(*MatchingCell));
 		}
 	}
 	return nullptr;
@@ -356,11 +356,11 @@ void UWorldPartitionRuntimeHash::PopulateRuntimeCell(UWorldPartitionRuntimeCell*
 
 UWorldPartitionRuntimeCell* UWorldPartitionRuntimeHash::GetCellForCookPackage(const FString& InCookPackageName) const
 {
-	if (UWorldPartitionRuntimeCell* const* MatchingCell = PackagesToGenerateForCook.Find(InCookPackageName))
+	if (const TObjectPtr<UWorldPartitionRuntimeCell>* MatchingCell = PackagesToGenerateForCook.Find(InCookPackageName))
 	{
 		if (ensure(*MatchingCell))
 		{
-			return const_cast<UWorldPartitionRuntimeCell*>(*MatchingCell);
+			return const_cast<UWorldPartitionRuntimeCell*>(ToRawPtr(*MatchingCell));
 		}
 	}
 	return nullptr;
diff --git a/Engine/Source/Runtime/Engine/Public/WorldPartition/WorldPartitionRuntimeHash.h b/Engine/Source/Runtime/Engine/Public/WorldPartition/WorldPartitionRuntimeHash.h
--- a/Engine/Source/Runtime/Engine/Public/WorldPartition/WorldPartitionRuntimeHash.h
+++ b/Engine/Source/Runtime/Engine/Public/WorldPartition/WorldPartitionRuntimeHash.h
@@ -121,7 +121,7 @@ protected:
 
 #if WITH_EDITORONLY_DATA
 	UPROPERTY(Transient)
-	TMap<FString, UWorldPartitionRuntimeCell*> PackagesToGenerateForCook;
+	TMap<FString, TObjectPtr<UWorldPartitionRuntimeCell>> PackagesToGenerateForCook;
 
 	UPROPERTY(Transient)
 	TObjectPtr<const UExternalDataLayerAsset> ExternalDataLayerAsset;
@@ -249,7 +249,7 @@ protected:
 
 	TArray<FAlwaysLoadedActorForPIE> AlwaysLoadedActorsForPIE;
 
-	TMap<FString, UWorldPartitionRuntimeCell*> PackagesToGenerateForCook;
+	TMap<FString, TObjectPtr<UWorldPartitionRuntimeCell>> PackagesToGenerateForCook;
 #endif
 
 protected:
