Fixed FWorldDataLayerActorDesc : Main WorldDataLayers now doesn't save the transient DataLayerInstanceDescs and EDL WorldDataLayers now saves its own list of DataLayerInstanceDescs.

diff --git a/Engine/Source/Runtime/Engine/Private/WorldPartition/DataLayer/DataLayerUtils.cpp b/Engine/Source/Runtime/Engine/Private/WorldPartition/DataLayer/DataLayerUtils.cpp
--- a/Engine/Source/Runtime/Engine/Private/WorldPartition/DataLayer/DataLayerUtils.cpp
+++ b/Engine/Source/Runtime/Engine/Private/WorldPartition/DataLayer/DataLayerUtils.cpp
@@ -101,7 +101,7 @@ FDataLayerInstanceNames FDataLayerUtils::ResolveDataLayerInstanceNames(const UDa
 }
 
 // For performance reasons, this function assumes that InActorDesc's DataLayerInstanceNames was already resolved.
-bool FDataLayerUtils::ResolveRuntimeDataLayerInstanceNames(const UDataLayerManager* InDataLayerManager, const IWorldPartitionActorDescInstanceView& InActorDescView, const FStreamingGenerationActorDescViewMap& InActorDescViewMap, FDataLayerInstanceNames& OutRuntimeDataLayerInstanceNames)
+bool FDataLayerUtils::ResolveRuntimeDataLayerInstanceNames(const UDataLayerManager* InDataLayerManager, const IWorldPartitionActorDescInstanceView& InActorDescView, const TArray<const FWorldDataLayersActorDesc*>& InWorldDataLayersActorDescs, FDataLayerInstanceNames& OutRuntimeDataLayerInstanceNames)
 {
 	const FDataLayerInstanceNames& ActorDescViewDataLayerInstanceNames = InActorDescView.GetDataLayerInstanceNames();
 	const FName EDLInstanceName = ActorDescViewDataLayerInstanceNames.GetExternalDataLayer();
@@ -130,7 +130,7 @@ bool FDataLayerUtils::ResolveRuntimeDataLayerInstanceNames(const UDataLayerManag
 	}
 
 	// Fallback on FWorldDataLayersActorDesc
-	TArray<const FWorldDataLayersActorDesc*> WorldDataLayersActorDescs = FindWorldDataLayerActorDescs(InActorDescViewMap);
+	const TArray<const FWorldDataLayersActorDesc*>& WorldDataLayersActorDescs = InWorldDataLayersActorDescs;
 	if (WorldDataLayersActorDescs.Num())
 	{
 		check(AreWorldDataLayersActorDescsSane(WorldDataLayersActorDescs));
@@ -185,16 +185,6 @@ const FDataLayerInstanceDesc* FDataLayerUtils::GetDataLayerInstanceDescFromAsset
 	return nullptr;
 }
 
-TArray<const FWorldDataLayersActorDesc*> FDataLayerUtils::FindWorldDataLayerActorDescs(const FStreamingGenerationActorDescViewMap& InActorDescViewMap)
-{
-	TArray<const FWorldDataLayersActorDesc*> WorldDataLayersActorDescs;
-	TArray<const FStreamingGenerationActorDescView*> WorldDataLayerViews = InActorDescViewMap.FindByExactNativeClass<AWorldDataLayers>();
-	Algo::TransformIf(WorldDataLayerViews, WorldDataLayersActorDescs,
-		[](const FStreamingGenerationActorDescView* WorldDataLayersActorDescView) { return ((FWorldDataLayersActorDesc*)WorldDataLayersActorDescView->GetActorDesc())->IsValid(); },
-		[](const FStreamingGenerationActorDescView* WorldDataLayersActorDescView) { return (FWorldDataLayersActorDesc*)WorldDataLayersActorDescView->GetActorDesc(); });
-	return WorldDataLayersActorDescs;
-}
-
 bool FDataLayerUtils::AreWorldDataLayersActorDescsSane(const TArray<const FWorldDataLayersActorDesc*>& InWorldDataLayersActorDescs)
 {
 	// Deprecation handling: Case where pre 5.1 map, without WorldDataLayersActorDescs, are resolving their data layer outside of their world (changelist validation). They need proper Actor descs to resolve.
diff --git a/Engine/Source/Runtime/Engine/Private/WorldPartition/DataLayer/WorldDataLayers.cpp b/Engine/Source/Runtime/Engine/Private/WorldPartition/DataLayer/WorldDataLayers.cpp
--- a/Engine/Source/Runtime/Engine/Private/WorldPartition/DataLayer/WorldDataLayers.cpp
+++ b/Engine/Source/Runtime/Engine/Private/WorldPartition/DataLayer/WorldDataLayers.cpp
@@ -1579,7 +1579,7 @@ bool AWorldDataLayers::ShouldLevelKeepRefIfExternal() const
 
 bool AWorldDataLayers::IsEditorOnly() const
 {
-	return Super::IsEditorOnly() || IsExternalDataLayerWorldDataLayers();
+	return IsExternalDataLayerWorldDataLayers();
 }
 
 #endif
diff --git a/Engine/Source/Runtime/Engine/Private/WorldPartition/DataLayer/WorldDataLayersActorDesc.cpp b/Engine/Source/Runtime/Engine/Private/WorldPartition/DataLayer/WorldDataLayersActorDesc.cpp
--- a/Engine/Source/Runtime/Engine/Private/WorldPartition/DataLayer/WorldDataLayersActorDesc.cpp
+++ b/Engine/Source/Runtime/Engine/Private/WorldPartition/DataLayer/WorldDataLayersActorDesc.cpp
@@ -218,14 +218,11 @@ void FWorldDataLayersActorDesc::Init(const AActor* InActor)
 	const AWorldDataLayers* WorldDataLayers = CastChecked<AWorldDataLayers>(InActor);
 	bUseExternalPackageDataLayerInstances = WorldDataLayers->IsUsingExternalPackageDataLayerInstances();
 	bIsExternalDataLayerWorldDataLayers = WorldDataLayers->IsExternalDataLayerWorldDataLayers();
-	if (!bUseExternalPackageDataLayerInstances)
+	const IDataLayerInstanceProvider* DataLayerInstanceProvider = CastChecked<IDataLayerInstanceProvider>(WorldDataLayers);
+	for (UDataLayerInstance* DataLayerInstance : DataLayerInstanceProvider->GetDataLayerInstances())
 	{
-		WorldDataLayers->ForEachDataLayerInstance([this](UDataLayerInstance* DataLayerInstance)
-		{
-			FDataLayerInstanceDesc& DataLayerInstanceDesc = DataLayerInstances.Emplace_GetRef();
-			DataLayerInstanceDesc.Init(DataLayerInstance);
-			return true;
-		});
+		FDataLayerInstanceDesc& DataLayerInstanceDesc = DataLayerInstances.Emplace_GetRef();
+		DataLayerInstanceDesc.Init(DataLayerInstance);
 	}
 
 	bIsValid = true;
diff --git a/Engine/Source/Runtime/Engine/Private/WorldPartition/RuntimeHashSet/WorldPartitionRuntimeHashSetHLODGeneration.cpp b/Engine/Source/Runtime/Engine/Private/WorldPartition/RuntimeHashSet/WorldPartitionRuntimeHashSetHLODGeneration.cpp
--- a/Engine/Source/Runtime/Engine/Private/WorldPartition/RuntimeHashSet/WorldPartitionRuntimeHashSetHLODGeneration.cpp
+++ b/Engine/Source/Runtime/Engine/Private/WorldPartition/RuntimeHashSet/WorldPartitionRuntimeHashSetHLODGeneration.cpp
@@ -306,7 +306,7 @@ bool UWorldPartitionRuntimeHashSet::SetupHLODActors(const IStreamingGenerationCo
 				ActorSetInstance.ActorSet = ActorSet;
 
 				FDataLayerInstanceNames RuntimeDataLayerInstanceNames;
-				if (FDataLayerUtils::ResolveRuntimeDataLayerInstanceNames(DataLayerManager, *HLODActorDescView, *BaseActorSetContainerInstance->ActorDescViewMap, RuntimeDataLayerInstanceNames))
+				if (FDataLayerUtils::ResolveRuntimeDataLayerInstanceNames(DataLayerManager, *HLODActorDescView, *BaseActorSetContainerInstance->DataLayerResolvers, RuntimeDataLayerInstanceNames))
 				{
 					HLODActorDescView->SetRuntimeDataLayerInstanceNames(RuntimeDataLayerInstanceNames);
 					ActorSetInstance.DataLayers = DataLayerManager->GetRuntimeDataLayerInstances(RuntimeDataLayerInstanceNames.ToArray());
diff --git a/Engine/Source/Runtime/Engine/Private/WorldPartition/RuntimeSpatialHash/RuntimeSpatialHashHLOD.cpp b/Engine/Source/Runtime/Engine/Private/WorldPartition/RuntimeSpatialHash/RuntimeSpatialHashHLOD.cpp
--- a/Engine/Source/Runtime/Engine/Private/WorldPartition/RuntimeSpatialHash/RuntimeSpatialHashHLOD.cpp
+++ b/Engine/Source/Runtime/Engine/Private/WorldPartition/RuntimeSpatialHash/RuntimeSpatialHashHLOD.cpp
@@ -525,7 +525,7 @@ bool UWorldPartitionRuntimeSpatialHash::SetupHLODActors(const IStreamingGenerati
 				FStreamingGenerationActorDescView* ActorDescView = NonConstActorDescViewMap->Emplace(HLODActorDescInstance);
 				
 				FDataLayerInstanceNames RuntimeDataLayerInstanceNames;
-				if (FDataLayerUtils::ResolveRuntimeDataLayerInstanceNames(DataLayerManager, *ActorDescView, *BaseActorSetContainerInstance->ActorDescViewMap, RuntimeDataLayerInstanceNames))
+				if (FDataLayerUtils::ResolveRuntimeDataLayerInstanceNames(DataLayerManager, *ActorDescView, *BaseActorSetContainerInstance->DataLayerResolvers, RuntimeDataLayerInstanceNames))
 				{
 					ActorDescView->SetRuntimeDataLayerInstanceNames(RuntimeDataLayerInstanceNames);
 				}
diff --git a/Engine/Source/Runtime/Engine/Private/WorldPartition/WorldPartitionStreamingGeneration.cpp b/Engine/Source/Runtime/Engine/Private/WorldPartition/WorldPartitionStreamingGeneration.cpp
--- a/Engine/Source/Runtime/Engine/Private/WorldPartition/WorldPartitionStreamingGeneration.cpp
+++ b/Engine/Source/Runtime/Engine/Private/WorldPartition/WorldPartitionStreamingGeneration.cpp
@@ -29,6 +29,8 @@
 #include "WorldPartition/DataLayer/ExternalDataLayerHelper.h"
 #include "WorldPartition/DataLayer/ExternalDataLayerManager.h"
 #include "WorldPartition/DataLayer/DataLayerInstanceNames.h"
+#include "WorldPartition/DataLayer/WorldDataLayersActorDesc.h"
+#include "WorldPartition/DataLayer/WorldDataLayers.h"
 #include "WorldPartition/ErrorHandling/WorldPartitionStreamingGenerationNullErrorHandler.h"
 #include "WorldPartition/ErrorHandling/WorldPartitionStreamingGenerationLogErrorHandler.h"
 #include "WorldPartition/HLOD/HLODActor.h"
@@ -450,6 +452,7 @@ class FWorldPartitionStreamingGenerator
 
 				FActorSetContainerInstance& ActorSetContainer = ActorSetContainerInstances[ContainerIndex];
 				ActorSetContainer.ActorDescViewMap = &ContainerDescriptor.ActorDescViewMap.Get();
+				ActorSetContainer.DataLayerResolvers = &ContainerDescriptor.DataLayerResolvers;
 				ActorSetContainer.ContainerInstanceCollection = ContainerDescriptor.ContainerInstanceCollection.Get();
 
 				ActorSetContainer.ActorSets.Empty(ContainerDescriptor.Clusters.Num());
@@ -592,6 +595,9 @@ class FWorldPartitionStreamingGenerator
 		/** The actor descriptor views for for this descriptor (TUniqueObj so it is moveable without having to update FStreamingGenerationActorDescView::ActorDescViewMap pointer */
 		TUniqueObj<FStreamingGenerationActorDescViewMap> ActorDescViewMap;
 
+		/** List of FWorldDataLayerActorDesc used to resolve Data Layers when DataLayerManager is not available */
+		TArray<const FWorldDataLayersActorDesc*> DataLayerResolvers;
+
 		/** Set of editor-only actors that are not part of the actor descriptor views */
 		TSet<FGuid> EditorOnlyActorDescSet;
 
@@ -691,20 +697,21 @@ class FWorldPartitionStreamingGenerator
 		}
 	}
 
-	void ResolveRuntimeDataLayers(FStreamingGenerationActorDescView& ActorDescView, const FStreamingGenerationActorDescViewMap& ActorDescViewMap)
+	void ResolveRuntimeDataLayers(FStreamingGenerationActorDescView& ActorDescView, const TArray<const FWorldDataLayersActorDesc*>& InDataLayerResolvers)
 	{
 		// Resolve DataLayerInstanceNames of ActorDescView only when necessary (i.e. when container is a template)
 		if (!ActorDescView.HasResolvedDataLayerInstanceNames())
 		{
 			// Build a WorldDataLayerActorDescs if DataLayerManager can't resolve Data Layers (i.e. when validating changelists and World is not loaded)
+			static const TArray<const FWorldDataLayersActorDesc*> Empty;
 			const bool bDataLayerManagerCanResolve = DataLayerManager && DataLayerManager->CanResolveDataLayers();
-			const TArray<const FWorldDataLayersActorDesc*> WorldDataLayerActorDescs = !bDataLayerManagerCanResolve ? FDataLayerUtils::FindWorldDataLayerActorDescs(ActorDescViewMap) : TArray<const FWorldDataLayersActorDesc*>();
-			const FDataLayerInstanceNames DataLayerInstanceNames = FDataLayerUtils::ResolveDataLayerInstanceNames(DataLayerManager, ActorDescView.GetActorDesc(), WorldDataLayerActorDescs);
+			const TArray<const FWorldDataLayersActorDesc*>& DataLayerResolvers = !bDataLayerManagerCanResolve ? InDataLayerResolvers : Empty;
+			const FDataLayerInstanceNames DataLayerInstanceNames = FDataLayerUtils::ResolveDataLayerInstanceNames(DataLayerManager, ActorDescView.GetActorDesc(), DataLayerResolvers);
 			ActorDescView.SetDataLayerInstanceNames(DataLayerInstanceNames);
 		}
 
 		FDataLayerInstanceNames RuntimeDataLayerInstanceNames;
-		if (FDataLayerUtils::ResolveRuntimeDataLayerInstanceNames(DataLayerManager, ActorDescView, ActorDescViewMap, RuntimeDataLayerInstanceNames))
+		if (FDataLayerUtils::ResolveRuntimeDataLayerInstanceNames(DataLayerManager, ActorDescView, InDataLayerResolvers, RuntimeDataLayerInstanceNames))
 		{
 			ActorDescView.SetRuntimeDataLayerInstanceNames(RuntimeDataLayerInstanceNames);
 		}
@@ -739,6 +746,7 @@ class FWorldPartitionStreamingGenerator
 	{
 		const FStreamingGenerationContainerInstanceCollection& InActorDescCollection = *InContainerCollectionInstanceDescriptor.ContainerInstanceCollection;
 		FStreamingGenerationActorDescViewMap& OutActorDescViewMap = InContainerCollectionInstanceDescriptor.ActorDescViewMap.Get();
+		TArray<const FWorldDataLayersActorDesc*>& OutDataLayerResolvers = InContainerCollectionInstanceDescriptor.DataLayerResolvers;
 		TSet<FGuid>& OutEditorOnlyActorDescSet = InContainerCollectionInstanceDescriptor.EditorOnlyActorDescSet;
 		const FActorContainerID& InContainerID = InContainerCollectionInstanceDescriptor.ID;
 		TArray<FStreamingGenerationActorDescView>& OutContainerInstances = InContainerCollectionInstanceDescriptor.ContainerCollectionInstanceViews;
@@ -798,6 +806,15 @@ class FWorldPartitionStreamingGenerator
 			// @todo_ow: this is to validate that new parenting of container instance code is equivalent
 			check(Iterator->GetContainerInstance()->GetContainerID() == InContainerID);
 
+			if (Iterator->GetActorNativeClass()->IsChildOf<AWorldDataLayers>())
+			{
+				const FWorldDataLayersActorDesc* WorldDataLayersActorDesc = (const FWorldDataLayersActorDesc*)Iterator->GetActorDesc();
+				if (WorldDataLayersActorDesc->IsValid())
+				{
+					OutDataLayerResolvers.Add(WorldDataLayersActorDesc);
+				}
+			}
+
 			AActor* Actor = nullptr;
 			if (ShouldRegisterActorDesc(*Iterator, &Actor))
 			{
@@ -1059,7 +1076,7 @@ class FWorldPartitionStreamingGenerator
 		{
 			ResolveRuntimeSpatiallyLoaded(ActorDescView);
 			ResolveRuntimeGrid(ActorDescView);
-			ResolveRuntimeDataLayers(ActorDescView, ContainerCollectionInstanceDescriptor.ActorDescViewMap.Get());
+			ResolveRuntimeDataLayers(ActorDescView, ContainerCollectionInstanceDescriptor.DataLayerResolvers);
 			ResolveHLODLayer(ActorDescView, WorldPartitionContext ? FSoftObjectPath(WorldPartitionContext->GetDefaultHLODLayer()) : FSoftObjectPath());
 			ResolveParentView(ActorDescView, ContainerCollectionInstanceDescriptor.ActorDescViewMap.Get(), ContainerCollectionInstanceDescriptor.EditorOnlyActorDescSet, ContainerCollectionInstanceDescriptor.EditorOnlyParentActorTransforms);
 		};
diff --git a/Engine/Source/Runtime/Engine/Public/WorldPartition/DataLayer/DataLayerUtils.h b/Engine/Source/Runtime/Engine/Public/WorldPartition/DataLayer/DataLayerUtils.h
--- a/Engine/Source/Runtime/Engine/Public/WorldPartition/DataLayer/DataLayerUtils.h
+++ b/Engine/Source/Runtime/Engine/Public/WorldPartition/DataLayer/DataLayerUtils.h
@@ -35,21 +35,22 @@ public:
 
 	PRAGMA_DISABLE_DEPRECATION_WARNINGS
 
-	UE_DEPRECATED(5.4, "Use FStreamingGenerationActorDescView version instead")
+	UE_DEPRECATED(5.4, "Use IWorldPartitionActorDescInstanceView version instead")
 	static ENGINE_API bool ResolveRuntimeDataLayerInstanceNames(const UDataLayerManager* InDataLayerManager, const class FWorldPartitionActorDescView& InActorDescView, const FStreamingGenerationActorDescViewMap& ActorDescViewMap, TArray<FName>& OutRuntimeDataLayerInstanceNames) { return false; }
 
+	UE_DEPRECATED(5.4, "This function is no longer used")
+	static ENGINE_API TArray<const FWorldDataLayersActorDesc*> FindWorldDataLayerActorDescs(const FStreamingGenerationActorDescViewMap& ActorDescViewMap) { return TArray<const FWorldDataLayersActorDesc*>(); }
+
 	PRAGMA_ENABLE_DEPRECATION_WARNINGS
 
 	static ENGINE_API FDataLayerInstanceNames ResolveDataLayerInstanceNames(const UDataLayerManager* InDataLayerManager, const FWorldPartitionActorDesc* InActorDesc, const TArray<const FWorldDataLayersActorDesc*>& InWorldDataLayersActorDescs = TArray<const FWorldDataLayersActorDesc*>());
 
-	static ENGINE_API bool ResolveRuntimeDataLayerInstanceNames(const UDataLayerManager* InDataLayerManager, const IWorldPartitionActorDescInstanceView& InActorDescView, const FStreamingGenerationActorDescViewMap& ActorDescViewMap, FDataLayerInstanceNames& OutRuntimeDataLayerInstanceNames);
+	static ENGINE_API bool ResolveRuntimeDataLayerInstanceNames(const UDataLayerManager* InDataLayerManager, const IWorldPartitionActorDescInstanceView& InActorDescView, const TArray<const FWorldDataLayersActorDesc*>& InWorldDataLayersActorDescs, FDataLayerInstanceNames& OutRuntimeDataLayerInstanceNames);
 
 	static ENGINE_API const FDataLayerInstanceDesc* GetDataLayerInstanceDescFromInstanceName(const TArray<const FWorldDataLayersActorDesc*>& InWorldDataLayersActorDescs, const FName& DataLayerInstanceName);
 
 	static ENGINE_API const FDataLayerInstanceDesc* GetDataLayerInstanceDescFromAssetPath(const TArray<const FWorldDataLayersActorDesc*>& InWorldDataLayersActorDescs, const FName& DataLayerAssetPath);
 
-	static ENGINE_API TArray<const FWorldDataLayersActorDesc*> FindWorldDataLayerActorDescs(const FStreamingGenerationActorDescViewMap& ActorDescViewMap);
-
 	static ENGINE_API bool AreWorldDataLayersActorDescsSane(const TArray<const FWorldDataLayersActorDesc*>& InWorldDataLayersActorDescs);
 
 	static ENGINE_API FString GenerateUniqueDataLayerShortName(const UDataLayerManager* InDataLayerManager, const FString& InNewShortName);
diff --git a/Engine/Source/Runtime/Engine/Public/WorldPartition/WorldPartitionStreamingGenerationContext.h b/Engine/Source/Runtime/Engine/Public/WorldPartition/WorldPartitionStreamingGenerationContext.h
--- a/Engine/Source/Runtime/Engine/Public/WorldPartition/WorldPartitionStreamingGenerationContext.h
+++ b/Engine/Source/Runtime/Engine/Public/WorldPartition/WorldPartitionStreamingGenerationContext.h
@@ -16,6 +16,7 @@ class FWorldPartitionStreamingGenerator;
 class FStreamingGenerationContainerInstanceCollection;
 class UActorDescContainer;
 class UDataLayerInstance;
+class FWorldDataLayersActorDesc;
 struct FWorldPartitionRuntimeContainerResolver;
 
 #if WITH_EDITOR
@@ -56,6 +57,7 @@ public:
 	{
 		FActorSetContainerInstance()
 			: ActorDescViewMap(nullptr)
+			, DataLayerResolvers(nullptr)
 			, ContainerInstanceCollection(nullptr)
 		{}
 
@@ -64,6 +66,7 @@ public:
 		FActorSetContainerInstance& operator=(const FActorSetContainerInstance&) = delete;
 
 		const FStreamingGenerationActorDescViewMap* ActorDescViewMap;
+		const TArray<const FWorldDataLayersActorDesc*>* DataLayerResolvers;
 		const FStreamingGenerationContainerInstanceCollection* ContainerInstanceCollection; // Only used by UWorldPartitionRuntimeSpatialHash::SetupHLODActors
 		TArray<TUniquePtr<FActorSet>> ActorSets;
 	};
