Logging to help track down dangling actor pointer in FPhysScene_Chaos::AsyncPhysicsTickCallback

diff --git a/Engine/Source/Runtime/Engine/Private/PhysicsEngine/Experimental/PhysScene_Chaos.cpp b/Engine/Source/Runtime/Engine/Private/PhysicsEngine/Experimental/PhysScene_Chaos.cpp
--- a/Engine/Source/Runtime/Engine/Private/PhysicsEngine/Experimental/PhysScene_Chaos.cpp
+++ b/Engine/Source/Runtime/Engine/Private/PhysicsEngine/Experimental/PhysScene_Chaos.cpp
@@ -2552,13 +2552,25 @@ void FPhysScene_Chaos::UnregisterAsyncPhysicsTickComponent(UActorComponent* Comp
 void FPhysScene_Chaos::RegisterAsyncPhysicsTickActor(AActor* Actor)
 {
 	EnableAsyncPhysicsTickCallback();
-	AsyncPhysicsTickCallback->AsyncPhysicsTickActors.Add(Actor);
+
+	bool bAlreadyRegistered = false;
+	AsyncPhysicsTickCallback->AsyncPhysicsTickActors.Add(Actor, &bAlreadyRegistered);
+
+	if (!bAlreadyRegistered)
+	{ 
+		UE_LOG(LogChaos, Log, TEXT("RegisterAsyncPhysicsTickActor %s @ 0x%x"), *AActor::GetDebugName(Actor), (SIZE_T)Actor);
+	}
 }
 
 void FPhysScene_Chaos::UnregisterAsyncPhysicsTickActor(AActor* Actor)
 {
 	if (AsyncPhysicsTickCallback)
 	{
+		if (AsyncPhysicsTickCallback->AsyncPhysicsTickActors.Contains(Actor))
+		{
+			UE_LOG(LogChaos, Log, TEXT("UnregisterAsyncPhysicsTickActor %s @ 0x%x"), *AActor::GetDebugName(Actor), (SIZE_T)Actor);
+		}
+
 		AsyncPhysicsTickCallback->AsyncPhysicsTickActors.Remove(Actor);
 	}
 }
