Blueprint function to enable/disable MACD on a mesh component

diff --git a/Engine/Source/Runtime/Engine/Classes/Components/PrimitiveComponent.h b/Engine/Source/Runtime/Engine/Classes/Components/PrimitiveComponent.h
--- a/Engine/Source/Runtime/Engine/Classes/Components/PrimitiveComponent.h
+++ b/Engine/Source/Runtime/Engine/Classes/Components/PrimitiveComponent.h
@@ -2711,6 +2711,14 @@ public:
 	UFUNCTION(BlueprintCallable, Category = "Physics")
 	ENGINE_API virtual void SetAllUseCCD(bool InUseCCD);
 
+	/** [EXPERIMENTAL] Set whether this component should use Motion-Aware Collision Detection */
+	UFUNCTION(BlueprintCallable, Category = "Physics")
+	ENGINE_API virtual void SetUseMACD(bool InUseMACD, FName BoneName = NAME_None);
+
+	/** [EXPERIMENTAL] Set whether all bodies in this component should use Motion-Aware Collision Detection */
+	UFUNCTION(BlueprintCallable, Category = "Physics")
+	ENGINE_API virtual void SetAllUseMACD(bool InUseMACD);
+
 	/**
 	 *	Force all bodies in this component to sleep.
 	 */
diff --git a/Engine/Source/Runtime/Engine/Private/PhysicsEngine/BodyInstance.cpp b/Engine/Source/Runtime/Engine/Private/PhysicsEngine/BodyInstance.cpp
--- a/Engine/Source/Runtime/Engine/Private/PhysicsEngine/BodyInstance.cpp
+++ b/Engine/Source/Runtime/Engine/Private/PhysicsEngine/BodyInstance.cpp
@@ -3730,7 +3730,7 @@ void FBodyInstance::SetUseMACD(bool bInUseMACD)
 			{
 				if (FPhysicsInterface::IsValid(Actor) && FPhysicsInterface::IsRigidBody(Actor))
 				{
-					FPhysicsInterface::SetCcdEnabled_AssumesLocked(Actor, bInUseMACD);
+					FPhysicsInterface::SetMACDEnabled_AssumesLocked(Actor, bInUseMACD);
 				}
 			});
 	}
diff --git a/Engine/Source/Runtime/Engine/Private/PrimitiveComponentPhysics.cpp b/Engine/Source/Runtime/Engine/Private/PrimitiveComponentPhysics.cpp
--- a/Engine/Source/Runtime/Engine/Private/PrimitiveComponentPhysics.cpp
+++ b/Engine/Source/Runtime/Engine/Private/PrimitiveComponentPhysics.cpp
@@ -669,6 +669,21 @@ void UPrimitiveComponent::SetAllUseCCD(bool bInUseCCD)
 	SetUseCCD(bInUseCCD, NAME_None);
 }
 
+
+void UPrimitiveComponent::SetUseMACD(bool bInUseMACD, FName BoneName)
+{
+	FBodyInstance* BI = GetBodyInstance(BoneName);
+	if (BI)
+	{
+		BI->SetUseMACD(bInUseMACD);
+	}
+}
+
+void UPrimitiveComponent::SetAllUseMACD(bool bInUseMACD)
+{
+	SetUseMACD(bInUseMACD, NAME_None);
+}
+
 void UPrimitiveComponent::PutRigidBodyToSleep(FName BoneName)
 {
 	FBodyInstance* BI = GetBodyInstance(BoneName);
