PR #11513: Expose UActorContainer for public linking

diff --git a/Engine/Source/Runtime/Engine/Classes/Engine/Level.h b/Engine/Source/Runtime/Engine/Classes/Engine/Level.h
--- a/Engine/Source/Runtime/Engine/Classes/Engine/Level.h
+++ b/Engine/Source/Runtime/Engine/Classes/Engine/Level.h
@@ -64,7 +64,7 @@ private:
 
 // Actor container class used to duplicate actors during cells streaming in PIE
 UCLASS()
-class UActorContainer : public UObject
+class ENGINE_API UActorContainer : public UObject
 {
 	GENERATED_BODY()
 
