Fixed the "eatmem" debug command so it triggers an OOM on Windows.

It takes very long to trigger the OOM handler at the current leak rate, and since the debug command is blocking the hang handler triggers before the OOM handler gets a chance to kick in.
Memset is not necessary for Windows as the OS will not allow overcommitting. Using Memset on the whole block causes the OS to start swapping, slowing down the process and also resulting in a hang before the OOM.

#rnx

diff --git a/Engine/Source/Runtime/Engine/Private/UnrealEngine.cpp b/Engine/Source/Runtime/Engine/Private/UnrealEngine.cpp
--- a/Engine/Source/Runtime/Engine/Private/UnrealEngine.cpp
+++ b/Engine/Source/Runtime/Engine/Private/UnrealEngine.cpp
@@ -10945,8 +10945,15 @@ bool UEngine::PerformError(const TCHAR* Cmd, FOutputDevice& Ar)
 			FGenericCrashContext::SetCrashTrigger(ECrashTrigger::Debug);
 			while (1)
 			{
-				void* Eat = FMemory::Malloc(65536);
-				FMemory::Memset(Eat, 0, 65536);
+#				if PLATFORM_WINDOWS
+					/* Windows does not allow overcommitting, so allocating the memory is sufficient to cause an OOM.
+					   Running memset on the memory block will cause Windows to start swapping memory to/from the page file
+					   and will ultimately result in an application hang instead of OOM. */
+					(void)FMemory::Malloc(65536 * 1024);
+#				else
+					void* Eat = FMemory::Malloc(65536);
+					FMemory::Memset(Eat, 0, 65536);
+#				endif
 			}
 		}
 
