Fix NaN ensure caused by LastPosition being used before initialization in FStreamingSourceVelocity::GetAverageVelocity

#tests Editor
#rnx

diff --git a/Engine/Source/Runtime/Engine/Private/WorldPartition/WorldPartitionSubsystem.cpp b/Engine/Source/Runtime/Engine/Private/WorldPartition/WorldPartitionSubsystem.cpp
--- a/Engine/Source/Runtime/Engine/Private/WorldPartition/WorldPartitionSubsystem.cpp
+++ b/Engine/Source/Runtime/Engine/Private/WorldPartition/WorldPartitionSubsystem.cpp
@@ -1577,6 +1577,7 @@ FStreamingSourceVelocity::FStreamingSourceVelocity(const FName& InSourceName)
 	: bIsValid(false)
 	, SourceName(InSourceName)
 	, LastUpdateTime(-1.0)
+	, LastPosition(FVector::Zero())
 {}
 
 FVector FStreamingSourceVelocity::GetAverageVelocity(const FVector& NewPosition, double CurrentTime)
