[Skeletal Mesh Component] Fixed MeshDeformerSet and MeshDeformerInstanceSet not accounting for skeletal meshes with more than 8 LOD

diff --git a/Engine/Source/Runtime/Engine/Classes/Components/SkinnedMeshComponent.h b/Engine/Source/Runtime/Engine/Classes/Components/SkinnedMeshComponent.h
--- a/Engine/Source/Runtime/Engine/Classes/Components/SkinnedMeshComponent.h
+++ b/Engine/Source/Runtime/Engine/Classes/Components/SkinnedMeshComponent.h
@@ -196,8 +196,11 @@ struct FMeshDeformerSet
 	/**
 	 * Indexed by mesh LOD. Each element is either a valid index into the Deformers array, or
 	 * INDEX_NONE to signify that no deformer should be used for this LOD.
+	 * MAX_MESH_LOD_COUNT here cannot be a hard limit as editor allows more LODs to be created
+	 * however there may be other systems assuming MAX_MESH_LOD_COUNT is respected so generally it is not a good idea
+	 * to exceed it
 	 */
-	TArray<int8, TFixedAllocator<MAX_MESH_LOD_COUNT>> DeformerIndexForLOD;
+	TArray<int8, TInlineAllocator<MAX_MESH_LOD_COUNT>> DeformerIndexForLOD;
 };
 
 /** Same as FMeshDeformerSet, except for mesh deformer instances */
@@ -209,7 +212,7 @@ struct FMeshDeformerInstanceSet
 	UPROPERTY()
 	TArray<TObjectPtr<UMeshDeformerInstance>> DeformerInstances;
 
-	TArray<int8, TFixedAllocator<MAX_MESH_LOD_COUNT>> InstanceIndexForLOD;
+	TArray<int8> InstanceIndexForLOD;
 };
 
 /** The map of external morph sets registered on the skinned mesh component. */
