Shader Compiler:  User facing workaround for platform specific register overflow bugs related to preshaders.  Added setting to Material to allow the r.Material.PreshaderGapInterval to be overridden for a specific material (instead of needing to set it globally).  Could be removed when the bug in the platform specific compiler gets fixed.

#rnx

diff --git a/Engine/Source/Runtime/Engine/Classes/Materials/Material.h b/Engine/Source/Runtime/Engine/Classes/Materials/Material.h
--- a/Engine/Source/Runtime/Engine/Classes/Materials/Material.h
+++ b/Engine/Source/Runtime/Engine/Classes/Materials/Material.h
@@ -1027,6 +1027,10 @@ public:
 	UPROPERTY(EditAnywhere, BlueprintReadWrite, Category = PostProcessMaterial, meta = (DisplayName = "Is Blendable"))
 	uint8 bIsBlendable : 1;
 
+	/** If non-zero, overrides r.Material.PreshaderGapInterval for this material.  Workaround for a platform specific register overflow bug. */
+	UPROPERTY(EditAnywhere, Category = Material, AdvancedDisplay)
+	uint16 PreshaderGap;
+
 	/** true if we have printed a warning about material usage for a given usage flag. */
 	UPROPERTY(transient, duplicatetransient)
 	uint32 UsageFlagWarnings;
diff --git a/Engine/Source/Runtime/Engine/Private/Materials/HLSLMaterialTranslator.cpp b/Engine/Source/Runtime/Engine/Private/Materials/HLSLMaterialTranslator.cpp
--- a/Engine/Source/Runtime/Engine/Private/Materials/HLSLMaterialTranslator.cpp
+++ b/Engine/Source/Runtime/Engine/Private/Materials/HLSLMaterialTranslator.cpp
@@ -200,7 +200,7 @@ static int32 GPreshaderGapInterval = 32;
 static FAutoConsoleVariableRef CVarPreshaderGapInterval(
 	TEXT("r.Material.PreshaderGapInterval"),
 	GPreshaderGapInterval,
-	TEXT("Insert an empty element in the preshader buffer every specified number of elements in the buffer.  Workaround for a shader compiler bug."));
+	TEXT("Insert an empty element in the preshader buffer every specified number of elements in the buffer.  Workaround for a shader compiler register overflow bug."));
 
 /* Controls whether DDC caching of material translation results should be forcefully disabled. */
 #define FORCE_DISABLE_MATERIAL_TRANSLATION_DDC false
@@ -3874,12 +3874,20 @@ int32 FHLSLMaterialTranslator::AccessUniformExpression(int32 Index)
 			}
 
 			// Optionally insert an empty vector element (gap) at the given interval, to work around a platform specific shader compiler bug
-			if (GPreshaderGapInterval > 0)
+			int32 PreshaderGapInterval = Material->GetPreshaderGap();
+			if (!PreshaderGapInterval)
 			{
+				PreshaderGapInterval = GPreshaderGapInterval;
+			}
+			if (PreshaderGapInterval > 0)
+			{
+				// Set a minimum of 4 (one out of every four elements is padding)
+				PreshaderGapInterval = FMath::Max(PreshaderGapInterval, 4);
+
 				// Check if we are on the vector element that is the end of the gap interval, or we will pass the end of the gap interval if the
 				// preshader is larger than a single vector.  The divides will change value when a modulus gap interval boundary is crossed.
 				uint32 PreshaderGapOffset = FMath::Max(4u, bIsLWC ? NumComponents * 2u : NumComponents);
-				if (UniformPreshaderOffset / (GPreshaderGapInterval*4) != (UniformPreshaderOffset + PreshaderGapOffset) / (GPreshaderGapInterval*4))
+				if (UniformPreshaderOffset / (PreshaderGapInterval *4) != (UniformPreshaderOffset + PreshaderGapOffset) / (PreshaderGapInterval *4))
 				{
 					UniformPreshaderOffset += 4u;
 				}
diff --git a/Engine/Source/Runtime/Engine/Private/Materials/Material.cpp b/Engine/Source/Runtime/Engine/Private/Materials/Material.cpp
--- a/Engine/Source/Runtime/Engine/Private/Materials/Material.cpp
+++ b/Engine/Source/Runtime/Engine/Private/Materials/Material.cpp
@@ -1009,6 +1009,7 @@ UMaterial::UMaterial(const FObjectInitializer& ObjectInitializer)
 	BlendablePriority = 0;
 	BlendableOutputAlpha = false;
 	bIsBlendable = true;
+	PreshaderGap = 0;
 	bEnableStencilTest = false;
 	bUsedWithVolumetricCloud = false;
 	bUsedWithHeterogeneousVolumes = false;
diff --git a/Engine/Source/Runtime/Engine/Private/Materials/MaterialShared.cpp b/Engine/Source/Runtime/Engine/Private/Materials/MaterialShared.cpp
--- a/Engine/Source/Runtime/Engine/Private/Materials/MaterialShared.cpp
+++ b/Engine/Source/Runtime/Engine/Private/Materials/MaterialShared.cpp
@@ -2219,6 +2219,11 @@ bool FMaterialResource::HasRenderTracePhysicalMaterialOutputs() const
 	return Material->GetRenderTracePhysicalMaterialOutputs().Num() > 0;
 }
 
+uint16 FMaterialResource::GetPreshaderGap() const
+{
+	return Material->PreshaderGap;
+}
+
 UMaterialInterface* FMaterialResource::GetMaterialInterface() const 
 { 
 	return MaterialInstance ? (UMaterialInterface*)MaterialInstance : (UMaterialInterface*)Material;
@@ -5391,6 +5396,7 @@ FMaterialShaderParameters::FMaterialShaderParameters(const FMaterial* InMaterial
 	BlendMode = InMaterial->GetBlendMode();
 	FeatureLevel = InMaterial->GetFeatureLevel();
 	QualityLevel = InMaterial->GetQualityLevel();
+	PreshaderGap = InMaterial->GetPreshaderGap();
 	BlendableLocation = InMaterial->GetBlendableLocation();
 	NumCustomizedUVs = InMaterial->GetNumCustomizedUVs();
 	StencilCompare = InMaterial->GetStencilCompare();
diff --git a/Engine/Source/Runtime/Engine/Public/MaterialShared.h b/Engine/Source/Runtime/Engine/Public/MaterialShared.h
--- a/Engine/Source/Runtime/Engine/Public/MaterialShared.h
+++ b/Engine/Source/Runtime/Engine/Public/MaterialShared.h
@@ -901,6 +901,7 @@ struct FMaterialShaderParameters
 	EBlendMode BlendMode;
 	ERHIFeatureLevel::Type FeatureLevel;
 	EMaterialQualityLevel::Type QualityLevel;
+	uint16 PreshaderGap;
 	int32 BlendableLocation;
 	int32 NumCustomizedUVs;
 	uint32 StencilCompare;
@@ -2097,6 +2098,7 @@ public:
 	virtual bool CastsRayTracedShadows() const { return true; }
 	virtual bool IsTessellationEnabled() const { return false; }
 	virtual bool HasRenderTracePhysicalMaterialOutputs() const { return false; }
+	virtual uint16 GetPreshaderGap() const { return 0; }
 	virtual EMaterialShadingRate GetShadingRate() const { return MSR_1x1; }
 	virtual int32 GetNeuralProfileId() const { return INDEX_NONE; }
 	virtual bool IsVariableRateShadingAllowed() const { return true; }
@@ -2840,6 +2842,7 @@ public:
 	ENGINE_API virtual bool CastsRayTracedShadows() const override;
 	ENGINE_API virtual bool IsTessellationEnabled() const override;
 	ENGINE_API virtual bool HasRenderTracePhysicalMaterialOutputs() const override;
+	ENGINE_API virtual uint16 GetPreshaderGap() const override;
 	ENGINE_API virtual UMaterialInterface* GetMaterialInterface() const override;
 	ENGINE_API virtual int32 GetNeuralProfileId() const override;
 	/**
