Remove extra printf param

#rnx

diff --git a/Engine/Source/Runtime/Engine/Private/DataChannel.cpp b/Engine/Source/Runtime/Engine/Private/DataChannel.cpp
--- a/Engine/Source/Runtime/Engine/Private/DataChannel.cpp
+++ b/Engine/Source/Runtime/Engine/Private/DataChannel.cpp
@@ -4260,7 +4260,7 @@ bool UActorChannel::ValidateReplicatedSubObjects()
 		const DataChannelInternal::FSubObjectReplicatedInfo& Info = DataChannelInternal::ReplicatedSubObjectsTracker[i];
 
 		ensureMsgf(false, TEXT("%s was replicated only in the registered subobject list in %s. Not in the legacy path."), 
-			*Info.Describe(Actor), *Info.Describe(Actor), *Connection->GetDriver()->NetDriverName.ToString());
+			*Info.Describe(Actor), *Connection->GetDriver()->NetDriverName.ToString());
 		UE_CLOG(bLogAllErrors, LogNet, Warning, TEXT("%s was replicated only in the registered subobject list. Not in the legacy path."), *Info.Describe(Actor));
 	}
 
