[shaders] log exit code for SCW processes when unexpected termination occurs; only delete output files when relaunching if we previously hit this condition (as in this case we'll manually re-run all pending jobs locally, while abandoned workers may still carry on and write output files which will subsequently be associated with the wrong jobs)

diff --git a/Engine/Source/Runtime/Engine/Private/ShaderCompiler/ShaderCompiler.cpp b/Engine/Source/Runtime/Engine/Private/ShaderCompiler/ShaderCompiler.cpp
--- a/Engine/Source/Runtime/Engine/Private/ShaderCompiler/ShaderCompiler.cpp
+++ b/Engine/Source/Runtime/Engine/Private/ShaderCompiler/ShaderCompiler.cpp
@@ -3905,6 +3905,9 @@ bool FShaderCompileThreadRunnable::LaunchWorkersIfNeeded()
 		LastCheckForWorkersTime = CurrentTime;
 	}
 
+	bool bClearStaleOutputs = Manager->bClearStaleWorkerOutputs;
+	Manager->bClearStaleWorkerOutputs = false;
+
 	FScopeLock WorkerScopeLock(&WorkerInfosLock);
 	for (int32 WorkerIndex = 0; WorkerIndex < WorkerInfos.Num(); WorkerIndex++)
 	{
@@ -3933,6 +3936,8 @@ bool FShaderCompileThreadRunnable::LaunchWorkersIfNeeded()
 			if (CurrentWorkerInfo.WorkerProcess.IsValid())
 			{
 				// shader compiler exited one way or another, so clear out the stale PID.
+				int32 ReturnCode = 0;
+				FPlatformProcess::GetProcReturnCode(CurrentWorkerInfo.WorkerProcess, &ReturnCode);
 				FPlatformProcess::CloseProc(CurrentWorkerInfo.WorkerProcess);
 				CurrentWorkerInfo.WorkerProcess = FProcHandle();
 
@@ -3948,7 +3953,7 @@ bool FShaderCompileThreadRunnable::LaunchWorkersIfNeeded()
 					}
 					else
 					{
-						UE_LOG(LogShaderCompilers, Error, TEXT("ShaderCompileWorker terminated unexpectedly!  Falling back to directly compiling which will be very slow.  Thread %u."), WorkerIndex);
+						UE_LOG(LogShaderCompilers, Error, TEXT("ShaderCompileWorker terminated unexpectedly, return code %d! Falling back to directly compiling which will be very slow.  Thread %u."), ReturnCode, WorkerIndex);
 						LogQueuedCompileJobs(CurrentWorkerInfo.QueuedJobs, -1);
 
 						bAbandonWorkers = true;
@@ -3962,9 +3967,12 @@ bool FShaderCompileThreadRunnable::LaunchWorkersIfNeeded()
 				const FString WorkingDirectory = Manager->ShaderBaseWorkingDirectory + FString::FromInt(WorkerIndex) + TEXT("/");
 				FString InputFileName(TEXT("WorkerInputOnly.in"));
 				FString OutputFileName(TEXT("WorkerOutputOnly.out"));
-
-				// Delete any potential stale output files; these can persist if we had previously abandoned running with workers due to an unexpected SCW termination.
-				IFileManager::Get().Delete(*(WorkingDirectory / OutputFileName));
+				
+				if (bClearStaleOutputs)
+				{
+					// Delete any potential stale output files; these can persist if we had previously abandoned running with workers due to an unexpected SCW termination.
+					IFileManager::Get().Delete(*(WorkingDirectory / OutputFileName));
+				}
 
 				// Store the handle with this thread so that we will know not to launch it again
 				CurrentWorkerInfo.WorkerProcess = Manager->LaunchWorker(WorkingDirectory, Manager->ProcessId, WorkerIndex, InputFileName, OutputFileName);
@@ -4470,6 +4478,7 @@ int32 FShaderCompileThreadRunnable::CompilingLoop()
 				if (Manager->NumSingleThreadedRunsBeforeRetry == 0)
 				{
 					UE_LOG(LogShaderCompilers, Display, TEXT("Retry shader compiling through workers."));
+					Manager->bClearStaleWorkerOutputs = true;
 					Manager->bAllowCompilingThroughWorkers = true;
 				}
 			}
diff --git a/Engine/Source/Runtime/Engine/Public/ShaderCompiler.h b/Engine/Source/Runtime/Engine/Public/ShaderCompiler.h
--- a/Engine/Source/Runtime/Engine/Public/ShaderCompiler.h
+++ b/Engine/Source/Runtime/Engine/Public/ShaderCompiler.h
@@ -732,6 +732,8 @@ private:
 	bool bPromptToRetryFailedShaderCompiles;
 	/** If enabled when we enter the prompt to retry we will break in the debugger if one is attached rather than prompting. */
 	bool bDebugBreakOnPromptToRetryShaderCompile = false;
+	/** If true, when relaunching workers will clear out any existing output files. This is only used when workers are abandoned due to an unexpected termination */
+	bool bClearStaleWorkerOutputs = false;
 	/** Whether to log out shader job completion times on the worker thread.  Useful for tracking down which global shader is taking a long time. */
 	bool bLogJobCompletionTimes;
 	/** Target execution time for ProcessAsyncResults.  Larger values speed up async shader map processing but cause more hitchiness while async compiling is happening. */
