[WorldPartition] Instead of gathering level script references to external actors, gathers them from the world to make sure to include references from non-external actors and other pointers from the world object hierarchy.

diff --git a/Engine/Plugins/Editor/DataValidation/Source/DataValidation/Private/WorldPartitionChangelistValidator.cpp b/Engine/Plugins/Editor/DataValidation/Source/DataValidation/Private/WorldPartitionChangelistValidator.cpp
--- a/Engine/Plugins/Editor/DataValidation/Source/DataValidation/Private/WorldPartitionChangelistValidator.cpp
+++ b/Engine/Plugins/Editor/DataValidation/Source/DataValidation/Private/WorldPartitionChangelistValidator.cpp
@@ -356,24 +356,22 @@ void UWorldPartitionChangelistValidator::OnInvalidReferenceRuntimeGrid(const IWo
 	}
 }
 
-void UWorldPartitionChangelistValidator::OnInvalidReferenceLevelScriptStreamed(const IWorldPartitionActorDescInstanceView& ActorDescView)
-{
-	if (Filter(ActorDescView))
-	{		
-		FText CurrentError = FText::Format(LOCTEXT("DataValidation.Changelist.WorldPartition.InvalidReferenceLevelScriptStreamed", "Level script blueprint references streamed actor {0}."),
-											FText::FromString(GetFullActorName(ActorDescView)));
-		
-		AssetFails(CurrentAsset, CurrentError);
-	}
-}
-
-void UWorldPartitionChangelistValidator::OnInvalidReferenceLevelScriptDataLayers(const IWorldPartitionActorDescInstanceView& ActorDescView)
+void UWorldPartitionChangelistValidator::OnInvalidWorldReference(const IWorldPartitionActorDescInstanceView& ActorDescView, EWorldReferenceInvalidReason Reason)
 {
 	if (Filter(ActorDescView))
 	{
-		FText CurrentError = FText::Format(LOCTEXT("DataValidation.Changelist.WorldPartition.InvalidReferenceLevelScriptDataLayers", "Level script blueprint references streamed actor {0} with a non empty set of data layers."),
-											FText::FromString(GetFullActorName(ActorDescView)));
+		FText CurrentError;
 
+		switch(Reason)
+		{
+		case EWorldReferenceInvalidReason::ReferencedActorIsSpatiallyLoaded:
+			CurrentError = FText::Format(LOCTEXT("DataValidation.Changelist.WorldPartition.InvalidWorldReferenceSpatiallyLoaded", "World references spatially loaded actor {0}."), FText::FromString(GetFullActorName(ActorDescView)));
+			break;
+		case EWorldReferenceInvalidReason::ReferencedActorHasDataLayers:
+			CurrentError = FText::Format(LOCTEXT("DataValidation.Changelist.WorldPartition.InvalidWorldReferenceDataLayers", "World references actor {0} with data layers."), FText::FromString(GetFullActorName(ActorDescView)));
+			break;
+		}
+		
 		AssetFails(CurrentAsset, CurrentError);
 	}
 }
diff --git a/Engine/Plugins/Editor/DataValidation/Source/DataValidation/Private/WorldPartitionChangelistValidator.h b/Engine/Plugins/Editor/DataValidation/Source/DataValidation/Private/WorldPartitionChangelistValidator.h
--- a/Engine/Plugins/Editor/DataValidation/Source/DataValidation/Private/WorldPartitionChangelistValidator.h
+++ b/Engine/Plugins/Editor/DataValidation/Source/DataValidation/Private/WorldPartitionChangelistValidator.h
@@ -38,8 +38,7 @@ protected:
 	virtual void OnInvalidReference(const IWorldPartitionActorDescInstanceView& ActorDescView, const FGuid& ReferenceGuid, IWorldPartitionActorDescInstanceView* ReferenceActorDescView) override;
 	virtual void OnInvalidReferenceGridPlacement(const IWorldPartitionActorDescInstanceView& ActorDescView, const IWorldPartitionActorDescInstanceView& ReferenceActorDescView) override;
 	virtual void OnInvalidReferenceDataLayers(const IWorldPartitionActorDescInstanceView& ActorDescView, const IWorldPartitionActorDescInstanceView& ReferenceActorDescView, EDataLayerInvalidReason Reason) override;
-	virtual void OnInvalidReferenceLevelScriptStreamed(const IWorldPartitionActorDescInstanceView& ActorDescView) override;
-	virtual void OnInvalidReferenceLevelScriptDataLayers(const IWorldPartitionActorDescInstanceView& ActorDescView) override;
+	virtual void OnInvalidWorldReference(const IWorldPartitionActorDescInstanceView& ActorDescView, EWorldReferenceInvalidReason Reason) override;
 	virtual void OnInvalidReferenceRuntimeGrid(const IWorldPartitionActorDescInstanceView& ActorDescView, const IWorldPartitionActorDescInstanceView& ReferenceActorDescView) override;
 	virtual void OnInvalidReferenceDataLayerAsset(const UDataLayerInstanceWithAsset* DataLayerInstance) override;
 	virtual void OnInvalidDataLayerAssetType(const UDataLayerInstanceWithAsset* DataLayerInstance, const UDataLayerAsset* DataLayerAsset) override;
diff --git a/Engine/Source/Runtime/Engine/Classes/Engine/Level.h b/Engine/Source/Runtime/Engine/Classes/Engine/Level.h
--- a/Engine/Source/Runtime/Engine/Classes/Engine/Level.h
+++ b/Engine/Source/Runtime/Engine/Classes/Engine/Level.h
@@ -710,8 +710,8 @@ public:
 	ENGINE_API static bool GetLevelBoundsFromAsset(const FAssetData& Asset, FBox& OutLevelBounds);
 	ENGINE_API static bool GetLevelBoundsFromPackage(FName LevelPackage, FBox& OutLevelBounds);
 
-	ENGINE_API static bool GetLevelScriptExternalActorsReferencesFromAsset(const FAssetData& Asset, TArray<FGuid>& OutLevelScriptExternalActorsReferences);
-	ENGINE_API static bool GetLevelScriptExternalActorsReferencesFromPackage(FName LevelPackage, TArray<FGuid>& OutLevelScriptExternalActorsReferences);
+	ENGINE_API static bool GetWorldExternalActorsReferencesFromAsset(const FAssetData& Asset, TArray<FGuid>& OutWorldExternalActorsReferences);
+	ENGINE_API static bool GetWorldExternalActorsReferencesFromPackage(FName LevelPackage, TArray<FGuid>& OutWorldExternalActorsReferences);
 
 	ENGINE_API static bool GetIsLevelPartitionedFromAsset(const FAssetData& Asset);
 	ENGINE_API static bool GetIsLevelPartitionedFromPackage(FName LevelPackage);
@@ -729,6 +729,10 @@ public:
 	static bool GetPartitionedLevelCanBeUsedByLevelInstanceFromAsset(const FAssetData& Asset) { return true; }
 	UE_DEPRECATED(5.3, "GetPartitionedLevelCanBeUsedByLevelInstanceFromPackage is deprecated.")
 	static bool GetPartitionedLevelCanBeUsedByLevelInstanceFromPackage(FName LevelPackage) { return true; }
+	UE_DEPRECATED(5.4, "GetLevelScriptExternalActorsReferencesFromAsset is deprecated.")
+	ENGINE_API static bool GetLevelScriptExternalActorsReferencesFromAsset(const FAssetData& Asset, TArray<FGuid>& OutLevelScriptExternalActorsReferences) { return false;}
+	UE_DEPRECATED(5.4, "GetLevelScriptExternalActorsReferencesFromPackageis deprecated.")
+	ENGINE_API static bool GetLevelScriptExternalActorsReferencesFromPackage(FName LevelPackage, TArray<FGuid>& OutLevelScriptExternalActorsReferences) { return false; }
 
 	ENGINE_API static FVector GetLevelInstancePivotOffsetFromAsset(const FAssetData& Asset);
 	ENGINE_API static FVector GetLevelInstancePivotOffsetFromPackage(FName LevelPackage);
diff --git a/Engine/Source/Runtime/Engine/Private/Level.cpp b/Engine/Source/Runtime/Engine/Private/Level.cpp
--- a/Engine/Source/Runtime/Engine/Private/Level.cpp
+++ b/Engine/Source/Runtime/Engine/Private/Level.cpp
@@ -559,7 +559,7 @@ void ULevel::Serialize( FArchive& Ar )
 		bContainsStableActorGUIDs = Ar.CustomVer(FFortniteMainBranchObjectVersion::GUID) >= FFortniteMainBranchObjectVersion::ContainsStableActorGUIDs;
 #endif
 	}
-	else if (Ar.IsSaving() && Ar.IsPersistent())
+	else if (Ar.IsPersistent() && (Ar.IsSaving() || Ar.IsObjectReferenceCollector()))
 	{
 		UPackage* LevelPackage = GetOutermost();
 		TArray<AActor*> EmbeddedActors;
@@ -2552,16 +2552,16 @@ bool ULevel::GetIsStreamingDisabledFromAsset(const FAssetData& Asset)
 	return false;
 }
 
-bool ULevel::GetLevelScriptExternalActorsReferencesFromAsset(const FAssetData& Asset, TArray<FGuid>& OutLevelScriptExternalActorsReferences)
+bool ULevel::GetWorldExternalActorsReferencesFromAsset(const FAssetData& Asset, TArray<FGuid>& OutWorldExternalActorsReferences)
 {
-	FString LevelScriptExternalActorsReferencesStr;
-	static const FName NAME_LevelScriptExternalActorsReferences(TEXT("LevelScriptExternalActorsReferences"));
-	if (Asset.GetTagValue(NAME_LevelScriptExternalActorsReferences, LevelScriptExternalActorsReferencesStr))
+	FString WorldExternalActorsReferencesStr;
+	static const FName NAME_WorldExternalActorsReferences(TEXT("WorldExternalActorsReferences"));
+	if (Asset.GetTagValue(NAME_WorldExternalActorsReferences, WorldExternalActorsReferencesStr))
 	{
-		TArray<FString> LevelScriptReferencesStr;
-		if (LevelScriptExternalActorsReferencesStr.ParseIntoArray(LevelScriptReferencesStr, TEXT(",")))
+		TArray<FString> WorldReferencesStr;
+		if (WorldExternalActorsReferencesStr.ParseIntoArray(WorldReferencesStr, TEXT(",")))
 		{
-			Algo::Transform(LevelScriptReferencesStr, OutLevelScriptExternalActorsReferences, [](const FString& GuidStr) { return FGuid(GuidStr); });
+			Algo::Transform(WorldReferencesStr, OutWorldExternalActorsReferences, [](const FString& GuidStr) { return FGuid(GuidStr); });
 			return true;
 		}
 	}
@@ -2576,11 +2576,11 @@ bool ULevel::GetIsStreamingDisabledFromPackage(FName LevelPackage)
 	});
 }
 
-bool ULevel::GetLevelScriptExternalActorsReferencesFromPackage(FName LevelPackage, TArray<FGuid>& OutLevelScriptExternalActorsReferences)
+bool ULevel::GetWorldExternalActorsReferencesFromPackage(FName LevelPackage, TArray<FGuid>& OutWorldExternalActorsReferences)
 {
-	return LevelAssetRegistryHelper::GetLevelInfoFromAssetRegistry(LevelPackage, [&OutLevelScriptExternalActorsReferences](const FAssetData& Asset)
+	return LevelAssetRegistryHelper::GetLevelInfoFromAssetRegistry(LevelPackage, [&OutWorldExternalActorsReferences](const FAssetData& Asset)
 	{
-		return GetLevelScriptExternalActorsReferencesFromAsset(Asset, OutLevelScriptExternalActorsReferences);
+		return GetWorldExternalActorsReferencesFromAsset(Asset, OutWorldExternalActorsReferences);
 	});
 }
 
diff --git a/Engine/Source/Runtime/Engine/Private/WorldPartition/ErrorHandling/WorldPartitionStreamingGenerationLogErrorHandler.cpp b/Engine/Source/Runtime/Engine/Private/WorldPartition/ErrorHandling/WorldPartitionStreamingGenerationLogErrorHandler.cpp
--- a/Engine/Source/Runtime/Engine/Private/WorldPartition/ErrorHandling/WorldPartitionStreamingGenerationLogErrorHandler.cpp
+++ b/Engine/Source/Runtime/Engine/Private/WorldPartition/ErrorHandling/WorldPartitionStreamingGenerationLogErrorHandler.cpp
@@ -47,14 +47,17 @@ void FStreamingGenerationLogErrorHandler::OnInvalidReferenceRuntimeGrid(const IW
 	UE_ASSET_LOG_ACTORDESCVIEW(LogWorldPartition, Log, ActorDescView, TEXT("Actor %s references an actor in a different runtime grid %s"), *GetActorName(ActorDescView), *GetActorName(ReferenceActorDescView));
 }
 
-void FStreamingGenerationLogErrorHandler::OnInvalidReferenceLevelScriptStreamed(const IWorldPartitionActorDescInstanceView& ActorDescView)
+void FStreamingGenerationLogErrorHandler::OnInvalidWorldReference(const IWorldPartitionActorDescInstanceView& ActorDescView, EWorldReferenceInvalidReason Reason)
 {
-	UE_ASSET_LOG_ACTORDESCVIEW(LogWorldPartition, Log, ActorDescView, TEXT("Level Script Blueprint references streamed actor %s"), *GetActorName(ActorDescView));
-}
-
-void FStreamingGenerationLogErrorHandler::OnInvalidReferenceLevelScriptDataLayers(const IWorldPartitionActorDescInstanceView& ActorDescView)
-{
-	UE_ASSET_LOG_ACTORDESCVIEW(LogWorldPartition, Log, ActorDescView, TEXT("Level Script Blueprint references streamed actor %s with a non empty set of data layers"), *GetActorName(ActorDescView));
+	switch(Reason)
+	{
+	case EWorldReferenceInvalidReason::ReferencedActorIsSpatiallyLoaded:
+		UE_ASSET_LOG_ACTORDESCVIEW(LogWorldPartition, Log, ActorDescView, TEXT("World references spatially loaded actor %s"), *GetActorName(ActorDescView));
+		break;
+	case EWorldReferenceInvalidReason::ReferencedActorHasDataLayers:
+		UE_ASSET_LOG_ACTORDESCVIEW(LogWorldPartition, Log, ActorDescView, TEXT("World references actor %s with data layers"), *GetActorName(ActorDescView));
+		break;
+	}	
 }
 
 void FStreamingGenerationLogErrorHandler::OnInvalidReferenceDataLayerAsset(const UDataLayerInstanceWithAsset* DataLayerInstance)
diff --git a/Engine/Source/Runtime/Engine/Private/WorldPartition/ErrorHandling/WorldPartitionStreamingGenerationProxyErrorHandler.cpp b/Engine/Source/Runtime/Engine/Private/WorldPartition/ErrorHandling/WorldPartitionStreamingGenerationProxyErrorHandler.cpp
--- a/Engine/Source/Runtime/Engine/Private/WorldPartition/ErrorHandling/WorldPartitionStreamingGenerationProxyErrorHandler.cpp
+++ b/Engine/Source/Runtime/Engine/Private/WorldPartition/ErrorHandling/WorldPartitionStreamingGenerationProxyErrorHandler.cpp
@@ -28,14 +28,9 @@ void FStreamingGenerationProxyErrorHandler::OnInvalidReferenceRuntimeGrid(const
 	InnerErrorHandler->OnInvalidReferenceRuntimeGrid(ActorDescView, ReferenceActorDescView);
 }
 
-void FStreamingGenerationProxyErrorHandler::OnInvalidReferenceLevelScriptStreamed(const IWorldPartitionActorDescInstanceView& ActorDescView)
+void FStreamingGenerationProxyErrorHandler::OnInvalidWorldReference(const IWorldPartitionActorDescInstanceView& ActorDescView, EWorldReferenceInvalidReason Reason)
 {
-	InnerErrorHandler->OnInvalidReferenceLevelScriptStreamed(ActorDescView);
-}
-
-void FStreamingGenerationProxyErrorHandler::OnInvalidReferenceLevelScriptDataLayers(const IWorldPartitionActorDescInstanceView& ActorDescView)
-{
-	InnerErrorHandler->OnInvalidReferenceLevelScriptDataLayers(ActorDescView);
+	InnerErrorHandler->OnInvalidWorldReference(ActorDescView, Reason);
 }
 
 void FStreamingGenerationProxyErrorHandler::OnInvalidReferenceDataLayerAsset(const UDataLayerInstanceWithAsset* DataLayerInstance)
diff --git a/Engine/Source/Runtime/Engine/Private/WorldPartition/ErrorHandling/WorldPartitionStreamingGenerationTokenizedMessageErrorHandler.cpp b/Engine/Source/Runtime/Engine/Private/WorldPartition/ErrorHandling/WorldPartitionStreamingGenerationTokenizedMessageErrorHandler.cpp
--- a/Engine/Source/Runtime/Engine/Private/WorldPartition/ErrorHandling/WorldPartitionStreamingGenerationTokenizedMessageErrorHandler.cpp
+++ b/Engine/Source/Runtime/Engine/Private/WorldPartition/ErrorHandling/WorldPartitionStreamingGenerationTokenizedMessageErrorHandler.cpp
@@ -96,26 +96,25 @@ void ITokenizedMessageErrorHandler::OnInvalidReferenceRuntimeGrid(const IWorldPa
 	HandleTokenizedMessage(MoveTemp(Message));
 }
 
-void ITokenizedMessageErrorHandler::OnInvalidReferenceLevelScriptStreamed(const IWorldPartitionActorDescInstanceView& ActorDescView)
+void ITokenizedMessageErrorHandler::OnInvalidWorldReference(const IWorldPartitionActorDescInstanceView& ActorDescView, EWorldReferenceInvalidReason Reason)
 {
 	TSharedRef<FTokenizedMessage> Message = FTokenizedMessage::Create(EMessageSeverity::Error);
-	Message->AddToken(FTextToken::Create(LOCTEXT("TokenMessage_WorldPartition_LevelScriptBlueprintStreamedActorReference", "Level Script Blueprint references streamed actor")))
-		->AddToken(FActorToken::Create(ActorDescView.GetActorSoftPath().ToString(), ActorDescView.GetGuid(), FText::FromString(GetActorName(ActorDescView))))
-		->AddToken(FMapErrorToken::Create(TEXT("WorldPartition_LevelScriptBlueprintRefefenceStreamed_CheckForErrors")));
-
-	HandleTokenizedMessage(MoveTemp(Message));
-	
-}
 
-void ITokenizedMessageErrorHandler::OnInvalidReferenceLevelScriptDataLayers(const IWorldPartitionActorDescInstanceView& ActorDescView)
-{
-	TSharedRef<FTokenizedMessage> Message = FTokenizedMessage::Create(EMessageSeverity::Error);
-	Message->AddToken(FTextToken::Create(LOCTEXT("TokenMessage_WorldPartition_LevelScriptBlueprintActorReference", "Level Script Blueprint references actor")))
-		->AddToken(FActorToken::Create(ActorDescView.GetActorSoftPath().ToString(), ActorDescView.GetGuid(), FText::FromString(GetActorName(ActorDescView))))
-		->AddToken(FTextToken::Create(LOCTEXT("TokenMessage_WorldPartition_LevelScriptBlueprintDataLayerReference", "with a non empty set of data layers")))
-		->AddToken(FMapErrorToken::Create(TEXT("WorldPartition_LevelScriptBlueprintRefefenceDataLayer_CheckForErrors")));
+	switch(Reason)
+	{
+	case EWorldReferenceInvalidReason::ReferencedActorIsSpatiallyLoaded:
+		Message->AddToken(FTextToken::Create(LOCTEXT("TokenMessage_WorldPartition_WorldReferenceSpatiallyLoadedActor", "World references spatially loaded actor")))
+			->AddToken(FActorToken::Create(ActorDescView.GetActorSoftPath().ToString(), ActorDescView.GetGuid(), FText::FromString(GetActorName(ActorDescView))))
+			->AddToken(FMapErrorToken::Create(TEXT("WorldPartition_WorldReferenceSpatiallyLoadedActor_CheckForErrors")));
+		break;
+	case EWorldReferenceInvalidReason::ReferencedActorHasDataLayers:
+		Message->AddToken(FTextToken::Create(LOCTEXT("TokenMessage_WorldPartition_WorldReferenceActorWithDataLayers", "World references actor with data layers")))
+			->AddToken(FActorToken::Create(ActorDescView.GetActorSoftPath().ToString(), ActorDescView.GetGuid(), FText::FromString(GetActorName(ActorDescView))))
+			->AddToken(FMapErrorToken::Create(TEXT("WorldPartition_WorldReferenceActorWithDataLayers_CheckForErrors")));
+		break;
+	}
 
-	HandleTokenizedMessage(MoveTemp(Message));
+	HandleTokenizedMessage(MoveTemp(Message));	
 }
 
 void ITokenizedMessageErrorHandler::OnInvalidReferenceDataLayerAsset(const UDataLayerInstanceWithAsset* DataLayerInstance)
diff --git a/Engine/Source/Runtime/Engine/Private/WorldPartition/WorldPartition.cpp b/Engine/Source/Runtime/Engine/Private/WorldPartition/WorldPartition.cpp
--- a/Engine/Source/Runtime/Engine/Private/WorldPartition/WorldPartition.cpp
+++ b/Engine/Source/Runtime/Engine/Private/WorldPartition/WorldPartition.cpp
@@ -1917,26 +1917,23 @@ void UWorldPartition::AppendAssetRegistryTags(FAssetRegistryTagsContext Context)
 		Context.AddTag(FAssetRegistryTag(NAME_LevelHasStreamingDisabled, TEXT("1"), FAssetRegistryTag::TT_Hidden));
 	}
 
-	// Append level script references so we can perform changelists validations without loading the world
-	if (const ULevelScriptBlueprint* LevelScriptBlueprint = GetWorld()->PersistentLevel->GetLevelScriptBlueprint(true))
-	{
-		const ActorsReferencesUtils::FGetActorReferencesParams Params = ActorsReferencesUtils::FGetActorReferencesParams((UObject*)LevelScriptBlueprint)
-			.SetRequiredFlags(RF_HasExternalPackage);
-		TArray<ActorsReferencesUtils::FActorReference> LevelScriptExternalActorReferences = ActorsReferencesUtils::GetActorReferences(Params);
+	// Append world references so we can perform changelists validations without loading it
+	const ActorsReferencesUtils::FGetActorReferencesParams Params = ActorsReferencesUtils::FGetActorReferencesParams(GetWorld())
+		.SetRequiredFlags(RF_HasExternalPackage);
+	TArray<ActorsReferencesUtils::FActorReference> WorldExternalActorReferences = ActorsReferencesUtils::GetActorReferences(Params);
 		
-		if (LevelScriptExternalActorReferences.Num())
+	if (WorldExternalActorReferences.Num())
+	{
+		FStringBuilderBase StringBuilder;
+		for (const ActorsReferencesUtils::FActorReference& ActorReference : WorldExternalActorReferences)
 		{
-			FStringBuilderBase StringBuilder;
-			for (const ActorsReferencesUtils::FActorReference& ActorReference : LevelScriptExternalActorReferences)
-			{
-				StringBuilder.Append(ActorReference.Actor->GetActorGuid().ToString(EGuidFormats::Short));
-				StringBuilder.AppendChar(TEXT(','));
-			}
-			StringBuilder.RemoveSuffix(1);
-
-			static const FName NAME_LevelScriptExternalActorsReferences(TEXT("LevelScriptExternalActorsReferences"));
-			Context.AddTag(FAssetRegistryTag(NAME_LevelScriptExternalActorsReferences, StringBuilder.ToString(), FAssetRegistryTag::TT_Hidden));
+			StringBuilder.Append(ActorReference.Actor->GetActorGuid().ToString(EGuidFormats::Short));
+			StringBuilder.AppendChar(TEXT(','));
 		}
+		StringBuilder.RemoveSuffix(1);
+
+		static const FName NAME_WorldExternalActorsReferences(TEXT("WorldExternalActorsReferences"));
+		Context.AddTag(FAssetRegistryTag(NAME_WorldExternalActorsReferences, StringBuilder.ToString(), FAssetRegistryTag::TT_Hidden));
 	}
 }
 
diff --git a/Engine/Source/Runtime/Engine/Private/WorldPartition/WorldPartitionStreamingGeneration.cpp b/Engine/Source/Runtime/Engine/Private/WorldPartition/WorldPartitionStreamingGeneration.cpp
--- a/Engine/Source/Runtime/Engine/Private/WorldPartition/WorldPartitionStreamingGeneration.cpp
+++ b/Engine/Source/Runtime/Engine/Private/WorldPartition/WorldPartitionStreamingGeneration.cpp
@@ -1101,17 +1101,14 @@ class FWorldPartitionStreamingGenerator
 		const bool bIsMainContainerNonContentBundle = ContainerCollectionInstanceDescriptor.ID.IsMainContainer() && !ContainerCollectionInstanceDescriptor.ContentBundleID.IsValid();
 		if (bIsMainContainerNonContentBundle)
 		{
-			TArray<FGuid> LevelScriptReferences;
+			TArray<FGuid> WorldReferences;
 			if (WorldPartitionContext)
 			{
-				// Gather all references to external actors from the level script and make them always loaded
-				if (ULevelScriptBlueprint* LevelScriptBlueprint = WorldPartitionContext->GetTypedOuter<UWorld>()->PersistentLevel->GetLevelScriptBlueprint(true))
-				{
-					const ActorsReferencesUtils::FGetActorReferencesParams Params = ActorsReferencesUtils::FGetActorReferencesParams(LevelScriptBlueprint)
-						.SetRequiredFlags(RF_HasExternalPackage);
-					TArray<ActorsReferencesUtils::FActorReference> LevelScriptExternalActorReferences = ActorsReferencesUtils::GetActorReferences(Params);
-					Algo::Transform(LevelScriptExternalActorReferences, LevelScriptReferences, [](const ActorsReferencesUtils::FActorReference& ActorReference) { return ActorReference.Actor->GetActorGuid(); });
-				}
+				// Gather all references to external actors from the world and make them non-spatially loaded
+				const ActorsReferencesUtils::FGetActorReferencesParams Params = ActorsReferencesUtils::FGetActorReferencesParams(WorldPartitionContext->GetTypedOuter<UWorld>())
+					.SetRequiredFlags(RF_HasExternalPackage);
+				TArray<ActorsReferencesUtils::FActorReference> WorldExternalActorReferences = ActorsReferencesUtils::GetActorReferences(Params);
+				Algo::Transform(WorldExternalActorReferences, WorldReferences, [](const ActorsReferencesUtils::FActorReference& ActorReference) { return ActorReference.Actor->GetActorGuid(); });
 
 				// Validate data layers
 				if (DataLayerManager)
@@ -1125,22 +1122,22 @@ class FWorldPartitionStreamingGenerator
 			}
 			else
 			{
-				ULevel::GetLevelScriptExternalActorsReferencesFromPackage(ContainerCollectionInstanceDescriptor.ContainerInstanceCollection->GetBaseContainerInstancePackageName(), LevelScriptReferences);
+				ULevel::GetWorldExternalActorsReferencesFromPackage(ContainerCollectionInstanceDescriptor.ContainerInstanceCollection->GetBaseContainerInstancePackageName(), WorldReferences);
 			}
 
-			for (const FGuid& LevelScriptReferenceActorGuid : LevelScriptReferences)
+			for (const FGuid& LevelScriptReferenceActorGuid : WorldReferences)
 			{
 				if (FStreamingGenerationActorDescView* ActorDescView = ContainerCollectionInstanceDescriptor.ActorDescViewMap->FindByGuid(LevelScriptReferenceActorGuid))
 				{
 					if (ActorDescView->GetIsSpatiallyLoaded())
 					{
-						ErrorHandler->OnInvalidReferenceLevelScriptStreamed(*ActorDescView);
+						ErrorHandler->OnInvalidWorldReference(*ActorDescView, IStreamingGenerationErrorHandler::EWorldReferenceInvalidReason::ReferencedActorIsSpatiallyLoaded);
 						ActorDescView->SetForcedNonSpatiallyLoaded();
 					}
 
 					if (ActorDescView->GetRuntimeDataLayerInstanceNames().Num())
 					{
-						ErrorHandler->OnInvalidReferenceLevelScriptDataLayers(*ActorDescView);
+						ErrorHandler->OnInvalidWorldReference(*ActorDescView, IStreamingGenerationErrorHandler::EWorldReferenceInvalidReason::ReferencedActorHasDataLayers);
 						ActorDescView->SetForcedNoDataLayers();
 					}
 				}
@@ -2321,4 +2318,4 @@ void UWorldPartition::CheckForErrors(const FCheckForErrorsParams& InParams)
 
 #endif // WITH_EDITOR
 
-#undef LOCTEXT_NAMESPACE
\ No newline at end of file
+#undef LOCTEXT_NAMESPACE
diff --git a/Engine/Source/Runtime/Engine/Public/WorldPartition/ErrorHandling/WorldPartitionStreamingGenerationErrorHandler.h b/Engine/Source/Runtime/Engine/Public/WorldPartition/ErrorHandling/WorldPartitionStreamingGenerationErrorHandler.h
--- a/Engine/Source/Runtime/Engine/Public/WorldPartition/ErrorHandling/WorldPartitionStreamingGenerationErrorHandler.h
+++ b/Engine/Source/Runtime/Engine/Public/WorldPartition/ErrorHandling/WorldPartitionStreamingGenerationErrorHandler.h
@@ -33,7 +33,7 @@ public:
 	virtual void OnInvalidReferenceGridPlacement(const IWorldPartitionActorDescInstanceView& ActorDescView, const IWorldPartitionActorDescInstanceView& ReferenceActorDescView) = 0;
 
 	/**
-	 * Used to identify an actor data layer error
+	 * Called when there's an error with the data layers used by an actor and one or more referenced actor.
 	 */
 	enum class EDataLayerInvalidReason
 	{
@@ -41,9 +41,6 @@ public:
 		ReferencedActorDifferentExternalDataLayer
 	};
 
-	/**
-	 * Called when there's an error with the data layers used by an actor and one or more referenced actor.
-	 */
 	virtual void OnInvalidReferenceDataLayers(const IWorldPartitionActorDescInstanceView& ActorDescView, const IWorldPartitionActorDescInstanceView& ReferenceActorDescView, EDataLayerInvalidReason Reason) = 0;
 
 	/**
@@ -52,14 +49,15 @@ public:
 	virtual void OnInvalidReferenceRuntimeGrid(const IWorldPartitionActorDescInstanceView& ActorDescView, const IWorldPartitionActorDescInstanceView& ReferenceActorDescView) = 0;
 
 	/**
-	 * Called when the level script references a streamed actor.
+	 * Called when the world references a streamed actor.
 	 */
-	virtual void OnInvalidReferenceLevelScriptStreamed(const IWorldPartitionActorDescInstanceView& ActorDescView) = 0;
+	enum class EWorldReferenceInvalidReason
+	{
+		ReferencedActorIsSpatiallyLoaded,
+		ReferencedActorHasDataLayers
+	};
 
-	/**
-	 * Called when an actor descriptor references an actor using data layers.
-	 */
-	virtual void OnInvalidReferenceLevelScriptDataLayers(const IWorldPartitionActorDescInstanceView& ActorDescView) = 0;
+	virtual void OnInvalidWorldReference(const IWorldPartitionActorDescInstanceView& ActorDescView, EWorldReferenceInvalidReason Reason) = 0;
 
 	/**
 	 * Called when a data layer instance does not have a data layer asset
@@ -96,7 +94,6 @@ public:
 	 */
 	virtual void OnActorNeedsResave(const IWorldPartitionActorDescInstanceView& ActorDescView) = 0;
 
-
 	/**
 	 * Used to identify a level instance actor error
 	 */
@@ -112,12 +109,12 @@ public:
 	virtual void OnLevelInstanceInvalidWorldAsset(const IWorldPartitionActorDescInstanceView& ActorDescView, FName WorldAsset, ELevelInstanceInvalidReason Reason) = 0;
 
 	/**
-	 * Called when an actor descriptor references another actor with a different set of actor filter.
+	 * Called when an actor references another actor with a different set of actor filter.
 	 */
 	virtual void OnInvalidActorFilterReference(const IWorldPartitionActorDescInstanceView& ActorDescView, const IWorldPartitionActorDescInstanceView& ReferenceActorDescView) = 0;
 
 	/**
-	 * Called when an actor descriptor has an invalid HLOD layer.
+	 * Called when an actor has an invalid HLOD layer.
 	 */
 	virtual void OnInvalidHLODLayer(const IWorldPartitionActorDescInstanceView& ActorDescView) = 0;
 
@@ -125,8 +122,6 @@ public:
 	static ENGINE_API FString GetActorName(const IWorldPartitionActorDescInstanceView& ActorDescView);
 	static ENGINE_API FString GetFullActorName(const IWorldPartitionActorDescInstanceView& ActorDescView);
 	
-	PRAGMA_DISABLE_DEPRECATION_WARNINGS
-
 	UE_DEPRECATED(5.4, "Use OnInvalidReferenceDataLayers with EDataLayerInvalidReason instead.")
 	virtual void OnInvalidReferenceDataLayers(const IWorldPartitionActorDescInstanceView& ActorDescView, const IWorldPartitionActorDescInstanceView& ReferenceActorDescView) final {}
 
@@ -166,12 +161,16 @@ public:
 	UE_DEPRECATED(5.4, "Use IWorldPartitionActorDescInstanceView version instead")
 	virtual void OnInvalidHLODLayer(const FWorldPartitionActorDescView& ActorDescView) {}
 
+	UE_DEPRECATED(5.4, "Use OnInvalidWorldReference instead")
+	void OnInvalidReferenceLevelScriptStreamed(const IWorldPartitionActorDescInstanceView& ActorDescView) { OnInvalidWorldReference(ActorDescView, EWorldReferenceInvalidReason::ReferencedActorIsSpatiallyLoaded); }
+
+	UE_DEPRECATED(5.4, "Use OnInvalidWorldReference instead")
+	void OnInvalidReferenceLevelScriptDataLayers(const IWorldPartitionActorDescInstanceView& ActorDescView) { OnInvalidWorldReference(ActorDescView, EWorldReferenceInvalidReason::ReferencedActorHasDataLayers); }
+
 	UE_DEPRECATED(5.4, "Use IWorldPartitionActorDescInstanceView version instead")
 	static ENGINE_API FString GetActorName(const FWorldPartitionActorDescView& ActorDescView) { return FString(); }
 
 	UE_DEPRECATED(5.4, "Use IWorldPartitionActorDescInstanceView version instead")
 	static ENGINE_API FString GetFullActorName(const FWorldPartitionActorDescView& ActorDescView) { return FString(); }
-
-	PRAGMA_ENABLE_DEPRECATION_WARNINGS
 };
 #endif
diff --git a/Engine/Source/Runtime/Engine/Public/WorldPartition/ErrorHandling/WorldPartitionStreamingGenerationLogErrorHandler.h b/Engine/Source/Runtime/Engine/Public/WorldPartition/ErrorHandling/WorldPartitionStreamingGenerationLogErrorHandler.h
--- a/Engine/Source/Runtime/Engine/Public/WorldPartition/ErrorHandling/WorldPartitionStreamingGenerationLogErrorHandler.h
+++ b/Engine/Source/Runtime/Engine/Public/WorldPartition/ErrorHandling/WorldPartitionStreamingGenerationLogErrorHandler.h
@@ -11,8 +11,7 @@ public:
 	ENGINE_API virtual void OnInvalidReference(const IWorldPartitionActorDescInstanceView& ActorDescView, const FGuid& ReferenceGuid, IWorldPartitionActorDescInstanceView* ReferenceActorDescView) override;
 	ENGINE_API virtual void OnInvalidReferenceGridPlacement(const IWorldPartitionActorDescInstanceView& ActorDescView, const IWorldPartitionActorDescInstanceView& ReferenceActorDescView) override;
 	ENGINE_API virtual void OnInvalidReferenceDataLayers(const IWorldPartitionActorDescInstanceView& ActorDescView, const IWorldPartitionActorDescInstanceView& ReferenceActorDescView, EDataLayerInvalidReason Reason) override;
-	ENGINE_API virtual void OnInvalidReferenceLevelScriptStreamed(const IWorldPartitionActorDescInstanceView& ActorDescView) override;
-	ENGINE_API virtual void OnInvalidReferenceLevelScriptDataLayers(const IWorldPartitionActorDescInstanceView& ActorDescView) override;
+	ENGINE_API virtual void OnInvalidWorldReference(const IWorldPartitionActorDescInstanceView& ActorDescView, EWorldReferenceInvalidReason Reason) override;
 	ENGINE_API virtual void OnInvalidReferenceRuntimeGrid(const IWorldPartitionActorDescInstanceView& ActorDescView, const IWorldPartitionActorDescInstanceView& ReferenceActorDescView) override;
 	ENGINE_API virtual void OnInvalidReferenceDataLayerAsset(const UDataLayerInstanceWithAsset* DataLayerInstance) override;
 	ENGINE_API virtual void OnInvalidDataLayerAssetType(const UDataLayerInstanceWithAsset* DataLayerInstance, const UDataLayerAsset* DataLayerAsset) override;
diff --git a/Engine/Source/Runtime/Engine/Public/WorldPartition/ErrorHandling/WorldPartitionStreamingGenerationNullErrorHandler.h b/Engine/Source/Runtime/Engine/Public/WorldPartition/ErrorHandling/WorldPartitionStreamingGenerationNullErrorHandler.h
--- a/Engine/Source/Runtime/Engine/Public/WorldPartition/ErrorHandling/WorldPartitionStreamingGenerationNullErrorHandler.h
+++ b/Engine/Source/Runtime/Engine/Public/WorldPartition/ErrorHandling/WorldPartitionStreamingGenerationNullErrorHandler.h
@@ -11,8 +11,7 @@ public:
 	virtual void OnInvalidReference(const IWorldPartitionActorDescInstanceView& ActorDescView, const FGuid& ReferenceGuid, IWorldPartitionActorDescInstanceView* ReferenceActorDescView) override {}
 	virtual void OnInvalidReferenceGridPlacement(const IWorldPartitionActorDescInstanceView& ActorDescView, const IWorldPartitionActorDescInstanceView& ReferenceActorDescView) override {}
 	virtual void OnInvalidReferenceDataLayers(const IWorldPartitionActorDescInstanceView& ActorDescView, const IWorldPartitionActorDescInstanceView& ReferenceActorDescView, EDataLayerInvalidReason Reason) override {}
-	virtual void OnInvalidReferenceLevelScriptStreamed(const IWorldPartitionActorDescInstanceView& ActorDescView) override {}
-	virtual void OnInvalidReferenceLevelScriptDataLayers(const IWorldPartitionActorDescInstanceView& ActorDescView) override {}
+	virtual void OnInvalidWorldReference(const IWorldPartitionActorDescInstanceView& ActorDescView, EWorldReferenceInvalidReason Reason) override {}
 	virtual void OnInvalidReferenceRuntimeGrid(const IWorldPartitionActorDescInstanceView& ActorDescView, const IWorldPartitionActorDescInstanceView& ReferenceActorDescView) override {}
 	virtual void OnInvalidReferenceDataLayerAsset(const UDataLayerInstanceWithAsset* DataLayerInstance) override {}
 	virtual void OnInvalidDataLayerAssetType(const UDataLayerInstanceWithAsset* DataLayerInstance, const UDataLayerAsset* DataLayerAsset) override {}
diff --git a/Engine/Source/Runtime/Engine/Public/WorldPartition/ErrorHandling/WorldPartitionStreamingGenerationProxyErrorHandler.h b/Engine/Source/Runtime/Engine/Public/WorldPartition/ErrorHandling/WorldPartitionStreamingGenerationProxyErrorHandler.h
--- a/Engine/Source/Runtime/Engine/Public/WorldPartition/ErrorHandling/WorldPartitionStreamingGenerationProxyErrorHandler.h
+++ b/Engine/Source/Runtime/Engine/Public/WorldPartition/ErrorHandling/WorldPartitionStreamingGenerationProxyErrorHandler.h
@@ -15,8 +15,7 @@ public:
 	ENGINE_API virtual void OnInvalidReference(const IWorldPartitionActorDescInstanceView& ActorDescView, const FGuid& ReferenceGuid, IWorldPartitionActorDescInstanceView* ReferenceActorDescView) override;
 	ENGINE_API virtual void OnInvalidReferenceGridPlacement(const IWorldPartitionActorDescInstanceView& ActorDescView, const IWorldPartitionActorDescInstanceView& ReferenceActorDescView) override;
 	ENGINE_API virtual void OnInvalidReferenceDataLayers(const IWorldPartitionActorDescInstanceView& ActorDescView, const IWorldPartitionActorDescInstanceView& ReferenceActorDescView, EDataLayerInvalidReason Reason) override;
-	ENGINE_API virtual void OnInvalidReferenceLevelScriptStreamed(const IWorldPartitionActorDescInstanceView& ActorDescView) override;
-	ENGINE_API virtual void OnInvalidReferenceLevelScriptDataLayers(const IWorldPartitionActorDescInstanceView& ActorDescView) override;
+	ENGINE_API virtual void OnInvalidWorldReference(const IWorldPartitionActorDescInstanceView& ActorDescView, EWorldReferenceInvalidReason Reason) override;
 	ENGINE_API virtual void OnInvalidReferenceRuntimeGrid(const IWorldPartitionActorDescInstanceView& ActorDescView, const IWorldPartitionActorDescInstanceView& ReferenceActorDescView) override;
 	ENGINE_API virtual void OnInvalidReferenceDataLayerAsset(const UDataLayerInstanceWithAsset* DataLayerInstance) override;
 	ENGINE_API virtual void OnInvalidDataLayerAssetType(const UDataLayerInstanceWithAsset* DataLayerInstance, const UDataLayerAsset* DataLayerAsset) override;
diff --git a/Engine/Source/Runtime/Engine/Public/WorldPartition/ErrorHandling/WorldPartitionStreamingGenerationTokenizedMessageErrorHandler.h b/Engine/Source/Runtime/Engine/Public/WorldPartition/ErrorHandling/WorldPartitionStreamingGenerationTokenizedMessageErrorHandler.h
--- a/Engine/Source/Runtime/Engine/Public/WorldPartition/ErrorHandling/WorldPartitionStreamingGenerationTokenizedMessageErrorHandler.h
+++ b/Engine/Source/Runtime/Engine/Public/WorldPartition/ErrorHandling/WorldPartitionStreamingGenerationTokenizedMessageErrorHandler.h
@@ -12,8 +12,7 @@ class ITokenizedMessageErrorHandler : public IStreamingGenerationErrorHandler
 	ENGINE_API virtual void OnInvalidReference(const IWorldPartitionActorDescInstanceView& ActorDescView, const FGuid& ReferenceGuid, IWorldPartitionActorDescInstanceView* ReferenceActorDescView) override;
 	ENGINE_API virtual void OnInvalidReferenceGridPlacement(const IWorldPartitionActorDescInstanceView& ActorDescView, const IWorldPartitionActorDescInstanceView& ReferenceActorDescView) override;
 	ENGINE_API virtual void OnInvalidReferenceDataLayers(const IWorldPartitionActorDescInstanceView& ActorDescView, const IWorldPartitionActorDescInstanceView& ReferenceActorDescView, EDataLayerInvalidReason Reason) override;
-	ENGINE_API virtual void OnInvalidReferenceLevelScriptStreamed(const IWorldPartitionActorDescInstanceView& ActorDescView) override;
-	ENGINE_API virtual void OnInvalidReferenceLevelScriptDataLayers(const IWorldPartitionActorDescInstanceView& ActorDescView) override;
+	ENGINE_API virtual void OnInvalidWorldReference(const IWorldPartitionActorDescInstanceView& ActorDescView, EWorldReferenceInvalidReason Reason) override;
 	ENGINE_API virtual void OnInvalidReferenceRuntimeGrid(const IWorldPartitionActorDescInstanceView& ActorDescView, const IWorldPartitionActorDescInstanceView& ReferenceActorDescView) override;
 	ENGINE_API virtual void OnInvalidReferenceDataLayerAsset(const UDataLayerInstanceWithAsset* DataLayerInstance) override;
 	ENGINE_API virtual void OnInvalidDataLayerAssetType(const UDataLayerInstanceWithAsset* DataLayerInstance, const UDataLayerAsset* DataLayerAsset) override;
