[shaders]
- modify the input hash debug dump mechanism to output an empty "debughash_<hash>" file instead of a txt file with the hash in contents, and always dump these files for the instance of the job that actually compiled
- the existing cvar will now just make it so these files are also dumped for jobs which hit in DDC or the job cache; we don't do this by default so there's only a single match for the debug hash for any given shader normally and it is inside the folder containing the full debug info, including those artifacts which are only output as a side effect of the compile step
- add the same hash as the first line in the stripped source code, so "debughash_<hash>" can be used as a search term in Everything to quickly find debug info associated with a shader (i.e. when looking at a capture in renderdoc or similar)

diff --git a/Engine/Source/Developer/ShaderCompilerCommon/Private/ShaderCompilerCommon.cpp b/Engine/Source/Developer/ShaderCompilerCommon/Private/ShaderCompilerCommon.cpp
--- a/Engine/Source/Developer/ShaderCompilerCommon/Private/ShaderCompilerCommon.cpp
+++ b/Engine/Source/Developer/ShaderCompilerCommon/Private/ShaderCompilerCommon.cpp
@@ -1847,10 +1847,29 @@ namespace UE::ShaderCompilerCommon
 				FFileHelper::SaveStringToFile(Merged, *GetDebugFileName(Input, Options, TEXT("Diagnostics.txt")), FFileHelper::EEncodingOptions::ForceAnsi);
 			}
 		}
+		
+		// delete old DebugHash_* files so we don't clutter the debug info folder (these change every time the deadstripped source code changes)
+		IFileManager::Get().IterateDirectory(*Input.DumpDebugInfoPath, [](const TCHAR* FilenameOrDirectory, bool bIsDirectory)
+			{
+				if (!bIsDirectory)
+				{
+					FStringView Filename = FPathViews::GetCleanFilename(FilenameOrDirectory);
+					if (Filename.StartsWith(kShaderSourceDebugHashPrefix))
+					{
+						IFileManager::Get().Delete(FilenameOrDirectory);
+					}
+				}
+				return true;
+			});
 
-		if (EnumHasAnyFlags(Input.DebugInfoFlags, EShaderDebugInfoFlags::InputHash))
+		if (EnumHasAnyFlags(Input.DebugInfoFlags, EShaderDebugInfoFlags::InputHash) || Output.CompileTime > 0.0f)
 		{
-			FFileHelper::SaveStringToFile(LexToString(Input.Hash), *GetDebugFileName(Input, Options, TEXT("InputHash.txt")), FFileHelper::EEncodingOptions::ForceAnsi);
+			// If compile time was > 0, this was the copy of the job that actually compiled; we write an empty file with the shader hash in it so it can be
+			// found easily using the ShaderHash comment printed on the first line of the stripped source code.
+			// We don't do this for cache hits or duplicate jobs unless explicitly requested (via the InputHash debug info flags), so that _all_ debug artifacts
+			// (including those only generated by the compile process) are available in the folder containing this file.
+			FString InputHashStr = LexToString(Input.Hash);
+			FFileHelper::SaveStringToFile(FStringView(), *GetDebugFileName(Input, Options, kShaderSourceDebugHashPrefix, *InputHashStr));
 		}
 
 		if (EnumHasAnyFlags(Input.DebugInfoFlags, EShaderDebugInfoFlags::ShaderCodeBinary))
diff --git a/Engine/Source/Runtime/Engine/Private/ShaderCompiler/ShaderCompiler.cpp b/Engine/Source/Runtime/Engine/Private/ShaderCompiler/ShaderCompiler.cpp
--- a/Engine/Source/Runtime/Engine/Private/ShaderCompiler/ShaderCompiler.cpp
+++ b/Engine/Source/Runtime/Engine/Private/ShaderCompiler/ShaderCompiler.cpp
@@ -396,10 +396,9 @@ private:
 /** Map element type for job cache */
 struct FShaderJobData
 {
-	using FJobInputHash = FBlake3Hash;
 	using FJobOutputHash = FBlake3Hash;
 
-	FJobInputHash InputHash;
+	FShaderCompilerInputHash InputHash;
 
 	/** Output hash will be zero if output data has not been written yet, or can be cleared if output data has been removed */
 	FJobOutputHash OutputHash;
@@ -447,8 +446,8 @@ public:
 		Reserve(FShaderJobDataBlock::BlockSize);
 	}
 
-	FShaderJobData* Find(const FShaderJobData::FJobInputHash& Key);
-	FShaderJobCacheRef FindOrAdd(const FShaderJobData::FJobInputHash& Key);
+	FShaderJobData* Find(const FShaderCompilerInputHash& Key);
+	FShaderJobCacheRef FindOrAdd(const FShaderCompilerInputHash& Key);
 
 	FORCEINLINE int32 Num() const
 	{
@@ -754,7 +753,6 @@ public:
 	int32 GetPendingJobs(EShaderCompilerWorkerType InWorkerType, EShaderCompileJobPriority InPriority, int32 MinNumJobs, int32 MaxNumJobs, TArray<FShaderCommonCompileJobPtr>& OutJobs);
 
 private:
-	using FJobInputHash = FShaderCommonCompileJob::FInputHash;
 	using FJobOutputHash = FBlake3Hash;
 	using FJobCachedOutput = FSharedBuffer;
 	using FStoredOutput = FShaderJobCacheStoredOutput;
@@ -770,10 +768,10 @@ private:
 	void InternalSetPriority(FShaderCommonCompileJob* Job, EShaderCompileJobPriority InPriority);
 
 	/** Looks for or adds an entry for the given hash in the cache.  Returns cached output if it exists, or may initialize DDC request if one has been issued. */
-	FShaderJobCacheRef FindOrAdd(const FJobInputHash& Hash, EShaderCompileJobPriority JobPriority, const bool bCheckDDC, TPimplPtr<UE::DerivedData::FRequestOwner>& InoutRequestOwner, FJobCachedOutput*& OutCachedOutput);
+	FShaderJobCacheRef FindOrAdd(const FShaderCompilerInputHash& Hash, EShaderCompileJobPriority JobPriority, const bool bCheckDDC, TPimplPtr<UE::DerivedData::FRequestOwner>& InoutRequestOwner, FJobCachedOutput*& OutCachedOutput);
 
 	/** Find an existing item in the cache. */
-	FShaderJobData* Find(const FJobInputHash& Hash);
+	FShaderJobData* Find(const FShaderCompilerInputHash& Hash);
 
 	/** Add a reference to a duplicate job (to the DuplicateJobs array) */
 	void AddDuplicateJob(FShaderCommonCompileJob* DuplicateJob);
@@ -782,7 +780,7 @@ private:
 	void RemoveDuplicateJob(FShaderCommonCompileJob* DuplicateJob);
 
 	/** Adds a job output to the cache */
-	void AddJobOutput(FShaderJobData& JobData, const FShaderCommonCompileJob* FinishedJob, const FJobInputHash& Hash, const FJobCachedOutput& Contents, int32 InitialHitCount, const bool bAddToDDC);
+	void AddJobOutput(FShaderJobData& JobData, const FShaderCommonCompileJob* FinishedJob, const FShaderCompilerInputHash& Hash, const FJobCachedOutput& Contents, int32 InitialHitCount, const bool bAddToDDC);
 
 	/** Returns memory used by the cache*/
 	uint64 GetAllocatedMemory() const;
@@ -986,7 +984,7 @@ static FShaderJobData& GetShaderJobData(const FShaderJobCacheRef& CacheRef)
 	return CacheRef.Block->Data[CacheRef.IndexInBlock];
 }
 
-FShaderJobData* FShaderJobDataMap::Find(const FShaderJobData::FJobInputHash& Key)
+FShaderJobData* FShaderJobDataMap::Find(const FShaderCompilerInputHash& Key)
 {
 	// Search for key with linear probing
 	for (uint32 TableIndex = GetTypeHash(Key) & HashTableMask; HashTable[TableIndex] != INDEX_NONE; TableIndex = (TableIndex + 1) & HashTableMask)
@@ -999,7 +997,7 @@ FShaderJobData* FShaderJobDataMap::Find(const FShaderJobData::FJobInputHash& Key
 	return nullptr;
 }
 
-FShaderJobCacheRef FShaderJobDataMap::FindOrAdd(const FShaderJobData::FJobInputHash& Key)
+FShaderJobCacheRef FShaderJobDataMap::FindOrAdd(const FShaderCompilerInputHash& Key)
 {
 	// Search for key with linear probing
 	uint32 TableIndex;
@@ -1613,7 +1611,7 @@ void FShaderJobCache::SubmitJob(FShaderCommonCompileJob* Job)
 	{
 		bJobCacheLocked = true;
 
-		const FShaderCommonCompileJob::FInputHash& InputHash = Job->GetInputHash();
+		const FShaderCompilerInputHash& InputHash = Job->GetInputHash();
 
 		const bool bCheckDDC = GShaderCompilerPerShaderDDCGlobal || !(Job->bIsDefaultMaterial || Job->bIsGlobalShader);
 
@@ -1696,7 +1694,7 @@ void FShaderJobCache::SubmitJob(FShaderCommonCompileJob* Job)
 	else if (ShaderCompiler::IsJobCacheDebugValidateEnabled())
 	{
 		FSharedBuffer* ExistingOutput;
-		const FShaderCommonCompileJob::FInputHash& InputHash = Job->GetInputHash();
+		const FShaderCompilerInputHash& InputHash = Job->GetInputHash();
 		const bool bCheckDDC = !(Job->bIsDefaultMaterial || Job->bIsGlobalShader);
 		JobLock.WriteLock();
 		Job->JobCacheRef = FindOrAdd(InputHash, Job->Priority, bCheckDDC, Job->RequestOwner, ExistingOutput);
@@ -1864,7 +1862,7 @@ void FShaderJobCache::AddToCacheAndProcessPending(FShaderCommonCompileJob* Finis
 
 	ensureMsgf(FinishedJob->bInputHashSet, TEXT("Finished job didn't have input hash set, was shader compiler jobs cache toggled runtime?"));
 
-	const FShaderCommonCompileJob::FInputHash& InputHash = FinishedJob->GetInputHash();
+	const FShaderCompilerInputHash& InputHash = FinishedJob->GetInputHash();
 	TArray<uint8> Output;
 	FMemoryWriter Writer(Output);
 	FinishedJob->SerializeOutput(Writer);
@@ -10855,7 +10853,7 @@ namespace
 }
 #endif
 
-FShaderJobCacheRef FShaderJobCache::FindOrAdd(const FJobInputHash& Hash, EShaderCompileJobPriority JobPriority, const bool bCheckDDC, TPimplPtr<UE::DerivedData::FRequestOwner>& InoutRequestOwner, FJobCachedOutput*& OutCachedOutput)
+FShaderJobCacheRef FShaderJobCache::FindOrAdd(const FShaderCompilerInputHash& Hash, EShaderCompileJobPriority JobPriority, const bool bCheckDDC, TPimplPtr<UE::DerivedData::FRequestOwner>& InoutRequestOwner, FJobCachedOutput*& OutCachedOutput)
 {
 	LLM_SCOPE_BYTAG(ShaderCompiler);
 
@@ -11134,7 +11132,7 @@ FShaderJobCacheRef FShaderJobCache::FindOrAdd(const FJobInputHash& Hash, EShader
 	return JobCacheRef;
 }
 
-FShaderJobData* FShaderJobCache::Find(const FJobInputHash& Hash)
+FShaderJobData* FShaderJobCache::Find(const FShaderCompilerInputHash& Hash)
 {
 	check(ShaderCompiler::IsJobCacheEnabled());
 
@@ -11195,7 +11193,7 @@ FShaderJobCache::~FShaderJobCache()
 	}
 }
 
-void FShaderJobCache::AddJobOutput(FShaderJobData& JobData, const FShaderCommonCompileJob* FinishedJob, const FJobInputHash& Hash, const FJobCachedOutput& Contents, int32 InitialHitCount, const bool bAddToDDC)
+void FShaderJobCache::AddJobOutput(FShaderJobData& JobData, const FShaderCommonCompileJob* FinishedJob, const FShaderCompilerInputHash& Hash, const FJobCachedOutput& Contents, int32 InitialHitCount, const bool bAddToDDC)
 {
 	TRACE_CPUPROFILER_EVENT_SCOPE(FShaderJobCache::Add);
 
diff --git a/Engine/Source/Runtime/RenderCore/Private/ShaderCore.cpp b/Engine/Source/Runtime/RenderCore/Private/ShaderCore.cpp
--- a/Engine/Source/Runtime/RenderCore/Private/ShaderCore.cpp
+++ b/Engine/Source/Runtime/RenderCore/Private/ShaderCore.cpp
@@ -1491,6 +1491,17 @@ public:
 			// and blank lines to improve deduplication (and populates data required to remap diagnostic messages to correct line numbers)
 			Job.PreprocessOutput.StripCode(Job.Input.NeedsOriginalShaderSource());
 
+			FShaderCompilerInputHash Hash = Job.GetInputHash();
+			// Replace the placeholder debug hash value appended in StripCode with the real job input hash
+			static const FShaderSource::FStringType DebugHashStr(kShaderSourceDebugHashPrefix);
+			FShaderSource::FViewType SourceView = Job.PreprocessOutput.GetSourceViewAnsi();
+			int32 DebugHashLoc = SourceView.Find(DebugHashStr) + DebugHashStr.Len();
+			int32 NewlineLoc = SourceView.Find(SHADER_SOURCE_VIEWLITERAL("\n"), DebugHashLoc);
+			TAnsiStringBuilder<2 * sizeof(FShaderCompilerInputHash::ByteArray) + 1> HashStr;
+			HashStr << Hash;
+			check(NewlineLoc - DebugHashLoc == HashStr.Len());
+			FMemory::Memcpy(Job.PreprocessOutput.EditSource().GetData() + DebugHashLoc, HashStr.GetData(), sizeof(FShaderSource::CharType) * HashStr.Len());
+
 			// always compress the code after stripping to minimize memory footprint
 			Job.PreprocessOutput.CompressCode();
 		}
@@ -3747,26 +3758,26 @@ FArchive& operator<<(FArchive& Ar, FShaderCompilerInput& Input)
 	return Ar;
 }
 
-FShaderCommonCompileJob::FInputHash FShaderPipelineCompileJob::GetInputHash()
+FShaderCompilerInputHash FShaderPipelineCompileJob::GetInputHash()
 {
 	if (bInputHashSet)
 	{
 		return InputHash;
 	}
-	static_assert(sizeof(FShaderCommonCompileJob::FInputHash) == 32);
+	static_assert(sizeof(FShaderCompilerInputHash) == 32);
 	int256 CombinedHash = 0u;
 	for (int32 Index = 0; Index < StageJobs.Num(); ++Index)
 	{
 		if (StageJobs[Index])
 		{
-			const FShaderCommonCompileJob::FInputHash StageHash = StageJobs[Index]->GetInputHash();
-			const FShaderCommonCompileJob::FInputHash::ByteArray& StageHashBytes = StageHash.GetBytes();
+			const FShaderCompilerInputHash StageHash = StageJobs[Index]->GetInputHash();
+			const FShaderCompilerInputHash::ByteArray& StageHashBytes = StageHash.GetBytes();
 			static_assert(sizeof(StageHashBytes) == sizeof(int256));
 			CombinedHash += int256(StageHashBytes, sizeof(StageHashBytes));
 		}
 	}
 
-	InputHash = FShaderCommonCompileJob::FInputHash(reinterpret_cast<FShaderCommonCompileJob::FInputHash::ByteArray&>(*CombinedHash.GetBits()));
+	InputHash = FShaderCompilerInputHash(reinterpret_cast<FShaderCompilerInputHash::ByteArray&>(*CombinedHash.GetBits()));
 	bInputHashSet = true;
 	return InputHash;
 }
@@ -3807,7 +3818,7 @@ struct FShaderVirtualFileContents
 	{}
 };
 
-FShaderCommonCompileJob::FInputHash FShaderCompileJob::GetInputHash()
+FShaderCompilerInputHash FShaderCompileJob::GetInputHash()
 {
 	if (bInputHashSet)
 	{
diff --git a/Engine/Source/Runtime/RenderCore/Private/ShaderPreprocessTypes.cpp b/Engine/Source/Runtime/RenderCore/Private/ShaderPreprocessTypes.cpp
--- a/Engine/Source/Runtime/RenderCore/Private/ShaderPreprocessTypes.cpp
+++ b/Engine/Source/Runtime/RenderCore/Private/ShaderPreprocessTypes.cpp
@@ -9,6 +9,12 @@
 const FShaderSource::CharType* FilenameSentinel = SHADER_SOURCE_LITERAL("__UE_FILENAME_SENTINEL__");
 static const int FilenameSentinelLen = FShaderSource::FCStringType::Strlen(FilenameSentinel);
 static const FShaderSource::FStringType LineDirectiveSentinel = FShaderSource::FStringType::Printf(SHADER_SOURCE_LITERAL("#line 1 \"%s\"\n"), FilenameSentinel);
+static const FShaderSource::FStringType InputHashCommentPlaceholder = FShaderSource::FStringType::Printf
+(
+	SHADER_SOURCE_LITERAL("// %ls%ls\n"),
+	kShaderSourceDebugHashPrefix,
+	*LexToString(FShaderCompilerInputHash())
+);
 
 void FShaderDiagnosticRemapper::Remap(FShaderCompilerError& Diagnostic) const
 {
@@ -153,9 +159,10 @@ inline void SkipNewLine(const FShaderSource::CharType*& Current, const FShaderSo
 void FShaderPreprocessOutput::StripCode(bool bCopyOriginalPreprocessdSource)
 {
 	// Reserve worst case slack (i.e. assuming there is nothing to strip) to avoid reallocation
-	FShaderSource PreprocessedSourceStripped(LineDirectiveSentinel.GetCharArray().GetData(), PreprocessedSource.Len());
+	FShaderSource PreprocessedSourceStripped(InputHashCommentPlaceholder, LineDirectiveSentinel.Len() + PreprocessedSource.Len());
 	FShaderSource::CharType* OutStrippedData = PreprocessedSourceStripped.GetData();
-	FShaderSource::CharType* OutStripped = OutStrippedData + LineDirectiveSentinel.Len();
+	FMemory::Memcpy(OutStrippedData + InputHashCommentPlaceholder.Len(), LineDirectiveSentinel.GetCharArray().GetData(), sizeof(FShaderSource::CharType) * LineDirectiveSentinel.Len());
+	FShaderSource::CharType* OutStripped = OutStrippedData + InputHashCommentPlaceholder.Len() + LineDirectiveSentinel.Len();
 
 	const FShaderSource::CharType* Begin = PreprocessedSource.GetData(), * Current = Begin;
 	const FShaderSource::CharType* End = Current + PreprocessedSource.Len();
diff --git a/Engine/Source/Runtime/RenderCore/Public/ShaderCompilerCore.h b/Engine/Source/Runtime/RenderCore/Public/ShaderCompilerCore.h
--- a/Engine/Source/Runtime/RenderCore/Public/ShaderCompilerCore.h
+++ b/Engine/Source/Runtime/RenderCore/Public/ShaderCompilerCore.h
@@ -229,6 +229,8 @@ enum class EShaderDebugInfoFlags : uint8
 };
 ENUM_CLASS_FLAGS(EShaderDebugInfoFlags)
 
+using FShaderCompilerInputHash = FBlake3Hash;
+
 /** Struct that gathers all readonly inputs needed for the compilation of a single shader. */
 struct FShaderCompilerInput
 {
@@ -281,7 +283,7 @@ struct FShaderCompilerInput
 	FString DebugDescription;
 
 	// Hash of this input (used as the key for the shader job cache)
-	FBlake3Hash Hash;
+	FShaderCompilerInputHash Hash;
 
 	// Compilation Environment
 	FShaderCompilerEnvironment Environment;
diff --git a/Engine/Source/Runtime/RenderCore/Public/ShaderCompilerJobTypes.h b/Engine/Source/Runtime/RenderCore/Public/ShaderCompilerJobTypes.h
--- a/Engine/Source/Runtime/RenderCore/Public/ShaderCompilerJobTypes.h
+++ b/Engine/Source/Runtime/RenderCore/Public/ShaderCompilerJobTypes.h
@@ -79,7 +79,7 @@ public:
 	FShaderCommonCompileJob* NextLink = nullptr;
 	FShaderCommonCompileJob** PrevLink = nullptr;
 
-	using FInputHash = FBlake3Hash;
+	using FInputHash UE_DEPRECATED(5.5, "Use FShaderCompilerInputHash instead of FShaderCommonCompileJob::FInputHash") = FShaderCompilerInputHash;
 
 	FPendingShaderMapCompileResultsPtr PendingShaderMap;
 
@@ -113,7 +113,7 @@ public:
 	/** Whether or not we are a global shader. */
 	uint8 bIsGlobalShader : 1;
 	/** Hash of all the job inputs */
-	FInputHash InputHash;
+	FShaderCompilerInputHash InputHash;
 
 	/** In-engine timestamp of being added to a pending queue. Not set for jobs that are satisfied from the jobs cache */
 	double TimeAddedToPendingQueue = 0.0;
@@ -147,7 +147,7 @@ public:
 	}
 
 	/** Returns hash of all inputs for this job (needed for caching). */
-	virtual FInputHash GetInputHash() { return FInputHash(); }
+	virtual FShaderCompilerInputHash GetInputHash() { return FShaderCompilerInputHash(); }
 
 	/** Serializes (and deserializes) the output for caching purposes. */
 	virtual void SerializeOutput(FArchive& Ar) {}
@@ -251,7 +251,7 @@ public:
 	// List of pipelines that are sharing this job.
 	TMap<const FVertexFactoryType*, TArray<const FShaderPipelineType*>> SharingPipelines;
 
-	virtual RENDERCORE_API FInputHash GetInputHash() override;
+	virtual RENDERCORE_API FShaderCompilerInputHash GetInputHash() override;
 	virtual RENDERCORE_API void SerializeOutput(FArchive& Ar) override;
 
 	virtual RENDERCORE_API void OnComplete() override;
@@ -312,7 +312,7 @@ public:
 	FShaderPipelineCompileJobKey Key;
 	TArray<TRefCountPtr<FShaderCompileJob>> StageJobs;
 
-	virtual RENDERCORE_API FInputHash GetInputHash() override;
+	virtual RENDERCORE_API FShaderCompilerInputHash GetInputHash() override;
 	virtual RENDERCORE_API void SerializeOutput(FArchive& Ar) override;
 	virtual RENDERCORE_API void OnComplete() override;
 	virtual RENDERCORE_API void AppendDebugName(FStringBuilderBase& OutName) const override;
diff --git a/Engine/Source/Runtime/RenderCore/Public/ShaderPreprocessTypes.h b/Engine/Source/Runtime/RenderCore/Public/ShaderPreprocessTypes.h
--- a/Engine/Source/Runtime/RenderCore/Public/ShaderPreprocessTypes.h
+++ b/Engine/Source/Runtime/RenderCore/Public/ShaderPreprocessTypes.h
@@ -13,6 +13,8 @@ struct FShaderCompilerEnvironment;
 struct FShaderCompilerError;
 struct FShaderCompilerInput;
 
+static constexpr const TCHAR* kShaderSourceDebugHashPrefix = TEXT("DebugHash_");
+
 /*
 * Helper class used to remap compiler diagnostic messages from stripped preprocessed source (i.e. source with all whitespace normalized
 * and comments and line directives removed) back to line numbers/locations from the original source. 
