Do not attempt to change fullscreen mode if there is no viewport.

diff --git a/Engine/Source/Runtime/Engine/Private/GameViewportClient.cpp b/Engine/Source/Runtime/Engine/Private/GameViewportClient.cpp
--- a/Engine/Source/Runtime/Engine/Private/GameViewportClient.cpp
+++ b/Engine/Source/Runtime/Engine/Private/GameViewportClient.cpp
@@ -3803,6 +3803,11 @@ bool UGameViewportClient::SetDisplayConfiguration(const FIntPoint* Dimensions, E
 
 bool UGameViewportClient::HandleToggleFullscreenCommand()
 {
+	if (!Viewport)
+	{
+		return true;
+	}
+
 	static auto CVar = IConsoleManager::Get().FindTConsoleVariableDataInt(TEXT("r.FullScreenMode"));
 	check(CVar);
 	auto FullScreenMode = CVar->GetValueOnGameThread() == 0 ? EWindowMode::Fullscreen : EWindowMode::WindowedFullscreen;
