Fixed ExternalDataLayerSubsystem logs and messages.

diff --git a/Engine/Source/Runtime/Engine/Private/WorldPartition/DataLayer/ExternalDataLayerEngineSubsystem.cpp b/Engine/Source/Runtime/Engine/Private/WorldPartition/DataLayer/ExternalDataLayerEngineSubsystem.cpp
--- a/Engine/Source/Runtime/Engine/Private/WorldPartition/DataLayer/ExternalDataLayerEngineSubsystem.cpp
+++ b/Engine/Source/Runtime/Engine/Private/WorldPartition/DataLayer/ExternalDataLayerEngineSubsystem.cpp
@@ -150,7 +150,7 @@ bool UExternalDataLayerEngineSubsystem::CanWorldInjectExternalDataLayerAsset(con
 	{
 		if (OutFailureReason)
 		{
-			*OutFailureReason = FText::Format(LOCTEXT("CantInjectNotRegisteredExternalDataLayerAsset", "External Data Layer Asset {0} not registered"), FText::FromString(InExternalDataLayerAsset->GetName()));
+			*OutFailureReason = FText::Format(LOCTEXT("CantInjectNotRegisteredExternalDataLayerAsset", "External Data Layer Asset {0} is not registered"), FText::FromString(InExternalDataLayerAsset->GetName()));
 		}
 		return false;
 	}
@@ -171,7 +171,7 @@ bool UExternalDataLayerEngineSubsystem::CanWorldInjectExternalDataLayerAsset(con
 	{
 		if (OutFailureReason)
 		{
-			*OutFailureReason = FText::Format(LOCTEXT("CantInjectNotActiveExternalDataLayerAsset", "External Data Layer Asset {0} not active"), FText::FromString(InExternalDataLayerAsset->GetName()));
+			*OutFailureReason = FText::Format(LOCTEXT("CantInjectNotActiveExternalDataLayerAsset", "External Data Layer Asset {0} is not active"), FText::FromString(InExternalDataLayerAsset->GetName()));
 		}
 		return false;
 	}
@@ -187,7 +187,7 @@ void UExternalDataLayerEngineSubsystem::RegisterExternalDataLayerAsset(const UEx
 	{
 		if ((InExternalDataLayerAsset != ExternalDataLayerAsset) && (InExternalDataLayerAsset->GetUID() == ExternalDataLayerAsset->GetUID()))
 		{
-			UE_LOG(LogWorldPartition, Error, TEXT("ExternalDataLayerAsset %s is already registered with UID %s. Can't register ExternalDataLayerAsset %s under same UID."), 
+			UE_LOG(LogWorldPartition, Error, TEXT("External Data Layer Asset %s is already registered with UID %s. Can't register External Data Layer Asset %s under the same UID."), 
 				*ExternalDataLayerAsset->GetPathName(), *ExternalDataLayerAsset->GetUID().ToString(), *InExternalDataLayerAsset->GetPathName());
 			return;
 		}
@@ -209,7 +209,7 @@ void UExternalDataLayerEngineSubsystem::ActivateExternalDataLayerAsset(const UEx
 	check(IsValid(InExternalDataLayerAsset));
 	if (!IsExternalDataLayerAssetRegistered(InExternalDataLayerAsset, InClient))
 	{
-		UE_LOG(LogWorldPartition, Error, TEXT("ExternalDataLayerAsset %s is not registered."), *InExternalDataLayerAsset->GetPathName());
+		UE_LOG(LogWorldPartition, Error, TEXT("External Data Layer Asset %s is not registered."), *InExternalDataLayerAsset->GetPathName());
 		return;
 	}
 
