NetDriver
* GetDescription now logs the netdriver definition next to it's name

#rnx

diff --git a/Engine/Source/Runtime/Engine/Classes/Engine/NetDriver.h b/Engine/Source/Runtime/Engine/Classes/Engine/NetDriver.h
--- a/Engine/Source/Runtime/Engine/Classes/Engine/NetDriver.h
+++ b/Engine/Source/Runtime/Engine/Classes/Engine/NetDriver.h
@@ -1731,7 +1731,7 @@ public:
 	/** @return String that uniquely describes the net driver instance */
 	FString GetDescription() const
 	{ 
-		return FString::Printf(TEXT("%s %s%s"), *NetDriverName.ToString(), *GetName(), bIsPeer ? TEXT("(PEER)") : TEXT(""));
+		return FString::Printf(TEXT("Name:%s Def:%s %s%s"), *NetDriverName.ToString(), *NetDriverDefinition.ToString(), *GetName(), bIsPeer ? TEXT("(PEER)") : TEXT(""));
 	}
 
 	/** @return true if this netdriver is handling accepting connections */
