Fixed crash when unloading a partition world using External Data Layers combined with its main WorldDataLayers using Data Layer Instance External Packaging.

diff --git a/Engine/Source/Runtime/Engine/Private/WorldPartition/DataLayer/WorldDataLayers.cpp b/Engine/Source/Runtime/Engine/Private/WorldPartition/DataLayer/WorldDataLayers.cpp
--- a/Engine/Source/Runtime/Engine/Private/WorldPartition/DataLayer/WorldDataLayers.cpp
+++ b/Engine/Source/Runtime/Engine/Private/WorldPartition/DataLayer/WorldDataLayers.cpp
@@ -644,7 +644,8 @@ void AWorldDataLayers::OnLoadedActorRemovedFromLevel()
 
 	if (IsUsingExternalPackageDataLayerInstances())
 	{
-		ForEachDataLayerInstance([this](UDataLayerInstance* DataLayerInstance)
+		// Iterate on Data Layer Instances excluding Transient Data Layer Instances
+		for (UDataLayerInstance* DataLayerInstance : GetDataLayerInstances())
 		{
 			check(DataLayerInstance->IsPackageExternal())
 			ForEachObjectWithPackage(DataLayerInstance->GetPackage(), [](UObject* Object)
@@ -655,9 +656,7 @@ void AWorldDataLayers::OnLoadedActorRemovedFromLevel()
 				}
 				return true;
 			}, false);
-
-			return true;
-		});
+		}
 	}
 }
 
