Add additional Orthographic Camera settings for further user control granularity per camera instead of relying on CVars, and manually set the PT_OrthoCamera test to ignore the automatic plane resolution logic to avoid regression.

No serialization is occuring because these should be enabled on all ortho cameras by default, even older ones as old ortho scenes did not work correctly in any sample project, however some edge cases are expected. With serialization, many cameras would need their settings manually updated, whereas without serialization, the edge cases in older cameras are rare because Ortho has not worked for so long and in general these settings will improve those cameras.
SceneCaptures are unaffected as their default logic should always be false due to their much wider use cases.

#virtualized

diff --git a/Engine/Source/Editor/DetailCustomizations/Private/CameraDetails.cpp b/Engine/Source/Editor/DetailCustomizations/Private/CameraDetails.cpp
--- a/Engine/Source/Editor/DetailCustomizations/Private/CameraDetails.cpp
+++ b/Engine/Source/Editor/DetailCustomizations/Private/CameraDetails.cpp
@@ -110,12 +110,21 @@ void FCameraDetails::CustomizeDetails( IDetailLayoutBuilder& DetailLayout )
 		IDetailPropertyRow& AutoCalculateOrthoPlanesRow = CameraCategory.AddProperty(DetailLayout.GetProperty(GET_MEMBER_NAME_CHECKED(UCameraComponent, bAutoCalculateOrthoPlanes)));
 		AutoCalculateOrthoPlanesRow.Visibility(OrthographicVisibility);
 
+		IDetailPropertyRow& AutoPlanesShift = CameraCategory.AddProperty(DetailLayout.GetProperty(GET_MEMBER_NAME_CHECKED(UCameraComponent, AutoPlaneShift)));
+		AutoPlanesShift.Visibility(OrthographicVisibility);
+
 		IDetailPropertyRow& OrthoNearClipPlaneRow = CameraCategory.AddProperty(DetailLayout.GetProperty(GET_MEMBER_NAME_CHECKED(UCameraComponent, OrthoNearClipPlane)));
 		OrthoNearClipPlaneRow.Visibility(OrthographicVisibility);
 
 		IDetailPropertyRow& OrthoFarClipPlaneRow = CameraCategory.AddProperty(DetailLayout.GetProperty(GET_MEMBER_NAME_CHECKED(UCameraComponent, OrthoFarClipPlane)));
 		OrthoFarClipPlaneRow.Visibility(OrthographicVisibility);
 
+		IDetailPropertyRow& UpdateOrthoPlanesRow = CameraCategory.AddProperty(DetailLayout.GetProperty(GET_MEMBER_NAME_CHECKED(UCameraComponent, bUpdateOrthoPlanes)));
+		UpdateOrthoPlanesRow.Visibility(OrthographicVisibility);
+
+		IDetailPropertyRow& CameraHeightAsViewTargetRow = CameraCategory.AddProperty(DetailLayout.GetProperty(GET_MEMBER_NAME_CHECKED(UCameraComponent, bUseCameraHeightAsViewTarget)));
+		CameraHeightAsViewTargetRow.Visibility(OrthographicVisibility);
+
 		// Aspect ratio
 		IDetailPropertyRow& AspectRatioRow = CameraCategory.AddProperty(AspectRatioProperty);
 
diff --git a/Engine/Source/Editor/UnrealEd/Private/EditorViewportClient.cpp b/Engine/Source/Editor/UnrealEd/Private/EditorViewportClient.cpp
--- a/Engine/Source/Editor/UnrealEd/Private/EditorViewportClient.cpp
+++ b/Engine/Source/Editor/UnrealEd/Private/EditorViewportClient.cpp
@@ -1200,12 +1200,16 @@ FSceneView* FEditorViewportClient::CalcSceneView(FSceneViewFamily* ViewFamily, c
 			if(!ViewFamily->EngineShowFlags.Wireframe)
 			{
 				FMinimalViewInfo CalculatePlanesViewInfo;
-				CalculatePlanesViewInfo.ProjectionMode = ECameraProjectionMode::Orthographic;
-				CalculatePlanesViewInfo.bAutoCalculateOrthoPlanes = true;
-				CalculatePlanesViewInfo.OrthoWidth = OrthoWidth;
 				CalculatePlanesViewInfo.Rotation = ViewInitOptions.ViewRotationMatrix.Rotator();
 				CalculatePlanesViewInfo.AspectRatio = AspectRatio;
 				CalculatePlanesViewInfo.bConstrainAspectRatio = false;
+
+				CalculatePlanesViewInfo.ProjectionMode = ECameraProjectionMode::Orthographic;				
+				CalculatePlanesViewInfo.OrthoWidth = OrthoWidth;
+				CalculatePlanesViewInfo.bAutoCalculateOrthoPlanes = true;
+				CalculatePlanesViewInfo.AutoPlaneShift = 0.0f;
+				CalculatePlanesViewInfo.bUpdateOrthoPlanes = true;
+				CalculatePlanesViewInfo.bUseCameraHeightAsViewTarget = true;
 				CalculatePlanesViewInfo.OrthoNearClipPlane = OrthoWidth * -CVarOrthoEditorDebugClipPlaneScale.GetValueOnAnyThread();
 				CalculatePlanesViewInfo.OrthoFarClipPlane = FarPlane - NearPlane + CalculatePlanesViewInfo.OrthoNearClipPlane;
 
diff --git a/Engine/Source/Runtime/Engine/Classes/Camera/CameraComponent.h b/Engine/Source/Runtime/Engine/Classes/Camera/CameraComponent.h
--- a/Engine/Source/Runtime/Engine/Classes/Camera/CameraComponent.h
+++ b/Engine/Source/Runtime/Engine/Classes/Camera/CameraComponent.h
@@ -46,6 +46,12 @@ class UCameraComponent : public USceneComponent
 	UFUNCTION(BlueprintCallable, Category = Camera)
 	void SetAutoCalculateOrthoPlanes(bool bAutoCalculate) { bAutoCalculateOrthoPlanes = bAutoCalculate; }
 
+	/** Manually adjusts the planes of this camera, maintaining the distance between them. Positive moves out to the farplane, negative towards the near plane */
+	UPROPERTY(Interp, EditAnywhere, BlueprintReadWrite, Category = CameraSettings, meta = (EditCondition = "bAutoCalculateOrthoPlanes"))
+	float AutoPlaneShift;
+	UFUNCTION(BlueprintCallable, Category = Camera)
+	void SetAutoPlaneShift(float InAutoPlaneShift) { AutoPlaneShift = InAutoPlaneShift; }
+
 	/** The near plane distance of the orthographic view (in world units) */
 	UPROPERTY(Interp, EditAnywhere, BlueprintReadWrite, Category = CameraSettings, meta = (EditCondition = "!bAutoCalculateOrthoPlanes"))
 	float OrthoNearClipPlane;
@@ -58,6 +64,18 @@ class UCameraComponent : public USceneComponent
 	UFUNCTION(BlueprintCallable, Category = Camera)
 	void SetOrthoFarClipPlane(float InOrthoFarClipPlane) { OrthoFarClipPlane = InOrthoFarClipPlane; }
 
+	/** Adjusts the near/far planes and the view origin of the current camera automatically to avoid clipping and light artefacting*/
+	UPROPERTY(Interp, EditAnywhere, BlueprintReadWrite, Category = CameraSettings)
+	bool bUpdateOrthoPlanes;
+	UFUNCTION(BlueprintCallable, Category = Camera)
+	void SetUpdateOrthoPlanes(bool bInUpdateOrthoPlanes) { bUpdateOrthoPlanes = bInUpdateOrthoPlanes; }
+
+	/** If UpdateOrthoPlanes is enabled, this setting will use the cameras current height to compensate the distance to the general view (as a pseudo distance to view target when one isn't present) */
+	UPROPERTY(Interp, EditAnywhere, BlueprintReadWrite, Category = CameraSettings, meta = (EditCondition = "bUpdateOrthoPlanes"))
+	bool bUseCameraHeightAsViewTarget;
+	UFUNCTION(BlueprintCallable, Category = Camera)
+	void SetUseCameraHeightAsViewTarget(bool bInUseCameraHeightAsViewTarget) { bUseCameraHeightAsViewTarget = bInUseCameraHeightAsViewTarget; }
+
 	/** Aspect Ratio (Width/Height) */
 	UPROPERTY(Interp, EditAnywhere, BlueprintReadWrite, Category = CameraSettings, meta = (ClampMin = "0.1", ClampMax = "100.0", EditCondition = "bConstrainAspectRatio"))
 	float AspectRatio;
diff --git a/Engine/Source/Runtime/Engine/Classes/Camera/CameraTypes.h b/Engine/Source/Runtime/Engine/Classes/Camera/CameraTypes.h
--- a/Engine/Source/Runtime/Engine/Classes/Camera/CameraTypes.h
+++ b/Engine/Source/Runtime/Engine/Classes/Camera/CameraTypes.h
@@ -61,6 +61,18 @@ struct FMinimalViewInfo
 	UPROPERTY(EditAnywhere, BlueprintReadWrite, Category=Camera)
 	bool bAutoCalculateOrthoPlanes;
 
+	/** Manually adjusts the planes of this camera, maintaining the distance between them. Positive moves out to the farplane, negative towards the near plane */
+	UPROPERTY(EditAnywhere, BlueprintReadWrite, Category = Camera)
+	float AutoPlaneShift;
+
+	/** Adjusts the near/far planes and the view origin of the current camera automatically to avoid clipping and light artefacting*/
+	UPROPERTY(EditAnywhere, BlueprintReadWrite, Category = Camera)
+	bool bUpdateOrthoPlanes;
+
+	/** If UpdateOrthoPlanes is enabled, this setting will use the cameras current height to compensate the distance to the general view (as a pseudo distance to view target when one isn't present) */
+	UPROPERTY(EditAnywhere, BlueprintReadWrite, Category = Camera)
+	bool bUseCameraHeightAsViewTarget;
+
 	/** The near plane distance of the orthographic view (in world units) */
 	UPROPERTY(Interp, EditAnywhere, BlueprintReadWrite, Category=Camera)
 	float OrthoNearClipPlane;
diff --git a/Engine/Source/Runtime/Engine/Classes/Camera/PlayerCameraManager.h b/Engine/Source/Runtime/Engine/Classes/Camera/PlayerCameraManager.h
--- a/Engine/Source/Runtime/Engine/Classes/Camera/PlayerCameraManager.h
+++ b/Engine/Source/Runtime/Engine/Classes/Camera/PlayerCameraManager.h
@@ -409,9 +409,21 @@ public:
 	uint32 bIsOrthographic : 1;
 
 	/** True when this camera should automatically calculated the Near+Far planes */
-	UPROPERTY(EditAnywhere, BlueprintReadWrite, Category = PlayerCameraManager)
+	UPROPERTY(EditAnywhere, BlueprintReadWrite, Category = PlayerCameraManager, meta = (EditCondition = "bIsOrthographic"))
 	uint32 bAutoCalculateOrthoPlanes : 1;
 
+	/** Manually adjusts the planes of this camera, maintaining the distance between them. Positive moves out to the farplane, negative towards the near plane */
+	UPROPERTY(EditAnywhere, BlueprintReadWrite, Category = PlayerCameraManager, meta = (EditCondition = "bIsOrthographic && bAutoCalculateOrthoPlanes"))
+	float AutoPlaneShift;
+
+	/** Adjusts the near/far planes and the view origin of the current camera automatically to avoid clipping and light artefacting*/
+	UPROPERTY(EditAnywhere, BlueprintReadWrite, Category = PlayerCameraManager, meta = (EditCondition = "bIsOrthographic"))
+	uint32 bUpdateOrthoPlanes : 1;
+
+	/** If UpdateOrthoPlanes is enabled, this setting will use the cameras current height to compensate the distance to the general view (as a pseudo distance to view target when one isn't present) */
+	UPROPERTY(EditAnywhere, BlueprintReadWrite, Category = PlayerCameraManager, meta = (EditCondition = "bIsOrthographic && bUpdateOrthoPlanes"))
+	uint32 bUseCameraHeightAsViewTarget : 1;
+
 	/** True if black bars should be added if the destination view has a different aspect ratio (only used when a view target doesn't specify whether or not to constrain the aspect ratio; most of the time the value from a camera component is used instead) */
 	UPROPERTY(EditAnywhere, BlueprintReadWrite, Category=PlayerCameraManager)
 	uint32 bDefaultConstrainAspectRatio : 1;
diff --git a/Engine/Source/Runtime/Engine/Classes/Components/SceneCaptureComponent2D.h b/Engine/Source/Runtime/Engine/Classes/Components/SceneCaptureComponent2D.h
--- a/Engine/Source/Runtime/Engine/Classes/Components/SceneCaptureComponent2D.h
+++ b/Engine/Source/Runtime/Engine/Classes/Components/SceneCaptureComponent2D.h
@@ -27,21 +27,29 @@ public:
 	TEnumAsByte<ECameraProjectionMode::Type> ProjectionType;
 
 	/** Camera field of view (in degrees). */
-	UPROPERTY(interp, EditAnywhere, BlueprintReadWrite, Category=Projection, meta=(DisplayName = "Field of View", UIMin = "5.0", UIMax = "170", ClampMin = "0.001", ClampMax = "360.0"))
+	UPROPERTY(interp, EditAnywhere, BlueprintReadWrite, Category=Projection, meta=(DisplayName = "Field of View", UIMin = "5.0", UIMax = "170", ClampMin = "0.001", ClampMax = "360.0", editcondition = "ProjectionType==0"))
 	float FOVAngle;
 
 	/** The desired width (in world units) of the orthographic view (ignored in Perspective mode) */
-	UPROPERTY(EditAnywhere, BlueprintReadWrite, Category=Projection)
+	UPROPERTY(EditAnywhere, BlueprintReadWrite, Category=Projection, meta = (editcondition = "ProjectionType==1"))
 	float OrthoWidth;
 
 	/** Automatically determine a min/max Near/Far clip plane position depending on OrthoWidth value */
-	UPROPERTY(EditAnywhere, BlueprintReadWrite, Category=Projection)
+	UPROPERTY(EditAnywhere, BlueprintReadWrite, Category=Projection, meta = (editcondition = "ProjectionType==1"))
 	bool bAutoCalculateOrthoPlanes;
 
-	/** Adjusts the view origin to the near plane location, and forces near plane to 0. Helps with artefacting when auto-adjusting the plane range. */
-	UPROPERTY(EditAnywhere, BlueprintReadWrite, Category=Projection)
+	/** Manually adjusts the planes of this camera, maintaining the distance between them. Positive moves out to the farplane, negative towards the near plane */
+	UPROPERTY(EditAnywhere, BlueprintReadWrite, Category=Projection, meta = (editcondition = "ProjectionType==1 && bAutoCalculateOrthoPlanes"))
+	float AutoPlaneShift;
+
+	/** Adjusts the near/far planes and the view origin of the current camera automatically to avoid clipping and light artefacting*/
+	UPROPERTY(EditAnywhere, BlueprintReadWrite, Category=Projection, meta = (editcondition = "ProjectionType==1"))
 	bool bUpdateOrthoPlanes;
 
+	/** If UpdateOrthoPlanes is enabled, this setting will use the cameras current height to compensate the distance to the general view (as a pseudo distance to view target when one isn't present) */
+	UPROPERTY(EditAnywhere, BlueprintReadWrite, Category=Projection, meta = (editcondition = "ProjectionType==1 && bUpdateOrthoPlanes"))
+	bool bUseCameraHeightAsViewTarget;
+
 	/** Output render target of the scene capture that can be read in materials. */
 	UPROPERTY(EditAnywhere, BlueprintReadWrite, Category=SceneCapture)
 	TObjectPtr<class UTextureRenderTarget2D> TextureTarget;
diff --git a/Engine/Source/Runtime/Engine/Private/Camera/CameraComponent.cpp b/Engine/Source/Runtime/Engine/Private/Camera/CameraComponent.cpp
--- a/Engine/Source/Runtime/Engine/Private/Camera/CameraComponent.cpp
+++ b/Engine/Source/Runtime/Engine/Private/Camera/CameraComponent.cpp
@@ -48,6 +48,9 @@ UCameraComponent::UCameraComponent(const FObjectInitializer& ObjectInitializer)
 	AspectRatio = 1.777778f;
 	OrthoWidth = DEFAULT_ORTHOWIDTH;
 	bAutoCalculateOrthoPlanes = true;
+	AutoPlaneShift = 0.0f;
+	bUpdateOrthoPlanes = true;
+	bUseCameraHeightAsViewTarget = true;
 	OrthoNearClipPlane = DEFAULT_ORTHONEARPLANE;
 	OrthoFarClipPlane = DEFAULT_ORTHOFARPLANE;
 	bConstrainAspectRatio = false;
@@ -420,6 +423,10 @@ void UCameraComponent::GetCameraView(float DeltaTime, FMinimalViewInfo& DesiredV
 	DesiredView.OrthoNearClipPlane = OrthoNearClipPlane;
 	DesiredView.OrthoFarClipPlane = OrthoFarClipPlane;
 	DesiredView.bAutoCalculateOrthoPlanes = bAutoCalculateOrthoPlanes;
+	DesiredView.AutoPlaneShift = AutoPlaneShift;
+	DesiredView.bUpdateOrthoPlanes = bUpdateOrthoPlanes;
+	DesiredView.bUseCameraHeightAsViewTarget = bUseCameraHeightAsViewTarget;
+	
 	if (bAutoCalculateOrthoPlanes)
 	{
 		if (const AActor* ViewTarget = GetOwner())
diff --git a/Engine/Source/Runtime/Engine/Private/Camera/CameraStackTypes.cpp b/Engine/Source/Runtime/Engine/Private/Camera/CameraStackTypes.cpp
--- a/Engine/Source/Runtime/Engine/Private/Camera/CameraStackTypes.cpp
+++ b/Engine/Source/Runtime/Engine/Private/Camera/CameraStackTypes.cpp
@@ -233,7 +233,7 @@ void FMinimalViewInfo::CalculateProjectionMatrixGivenViewRectangle(FMinimalViewI
 	bool bOrthographic = ViewInfo.ProjectionMode == ECameraProjectionMode::Orthographic;
 	if(bOrthographic)
 	{
-		ViewInfo.AutoCalculateOrthoPlanes(InOutProjectionData);	
+		ViewInfo.AutoCalculateOrthoPlanes(InOutProjectionData);
 	}
 
 	// Create the projection matrix (and possibly constrain the view rectangle)
@@ -302,7 +302,7 @@ void FMinimalViewInfo::CalculateProjectionMatrixGivenViewRectangle(FMinimalViewI
 			float FarPlane = ViewInfo.OrthoFarClipPlane;
 			float NearPlane = ViewInfo.OrthoNearClipPlane;
 
-			InOutProjectionData.UpdateOrthoPlanes(NearPlane, FarPlane, OrthoWidth);
+			InOutProjectionData.UpdateOrthoPlanes(NearPlane, FarPlane, OrthoWidth, ViewInfo.bUseCameraHeightAsViewTarget);
 
 			const float ZScale = 1.0f / (FarPlane - NearPlane);
 			const float ZOffset = -NearPlane;
@@ -425,9 +425,9 @@ bool FMinimalViewInfo::AutoCalculateOrthoPlanes(FSceneViewProjectionData& InOutP
 		FarPlane = FMath::Clamp(FarPlane, OrthoHeight, MaxFPValue + NearPlane);
 
 		//The Planes can be scaled in the Z axis without restriction to ensure a user can capture their entire view.
-		const float AutoPlaneShift = CVarOrthoAutoPlaneShift.GetValueOnAnyThread();
-		OrthoNearClipPlane = NearPlane + AutoPlaneShift;
-		OrthoFarClipPlane = FarPlane + AutoPlaneShift;
+		const float GlobalAutoPlaneShift = CVarOrthoAutoPlaneShift.GetValueOnAnyThread();
+		OrthoNearClipPlane = NearPlane + AutoPlaneShift + GlobalAutoPlaneShift;
+		OrthoFarClipPlane = FarPlane + AutoPlaneShift + GlobalAutoPlaneShift;
 		InOutProjectionData.CameraToViewTarget = CameraToViewTarget;
 		return true;
 	}
diff --git a/Engine/Source/Runtime/Engine/Private/Components/SceneCaptureComponent.cpp b/Engine/Source/Runtime/Engine/Private/Components/SceneCaptureComponent.cpp
--- a/Engine/Source/Runtime/Engine/Private/Components/SceneCaptureComponent.cpp
+++ b/Engine/Source/Runtime/Engine/Private/Components/SceneCaptureComponent.cpp
@@ -569,7 +569,9 @@ USceneCaptureComponent2D::USceneCaptureComponent2D(const FObjectInitializer& Obj
 
 	OrthoWidth = DEFAULT_ORTHOWIDTH;
 	bAutoCalculateOrthoPlanes = true;
+	AutoPlaneShift = 0.0f;
 	bUpdateOrthoPlanes = false;
+	bUseCameraHeightAsViewTarget = false;
 
 	bUseCustomProjectionMatrix = false;
 	bAutoActivate = true;
@@ -702,7 +704,11 @@ void USceneCaptureComponent2D::GetCameraView(float DeltaTime, FMinimalViewInfo&
 	OutMinimalViewInfo.bConstrainAspectRatio = false;
 	OutMinimalViewInfo.ProjectionMode = ProjectionType;
 	OutMinimalViewInfo.OrthoWidth = OrthoWidth;
-	OutMinimalViewInfo.bAutoCalculateOrthoPlanes = bAutoCalculateOrthoPlanes;		
+	OutMinimalViewInfo.bAutoCalculateOrthoPlanes = bAutoCalculateOrthoPlanes;
+	OutMinimalViewInfo.AutoPlaneShift = AutoPlaneShift;
+	OutMinimalViewInfo.bUpdateOrthoPlanes = bUpdateOrthoPlanes;
+	OutMinimalViewInfo.bUseCameraHeightAsViewTarget = bUseCameraHeightAsViewTarget;
+
 	if (bAutoCalculateOrthoPlanes)
 	{
 		if(const AActor* ViewTarget = GetOwner())
diff --git a/Engine/Source/Runtime/Engine/Private/PlayerCameraManager.cpp b/Engine/Source/Runtime/Engine/Private/PlayerCameraManager.cpp
--- a/Engine/Source/Runtime/Engine/Private/PlayerCameraManager.cpp
+++ b/Engine/Source/Runtime/Engine/Private/PlayerCameraManager.cpp
@@ -41,6 +41,9 @@ APlayerCameraManager::APlayerCameraManager(const FObjectInitializer& ObjectIniti
 	bDefaultConstrainAspectRatio = false;
 	DefaultOrthoWidth = 512.0f;
 	bAutoCalculateOrthoPlanes = true;
+	AutoPlaneShift = 0.0f;
+	bUpdateOrthoPlanes = true;
+	bUseCameraHeightAsViewTarget = true;
 	SetHidden(true);
 	bReplicates = false;
 	FreeCamDistance = 256.0f;
@@ -366,6 +369,9 @@ void APlayerCameraManager::UpdateViewTarget(FTViewTarget& OutVT, float DeltaTime
 	OutVT.POV.ProjectionMode = bIsOrthographic ? ECameraProjectionMode::Orthographic : ECameraProjectionMode::Perspective;
 	OutVT.POV.PostProcessBlendWeight = 1.0f;
 	OutVT.POV.bAutoCalculateOrthoPlanes = bAutoCalculateOrthoPlanes;
+	OutVT.POV.AutoPlaneShift = AutoPlaneShift;
+	OutVT.POV.bUpdateOrthoPlanes = bUpdateOrthoPlanes;
+	OutVT.POV.bUseCameraHeightAsViewTarget = bUseCameraHeightAsViewTarget;
 
 	bool bDoNotApplyModifiers = false;
 
diff --git a/Engine/Source/Runtime/Engine/Private/SceneView.cpp b/Engine/Source/Runtime/Engine/Private/SceneView.cpp
--- a/Engine/Source/Runtime/Engine/Private/SceneView.cpp
+++ b/Engine/Source/Runtime/Engine/Private/SceneView.cpp
@@ -570,7 +570,7 @@ FVector4f CreateInvDeviceZToWorldZTransform(const FMatrix& ProjMatrix)
 	}
 }
 
-bool FSceneViewProjectionData::UpdateOrthoPlanes(FSceneViewProjectionData* InOutProjectionData, float& NearPlane, float& FarPlane, float HalfOrthoWidth)
+bool FSceneViewProjectionData::UpdateOrthoPlanes(FSceneViewProjectionData* InOutProjectionData, float& NearPlane, float& FarPlane, float HalfOrthoWidth, bool bUseCameraHeightAsViewTarget)
 {
 	if (!InOutProjectionData || !CVarAllowOrthoNearPlaneCorrection.GetValueOnAnyThread())
 	{
@@ -598,7 +598,7 @@ bool FSceneViewProjectionData::UpdateOrthoPlanes(FSceneViewProjectionData* InOut
 		* Depends on view direction being top down as when we get to a side view, the scale grows much larger and it becomes redundant.
 		*/
 		float CameraHeightAdjustment = 0.0f;
-		if (CVarOrthoCameraHeightAsViewTarget.GetValueOnAnyThread())
+		if (CVarOrthoCameraHeightAsViewTarget.GetValueOnAnyThread() && bUseCameraHeightAsViewTarget)
 		{
 			CameraHeightAdjustment = FMath::Abs(FMath::Min(InOutProjectionData->ViewOrigin.Z, HalfOrthoWidth)) * FMath::Abs((ViewForward.Dot(FVector(0, 0, -1.0f))));
 		}
@@ -612,10 +612,14 @@ bool FSceneViewProjectionData::UpdateOrthoPlanes(FSceneViewProjectionData* InOut
 
 bool FSceneViewProjectionData::UpdateOrthoPlanes(FMinimalViewInfo& MinimalViewInfo)
 {
-	return UpdateOrthoPlanes(MinimalViewInfo.OrthoNearClipPlane, MinimalViewInfo.OrthoFarClipPlane, MinimalViewInfo.OrthoWidth/2.0f);
+	if(MinimalViewInfo.bUpdateOrthoPlanes)
+	{
+		return UpdateOrthoPlanes(MinimalViewInfo.OrthoNearClipPlane, MinimalViewInfo.OrthoFarClipPlane, MinimalViewInfo.OrthoWidth/2.0f, MinimalViewInfo.bUseCameraHeightAsViewTarget);
+	}
+	return false;
 }
 
-bool FSceneViewProjectionData::UpdateOrthoPlanes()
+bool FSceneViewProjectionData::UpdateOrthoPlanes(bool bUseCameraHeightAsViewTarget)
 {
 	/**
 	* This function takes the existing projection matrix and moves the nearplane + view origin.
@@ -643,7 +647,7 @@ bool FSceneViewProjectionData::UpdateOrthoPlanes()
 	{
 		return false;
 	}
-	UpdateOrthoPlanes(NearPlane, FarPlane, 1.0f / HalfInvOrthoWidth);
+	UpdateOrthoPlanes(NearPlane, FarPlane, 1.0f / HalfInvOrthoWidth, bUseCameraHeightAsViewTarget);
 
 	const float ZScale = 1.0f / (FarPlane - NearPlane);
 	const float ZOffset = -NearPlane;
diff --git a/Engine/Source/Runtime/Engine/Public/SceneView.h b/Engine/Source/Runtime/Engine/Public/SceneView.h
--- a/Engine/Source/Runtime/Engine/Public/SceneView.h
+++ b/Engine/Source/Runtime/Engine/Public/SceneView.h
@@ -128,13 +128,13 @@ public:
 	}
 
 	// Function for correcting Ortho camera near plane locations to avoid artifacts behind camera view origin
-	static ENGINE_API bool UpdateOrthoPlanes(FSceneViewProjectionData* InOutProjectionData, float& NearPlane, float& FarPlane, float HalfOrthoWidth);
-	ENGINE_API inline bool UpdateOrthoPlanes(float& NearPlane, float& FarPlane, float HalfOrthoWidth)
+	static ENGINE_API bool UpdateOrthoPlanes(FSceneViewProjectionData* InOutProjectionData, float& NearPlane, float& FarPlane, float HalfOrthoWidth, bool bUseCameraHeightAsViewTarget);
+	ENGINE_API inline bool UpdateOrthoPlanes(float& NearPlane, float& FarPlane, float HalfOrthoWidth, bool bUseCameraHeightAsViewTarget)
 	{
-		return UpdateOrthoPlanes(this, NearPlane, FarPlane, HalfOrthoWidth);
+		return UpdateOrthoPlanes(this, NearPlane, FarPlane, HalfOrthoWidth, bUseCameraHeightAsViewTarget);
 	}
 	ENGINE_API bool UpdateOrthoPlanes(FMinimalViewInfo& MinimalViewInfo);
-	ENGINE_API bool UpdateOrthoPlanes();
+	ENGINE_API bool UpdateOrthoPlanes(bool bUseCameraHeightAsViewTarget);
 };
 
 /** Method used for primary screen percentage method. */
diff --git a/Engine/Source/Runtime/Renderer/Private/SceneCaptureRendering.cpp b/Engine/Source/Runtime/Renderer/Private/SceneCaptureRendering.cpp
--- a/Engine/Source/Runtime/Renderer/Private/SceneCaptureRendering.cpp
+++ b/Engine/Source/Runtime/Renderer/Private/SceneCaptureRendering.cpp
@@ -847,7 +847,7 @@ static FSceneRenderer* CreateSceneRendererForSceneCapture(
 	{
 		if (!SceneCaptureViewInfo.IsPerspectiveProjection() && SceneCaptureComponent2D->bUpdateOrthoPlanes)
 		{
-			SceneCaptureViewInfo.UpdateOrthoPlanes();
+			SceneCaptureViewInfo.UpdateOrthoPlanes(SceneCaptureComponent2D->bUseCameraHeightAsViewTarget);
 		}
 	}
 
