Moving USkeletalBodySetup out of "PhysicsEngine/PhysicsAsset.h" and into its own header "PhysicsEngine/SkeletalBodySetup.h"

diff --git a/Engine/Plugins/ChaosCloth/Source/ChaosCloth/Private/ChaosCloth/ChaosClothingSimulationCollider.cpp b/Engine/Plugins/ChaosCloth/Source/ChaosCloth/Private/ChaosCloth/ChaosClothingSimulationCollider.cpp
--- a/Engine/Plugins/ChaosCloth/Source/ChaosCloth/Private/ChaosCloth/ChaosClothingSimulationCollider.cpp
+++ b/Engine/Plugins/ChaosCloth/Source/ChaosCloth/Private/ChaosCloth/ChaosClothingSimulationCollider.cpp
@@ -17,6 +17,7 @@
 #include "Engine/SkeletalMesh.h"
 #include "PhysicsEngine/PhysicsAsset.h"
 #include "PhysicsEngine/AggregateGeom.h"
+#include "PhysicsEngine/SkeletalBodySetup.h"
 #include "HAL/PlatformMath.h"
 #include "Containers/ArrayView.h"
 #include "Containers/BitArray.h"
diff --git a/Engine/Plugins/ChaosClothAsset/Source/ChaosClothAssetEngine/Private/ChaosClothAsset/ClothAsset.cpp b/Engine/Plugins/ChaosClothAsset/Source/ChaosClothAssetEngine/Private/ChaosClothAsset/ClothAsset.cpp
--- a/Engine/Plugins/ChaosClothAsset/Source/ChaosClothAssetEngine/Private/ChaosClothAsset/ClothAsset.cpp
+++ b/Engine/Plugins/ChaosClothAsset/Source/ChaosClothAssetEngine/Private/ChaosClothAsset/ClothAsset.cpp
@@ -13,6 +13,7 @@
 #include "Animation/AnimationAsset.h"
 #endif
 #include "Engine/RendererSettings.h"
+#include "Engine/SkeletalMesh.h"
 #include "Engine/SkinnedAssetAsyncCompileUtils.h"
 #include "Features/IModularFeatures.h"
 #include "GeometryCollection/ManagedArrayCollection.h"
diff --git a/Engine/Plugins/Experimental/ChaosFlesh/Source/ChaosFleshNodes/Private/Dataflow/ChaosFleshKinematicInitializationNodes.cpp b/Engine/Plugins/Experimental/ChaosFlesh/Source/ChaosFleshNodes/Private/Dataflow/ChaosFleshKinematicInitializationNodes.cpp
--- a/Engine/Plugins/Experimental/ChaosFlesh/Source/ChaosFleshNodes/Private/Dataflow/ChaosFleshKinematicInitializationNodes.cpp
+++ b/Engine/Plugins/Experimental/ChaosFlesh/Source/ChaosFleshNodes/Private/Dataflow/ChaosFleshKinematicInitializationNodes.cpp
@@ -13,6 +13,8 @@
 #include "GeometryCollection/Facades/CollectionVertexBoneWeightsFacade.h"
 #include "GeometryCollection/TransformCollection.h"
 #include "PhysicsEngine/PhysicsAsset.h"
+#include "PhysicsEngine/SkeletalBodySetup.h"
+#include "PhysicsEngine/SphylElem.h"
 #include "Rendering/SkeletalMeshRenderData.h"
 #include "BoneWeights.h"
 
diff --git a/Engine/Plugins/Experimental/ChaosVehiclesPlugin/Source/ChaosVehicles/Private/ChaosVehicleMovementComponent.cpp b/Engine/Plugins/Experimental/ChaosVehiclesPlugin/Source/ChaosVehicles/Private/ChaosVehicleMovementComponent.cpp
--- a/Engine/Plugins/Experimental/ChaosVehiclesPlugin/Source/ChaosVehicles/Private/ChaosVehicleMovementComponent.cpp
+++ b/Engine/Plugins/Experimental/ChaosVehiclesPlugin/Source/ChaosVehicles/Private/ChaosVehicleMovementComponent.cpp
@@ -1,10 +1,12 @@
 // Copyright Epic Games, Inc. All Rights Reserved.
 
 #include "ChaosVehicleMovementComponent.h"
+#include "BodySetupCore.h"
 #include "EngineGlobals.h"
 #include "GameFramework/Pawn.h"
 #include "GameFramework/PlayerController.h"
 #include "Engine/Engine.h"
+#include "Engine/SkinnedAsset.h"
 #include "CanvasItem.h"
 #include "Engine/Canvas.h"
 #include "Components/StaticMeshComponent.h"
diff --git a/Engine/Plugins/Experimental/ChaosVehiclesPlugin/Source/ChaosVehicles/Private/ChaosWheeledVehicleMovementComponent.cpp b/Engine/Plugins/Experimental/ChaosVehiclesPlugin/Source/ChaosVehicles/Private/ChaosWheeledVehicleMovementComponent.cpp
--- a/Engine/Plugins/Experimental/ChaosVehiclesPlugin/Source/ChaosVehicles/Private/ChaosWheeledVehicleMovementComponent.cpp
+++ b/Engine/Plugins/Experimental/ChaosVehiclesPlugin/Source/ChaosVehicles/Private/ChaosWheeledVehicleMovementComponent.cpp
@@ -13,6 +13,7 @@
 #include "GameFramework/Pawn.h"
 #include "PhysicalMaterials/PhysicalMaterial.h"
 #include "PhysicsEngine/PhysicsAsset.h"
+#include "PhysicsEngine/SkeletalBodySetup.h"
 #include "VehicleAnimationInstance.h"
 #include "ChaosVehicleManager.h"
 #include "ChaosVehicleWheel.h"
diff --git a/Engine/Plugins/Experimental/Mutable/Source/CustomizableObject/Private/MuCO/CustomizableObjectInstance.cpp b/Engine/Plugins/Experimental/Mutable/Source/CustomizableObject/Private/MuCO/CustomizableObjectInstance.cpp
--- a/Engine/Plugins/Experimental/Mutable/Source/CustomizableObject/Private/MuCO/CustomizableObjectInstance.cpp
+++ b/Engine/Plugins/Experimental/Mutable/Source/CustomizableObject/Private/MuCO/CustomizableObjectInstance.cpp
@@ -11,6 +11,7 @@
 #include "BoneControllers/AnimNode_RigidBody.h"
 #include "ClothConfig.h"
 #include "ClothingAsset.h"
+#include "Engine/SkeletalMesh.h"
 #include "Engine/SkeletalMeshSocket.h"
 #include "Engine/SkeletalMeshLODSettings.h"
 #include "Engine/Texture2DArray.h"
@@ -44,7 +45,9 @@
 #include "UObject/ObjectSaveContext.h"
 #include "UObject/UObjectIterator.h"
 #include "Widgets/Notifications/SNotificationList.h"
+#include "PhysicsEngine/AggregateGeom.h"
 #include "PhysicsEngine/PhysicsConstraintTemplate.h"
+#include "PhysicsEngine/SkeletalBodySetup.h"
 
 #include UE_INLINE_GENERATED_CPP_BY_NAME(CustomizableObjectInstance)
 
diff --git a/Engine/Plugins/Experimental/Mutable/Source/CustomizableObjectEditor/Private/MuCOE/GenerateMutableSource/GenerateMutableSourceMesh.cpp b/Engine/Plugins/Experimental/Mutable/Source/CustomizableObjectEditor/Private/MuCOE/GenerateMutableSource/GenerateMutableSourceMesh.cpp
--- a/Engine/Plugins/Experimental/Mutable/Source/CustomizableObjectEditor/Private/MuCOE/GenerateMutableSource/GenerateMutableSourceMesh.cpp
+++ b/Engine/Plugins/Experimental/Mutable/Source/CustomizableObjectEditor/Private/MuCOE/GenerateMutableSource/GenerateMutableSourceMesh.cpp
@@ -46,6 +46,7 @@
 #include "MuT/NodeMeshTable.h"
 #include "MuT/NodeMeshVariation.h"
 #include "PhysicsEngine/PhysicsAsset.h"
+#include "PhysicsEngine/SkeletalBodySetup.h"
 #include "Engine/SkinnedAssetCommon.h"
 
 #define LOCTEXT_NAMESPACE "CustomizableObjectEditor"
diff --git a/Engine/Plugins/Experimental/PhysicsControl/Source/PhysicsControl/Private/AnimNode_RigidBodyWithControl.cpp b/Engine/Plugins/Experimental/PhysicsControl/Source/PhysicsControl/Private/AnimNode_RigidBodyWithControl.cpp
--- a/Engine/Plugins/Experimental/PhysicsControl/Source/PhysicsControl/Private/AnimNode_RigidBodyWithControl.cpp
+++ b/Engine/Plugins/Experimental/PhysicsControl/Source/PhysicsControl/Private/AnimNode_RigidBodyWithControl.cpp
@@ -6,6 +6,7 @@
 #include "Animation/AnimInstanceProxy.h"
 #include "ClothCollisionSource.h"
 #include "Engine/OverlapResult.h"
+#include "Engine/SkeletalMesh.h"
 #include "GameFramework/Pawn.h"
 #include "HAL/Event.h"
 #include "HAL/LowLevelMemTracker.h"
@@ -21,6 +22,7 @@
 #include "Chaos/PBDJointConstraints.h"
 #include "PhysicsField/PhysicsFieldComponent.h"
 #include "PhysicsEngine/PhysicsSettings.h"
+#include "PhysicsEngine/SkeletalBodySetup.h"
 #include "Logging/MessageLog.h"
 
 #include "PhysicsProxy/SingleParticlePhysicsProxy.h"
diff --git a/Engine/Plugins/Experimental/PhysicsControl/Source/PhysicsControl/Private/PhysicsControlComponent.cpp b/Engine/Plugins/Experimental/PhysicsControl/Source/PhysicsControl/Private/PhysicsControlComponent.cpp
--- a/Engine/Plugins/Experimental/PhysicsControl/Source/PhysicsControl/Private/PhysicsControlComponent.cpp
+++ b/Engine/Plugins/Experimental/PhysicsControl/Source/PhysicsControl/Private/PhysicsControlComponent.cpp
@@ -9,6 +9,7 @@
 #include "PhysicsEngine/ConstraintInstance.h"
 #include "PhysicsEngine/BodyInstance.h"
 #include "PhysicsEngine/PhysicsAsset.h"
+#include "PhysicsEngine/SkeletalBodySetup.h"
 
 #include "Physics/PhysicsInterfaceCore.h"
 
diff --git a/Engine/Plugins/Experimental/PhysicsControl/Source/PhysicsControl/Private/PhysicsControlOperatorNameGeneration.cpp b/Engine/Plugins/Experimental/PhysicsControl/Source/PhysicsControl/Private/PhysicsControlOperatorNameGeneration.cpp
--- a/Engine/Plugins/Experimental/PhysicsControl/Source/PhysicsControl/Private/PhysicsControlOperatorNameGeneration.cpp
+++ b/Engine/Plugins/Experimental/PhysicsControl/Source/PhysicsControl/Private/PhysicsControlOperatorNameGeneration.cpp
@@ -6,6 +6,7 @@
 #include "PhysicsControlLimbData.h"
 #include "PhysicsControlRecord.h"
 #include "PhysicsEngine/PhysicsAsset.h"
+#include "PhysicsEngine/SkeletalBodySetup.h"
 
 #include "ReferenceSkeleton.h"
 
diff --git a/Engine/Plugins/Experimental/PhysicsControl/Source/PhysicsControl/Private/RigidBodyWithControl.cpp b/Engine/Plugins/Experimental/PhysicsControl/Source/PhysicsControl/Private/RigidBodyWithControl.cpp
--- a/Engine/Plugins/Experimental/PhysicsControl/Source/PhysicsControl/Private/RigidBodyWithControl.cpp
+++ b/Engine/Plugins/Experimental/PhysicsControl/Source/PhysicsControl/Private/RigidBodyWithControl.cpp
@@ -12,6 +12,7 @@
 #include "Chaos/PBDJointConstraints.h"
 #include "Chaos/PBDJointConstraintTypes.h"
 #include "Components/StaticMeshComponent.h"
+#include "Engine/SkeletalMesh.h"
 #include "Physics/ImmediatePhysics/ImmediatePhysicsActorHandle.h"
 #include "Physics/ImmediatePhysics/ImmediatePhysicsJointHandle.h"
 #include "Physics/ImmediatePhysics/ImmediatePhysicsSimulation.h"
diff --git a/Engine/Plugins/FX/Niagara/Source/Niagara/Private/DataInterface/NiagaraDataInterfacePhysicsAsset.cpp b/Engine/Plugins/FX/Niagara/Source/Niagara/Private/DataInterface/NiagaraDataInterfacePhysicsAsset.cpp
--- a/Engine/Plugins/FX/Niagara/Source/Niagara/Private/DataInterface/NiagaraDataInterfacePhysicsAsset.cpp
+++ b/Engine/Plugins/FX/Niagara/Source/Niagara/Private/DataInterface/NiagaraDataInterfacePhysicsAsset.cpp
@@ -11,7 +11,9 @@
 #include "NiagaraShaderParametersBuilder.h"
 #include "NiagaraSystem.h"
 #include "NiagaraSystemInstance.h"
+#include "PhysicsEngine/BodySetup.h"
 #include "PhysicsEngine/PhysicsAsset.h"
+#include "PhysicsEngine/SkeletalBodySetup.h"
 #include "RenderGraphBuilder.h"
 
 #include UE_INLINE_GENERATED_CPP_BY_NAME(NiagaraDataInterfacePhysicsAsset)
diff --git a/Engine/Plugins/FX/Niagara/Source/Niagara/Private/NiagaraDataInterfaceRigidMeshCollisionQuery.cpp b/Engine/Plugins/FX/Niagara/Source/Niagara/Private/NiagaraDataInterfaceRigidMeshCollisionQuery.cpp
--- a/Engine/Plugins/FX/Niagara/Source/Niagara/Private/NiagaraDataInterfaceRigidMeshCollisionQuery.cpp
+++ b/Engine/Plugins/FX/Niagara/Source/Niagara/Private/NiagaraDataInterfaceRigidMeshCollisionQuery.cpp
@@ -25,7 +25,9 @@
 #include "NiagaraSystemImpl.h"
 #include "NiagaraSystemInstance.h"
 #include "NiagaraSystemSimulation.h"
+#include "PhysicsEngine/BodySetup.h"
 #include "PhysicsEngine/PhysicsAsset.h"
+#include "PhysicsEngine/SkeletalBodySetup.h"
 #include "PrimitiveSceneInfo.h"
 #include "SceneInterface.h"
 #include "SceneView.h"
diff --git a/Engine/Plugins/Runtime/HairStrands/Source/HairStrandsCore/Private/GroomDeformerBuilder.cpp b/Engine/Plugins/Runtime/HairStrands/Source/HairStrandsCore/Private/GroomDeformerBuilder.cpp
--- a/Engine/Plugins/Runtime/HairStrands/Source/HairStrandsCore/Private/GroomDeformerBuilder.cpp
+++ b/Engine/Plugins/Runtime/HairStrands/Source/HairStrandsCore/Private/GroomDeformerBuilder.cpp
@@ -11,6 +11,7 @@
 #include "Rendering/SkeletalMeshModel.h"
 #include "PhysicsEngine/PhysicsAsset.h"
 #include "PhysicsEngine/PhysicsConstraintTemplate.h"
+#include "PhysicsEngine/SkeletalBodySetup.h"
 #include "GroomAsset.h"
 #include "UObject/Package.h"
 
diff --git a/Engine/Plugins/Runtime/HairStrands/Source/HairStrandsCore/Public/GroomDeformerBuilder.h b/Engine/Plugins/Runtime/HairStrands/Source/HairStrandsCore/Public/GroomDeformerBuilder.h
--- a/Engine/Plugins/Runtime/HairStrands/Source/HairStrandsCore/Public/GroomDeformerBuilder.h
+++ b/Engine/Plugins/Runtime/HairStrands/Source/HairStrandsCore/Public/GroomDeformerBuilder.h
@@ -1,14 +1,14 @@
-﻿// Copyright Epic Games, Inc. All Rights Reserved.
+// Copyright Epic Games, Inc. All Rights Reserved.
 
 #pragma once
 
-#include "CoreMinimal.h"
+#include "HAL/Platform.h"
 
 class UGroomAsset;
 class USkeletalMesh;
 
-struct HAIRSTRANDSCORE_API FGroomDeformerBuilder
+struct FGroomDeformerBuilder
 {
 	// Create skeletal mesh bones+physics asset from groom asset guides
-	static USkeletalMesh* CreateSkeletalMesh(UGroomAsset* GroomAsset);
+	static HAIRSTRANDSCORE_API USkeletalMesh* CreateSkeletalMesh(UGroomAsset* GroomAsset);
 };
diff --git a/Engine/Source/Developer/AudioFormatRad/AudioFormatRad.Build.cs b/Engine/Source/Developer/AudioFormatRad/AudioFormatRad.Build.cs
--- a/Engine/Source/Developer/AudioFormatRad/AudioFormatRad.Build.cs
+++ b/Engine/Source/Developer/AudioFormatRad/AudioFormatRad.Build.cs
@@ -69,26 +69,28 @@ public class AudioFormatRad : ModuleRules
 
         bool bIsWindowsLib = false;
         string[] ArchitectureStrings = GetRadAudioPlatformString(Target, out bIsWindowsLib);
-
-        string BasePath = Path.Combine(ModuleDirectory, "..", "..", "Runtime", "RadAudioCodec", "SDK", "Lib");
-        foreach (string Architecture in ArchitectureStrings)
-        {
-            // radaudio_ is the codec.
-            // rada_ is the container.
-            if (bIsWindowsLib)
-            {
-                PublicAdditionalLibraries.Add(Path.Combine(BasePath, "radaudio_decoder_" + Architecture + ".lib"));
-                PublicAdditionalLibraries.Add(Path.Combine(BasePath, "rada_decode_" + Architecture + ".lib"));
-                PublicAdditionalLibraries.Add(Path.Combine(BasePath, "radaudio_encoder_" + Architecture + ".lib"));
-                PublicAdditionalLibraries.Add(Path.Combine(BasePath, "rada_encode_" + Architecture + ".lib"));
-            }
-            else
-            {
-                PublicAdditionalLibraries.Add(Path.Combine(BasePath, "libradaudio_decoder_" + Architecture + ".a"));
-                PublicAdditionalLibraries.Add(Path.Combine(BasePath, "librada_decode_" + Architecture + ".a"));
-                PublicAdditionalLibraries.Add(Path.Combine(BasePath, "libradaudio_encoder_" + Architecture + ".a"));
-                PublicAdditionalLibraries.Add(Path.Combine(BasePath, "librada_encode_" + Architecture + ".a"));
-            }
-        }
+		if (ArchitectureStrings is not null)
+		{
+			string BasePath = Path.Combine(ModuleDirectory, "..", "..", "Runtime", "RadAudioCodec", "SDK", "Lib");
+			foreach (string Architecture in ArchitectureStrings)
+			{
+				// radaudio_ is the codec.
+				// rada_ is the container.
+				if (bIsWindowsLib)
+				{
+					PublicAdditionalLibraries.Add(Path.Combine(BasePath, "radaudio_decoder_" + Architecture + ".lib"));
+					PublicAdditionalLibraries.Add(Path.Combine(BasePath, "rada_decode_" + Architecture + ".lib"));
+					PublicAdditionalLibraries.Add(Path.Combine(BasePath, "radaudio_encoder_" + Architecture + ".lib"));
+					PublicAdditionalLibraries.Add(Path.Combine(BasePath, "rada_encode_" + Architecture + ".lib"));
+				}
+				else
+				{
+					PublicAdditionalLibraries.Add(Path.Combine(BasePath, "libradaudio_decoder_" + Architecture + ".a"));
+					PublicAdditionalLibraries.Add(Path.Combine(BasePath, "librada_decode_" + Architecture + ".a"));
+					PublicAdditionalLibraries.Add(Path.Combine(BasePath, "libradaudio_encoder_" + Architecture + ".a"));
+					PublicAdditionalLibraries.Add(Path.Combine(BasePath, "librada_encode_" + Architecture + ".a"));
+				}
+			}
+		}
 	}
 }
diff --git a/Engine/Source/Developer/PhysicsUtilities/Private/PhysicsAssetUtils.cpp b/Engine/Source/Developer/PhysicsUtilities/Private/PhysicsAssetUtils.cpp
--- a/Engine/Source/Developer/PhysicsUtilities/Private/PhysicsAssetUtils.cpp
+++ b/Engine/Source/Developer/PhysicsUtilities/Private/PhysicsAssetUtils.cpp
@@ -11,6 +11,7 @@
 #include "PhysicsEngine/BoxElem.h"
 #include "PhysicsEngine/SphereElem.h"
 #include "PhysicsEngine/SphylElem.h"
+#include "PhysicsEngine/SkeletalBodySetup.h"
 #include "PhysicsEngine/SkinnedLevelSetElem.h"
 #include "Animation/SkeletalMeshActor.h"
 #include "Components/SkeletalMeshComponent.h"
diff --git a/Engine/Source/Editor/DetailCustomizations/Private/BodySetupDetails.cpp b/Engine/Source/Editor/DetailCustomizations/Private/BodySetupDetails.cpp
--- a/Engine/Source/Editor/DetailCustomizations/Private/BodySetupDetails.cpp
+++ b/Engine/Source/Editor/DetailCustomizations/Private/BodySetupDetails.cpp
@@ -18,6 +18,7 @@
 #include "ObjectEditorUtils.h"
 #include "PhysicsEngine/BodySetup.h"
 #include "PhysicsEngine/PhysicsAsset.h"
+#include "PhysicsEngine/SkeletalBodySetup.h"
 #include "PropertyEditorModule.h"
 #include "PropertyHandle.h"
 #include "SlotBase.h"
diff --git a/Engine/Source/Editor/Persona/Private/PhysicsAssetRenderUtils.cpp b/Engine/Source/Editor/Persona/Private/PhysicsAssetRenderUtils.cpp
--- a/Engine/Source/Editor/Persona/Private/PhysicsAssetRenderUtils.cpp
+++ b/Engine/Source/Editor/Persona/Private/PhysicsAssetRenderUtils.cpp
@@ -11,6 +11,7 @@
 #include "Preferences/PhysicsAssetEditorOptions.h"
 #include "PhysicsEngine/PhysicsAsset.h"
 #include "PhysicsEngine/PhysicsConstraintTemplate.h"
+#include "PhysicsEngine/SkeletalBodySetup.h"
 #include "Chaos/Core.h"
 #include "SkeletalMeshTypes.h"
 #include "Styling/AppStyle.h"
diff --git a/Engine/Source/Editor/PhysicsAssetEditor/Private/PhysicsAssetDetailsCustomization.cpp b/Engine/Source/Editor/PhysicsAssetEditor/Private/PhysicsAssetDetailsCustomization.cpp
--- a/Engine/Source/Editor/PhysicsAssetEditor/Private/PhysicsAssetDetailsCustomization.cpp
+++ b/Engine/Source/Editor/PhysicsAssetEditor/Private/PhysicsAssetDetailsCustomization.cpp
@@ -11,6 +11,7 @@
 #include "PhysicsAssetEditorActions.h"
 #include "PhysicsAssetEditor.h"
 #include "PhysicsEngine/PhysicsConstraintTemplate.h"
+#include "PhysicsEngine/SkeletalBodySetup.h"
 #include "PhysicsAssetEditorSkeletalMeshComponent.h"
 #include "Widgets/Text/SInlineEditableTextBlock.h"
 #include "EditorFontGlyphs.h"
diff --git a/Engine/Source/Editor/PhysicsAssetEditor/Private/PhysicsAssetEditor.cpp b/Engine/Source/Editor/PhysicsAssetEditor/Private/PhysicsAssetEditor.cpp
--- a/Engine/Source/Editor/PhysicsAssetEditor/Private/PhysicsAssetEditor.cpp
+++ b/Engine/Source/Editor/PhysicsAssetEditor/Private/PhysicsAssetEditor.cpp
@@ -54,6 +54,7 @@
 #include "PhysicsEngine/PhysicsConstraintTemplate.h"
 #include "PhysicsEngine/ConstraintUtils.h"
 #include "PhysicsEngine/PhysicsAsset.h"
+#include "PhysicsEngine/SkeletalBodySetup.h"
 #include "Engine/Selection.h"
 #include "PersonaModule.h"
 #include "PersonaToolMenuContext.h"
diff --git a/Engine/Source/Editor/PhysicsAssetEditor/Private/PhysicsAssetEditorEditMode.cpp b/Engine/Source/Editor/PhysicsAssetEditor/Private/PhysicsAssetEditorEditMode.cpp
--- a/Engine/Source/Editor/PhysicsAssetEditor/Private/PhysicsAssetEditorEditMode.cpp
+++ b/Engine/Source/Editor/PhysicsAssetEditor/Private/PhysicsAssetEditorEditMode.cpp
@@ -10,6 +10,7 @@
 #include "PhysicsEngine/PhysicsAsset.h"
 #include "PhysicsEngine/PhysicsConstraintTemplate.h"
 #include "PhysicsEngine/PhysicsSettings.h"
+#include "PhysicsEngine/SkeletalBodySetup.h"
 #include "Preferences/PhysicsAssetEditorOptions.h"
 #include "IPersonaPreviewScene.h"
 #include "PhysicsAssetEditor.h"
diff --git a/Engine/Source/Editor/PhysicsAssetEditor/Private/PhysicsAssetEditorMode.cpp b/Engine/Source/Editor/PhysicsAssetEditor/Private/PhysicsAssetEditorMode.cpp
--- a/Engine/Source/Editor/PhysicsAssetEditor/Private/PhysicsAssetEditorMode.cpp
+++ b/Engine/Source/Editor/PhysicsAssetEditor/Private/PhysicsAssetEditorMode.cpp
@@ -1,6 +1,7 @@
 // Copyright Epic Games, Inc. All Rights Reserved.
 
 #include "PhysicsAssetEditorMode.h"
+#include "Engine/SkeletalMesh.h"
 #include "PhysicsAssetEditor.h"
 #include "ISkeletonTree.h"
 #include "IPersonaPreviewScene.h"
@@ -9,6 +10,7 @@
 #include "PhysicsAssetGraph/PhysicsAssetGraphSummoner.h"
 #include "PhysicsEngine/PhysicsAsset.h"
 #include "PhysicsEngine/PhysicsConstraintTemplate.h"
+#include "PhysicsEngine/SkeletalBodySetup.h"
 #include "PhysicsAssetEditorActions.h"
 #include "SEditorViewportToolBarMenu.h"
 #include "PhysicsAssetEditorProfilesSummoner.h"
diff --git a/Engine/Source/Editor/PhysicsAssetEditor/Private/PhysicsAssetEditorSharedData.cpp b/Engine/Source/Editor/PhysicsAssetEditor/Private/PhysicsAssetEditorSharedData.cpp
--- a/Engine/Source/Editor/PhysicsAssetEditor/Private/PhysicsAssetEditorSharedData.cpp
+++ b/Engine/Source/Editor/PhysicsAssetEditor/Private/PhysicsAssetEditorSharedData.cpp
@@ -26,6 +26,7 @@
 #include "PhysicsEngine/PhysicsConstraintTemplate.h"
 #include "PhysicsEngine/PhysicalAnimationComponent.h"
 #include "PhysicsEngine/PhysicsAsset.h"
+#include "PhysicsEngine/SkeletalBodySetup.h"
 #include "PhysicsAssetEditorAnimInstance.h"
 #include "IPersonaPreviewScene.h"
 #include "PhysicsPublic.h"
diff --git a/Engine/Source/Editor/PhysicsAssetEditor/Private/PhysicsAssetEditorSkeletalMeshComponent.cpp b/Engine/Source/Editor/PhysicsAssetEditor/Private/PhysicsAssetEditorSkeletalMeshComponent.cpp
--- a/Engine/Source/Editor/PhysicsAssetEditor/Private/PhysicsAssetEditorSkeletalMeshComponent.cpp
+++ b/Engine/Source/Editor/PhysicsAssetEditor/Private/PhysicsAssetEditorSkeletalMeshComponent.cpp
@@ -15,6 +15,7 @@
 #include "PhysicsAssetRenderUtils.h"
 #include "PhysicsEngine/PhysicsConstraintTemplate.h"
 #include "PhysicsEngine/PhysicsAsset.h"
+#include "PhysicsEngine/SkeletalBodySetup.h"
 #include "Chaos/Core.h"
 #include "SkeletalMeshTypes.h"
 #include "AnimPreviewInstance.h"
diff --git a/Engine/Source/Editor/PhysicsAssetEditor/Private/PhysicsAssetEditorSkeletonTreeBuilder.cpp b/Engine/Source/Editor/PhysicsAssetEditor/Private/PhysicsAssetEditorSkeletonTreeBuilder.cpp
--- a/Engine/Source/Editor/PhysicsAssetEditor/Private/PhysicsAssetEditorSkeletonTreeBuilder.cpp
+++ b/Engine/Source/Editor/PhysicsAssetEditor/Private/PhysicsAssetEditorSkeletonTreeBuilder.cpp
@@ -5,6 +5,7 @@
 #include "SkeletonTreePhysicsShapeItem.h"
 #include "SkeletonTreePhysicsConstraintItem.h"
 #include "PhysicsEngine/PhysicsAsset.h"
+#include "PhysicsEngine/SkeletalBodySetup.h"
 #include "IPersonaPreviewScene.h"
 #include "Animation/DebugSkelMeshComponent.h"
 #include "Misc/TextFilterExpressionEvaluator.h"
diff --git a/Engine/Source/Editor/PhysicsAssetEditor/Private/PhysicsAssetGraph/PhysicsAssetGraph.cpp b/Engine/Source/Editor/PhysicsAssetEditor/Private/PhysicsAssetGraph/PhysicsAssetGraph.cpp
--- a/Engine/Source/Editor/PhysicsAssetEditor/Private/PhysicsAssetGraph/PhysicsAssetGraph.cpp
+++ b/Engine/Source/Editor/PhysicsAssetEditor/Private/PhysicsAssetGraph/PhysicsAssetGraph.cpp
@@ -6,6 +6,7 @@
 #include "PhysicsEngine/PhysicsAsset.h"
 #include "IEditableSkeleton.h"
 #include "PhysicsEngine/PhysicsConstraintTemplate.h"
+#include "PhysicsEngine/SkeletalBodySetup.h"
 #include "PhysicsAssetGraph/PhysicsAssetGraphNode_Bone.h"
 #include "PhysicsAssetGraph/PhysicsAssetGraphNode_Constraint.h"
 #include "Algo/Transform.h"
diff --git a/Engine/Source/Editor/PhysicsAssetEditor/Private/PhysicsAssetGraph/PhysicsAssetGraphNode_Bone.cpp b/Engine/Source/Editor/PhysicsAssetEditor/Private/PhysicsAssetGraph/PhysicsAssetGraphNode_Bone.cpp
--- a/Engine/Source/Editor/PhysicsAssetEditor/Private/PhysicsAssetGraph/PhysicsAssetGraphNode_Bone.cpp
+++ b/Engine/Source/Editor/PhysicsAssetEditor/Private/PhysicsAssetGraph/PhysicsAssetGraphNode_Bone.cpp
@@ -2,7 +2,9 @@
 
 #include "PhysicsAssetGraph/PhysicsAssetGraphNode_Bone.h"
 #include "PhysicsAssetGraph/SBoneNode.h"
+#include "PhysicsEngine/AggregateGeom.h"
 #include "PhysicsEngine/PhysicsAsset.h"
+#include "PhysicsEngine/SkeletalBodySetup.h"
 
 #define LOCTEXT_NAMESPACE "UPhysicsAssetGraphNode_Bone"
 
diff --git a/Engine/Source/Editor/PhysicsAssetEditor/Private/PhysicsAssetGraph/SPhysicsAssetGraph.cpp b/Engine/Source/Editor/PhysicsAssetEditor/Private/PhysicsAssetGraph/SPhysicsAssetGraph.cpp
--- a/Engine/Source/Editor/PhysicsAssetEditor/Private/PhysicsAssetGraph/SPhysicsAssetGraph.cpp
+++ b/Engine/Source/Editor/PhysicsAssetEditor/Private/PhysicsAssetGraph/SPhysicsAssetGraph.cpp
@@ -14,6 +14,7 @@
 #include "PhysicsAssetGraph/PhysicsAssetGraphNode_Constraint.h"
 #include "PhysicsEngine/PhysicsAsset.h"
 #include "PhysicsEngine/PhysicsConstraintTemplate.h"
+#include "PhysicsEngine/SkeletalBodySetup.h"
 #include "ISkeletonTree.h"
 #include "ISkeletonTreeItem.h"
 #include "PhysicsAssetEditor.h"
diff --git a/Engine/Source/Editor/PhysicsAssetEditor/Private/SkeletonTreePhysicsBodyItem.cpp b/Engine/Source/Editor/PhysicsAssetEditor/Private/SkeletonTreePhysicsBodyItem.cpp
--- a/Engine/Source/Editor/PhysicsAssetEditor/Private/SkeletonTreePhysicsBodyItem.cpp
+++ b/Engine/Source/Editor/PhysicsAssetEditor/Private/SkeletonTreePhysicsBodyItem.cpp
@@ -4,6 +4,7 @@
 #include "Styling/AppStyle.h"
 #include "PhysicsAssetRenderUtils.h"
 #include "PhysicsEngine/PhysicsAsset.h"
+#include "PhysicsEngine/SkeletalBodySetup.h"
 
 #define LOCTEXT_NAMESPACE "FSkeletonTreePhysicsBodyItem"
 
diff --git a/Engine/Source/Editor/PhysicsAssetEditor/Private/SkeletonTreePhysicsShapeItem.cpp b/Engine/Source/Editor/PhysicsAssetEditor/Private/SkeletonTreePhysicsShapeItem.cpp
--- a/Engine/Source/Editor/PhysicsAssetEditor/Private/SkeletonTreePhysicsShapeItem.cpp
+++ b/Engine/Source/Editor/PhysicsAssetEditor/Private/SkeletonTreePhysicsShapeItem.cpp
@@ -6,6 +6,7 @@
 #include "Styling/AppStyle.h"
 #include "Widgets/Text/SInlineEditableTextBlock.h"
 #include "ScopedTransaction.h"
+#include "PhysicsEngine/SkeletalBodySetup.h"
 
 #define LOCTEXT_NAMESPACE "FSkeletonTreePhysicsShapeItem"
 
@@ -79,6 +80,22 @@ TSharedRef< SWidget > FSkeletonTreePhysicsShapeItem::GenerateWidgetForDataColumn
 	return SNullWidget::NullWidget;
 }
 
+FName FSkeletonTreePhysicsShapeItem::GetRowItemName() const
+{
+	FString NameAsString = GetNameAsString();
+	return *NameAsString;
+}
+
+UObject* FSkeletonTreePhysicsShapeItem::GetObject() const
+{
+	return BodySetup;
+}
+
+bool FSkeletonTreePhysicsShapeItem::CanRenameItem() const
+{
+	return true;
+}
+
 void FSkeletonTreePhysicsShapeItem::OnItemDoubleClicked()
 {
 	OnRenameRequested.ExecuteIfBound();
diff --git a/Engine/Source/Editor/PhysicsAssetEditor/Private/SkeletonTreePhysicsShapeItem.h b/Engine/Source/Editor/PhysicsAssetEditor/Private/SkeletonTreePhysicsShapeItem.h
--- a/Engine/Source/Editor/PhysicsAssetEditor/Private/SkeletonTreePhysicsShapeItem.h
+++ b/Engine/Source/Editor/PhysicsAssetEditor/Private/SkeletonTreePhysicsShapeItem.h
@@ -9,6 +9,7 @@
 #include "Widgets/Views/STableViewBase.h"
 #include "Widgets/Views/STableRow.h"
 #include "SkeletonTreeItem.h"
+#include "PhysicsEngine/BodySetup.h"
 #include "PhysicsEngine/ShapeElem.h"
 #include "PhysicsEngine/PhysicsAsset.h"
 
@@ -22,9 +23,9 @@ public:
 	/** ISkeletonTreeItem interface */
 	virtual void GenerateWidgetForNameColumn(TSharedPtr< SHorizontalBox > Box, const TAttribute<FText>& FilterText, FIsSelected InIsSelected) override;
 	virtual TSharedRef< SWidget > GenerateWidgetForDataColumn(const FName& DataColumnName, FIsSelected InIsSelected ) override;	
-	virtual FName GetRowItemName() const override { FString NameAsString = GetNameAsString(); return *NameAsString; }
-	virtual UObject* GetObject() const override { return BodySetup; }
-	virtual bool CanRenameItem() const override { return true; }
+	virtual FName GetRowItemName() const override;
+	virtual UObject* GetObject() const override;
+	virtual bool CanRenameItem() const override;
 	virtual void RequestRename() override;
 	virtual void OnItemDoubleClicked() override;
 
diff --git a/Engine/Source/Editor/SkeletalMeshEditor/Private/SkeletalMeshEditorSubsystem.cpp b/Engine/Source/Editor/SkeletalMeshEditor/Private/SkeletalMeshEditorSubsystem.cpp
--- a/Engine/Source/Editor/SkeletalMeshEditor/Private/SkeletalMeshEditorSubsystem.cpp
+++ b/Engine/Source/Editor/SkeletalMeshEditor/Private/SkeletalMeshEditorSubsystem.cpp
@@ -17,6 +17,7 @@
 #include "PhysicsAssetUtils.h"
 #include "PhysicsEngine/PhysicsAsset.h"
 #include "PhysicsEngine/PhysicsConstraintTemplate.h"
+#include "PhysicsEngine/SkeletalBodySetup.h"
 #include "Rendering/SkeletalMeshLODModel.h"
 #include "Rendering/SkeletalMeshLODRenderData.h"
 #include "Rendering/SkeletalMeshModel.h"
diff --git a/Engine/Source/Editor/UnrealEd/Private/ThumbnailRendering/PhysicsAssetThumbnailRenderer.cpp b/Engine/Source/Editor/UnrealEd/Private/ThumbnailRendering/PhysicsAssetThumbnailRenderer.cpp
--- a/Engine/Source/Editor/UnrealEd/Private/ThumbnailRendering/PhysicsAssetThumbnailRenderer.cpp
+++ b/Engine/Source/Editor/UnrealEd/Private/ThumbnailRendering/PhysicsAssetThumbnailRenderer.cpp
@@ -6,6 +6,7 @@
 #include "SceneView.h"
 #include "PhysicsEngine/PhysicsAsset.h"
 #include "ThumbnailHelpers.h"
+#include "Engine/SkeletalMesh.h"
 
 UPhysicsAssetThumbnailRenderer::UPhysicsAssetThumbnailRenderer(const FObjectInitializer& ObjectInitializer)
 	: Super(ObjectInitializer)
diff --git a/Engine/Source/Runtime/AnimGraphRuntime/Private/BoneControllers/AnimNode_RigidBody.cpp b/Engine/Source/Runtime/AnimGraphRuntime/Private/BoneControllers/AnimNode_RigidBody.cpp
--- a/Engine/Source/Runtime/AnimGraphRuntime/Private/BoneControllers/AnimNode_RigidBody.cpp
+++ b/Engine/Source/Runtime/AnimGraphRuntime/Private/BoneControllers/AnimNode_RigidBody.cpp
@@ -10,12 +10,14 @@
 
 #include "ClothCollisionSource.h"
 #include "Engine/OverlapResult.h"
+#include "Engine/SkeletalMesh.h"
 #include "GameFramework/Pawn.h"
 #include "HAL/Event.h"
 #include "HAL/LowLevelMemTracker.h"
 #include "PhysicsEngine/BodySetup.h"
 #include "PhysicsEngine/PhysicsAsset.h"
 #include "PhysicsEngine/PhysicsConstraintTemplate.h"
+#include "PhysicsEngine/SkeletalBodySetup.h"
 #include "GameFramework/PawnMovementComponent.h"
 #include "Physics/PhysicsInterfaceCore.h"
 #include "Physics/ImmediatePhysics/ImmediatePhysicsActorHandle.h"
diff --git a/Engine/Source/Runtime/Engine/Classes/PhysicsEngine/PhysicsAsset.h b/Engine/Source/Runtime/Engine/Classes/PhysicsEngine/PhysicsAsset.h
--- a/Engine/Source/Runtime/Engine/Classes/PhysicsEngine/PhysicsAsset.h
+++ b/Engine/Source/Runtime/Engine/Classes/PhysicsEngine/PhysicsAsset.h
@@ -2,16 +2,20 @@
 
 #pragma once
 
-#include "CoreMinimal.h"
 #include "UObject/ObjectMacros.h"
 #include "UObject/Object.h"
 #include "PhysicsEngine/RigidBodyIndexPair.h"
-#include "Engine/SkeletalMesh.h"
-#include "PhysicsEngine/PhysicalAnimationComponent.h"
-#include "PhysicsEngine/BodySetup.h"
 #include "Interfaces/Interface_PreviewMeshProvider.h"
 #include "PhysicsAsset.generated.h"
 
+namespace ECollisionEnabled { enum Type : int; }
+namespace EAggCollisionShape { enum Type : int; }
+
+struct FConstraintInstance;
+struct FConstraintInstanceAccessor;
+struct FReferenceSkeleton;
+
+class UMaterialInterface;
 class FMeshElementCollector;
 class USkeletalBodySetup;
 
@@ -407,68 +411,11 @@ private:
 	TArray<TObjectPtr<class UBodySetup>> BodySetup_DEPRECATED;
 };
 
-USTRUCT()
-struct FPhysicalAnimationProfile
-{
-	GENERATED_BODY()
-	
-	/** Profile name used to identify set of physical animation parameters */
-	UPROPERTY()
-	FName ProfileName;
-
-	/** Physical animation parameters used to drive animation */
-	UPROPERTY(EditAnywhere, Category = PhysicalAnimation)
-	FPhysicalAnimationData PhysicalAnimationData;
-};
-
-UCLASS(MinimalAPI)
-class USkeletalBodySetup : public UBodySetup
-{
-	GENERATED_BODY()
-public:
-	const FPhysicalAnimationProfile* FindPhysicalAnimationProfile(const FName ProfileName) const
-	{
-		return PhysicalAnimationData.FindByPredicate([ProfileName](const FPhysicalAnimationProfile& Profile){ return ProfileName == Profile.ProfileName; });
-	}
-
-	FPhysicalAnimationProfile* FindPhysicalAnimationProfile(const FName ProfileName)
-	{
-		return PhysicalAnimationData.FindByPredicate([ProfileName](const FPhysicalAnimationProfile& Profile) { return ProfileName == Profile.ProfileName; });
-	}
-
-	const TArray<FPhysicalAnimationProfile>& GetPhysicalAnimationProfiles() const
-	{
-		return PhysicalAnimationData;
-	}
-
-#if WITH_EDITOR
-	virtual void PostEditChangeProperty(struct FPropertyChangedEvent& PropertyChangedEvent);
-	ENGINE_API FName GetCurrentPhysicalAnimationProfileName() const;
-	
-	/** Creates a new physical animation profile entry */
-	ENGINE_API void AddPhysicalAnimationProfile(FName ProfileName);
-
-	/** Removes physical animation profile */
-	ENGINE_API void RemovePhysicalAnimationProfile(FName ProfileName);
-
-	ENGINE_API void UpdatePhysicalAnimationProfiles(const TArray<FName>& Profiles);
-
-	ENGINE_API void DuplicatePhysicalAnimationProfile(FName DuplicateFromName, FName DuplicateToName);
-
-	ENGINE_API void RenamePhysicalAnimationProfile(FName CurrentName, FName NewName);
-#endif
 
-#if WITH_EDITORONLY_DATA
-	//dummy place for customization inside phat. Profiles are ordered dynamically and we need a static place for detail customization
-	UPROPERTY(EditAnywhere, Category = PhysicalAnimation)
-	FPhysicalAnimationProfile CurrentPhysicalAnimationProfile;
+#if UE_ENABLE_INCLUDE_ORDER_DEPRECATED_IN_5_5
+#include "CoreMinimal.h"
+#include "Engine/SkeletalMesh.h"
+#include "PhysicsEngine/BodySetup.h"
+#include "PhysicsEngine/PhysicalAnimationComponent.h"
+#include "PhysicsEngine/SkeletalBodySetup.h"
 #endif
-
-	/** If true we ignore scale changes from animation. This is useful for subtle scale animations like breathing where the physics collision should remain unchanged*/
-	UPROPERTY(EditAnywhere, Category = BodySetup)
-	bool bSkipScaleFromAnimation;
-
-private:
-	UPROPERTY()
-	TArray<FPhysicalAnimationProfile> PhysicalAnimationData;
-};
diff --git a/Engine/Source/Runtime/Engine/Classes/PhysicsEngine/ShapeElem.h b/Engine/Source/Runtime/Engine/Classes/PhysicsEngine/ShapeElem.h
--- a/Engine/Source/Runtime/Engine/Classes/PhysicsEngine/ShapeElem.h
+++ b/Engine/Source/Runtime/Engine/Classes/PhysicsEngine/ShapeElem.h
@@ -12,7 +12,7 @@
 
 namespace EAggCollisionShape
 {
-	enum Type
+	enum Type : int
 	{
 		Sphere,
 		Box,
diff --git a/Engine/Source/Runtime/Engine/Classes/PhysicsEngine/SkeletalBodySetup.h b/Engine/Source/Runtime/Engine/Classes/PhysicsEngine/SkeletalBodySetup.h
new file mode 100644
--- /dev/null
+++ b/Engine/Source/Runtime/Engine/Classes/PhysicsEngine/SkeletalBodySetup.h
@@ -0,0 +1,74 @@
+// Copyright Epic Games, Inc. All Rights Reserved.
+
+#pragma once
+
+#include "PhysicsEngine/BodySetup.h"
+#include "PhysicsEngine/PhysicalAnimationComponent.h"
+
+#include "SkeletalBodySetup.generated.h"
+
+USTRUCT()
+struct FPhysicalAnimationProfile
+{
+	GENERATED_BODY()
+	
+	/** Profile name used to identify set of physical animation parameters */
+	UPROPERTY()
+	FName ProfileName;
+
+	/** Physical animation parameters used to drive animation */
+	UPROPERTY(EditAnywhere, Category = PhysicalAnimation)
+	FPhysicalAnimationData PhysicalAnimationData;
+};
+
+UCLASS(MinimalAPI)
+class USkeletalBodySetup : public UBodySetup
+{
+	GENERATED_BODY()
+public:
+	const FPhysicalAnimationProfile* FindPhysicalAnimationProfile(const FName ProfileName) const
+	{
+		return PhysicalAnimationData.FindByPredicate([ProfileName](const FPhysicalAnimationProfile& Profile){ return ProfileName == Profile.ProfileName; });
+	}
+
+	FPhysicalAnimationProfile* FindPhysicalAnimationProfile(const FName ProfileName)
+	{
+		return PhysicalAnimationData.FindByPredicate([ProfileName](const FPhysicalAnimationProfile& Profile) { return ProfileName == Profile.ProfileName; });
+	}
+
+	const TArray<FPhysicalAnimationProfile>& GetPhysicalAnimationProfiles() const
+	{
+		return PhysicalAnimationData;
+	}
+
+#if WITH_EDITOR
+	virtual void PostEditChangeProperty(struct FPropertyChangedEvent& PropertyChangedEvent);
+	ENGINE_API FName GetCurrentPhysicalAnimationProfileName() const;
+	
+	/** Creates a new physical animation profile entry */
+	ENGINE_API void AddPhysicalAnimationProfile(FName ProfileName);
+
+	/** Removes physical animation profile */
+	ENGINE_API void RemovePhysicalAnimationProfile(FName ProfileName);
+
+	ENGINE_API void UpdatePhysicalAnimationProfiles(const TArray<FName>& Profiles);
+
+	ENGINE_API void DuplicatePhysicalAnimationProfile(FName DuplicateFromName, FName DuplicateToName);
+
+	ENGINE_API void RenamePhysicalAnimationProfile(FName CurrentName, FName NewName);
+#endif
+
+#if WITH_EDITORONLY_DATA
+	//dummy place for customization inside phat. Profiles are ordered dynamically and we need a static place for detail customization
+	UPROPERTY(EditAnywhere, Category = PhysicalAnimation)
+	FPhysicalAnimationProfile CurrentPhysicalAnimationProfile;
+#endif
+
+	/** If true we ignore scale changes from animation. This is useful for subtle scale animations like breathing where the physics collision should remain unchanged*/
+	UPROPERTY(EditAnywhere, Category = BodySetup)
+	bool bSkipScaleFromAnimation;
+
+private:
+	UPROPERTY()
+	TArray<FPhysicalAnimationProfile> PhysicalAnimationData;
+};
diff --git a/Engine/Source/Runtime/Engine/Private/Components/SkeletalMeshComponent.cpp b/Engine/Source/Runtime/Engine/Private/Components/SkeletalMeshComponent.cpp
--- a/Engine/Source/Runtime/Engine/Private/Components/SkeletalMeshComponent.cpp
+++ b/Engine/Source/Runtime/Engine/Private/Components/SkeletalMeshComponent.cpp
@@ -11,10 +11,12 @@
 #include "Animation/AnimSingleNodeInstance.h"
 #include "Animation/AnimationSettings.h"
 #include "Animation/PoseSnapshot.h"
+#include "Engine/SkeletalMesh.h"
 #include "Engine/SkeletalMeshSocket.h"
 #include "AI/NavigationSystemHelpers.h"
 #include "Engine/SkinnedAsset.h"
 #include "PhysicsEngine/PhysicsSettings.h"
+#include "PhysicsEngine/SkeletalBodySetup.h"
 #include "Misc/Fork.h"
 #include "Particles/ParticleSystemComponent.h"
 #include "Physics/Experimental/PhysScene_Chaos.h"
diff --git a/Engine/Source/Runtime/Engine/Private/Components/SkinnedMeshComponent.cpp b/Engine/Source/Runtime/Engine/Private/Components/SkinnedMeshComponent.cpp
--- a/Engine/Source/Runtime/Engine/Private/Components/SkinnedMeshComponent.cpp
+++ b/Engine/Source/Runtime/Engine/Private/Components/SkinnedMeshComponent.cpp
@@ -21,6 +21,7 @@
 #include "Engine/SkeletalMeshSocket.h"
 #include "Engine/SkinnedAssetCommon.h"
 #include "PhysicsEngine/PhysicsAsset.h"
+#include "PhysicsEngine/SkeletalBodySetup.h"
 #include "Engine/CollisionProfile.h"
 #include "SkeletalMeshSceneProxy.h"
 #include "Animation/MeshDeformer.h"
diff --git a/Engine/Source/Runtime/Engine/Private/PhysicsEngine/ConstraintInstance.cpp b/Engine/Source/Runtime/Engine/Private/PhysicsEngine/ConstraintInstance.cpp
--- a/Engine/Source/Runtime/Engine/Private/PhysicsEngine/ConstraintInstance.cpp
+++ b/Engine/Source/Runtime/Engine/Private/PhysicsEngine/ConstraintInstance.cpp
@@ -1,6 +1,7 @@
 // Copyright Epic Games, Inc. All Rights Reserved.
 
 #include "PhysicsEngine/ConstraintInstance.h"
+#include "Engine/SkeletalMesh.h"
 #include "UObject/FrameworkObjectVersion.h"
 #include "Components/SkeletalMeshComponent.h"
 #include "PhysicsEngine/PhysicsConstraintComponent.h"
diff --git a/Engine/Source/Runtime/Engine/Private/PhysicsEngine/Experimental/PhysScene_Chaos.cpp b/Engine/Source/Runtime/Engine/Private/PhysicsEngine/Experimental/PhysScene_Chaos.cpp
--- a/Engine/Source/Runtime/Engine/Private/PhysicsEngine/Experimental/PhysScene_Chaos.cpp
+++ b/Engine/Source/Runtime/Engine/Private/PhysicsEngine/Experimental/PhysScene_Chaos.cpp
@@ -15,6 +15,7 @@
 #include "PhysicsEngine/ConstraintInstance.h"
 #include "PhysicsEngine/PhysicsCollisionHandler.h"
 #include "PhysicsEngine/PhysicsObjectExternalInterface.h"
+#include "PhysicsEngine/SkeletalBodySetup.h"
 #include "Chaos/PhysicsObjectInternalInterface.h"
 #include "Physics/Experimental/ChaosEventRelay.h"
 #include "EngineUtils.h"
diff --git a/Engine/Source/Runtime/Engine/Private/PhysicsEngine/PhysAnim.cpp b/Engine/Source/Runtime/Engine/Private/PhysicsEngine/PhysAnim.cpp
--- a/Engine/Source/Runtime/Engine/Private/PhysicsEngine/PhysAnim.cpp
+++ b/Engine/Source/Runtime/Engine/Private/PhysicsEngine/PhysAnim.cpp
@@ -14,6 +14,7 @@
 #include "Physics/Experimental/PhysScene_Chaos.h"
 #include "PhysicsEngine/PhysicsConstraintTemplate.h"
 #include "PhysicsEngine/PhysicsAsset.h"
+#include "PhysicsEngine/SkeletalBodySetup.h"
 
 CSV_DECLARE_CATEGORY_MODULE_EXTERN(CORE_API, Basic);
 
diff --git a/Engine/Source/Runtime/Engine/Private/PhysicsEngine/PhysDrawing.cpp b/Engine/Source/Runtime/Engine/Private/PhysicsEngine/PhysDrawing.cpp
--- a/Engine/Source/Runtime/Engine/Private/PhysicsEngine/PhysDrawing.cpp
+++ b/Engine/Source/Runtime/Engine/Private/PhysicsEngine/PhysDrawing.cpp
@@ -8,6 +8,7 @@
 #include "SceneView.h"
 #include "Materials/MaterialInstanceDynamic.h"
 #include "DynamicMeshBuilder.h"
+#include "PhysicsEngine/AggregateGeom.h"
 #include "PhysicsEngine/LevelSetElem.h"
 #include "PhysicsEngine/PhysicsConstraintTemplate.h"
 #include "PhysicsEngine/ConvexElem.h"
@@ -15,7 +16,10 @@
 #include "PhysicsEngine/SphereElem.h"
 #include "PhysicsEngine/SphylElem.h"
 #include "PhysicsEngine/PhysicsAsset.h"
+#include "PhysicsEngine/SkeletalBodySetup.h"
+#include "PhysicsEngine/SkinnedLevelSetElem.h"
 #include "PhysicsEngine/TaperedCapsuleElem.h"
+#include "ReferenceSkeleton.h"
 #include "StaticMeshResources.h"
 #include "Chaos/Levelset.h"
 #include "Chaos/UniformGrid.h"
diff --git a/Engine/Source/Runtime/Engine/Private/PhysicsEngine/PhysicalAnimationComponent.cpp b/Engine/Source/Runtime/Engine/Private/PhysicsEngine/PhysicalAnimationComponent.cpp
--- a/Engine/Source/Runtime/Engine/Private/PhysicsEngine/PhysicalAnimationComponent.cpp
+++ b/Engine/Source/Runtime/Engine/Private/PhysicsEngine/PhysicalAnimationComponent.cpp
@@ -4,8 +4,10 @@
 #include "Chaos/ImplicitFwd.h"
 #include "SceneManagement.h"
 #include "Components/SkeletalMeshComponent.h"
+#include "Engine/SkeletalMesh.h"
 #include "Physics/Experimental/PhysScene_Chaos.h"
 #include "PhysicsEngine/PhysicsAsset.h"
+#include "PhysicsEngine/SkeletalBodySetup.h"
 #include "Chaos/ChaosConstraintSettings.h"
 #include "Chaos/Sphere.h"
 #include "PhysicsEngine/ConstraintInstance.h"
diff --git a/Engine/Source/Runtime/Engine/Private/PhysicsEngine/PhysicsAsset.cpp b/Engine/Source/Runtime/Engine/Private/PhysicsEngine/PhysicsAsset.cpp
--- a/Engine/Source/Runtime/Engine/Private/PhysicsEngine/PhysicsAsset.cpp
+++ b/Engine/Source/Runtime/Engine/Private/PhysicsEngine/PhysicsAsset.cpp
@@ -6,12 +6,14 @@
 
 #include "PhysicsEngine/PhysicsAsset.h"
 #include "Animation/MirrorDataTable.h"
+#include "Engine/SkeletalMesh.h"
 #include "Engine/SkinnedAsset.h"
 #include "UObject/FrameworkObjectVersion.h"
 #include "Serialization/ObjectWriter.h"
 #include "Serialization/ObjectReader.h"
 #include "Components/SkeletalMeshComponent.h"
 #include "PhysicsEngine/PhysicsConstraintTemplate.h"
+#include "PhysicsEngine/SkeletalBodySetup.h"
 #include "UObject/AssetRegistryTagsContext.h"
 #include "UObject/ReleaseObjectVersion.h"
 #include "UObject/UObjectIterator.h"
@@ -79,90 +81,6 @@ void UPhysicsAsset::UpdateBoundsBodiesArray()
 	}
 }
 
-#if WITH_EDITOR
-void USkeletalBodySetup::PostEditChangeProperty(struct FPropertyChangedEvent& PropertyChangedEvent)
-{
-	UPhysicsAsset* OwningPhysAsset = Cast<UPhysicsAsset>(GetOuter());
-
-	if(PropertyChangedEvent.Property == nullptr || !OwningPhysAsset)
-	{
-		return;
-	}
-
-	if (FPhysicalAnimationProfile* PhysProfile = FindPhysicalAnimationProfile(OwningPhysAsset->CurrentPhysicalAnimationProfileName))
-	{
-			//changed any setting so copy dummy UI into profile location
-			PhysProfile->PhysicalAnimationData = CurrentPhysicalAnimationProfile.PhysicalAnimationData;
-	}
-
-	OwningPhysAsset->RefreshPhysicsAssetChange();
-}
-
-FName USkeletalBodySetup::GetCurrentPhysicalAnimationProfileName() const
-{
-	FName CurrentProfileName;
-	if(UPhysicsAsset* OwningPhysAsset = Cast<UPhysicsAsset>(GetOuter()))
-	{
-		CurrentProfileName = OwningPhysAsset->CurrentPhysicalAnimationProfileName;
-	}
-
-	return CurrentProfileName;
-}
-
-void USkeletalBodySetup::AddPhysicalAnimationProfile(FName ProfileName)
-{
-	FPhysicalAnimationProfile* NewProfile = new (PhysicalAnimationData) FPhysicalAnimationProfile();
-	NewProfile->ProfileName = ProfileName;
-}
-
-void USkeletalBodySetup::RemovePhysicalAnimationProfile(FName ProfileName)
-{
-	for(int32 ProfileIdx = 0; ProfileIdx < PhysicalAnimationData.Num(); ++ProfileIdx)
-	{
-		if(PhysicalAnimationData[ProfileIdx].ProfileName == ProfileName)
-		{
-			PhysicalAnimationData.RemoveAtSwap(ProfileIdx--);
-		}
-	}
-}
-
-void USkeletalBodySetup::UpdatePhysicalAnimationProfiles(const TArray<FName>& Profiles)
-{
-	for(int32 ProfileIdx = 0; ProfileIdx < PhysicalAnimationData.Num(); ++ProfileIdx)
-	{
-		if(Profiles.Contains(PhysicalAnimationData[ProfileIdx].ProfileName) == false)
-		{
-			PhysicalAnimationData.RemoveAtSwap(ProfileIdx--);
-		}
-	}
-}
-
-void USkeletalBodySetup::DuplicatePhysicalAnimationProfile(FName DuplicateFromName, FName DuplicateToName)
-{
-	for (FPhysicalAnimationProfile& ProfileHandle : PhysicalAnimationData)
-	{
-		if (ProfileHandle.ProfileName == DuplicateFromName)
-		{
-			FPhysicalAnimationProfile* Duplicate = new (PhysicalAnimationData) FPhysicalAnimationProfile(ProfileHandle);
-			Duplicate->ProfileName = DuplicateToName;
-			break;
-		}
-	}
-}
-
-void USkeletalBodySetup::RenamePhysicalAnimationProfile(FName CurrentName, FName NewName)
-{
-	for(FPhysicalAnimationProfile& ProfileHandle : PhysicalAnimationData)
-	{
-		if(ProfileHandle.ProfileName == CurrentName)
-		{
-			ProfileHandle.ProfileName = NewName;
-		}
-	}
-}
-
-#endif
-
 void UPhysicsAsset::UpdateBodySetupIndexMap()
 {
 	// update BodySetupIndexMap
diff --git a/Engine/Source/Runtime/Engine/Private/PhysicsEngine/PhysicsConstraintComponent.cpp b/Engine/Source/Runtime/Engine/Private/PhysicsEngine/PhysicsConstraintComponent.cpp
--- a/Engine/Source/Runtime/Engine/Private/PhysicsEngine/PhysicsConstraintComponent.cpp
+++ b/Engine/Source/Runtime/Engine/Private/PhysicsEngine/PhysicsConstraintComponent.cpp
@@ -9,6 +9,7 @@
 #include "PhysicsEngine/PhysicsAsset.h"
 #include "PhysicsEngine/ConstraintUtils.h"
 #include "PhysicsEngine/PhysicsObjectExternalInterface.h"
+#include "PhysicsEngine/SkeletalBodySetup.h"
 #include "PhysicsProxy/SingleParticlePhysicsProxy.h"
 #include "Components/BillboardComponent.h"
 #include "UObject/ICookInfo.h"
diff --git a/Engine/Source/Runtime/Engine/Private/PhysicsEngine/SkeletalBodySetup.cpp b/Engine/Source/Runtime/Engine/Private/PhysicsEngine/SkeletalBodySetup.cpp
new file mode 100644
--- /dev/null
+++ b/Engine/Source/Runtime/Engine/Private/PhysicsEngine/SkeletalBodySetup.cpp
@@ -0,0 +1,88 @@
+// Copyright Epic Games, Inc. All Rights Reserved.
+
+#include "PhysicsEngine/SkeletalBodySetup.h"
+#include "PhysicsEngine/PhysicsAsset.h"
+
+#if WITH_EDITOR
+void USkeletalBodySetup::PostEditChangeProperty(struct FPropertyChangedEvent& PropertyChangedEvent)
+{
+	UPhysicsAsset* OwningPhysAsset = Cast<UPhysicsAsset>(GetOuter());
+
+	if (PropertyChangedEvent.Property == nullptr || !OwningPhysAsset)
+	{
+		return;
+	}
+
+	if (FPhysicalAnimationProfile* PhysProfile = FindPhysicalAnimationProfile(OwningPhysAsset->CurrentPhysicalAnimationProfileName))
+	{
+		//changed any setting so copy dummy UI into profile location
+		PhysProfile->PhysicalAnimationData = CurrentPhysicalAnimationProfile.PhysicalAnimationData;
+	}
+
+	OwningPhysAsset->RefreshPhysicsAssetChange();
+}
+
+FName USkeletalBodySetup::GetCurrentPhysicalAnimationProfileName() const
+{
+	FName CurrentProfileName;
+	if (UPhysicsAsset* OwningPhysAsset = Cast<UPhysicsAsset>(GetOuter()))
+	{
+		CurrentProfileName = OwningPhysAsset->CurrentPhysicalAnimationProfileName;
+	}
+
+	return CurrentProfileName;
+}
+
+void USkeletalBodySetup::AddPhysicalAnimationProfile(FName ProfileName)
+{
+	FPhysicalAnimationProfile* NewProfile = new (PhysicalAnimationData) FPhysicalAnimationProfile();
+	NewProfile->ProfileName = ProfileName;
+}
+
+void USkeletalBodySetup::RemovePhysicalAnimationProfile(FName ProfileName)
+{
+	for (int32 ProfileIdx = 0; ProfileIdx < PhysicalAnimationData.Num(); ++ProfileIdx)
+	{
+		if (PhysicalAnimationData[ProfileIdx].ProfileName == ProfileName)
+		{
+			PhysicalAnimationData.RemoveAtSwap(ProfileIdx--);
+		}
+	}
+}
+
+void USkeletalBodySetup::UpdatePhysicalAnimationProfiles(const TArray<FName>& Profiles)
+{
+	for (int32 ProfileIdx = 0; ProfileIdx < PhysicalAnimationData.Num(); ++ProfileIdx)
+	{
+		if (Profiles.Contains(PhysicalAnimationData[ProfileIdx].ProfileName) == false)
+		{
+			PhysicalAnimationData.RemoveAtSwap(ProfileIdx--);
+		}
+	}
+}
+
+void USkeletalBodySetup::DuplicatePhysicalAnimationProfile(FName DuplicateFromName, FName DuplicateToName)
+{
+	for (FPhysicalAnimationProfile& ProfileHandle : PhysicalAnimationData)
+	{
+		if (ProfileHandle.ProfileName == DuplicateFromName)
+		{
+			FPhysicalAnimationProfile* Duplicate = new (PhysicalAnimationData) FPhysicalAnimationProfile(ProfileHandle);
+			Duplicate->ProfileName = DuplicateToName;
+			break;
+		}
+	}
+}
+
+void USkeletalBodySetup::RenamePhysicalAnimationProfile(FName CurrentName, FName NewName)
+{
+	for (FPhysicalAnimationProfile& ProfileHandle : PhysicalAnimationData)
+	{
+		if (ProfileHandle.ProfileName == CurrentName)
+		{
+			ProfileHandle.ProfileName = NewName;
+		}
+	}
+}
+
+#endif
diff --git a/Engine/Source/Runtime/Engine/Private/SkeletalMesh.cpp b/Engine/Source/Runtime/Engine/Private/SkeletalMesh.cpp
--- a/Engine/Source/Runtime/Engine/Private/SkeletalMesh.cpp
+++ b/Engine/Source/Runtime/Engine/Private/SkeletalMesh.cpp
@@ -40,7 +40,9 @@
 #include "SkeletalMeshSceneProxy.h"
 #include "SkeletalRenderPublic.h"
 #include "SceneManagement.h"
+#include "PhysicsEngine/BodySetup.h"
 #include "PhysicsEngine/PhysicsAsset.h"
+#include "PhysicsEngine/SkeletalBodySetup.h"
 #include "Engine/AssetUserData.h"
 #include "Animation/NodeMappingContainer.h"
 #include "Rendering/SkeletalMeshRenderData.h"
diff --git a/Engine/Source/Runtime/Engine/Private/SkeletalMeshComponentPhysics.cpp b/Engine/Source/Runtime/Engine/Private/SkeletalMeshComponentPhysics.cpp
--- a/Engine/Source/Runtime/Engine/Private/SkeletalMeshComponentPhysics.cpp
+++ b/Engine/Source/Runtime/Engine/Private/SkeletalMeshComponentPhysics.cpp
@@ -8,6 +8,7 @@
 #include "Components/InstancedStaticMeshComponent.h"
 #include "SkeletalMeshSceneProxy.h"
 #include "Engine/OverlapResult.h"
+#include "Engine/SkeletalMesh.h"
 #include "Engine/World.h"
 #include "Modules/ModuleManager.h"
 #include "Physics/Experimental/PhysScene_Chaos.h"
@@ -21,8 +22,10 @@
 #include "CollisionDebugDrawingPublic.h"
 
 
+#include "PhysicsEngine/BodySetup.h"
 #include "PhysicsEngine/PhysicsConstraintTemplate.h"
 #include "PhysicsEngine/PhysicsAsset.h"
+#include "PhysicsEngine/SkeletalBodySetup.h"
 #include "PhysicsEngine/SphylElem.h"
 #include "PhysicsEngine/TaperedCapsuleElem.h"
 #include "Rendering/SkeletalMeshLODRenderData.h"
diff --git a/Engine/Source/Runtime/Experimental/GeometryCollectionEngine/Private/GeometryCollection/PhysicsAssetSimulation.cpp b/Engine/Source/Runtime/Experimental/GeometryCollectionEngine/Private/GeometryCollection/PhysicsAssetSimulation.cpp
--- a/Engine/Source/Runtime/Experimental/GeometryCollectionEngine/Private/GeometryCollection/PhysicsAssetSimulation.cpp
+++ b/Engine/Source/Runtime/Experimental/GeometryCollectionEngine/Private/GeometryCollection/PhysicsAssetSimulation.cpp
@@ -16,6 +16,7 @@
 #include "Engine/SkeletalMeshSocket.h"
 #include "Chaos/ChaosPhysicalMaterial.h"
 #include "PhysicsEngine/PhysicsAsset.h"
+#include "PhysicsEngine/SkeletalBodySetup.h"
 #include "Rendering/SkeletalMeshRenderData.h"
 #include "PhysicsProxy/SkeletalMeshPhysicsProxy.h"
 
