[NetPhysics] Predictive Interpolation to handle instant physics tick offset correction.

diff --git a/Engine/Source/Runtime/Engine/Private/PhysicsEngine/PhysicsReplication.cpp b/Engine/Source/Runtime/Engine/Private/PhysicsEngine/PhysicsReplication.cpp
--- a/Engine/Source/Runtime/Engine/Private/PhysicsEngine/PhysicsReplication.cpp
+++ b/Engine/Source/Runtime/Engine/Private/PhysicsEngine/PhysicsReplication.cpp
@@ -227,7 +227,7 @@ namespace PhysicsReplicationCVars
 		float SleepSecondsClearTarget = 15.0f;
 		static FAutoConsoleVariableRef CVarSleepSecondsClearTarget(TEXT("np2.PredictiveInterpolation.SleepSecondsClearTarget"), SleepSecondsClearTarget, TEXT("Wait for the object to sleep for this many seconds before clearing the replication target, to ensure nothing wakes up the object just after it goes to sleep on the client."));
 		
-		int32 TargetTickAlignmentClampMultiplier = 1;
+		int32 TargetTickAlignmentClampMultiplier = 2;
 		static FAutoConsoleVariableRef CVarTargetTickAlignmentClampMultiplier(TEXT("np2.PredictiveInterpolation.TargetTickAlignmentClampMultiplier"), TargetTickAlignmentClampMultiplier, TEXT("Multiplier to adjust clamping of target alignment via TickCount. Multiplier is performed on AverageReceiveInterval."));
 	}
 
@@ -952,16 +952,19 @@ void FPhysicsReplicationAsync::UpdateAsyncTarget(const FPhysicsRepAsyncInputData
 
 		// Cache TickCount before updating it, force to 0 if ServerFrame is -1
 		const int32 PrevTickCount = (Target->ServerFrame < 0) ? 0 : Target->TickCount;
-		
+
 		// Cache SendInterval, only calculate if we have a valid Target->ServerFrame, else leave at 0.
 		const int32 SendInterval = (Target->ServerFrame <= 0) ? 0 : Input.ServerFrame - Target->ServerFrame;
-		
+
 		// Cache if this target was previously allowed to be altered, before this update
 		const bool bPrevAllowTargetAltering = Target->bAllowTargetAltering;
-		
+
+		// Cache if the physics frame offset has changed since last target
+		const bool bFrameOffsetCorrected = Target->FrameOffset != Input.FrameOffset;
+
 		// Set if the target is allowed to be altered after this update
 		Target->bAllowTargetAltering = !(Target->TargetState.Flags & ERigidBodyFlags::Sleeping) && !(Input.TargetState.Flags & ERigidBodyFlags::Sleeping);
-		
+
 		// Set Target->ReceiveInterval from either SendInterval or the number of physics ticks between receiving input states
 		if (SendInterval > 0)
 		{
@@ -1024,7 +1027,7 @@ void FPhysicsReplicationAsync::UpdateAsyncTarget(const FPhysicsRepAsyncInputData
 				*/
 
 				// Run target alignment if we have been allowed to alter the target during the last two target updates
-				if (!bFirstTarget && bPrevAllowTargetAltering && Target->bAllowTargetAltering)
+				if (!bFirstTarget && bPrevAllowTargetAltering && Target->bAllowTargetAltering && !bFrameOffsetCorrected)
 				{
 					const int32 AdjustedAverageReceiveInterval = FMath::CeilToInt(Target->AverageReceiveInterval) * PhysicsReplicationCVars::PredictiveInterpolationCVars::TargetTickAlignmentClampMultiplier;
 
diff --git a/Engine/Source/Runtime/Engine/Private/PlayerController.cpp b/Engine/Source/Runtime/Engine/Private/PlayerController.cpp
--- a/Engine/Source/Runtime/Engine/Private/PlayerController.cpp
+++ b/Engine/Source/Runtime/Engine/Private/PlayerController.cpp
@@ -6257,7 +6257,7 @@ void APlayerController::ServerSendLatestAsyncPhysicsTimestamp_Implementation(FAs
 	{
 		// Get the buffer offset amount that deviates from the target buffer (Note: the buffer is already added to ActualTimestamp.ServerFrame here and in the PredictedServerFrame received from the client)
 		// 0 means buffer is perfect, positive value means the buffer is too large, negative value means the buffer is too small
-		int32 CurrentFrameBufferOffset = PredictedServerFrame - ActualTimestamp.ServerFrame;
+		int32 CurrentFrameBufferOffset = Timestamp.ServerFrame - ActualTimestamp.ServerFrame;
 
 		if (NetworkPhysicsCvars::TimeDilationEscalation == false)
 		{
