[NetPhysics] Sleep all connected particles when replication wants to sleep

diff --git a/Engine/Source/Runtime/Engine/Private/PhysicsEngine/PhysicsReplication.cpp b/Engine/Source/Runtime/Engine/Private/PhysicsEngine/PhysicsReplication.cpp
--- a/Engine/Source/Runtime/Engine/Private/PhysicsEngine/PhysicsReplication.cpp
+++ b/Engine/Source/Runtime/Engine/Private/PhysicsEngine/PhysicsReplication.cpp
@@ -205,6 +205,9 @@ namespace PhysicsReplicationCVars
 
 		bool bCorrectConnectedBodiesFriction = true;
 		static FAutoConsoleVariableRef CVarCorrectConnectedBodiesFriction(TEXT("np2.PredictiveInterpolation.CorrectConnectedBodiesFriction"), bCorrectConnectedBodiesFriction, TEXT("When true, transform correction on any connected physics object will also recalculate their friction."));
+		
+		bool bSleepConnectedBodies = true;
+		static FAutoConsoleVariableRef CVarSleepConnectedBodies(TEXT("np2.PredictiveInterpolation.SleepConnectedBodies"), bSleepConnectedBodies, TEXT("When true, sleep state will be applied to any dynamic physics object connected to the replicated object."));
 
 		bool bDisableSoftSnap = false;
 		static FAutoConsoleVariableRef CVarDisableSoftSnap(TEXT("np2.PredictiveInterpolation.DisableSoftSnap"), bDisableSoftSnap, TEXT("When true, predictive interpolation will not use softsnap to correct the replication with when velocity fails. Hardsnap will still eventually kick in if replication can't reach the target."));
@@ -1492,6 +1495,10 @@ bool FPhysicsReplicationAsync::PredictiveInterpolation(Chaos::FPBDRigidParticleH
 		if (bOkToClear && bShouldSleep && bCanSimulate)
 		{
 			RigidsSolver->GetEvolution()->SetParticleObjectState(Handle, Chaos::EObjectStateType::Sleeping);
+			if (PhysicsReplicationCVars::PredictiveInterpolationCVars::bSleepConnectedBodies)
+			{
+				RigidsSolver->GetEvolution()->ApplySleepOnConnectedParticles(Handle);
+			}
 		}
 
 		// --- Should replication stop? ---
diff --git a/Engine/Source/Runtime/Experimental/Chaos/Private/Chaos/PBDRigidsEvolutionGBF.cpp b/Engine/Source/Runtime/Experimental/Chaos/Private/Chaos/PBDRigidsEvolutionGBF.cpp
--- a/Engine/Source/Runtime/Experimental/Chaos/Private/Chaos/PBDRigidsEvolutionGBF.cpp
+++ b/Engine/Source/Runtime/Experimental/Chaos/Private/Chaos/PBDRigidsEvolutionGBF.cpp
@@ -1512,6 +1512,21 @@ void FPBDRigidsEvolutionGBF::ApplyParticleTransformCorrectionImpl(FGeometryParti
 	}
 }
 
+void FPBDRigidsEvolutionGBF::ApplySleepOnConnectedParticles(FGeometryParticleHandle* InParticle)
+{
+	TArray<FGeometryParticleHandle*> ConnectedParticles = GetConnectedParticles(InParticle);
+	for (FGeometryParticleHandle* ConnectedParticle : ConnectedParticles)
+	{
+		if (FGenericParticleHandle(ConnectedParticle)->ObjectState() == EObjectStateType::Dynamic)
+		{
+			if (FPBDRigidParticleHandle* Rigid = ConnectedParticle->CastToRigidParticle())
+			{
+				SetParticleObjectState(Rigid, EObjectStateType::Sleeping);
+			}
+		}
+	}
+}
+
 TArray<FGeometryParticleHandle*> FPBDRigidsEvolutionGBF::GetConnectedParticles(FGeometryParticleHandle* InParticle)
 {
 	if (InParticle->ParticleConstraints().IsEmpty())
diff --git a/Engine/Source/Runtime/Experimental/Chaos/Public/Chaos/PBDRigidsEvolutionGBF.h b/Engine/Source/Runtime/Experimental/Chaos/Public/Chaos/PBDRigidsEvolutionGBF.h
--- a/Engine/Source/Runtime/Experimental/Chaos/Public/Chaos/PBDRigidsEvolutionGBF.h
+++ b/Engine/Source/Runtime/Experimental/Chaos/Public/Chaos/PBDRigidsEvolutionGBF.h
@@ -186,6 +186,11 @@ namespace Chaos
 		 */
 		CHAOS_API void ApplyParticleTransformCorrection(FGeometryParticleHandle* InParticle, const FVec3& InPos, const FRotation3& InRot, const bool bApplyToConnectedBodies, const bool bInRecalculateFrictionOnConnectedBodies = true);
 
+		/*
+		* [EXPERIMENTAL] Apply sleep state on connected particles that are dynamic
+		*/
+		CHAOS_API void ApplySleepOnConnectedParticles(FGeometryParticleHandle* InParticle);
+
 		/**
 		 * Called when a particle is moved. We need to reset some friction properties, sleeping properties, etc
 		 */
