Volmetric cloud
- clarified tooltip
- marked empty space skipping nodes as experimental

diff --git a/Engine/Source/Runtime/Engine/Classes/Materials/MaterialExpressionVolumetricAdvancedMaterialInput.h b/Engine/Source/Runtime/Engine/Classes/Materials/MaterialExpressionVolumetricAdvancedMaterialInput.h
--- a/Engine/Source/Runtime/Engine/Classes/Materials/MaterialExpressionVolumetricAdvancedMaterialInput.h
+++ b/Engine/Source/Runtime/Engine/Classes/Materials/MaterialExpressionVolumetricAdvancedMaterialInput.h
@@ -21,7 +21,7 @@ class UMaterialExpressionVolumetricAdvancedMaterialInput : public UMaterialExpre
 	//~ End UMaterialExpression Interface
 };
 
-UCLASS()
+UCLASS(Experimental)
 class UMaterialExpressionVolumetricCloudEmptySpaceSkippingInput : public UMaterialExpression
 {
 	GENERATED_UCLASS_BODY()
diff --git a/Engine/Source/Runtime/Engine/Classes/Materials/MaterialExpressionVolumetricAdvancedMaterialOutput.h b/Engine/Source/Runtime/Engine/Classes/Materials/MaterialExpressionVolumetricAdvancedMaterialOutput.h
--- a/Engine/Source/Runtime/Engine/Classes/Materials/MaterialExpressionVolumetricAdvancedMaterialOutput.h
+++ b/Engine/Source/Runtime/Engine/Classes/Materials/MaterialExpressionVolumetricAdvancedMaterialOutput.h
@@ -39,12 +39,10 @@ class UMaterialExpressionVolumetricAdvancedMaterialOutput : public UMaterialExpr
 	UPROPERTY(meta = (RequiredInput = "false", ToolTip = "Multi-scattering approximation: represents how much the phase will become isotropic for each successive octave. Evaluated per pixel. Valid range is [0,1], from anisotropic to isotropic phase. Defaults to ConstMultiScatteringEccentricity from properties panel if not specified. Evaluated per pixel (globally)."))
 	FExpressionInput MultiScatteringEccentricity;
 
-
 	/** This is a 3-components float vector. The X component must represent the participating medium conservative density. This is used to accelerate the ray marching by early skipping expensive material evaluation. For example, a simple top down 2D density texture would be enough to help by not evaluating the material in empty regions. The Y and Z components can contain parameters that can be recovered during the material evaluation using the VolumetricAdvancedMaterialInput node. Evaluated per sample. */
-	UPROPERTY(meta = (RequiredInput = "false", ToolTip = "This is a 3-components float vector. The X component must represent the participating medium conservative density. This is used to accelerate the ray marching by early skipping expensive material evaluation. For example, a simple top down 2D density texture would be enough to help by not evaluating the material in empty regions. The Y and Z components can contain parameters that can be recovered during the material evaluation using the VolumetricAdvancedMaterialInput node. Evaluated per sample."))
+	UPROPERTY(meta = (RequiredInput = "false", ToolTip = "This is a 4-components float vector. The X component must represent the participating medium conservative density. This is used to accelerate the ray marching by early skipping expensive material evaluation. For example, the bare minimum would be a top down 2D density texture skipping evaluating the full material in empty regions (better, would be a simple blobby volume representation of clouds). The Y, Z and W components can contain parameters that can be recovered during the material graph evaluation using the VolumetricAdvancedMaterialInput node. Evaluated per sample."))
 	FExpressionInput ConservativeDensity;
 
-
 	/** Only used if PhaseG is not hooked up. Parameter 'g' input to the phase function  describing how much forward(g<0) or backward (g>0) light scatter around. */
 	UPROPERTY(EditAnywhere, Category = "Phase", meta = (OverridingInputProperty = "PhaseG", UIMin = -0.999f, UIMax = 0.999f, ClampMin = -0.999f, ClampMax = 0.999f))
 	float ConstPhaseG;
@@ -116,7 +114,7 @@ public:
 };
 
 /** USed to help the cloud system to fast skip empty space areas when ray marching. */
-UCLASS(MinimalAPI, collapsecategories, hidecategories = Object)
+UCLASS(Experimental, MinimalAPI, collapsecategories, hidecategories = Object)
 class UMaterialExpressionVolumetricCloudEmptySpaceSkippingOutput : public UMaterialExpressionCustomOutput
 {
 	GENERATED_UCLASS_BODY()
diff --git a/Engine/Source/Runtime/Engine/Private/Materials/MaterialExpressions.cpp b/Engine/Source/Runtime/Engine/Private/Materials/MaterialExpressions.cpp
--- a/Engine/Source/Runtime/Engine/Private/Materials/MaterialExpressions.cpp
+++ b/Engine/Source/Runtime/Engine/Private/Materials/MaterialExpressions.cpp
@@ -23475,7 +23475,7 @@ int32 UMaterialExpressionVolumetricCloudEmptySpaceSkippingOutput::Compile(class
 
 void UMaterialExpressionVolumetricCloudEmptySpaceSkippingOutput::GetCaption(TArray<FString>& OutCaptions) const
 {
-	OutCaptions.Add(FString(TEXT("Volumetric Cloud Empty Space Skipping Output")));
+	OutCaptions.Add(FString(TEXT("Volumetric Cloud Empty Space Skipping Output <Experimental>")));
 }
 
 uint32 UMaterialExpressionVolumetricCloudEmptySpaceSkippingOutput::GetOutputType(int32 OutputIndex)
@@ -23546,7 +23546,7 @@ int32 UMaterialExpressionVolumetricCloudEmptySpaceSkippingInput::Compile(class F
 
 void UMaterialExpressionVolumetricCloudEmptySpaceSkippingInput::GetCaption(TArray<FString>& OutCaptions) const
 {
-	OutCaptions.Add(TEXT("Volumetric Cloud Empty Space Skipping Input"));
+	OutCaptions.Add(TEXT("Volumetric Cloud Empty Space Skipping Input <Experimental>"));
 }
 
 void UMaterialExpressionVolumetricCloudEmptySpaceSkippingInput::GetConnectorToolTip(int32 InputIndex, int32 OutputIndex, TArray<FString>& OutToolTip)
