[VisualLogger] removed ensure on the Editor world when getting timestamp since it might be called while it's getting garbage collected. Also added GWorld as a fallback in GetWorldForVisualLogger.

diff --git a/Engine/Source/Runtime/Engine/Private/VisualLogger/VisualLogger.cpp b/Engine/Source/Runtime/Engine/Private/VisualLogger/VisualLogger.cpp
--- a/Engine/Source/Runtime/Engine/Private/VisualLogger/VisualLogger.cpp
+++ b/Engine/Source/Runtime/Engine/Private/VisualLogger/VisualLogger.cpp
@@ -200,6 +200,12 @@ namespace
 			World =  GEngine->GetWorld();
 		}
 
+		// In last resort we fall back on the global world (i.e., logs produced while world gets GCed)
+		if (World == nullptr)
+		{
+			World = GWorld;
+		}
+
 		return World;
 	}
 }
@@ -268,10 +274,12 @@ double FVisualLogger::GetTimeStampForObject(const UObject* Object) const
 	UEditorEngine* EditorEngine = GIsEditor ? Cast<UEditorEngine>(GEngine) : nullptr;
 	if (EditorEngine)
 	{
-		// We will always have the Editor world to use.  This will ensure a consistent clock since it does not reset
-		// when more clients are added or removed and can exist before a PIE session is started.
+		// Using the Editor world to ensure a consistent clock since it does not reset when more clients are added or removed
+		// and can exist before a PIE session is started.
+
+		// We should always have the Editor world to use, but in some edge cases it can no longer be valid (e.g., world getting GCed).
 		WorldForTimeStamp = EditorEngine->GetEditorWorldContext().World();
-		if (ensureMsgf(WorldForTimeStamp, TEXT("We always expect to have an EditorWorld in Editor")))
+		if (WorldForTimeStamp)
 		{
 			using namespace UE::VisLog::Private;
 			if (EditorOnly::EditorBaseTimeStamp <= 0.0)
