-Remove need to specify a name for the streaming object for ContentBundles/ExternalDataLayers
-Was not strictly necessary and relied on a sucessful GC inbetween streaming generation passes

diff --git a/Engine/Source/Runtime/Engine/Private/WorldPartition/ContentBundle/ContentBundle.cpp b/Engine/Source/Runtime/Engine/Private/WorldPartition/ContentBundle/ContentBundle.cpp
--- a/Engine/Source/Runtime/Engine/Private/WorldPartition/ContentBundle/ContentBundle.cpp
+++ b/Engine/Source/Runtime/Engine/Private/WorldPartition/ContentBundle/ContentBundle.cpp
@@ -29,9 +29,10 @@ void FContentBundle::DoInitialize()
 	ExternalStreamingObjectPackage = LoadPackage(nullptr, *GetExternalStreamingObjectPackagePath(), LOAD_None);
 	if (ExternalStreamingObjectPackage != nullptr)
 	{
-		if (UObject* Object = StaticFindObjectFast(URuntimeHashExternalStreamingObjectBase::StaticClass(), ExternalStreamingObjectPackage, *GetExternalStreamingObjectName()))
+		ExternalStreamingObject = Cast<URuntimeHashExternalStreamingObjectBase>((UObject*)FindObjectWithOuter(ExternalStreamingObjectPackage, URuntimeHashExternalStreamingObjectBase::StaticClass()));
+
+		if (ExternalStreamingObject)
 		{
-			ExternalStreamingObject = CastChecked<URuntimeHashExternalStreamingObjectBase>(Object);
 			ExternalStreamingObject->OnStreamingObjectLoaded(GetInjectedWorld());
 		}
 		else
diff --git a/Engine/Source/Runtime/Engine/Private/WorldPartition/ContentBundle/ContentBundleBase.cpp b/Engine/Source/Runtime/Engine/Private/WorldPartition/ContentBundle/ContentBundleBase.cpp
--- a/Engine/Source/Runtime/Engine/Private/WorldPartition/ContentBundle/ContentBundleBase.cpp
+++ b/Engine/Source/Runtime/Engine/Private/WorldPartition/ContentBundle/ContentBundleBase.cpp
@@ -107,11 +107,6 @@ FString FContentBundleBase::GetExternalStreamingObjectPackagePath() const
 	return ContentBundlePaths::GetCookedContentBundleLevelFolder(*this) + ContentBundlePaths::GetGeneratedFolderName() + TEXT("/") + GetExternalStreamingObjectPackageName();
 }
 
-FString FContentBundleBase::GetExternalStreamingObjectName() const
-{
-	return SlugStringForValidName(GetDisplayName() + TEXT("_") + GetDescriptor()->GetGuid().ToString() + TEXT("_ExternalStreamingObject"));
-}
-
 UWorld* FContentBundleBase::GetInjectedWorld() const
 {
 	return InjectedWorld;
diff --git a/Engine/Source/Runtime/Engine/Private/WorldPartition/ContentBundle/ContentBundleEditor.cpp b/Engine/Source/Runtime/Engine/Private/WorldPartition/ContentBundle/ContentBundleEditor.cpp
--- a/Engine/Source/Runtime/Engine/Private/WorldPartition/ContentBundle/ContentBundleEditor.cpp
+++ b/Engine/Source/Runtime/Engine/Private/WorldPartition/ContentBundle/ContentBundleEditor.cpp
@@ -362,7 +362,7 @@ void FContentBundleEditor::GenerateStreaming(TArray<FString>* OutPackageToGenera
 	UWorldPartition* WorldPartition = GetInjectedWorld()->GetWorldPartition();
 	WorldPartition->GenerateContainerStreaming(Params, Context);
 
-	ExternalStreamingObject = WorldPartition->FlushStreamingToExternalStreamingObject(GetExternalStreamingObjectName());
+	ExternalStreamingObject = WorldPartition->FlushStreamingToExternalStreamingObject();
 
 	uint32 CellCount = 0;
 	ExternalStreamingObject->ForEachStreamingCells([&CellCount](const UWorldPartitionRuntimeCell& Cell)
diff --git a/Engine/Source/Runtime/Engine/Private/WorldPartition/Cook/WorldPartitionCook.cpp b/Engine/Source/Runtime/Engine/Private/WorldPartition/Cook/WorldPartitionCook.cpp
--- a/Engine/Source/Runtime/Engine/Private/WorldPartition/Cook/WorldPartitionCook.cpp
+++ b/Engine/Source/Runtime/Engine/Private/WorldPartition/Cook/WorldPartitionCook.cpp
@@ -146,10 +146,11 @@ bool UWorldPartition::PopulateGeneratedPackageForCook(IWorldPartitionCookPackage
 	auto GetCookPackageObjectPath = [this](IWorldPartitionCookPackageObject* InCookPackageObject)
 	{
 		const UExternalDataLayerAsset* ExternalDataLayerAsset = InCookPackageObject->GetExternalDataLayerAsset();
+		const URuntimeHashExternalStreamingObjectBase* StreamingObject = ExternalDataLayerManager->ExternalStreamingObjects.FindChecked(ExternalDataLayerAsset);
 		FStringBuilderBase CookPackageObjectPath;
 		CookPackageObjectPath += ExternalDataLayerManager->GetExternalStreamingObjectPackagePath(ExternalDataLayerAsset);
 		CookPackageObjectPath += TEXT(".");
-		CookPackageObjectPath += FExternalDataLayerHelper::GetExternalStreamingObjectName(ExternalDataLayerAsset);
+		CookPackageObjectPath += StreamingObject->GetName();
 		CookPackageObjectPath += TEXT(".");
 		CookPackageObjectPath += Cast<UObject>(InCookPackageObject)->GetName();
 		return FSoftObjectPath(*CookPackageObjectPath);
diff --git a/Engine/Source/Runtime/Engine/Private/WorldPartition/DataLayer/ExternalDataLayerHelper.cpp b/Engine/Source/Runtime/Engine/Private/WorldPartition/DataLayer/ExternalDataLayerHelper.cpp
--- a/Engine/Source/Runtime/Engine/Private/WorldPartition/DataLayer/ExternalDataLayerHelper.cpp
+++ b/Engine/Source/Runtime/Engine/Private/WorldPartition/DataLayer/ExternalDataLayerHelper.cpp
@@ -24,12 +24,6 @@ FString FExternalDataLayerHelper::GetExternalStreamingObjectPackageName(const UE
 	return FString::Printf(TEXT("StreamingObject_%X"), (uint32)InExternalDataLayerAsset->GetUID());
 }
 
-FString FExternalDataLayerHelper::GetExternalStreamingObjectName(const UExternalDataLayerAsset* InExternalDataLayerAsset)
-{
-	check(InExternalDataLayerAsset);
-	return SlugStringForValidName(InExternalDataLayerAsset->GetName() + TEXT("_") + InExternalDataLayerAsset->GetUID().ToString() + TEXT("_ExternalStreamingObject"));
-}
-
 bool FExternalDataLayerHelper::BuildExternalDataLayerRootPath(const FString& InEDLMountPoint, const FExternalDataLayerUID& InExternalDataLayerUID, FString& OutExternalDataLayerRootPath)
 {
 	if (InEDLMountPoint.IsEmpty() || !InExternalDataLayerUID.IsValid())
diff --git a/Engine/Source/Runtime/Engine/Private/WorldPartition/DataLayer/ExternalDataLayerManager.cpp b/Engine/Source/Runtime/Engine/Private/WorldPartition/DataLayer/ExternalDataLayerManager.cpp
--- a/Engine/Source/Runtime/Engine/Private/WorldPartition/DataLayer/ExternalDataLayerManager.cpp
+++ b/Engine/Source/Runtime/Engine/Private/WorldPartition/DataLayer/ExternalDataLayerManager.cpp
@@ -198,8 +198,8 @@ bool UExternalDataLayerManager::RegisterExternalStreamingObjectForGameWorld(cons
 		return false;
 	}
 
-	const FString ExternalStreamingObjectName = FExternalDataLayerHelper::GetExternalStreamingObjectName(InExternalDataLayerAsset);
-	URuntimeHashExternalStreamingObjectBase* ExternalStreamingObject = FindObject<URuntimeHashExternalStreamingObjectBase>(ExternalStreamingObjectPackage, *ExternalStreamingObjectName);
+	URuntimeHashExternalStreamingObjectBase* ExternalStreamingObject = Cast<URuntimeHashExternalStreamingObjectBase>((UObject*)FindObjectWithOuter(ExternalStreamingObjectPackage, URuntimeHashExternalStreamingObjectBase::StaticClass()));
+
 	if (!ExternalStreamingObject)
 	{
 		UE_LOG(LogWorldPartition, Error, TEXT("[EDL: %s] No external streaming object found in package %s."), *InExternalDataLayerAsset->GetName(), *ExternalStreamingObjectPackagePath);
@@ -890,8 +890,7 @@ URuntimeHashExternalStreamingObjectBase* UExternalDataLayerManager::CreateExtern
 		return nullptr;
 	}
 
-	const FString ExternalStreamingObjectName = FExternalDataLayerHelper::GetExternalStreamingObjectName(InExternalDataLayerAsset);
-	URuntimeHashExternalStreamingObjectBase* ExternalStreamingObject = GetOuterUWorldPartition()->FlushStreamingToExternalStreamingObject(ExternalStreamingObjectName);
+	URuntimeHashExternalStreamingObjectBase* ExternalStreamingObject = GetOuterUWorldPartition()->FlushStreamingToExternalStreamingObject();
 	if (!ExternalStreamingObject)
 	{
 		return nullptr;
diff --git a/Engine/Source/Runtime/Engine/Private/WorldPartition/WorldPartitionLevelHelper.cpp b/Engine/Source/Runtime/Engine/Private/WorldPartition/WorldPartitionLevelHelper.cpp
--- a/Engine/Source/Runtime/Engine/Private/WorldPartition/WorldPartitionLevelHelper.cpp
+++ b/Engine/Source/Runtime/Engine/Private/WorldPartition/WorldPartitionLevelHelper.cpp
@@ -21,6 +21,7 @@
 #include "WorldPartition/WorldPartition.h"
 #include "WorldPartition/WorldPartitionLog.h"
 #include "WorldPartition/WorldPartitionPackageHelper.h"
+#include "WorldPartition/WorldPartitionRuntimeHash.h"
 #include "WorldPartition/ContentBundle/ContentBundleEditor.h"
 #include "WorldPartition/WorldPartitionPropertyOverride.h"
 #include "GameFramework/WorldSettings.h"
@@ -372,7 +373,7 @@ bool FWorldPartitionLevelHelper::RemapLevelCellPathInContentBundle(ULevel* Level
 {
 	FString CellPath = ContentBundleEditor->GetExternalStreamingObjectPackagePath();
 	CellPath += TEXT(".");
-	CellPath += ContentBundleEditor->GetExternalStreamingObjectName();
+	CellPath += ContentBundleEditor->GetStreamingObject()->GetName();
 	CellPath += TEXT(".");
 	CellPath += Cell->GetName();
 	FSetWorldPartitionRuntimeCell SetWorldPartitionRuntimeCell(Level, FSoftObjectPath(CellPath));
diff --git a/Engine/Source/Runtime/Engine/Private/WorldPartition/WorldPartitionRuntimeHash.cpp b/Engine/Source/Runtime/Engine/Private/WorldPartition/WorldPartitionRuntimeHash.cpp
--- a/Engine/Source/Runtime/Engine/Private/WorldPartition/WorldPartitionRuntimeHash.cpp
+++ b/Engine/Source/Runtime/Engine/Private/WorldPartition/WorldPartitionRuntimeHash.cpp
@@ -134,7 +134,8 @@ bool URuntimeHashExternalStreamingObjectBase::OnPopulateGeneratorPackageForCook(
 
 bool URuntimeHashExternalStreamingObjectBase::OnPopulateGeneratedPackageForCook(UPackage* InPackage, TArray<UPackage*>& OutModifiedPackages)
 {
-	return Rename(nullptr, InPackage, REN_DontCreateRedirectors);
+	// We provide a new name for the URuntimeHashExternalStreamingObjectBase in the package so that we have a stable name (for cook determinism)
+	return Rename(TEXT("RuntimeHashExternalStreamingObjectBase"), InPackage, REN_DontCreateRedirectors);
 }
 
 void URuntimeHashExternalStreamingObjectBase::DumpStateLog(FHierarchicalLogArchive& Ar)
@@ -195,9 +196,9 @@ EWorldPartitionStreamingPerformance UWorldPartitionRuntimeHash::GetStreamingPerf
 	return EWorldPartitionStreamingPerformance::Good;
 }
 
-URuntimeHashExternalStreamingObjectBase* UWorldPartitionRuntimeHash::CreateExternalStreamingObject(TSubclassOf<URuntimeHashExternalStreamingObjectBase> InClass, UObject* InOuter, FName InName, UWorld* InOuterWorld)
+URuntimeHashExternalStreamingObjectBase* UWorldPartitionRuntimeHash::CreateExternalStreamingObject(TSubclassOf<URuntimeHashExternalStreamingObjectBase> InClass, UObject* InOuter, UWorld* InOuterWorld)
 {
-	URuntimeHashExternalStreamingObjectBase* StreamingObject = NewObject<URuntimeHashExternalStreamingObjectBase>(InOuter, InClass, InName, RF_Public);
+	URuntimeHashExternalStreamingObjectBase* StreamingObject = NewObject<URuntimeHashExternalStreamingObjectBase>(InOuter, InClass, NAME_None, RF_Public);
 	StreamingObject->OuterWorld = InOuterWorld;	
 	return StreamingObject;
 }
@@ -366,9 +367,9 @@ UWorldPartitionRuntimeCell* UWorldPartitionRuntimeHash::GetCellForCookPackage(co
 	return nullptr;
 }
 
-URuntimeHashExternalStreamingObjectBase* UWorldPartitionRuntimeHash::StoreStreamingContentToExternalStreamingObject(FName InStreamingObjectName)
+URuntimeHashExternalStreamingObjectBase* UWorldPartitionRuntimeHash::StoreStreamingContentToExternalStreamingObject()
 {
-	URuntimeHashExternalStreamingObjectBase* NewExternalStreamingObject = CreateExternalStreamingObject(GetExternalStreamingObjectClass(), GetOuterUWorldPartition(), InStreamingObjectName, GetTypedOuter<UWorld>());
+	URuntimeHashExternalStreamingObjectBase* NewExternalStreamingObject = CreateExternalStreamingObject(GetExternalStreamingObjectClass(), GetOuterUWorldPartition(), GetTypedOuter<UWorld>());
 	StoreStreamingContentToExternalStreamingObject(NewExternalStreamingObject);
 	return NewExternalStreamingObject;
 }
diff --git a/Engine/Source/Runtime/Engine/Private/WorldPartition/WorldPartitionStreamingGeneration.cpp b/Engine/Source/Runtime/Engine/Private/WorldPartition/WorldPartitionStreamingGeneration.cpp
--- a/Engine/Source/Runtime/Engine/Private/WorldPartition/WorldPartitionStreamingGeneration.cpp
+++ b/Engine/Source/Runtime/Engine/Private/WorldPartition/WorldPartitionStreamingGeneration.cpp
@@ -2020,9 +2020,9 @@ bool UWorldPartition::HasStreamingContent() const
 	return RuntimeHash && RuntimeHash->HasStreamingContent();
 }
 
-URuntimeHashExternalStreamingObjectBase* UWorldPartition::FlushStreamingToExternalStreamingObject(const FString& ExternalStreamingObjectName)
+URuntimeHashExternalStreamingObjectBase* UWorldPartition::FlushStreamingToExternalStreamingObject()
 {
-	URuntimeHashExternalStreamingObjectBase* ExternalStreamingObject = RuntimeHash->StoreStreamingContentToExternalStreamingObject(*ExternalStreamingObjectName);
+	URuntimeHashExternalStreamingObjectBase* ExternalStreamingObject = RuntimeHash->StoreStreamingContentToExternalStreamingObject();
 	check(ExternalStreamingObject);
 
 	StreamingPolicy->StoreStreamingContentToExternalStreamingObject(*ExternalStreamingObject);
diff --git a/Engine/Source/Runtime/Engine/Public/WorldPartition/ContentBundle/ContentBundleBase.h b/Engine/Source/Runtime/Engine/Public/WorldPartition/ContentBundle/ContentBundleBase.h
--- a/Engine/Source/Runtime/Engine/Public/WorldPartition/ContentBundle/ContentBundleBase.h
+++ b/Engine/Source/Runtime/Engine/Public/WorldPartition/ContentBundle/ContentBundleBase.h
@@ -34,7 +34,6 @@ public:
 	EContentBundleStatus GetStatus() const { return Status; }
 	ENGINE_API FString GetExternalStreamingObjectPackageName() const;
 	ENGINE_API FString GetExternalStreamingObjectPackagePath() const;
-	ENGINE_API FString GetExternalStreamingObjectName() const;
 
 protected:
 	ENGINE_API void SetStatus(EContentBundleStatus NewStatus);
diff --git a/Engine/Source/Runtime/Engine/Public/WorldPartition/DataLayer/ExternalDataLayerHelper.h b/Engine/Source/Runtime/Engine/Public/WorldPartition/DataLayer/ExternalDataLayerHelper.h
--- a/Engine/Source/Runtime/Engine/Public/WorldPartition/DataLayer/ExternalDataLayerHelper.h
+++ b/Engine/Source/Runtime/Engine/Public/WorldPartition/DataLayer/ExternalDataLayerHelper.h
@@ -18,9 +18,6 @@ public:
 	/** Returns the external streaming object package name */
 	static FString GetExternalStreamingObjectPackageName(const UExternalDataLayerAsset* InExternalDataLayerAsset);
 
-	/** Returns the external streaming object name */
-	static FString GetExternalStreamingObjectName(const UExternalDataLayerAsset* InExternalDataLayerAsset);
-
 	/** Return true if succeeds building the external data layer root path (OutExternalDataLayerRootPath) using the provided mount point and EDL UID.
 	 * Format is /{MountPoint}/{ExternalDataLayerFolder}/{EDL_UID}
 	 */
diff --git a/Engine/Source/Runtime/Engine/Public/WorldPartition/WorldPartition.h b/Engine/Source/Runtime/Engine/Public/WorldPartition/WorldPartition.h
--- a/Engine/Source/Runtime/Engine/Public/WorldPartition/WorldPartition.h
+++ b/Engine/Source/Runtime/Engine/Public/WorldPartition/WorldPartition.h
@@ -269,7 +269,10 @@ public:
 	ENGINE_API TUniquePtr<IStreamingGenerationContext> GenerateStreamingGenerationContext(const FGenerateStreamingParams& InParams, FGenerateStreamingContext& InContext);
 
 	ENGINE_API void FlushStreaming();
-	ENGINE_API URuntimeHashExternalStreamingObjectBase* FlushStreamingToExternalStreamingObject(const FString& ExternalStreamingObjectName);
+	ENGINE_API URuntimeHashExternalStreamingObjectBase* FlushStreamingToExternalStreamingObject();
+
+	UE_DEPRECATED(5.5, "FlushStreamingToExternalStreamingObject(const FString& ) is deprecated, use FlushStreamingToExternalStreamingObject() instead")
+	URuntimeHashExternalStreamingObjectBase* FlushStreamingToExternalStreamingObject(const FString& ExternalStreamingObjectName) { return FlushStreamingToExternalStreamingObject(); }
 
 	// Event when world partition was enabled/disabled in the world
 	DECLARE_MULTICAST_DELEGATE_OneParam(FWorldPartitionChangedEvent, UWorld*);
diff --git a/Engine/Source/Runtime/Engine/Public/WorldPartition/WorldPartitionRuntimeHash.h b/Engine/Source/Runtime/Engine/Public/WorldPartition/WorldPartitionRuntimeHash.h
--- a/Engine/Source/Runtime/Engine/Public/WorldPartition/WorldPartitionRuntimeHash.h
+++ b/Engine/Source/Runtime/Engine/Public/WorldPartition/WorldPartitionRuntimeHash.h
@@ -164,7 +164,11 @@ class UWorldPartitionRuntimeHash : public UObject
 	virtual void DrawPreview() const {}
 
 	ENGINE_API virtual bool HasStreamingContent() const { return false; }
-	ENGINE_API URuntimeHashExternalStreamingObjectBase* StoreStreamingContentToExternalStreamingObject(FName InStreamingObjectName);
+	ENGINE_API URuntimeHashExternalStreamingObjectBase* StoreStreamingContentToExternalStreamingObject();
+
+	UE_DEPRECATED(5.5, "StoreStreamingContentToExternalStreamingObject(FName) is deprecated, use StoreStreamingContentToExternalStreamingObject() instead")
+	URuntimeHashExternalStreamingObjectBase* StoreStreamingContentToExternalStreamingObject(FName InStreamingObjectName) { return StoreStreamingContentToExternalStreamingObject(); }
+
 	ENGINE_API virtual void FlushStreamingContent();
 	ENGINE_API virtual TSubclassOf<URuntimeHashExternalStreamingObjectBase> GetExternalStreamingObjectClass() const PURE_VIRTUAL(UWorldPartitionRuntimeHash::GetExternalStreamingObjectClass, return nullptr;);
 	ENGINE_API virtual void DumpStateLog(FHierarchicalLogArchive& Ar) const;
@@ -232,7 +236,7 @@ public:
 	virtual bool GetShouldMergeStreamingSourceInfo() const { return false; }
 
 protected:
-	static ENGINE_API URuntimeHashExternalStreamingObjectBase* CreateExternalStreamingObject(TSubclassOf<URuntimeHashExternalStreamingObjectBase> InClass, UObject* InOuter, FName InName, UWorld* InOuterWorld);
+	static ENGINE_API URuntimeHashExternalStreamingObjectBase* CreateExternalStreamingObject(TSubclassOf<URuntimeHashExternalStreamingObjectBase> InClass, UObject* InOuter, UWorld* InOuterWorld);
 	ENGINE_API UWorldPartitionRuntimeCell* CreateRuntimeCell(UClass* CellClass, UClass* CellDataClass, const FString& CellName, const FString& CellInstanceSuffix, UObject* InOuter = nullptr);
 	virtual EWorldPartitionStreamingPerformance GetStreamingPerformanceForCell(const UWorldPartitionRuntimeCell* Cell) const;
 
