Fixed PIE crash when launching a 2nd client in a separate process.

diff --git a/Engine/Source/Runtime/Engine/Private/WorldPartition/DataLayer/WorldDataLayers.cpp b/Engine/Source/Runtime/Engine/Private/WorldPartition/DataLayer/WorldDataLayers.cpp
--- a/Engine/Source/Runtime/Engine/Private/WorldPartition/DataLayer/WorldDataLayers.cpp
+++ b/Engine/Source/Runtime/Engine/Private/WorldPartition/DataLayer/WorldDataLayers.cpp
@@ -1363,7 +1363,7 @@ bool AWorldDataLayers::CanReferenceDataLayerAsset(const UDataLayerAsset* InDataL
 	{
 		FAssetReferenceFilterContext AssetReferenceFilterContext;
 		AssetReferenceFilterContext.ReferencingAssets.Add(FAssetData(InReferencingObject));
-		TSharedPtr<IAssetReferenceFilter> AssetReferenceFilter = GEditor->MakeAssetReferenceFilter(AssetReferenceFilterContext);
+		TSharedPtr<IAssetReferenceFilter> AssetReferenceFilter = GEditor ? GEditor->MakeAssetReferenceFilter(AssetReferenceFilterContext) : nullptr;
 		return AssetReferenceFilter.IsValid() ? AssetReferenceFilter->PassesFilter(FAssetData(InDataLayerAsset), OutReason) : true;
 	};
 
