Ensure new FMinimalViewInfo entries are initialized to avoid non-deterministic issues in Ortho cameras.

diff --git a/Engine/Source/Runtime/Engine/Classes/Camera/CameraTypes.h b/Engine/Source/Runtime/Engine/Classes/Camera/CameraTypes.h
--- a/Engine/Source/Runtime/Engine/Classes/Camera/CameraTypes.h
+++ b/Engine/Source/Runtime/Engine/Classes/Camera/CameraTypes.h
@@ -132,6 +132,9 @@ public:
 		, DesiredFOV(90.0f)
 		, OrthoWidth(512.0f)
 		, bAutoCalculateOrthoPlanes(true)
+		, AutoPlaneShift(0.0f)
+		, bUpdateOrthoPlanes(false)
+		, bUseCameraHeightAsViewTarget(false)
 		, OrthoNearClipPlane(0.0f)
 		, OrthoFarClipPlane(UE_OLD_WORLD_MAX)
 		, PerspectiveNearClipPlane(-1.0f)
