UTexture::ForceRebuildPlatformData will now build Final when passed an encode  speed of FinalIfAvailable, rather than fallingback to Fast.

This fixes an issue where failed fetches for FinalIfAvailable would enqueue deferred forcerebuilds of textures that will be rebuilt as Fast encoding. Subsequent runs of the editor will still miss on the same texture using Final encoding, and again kick off a rebuild for the already built Fast texture (since ForceRebuildPlatformData doesn't fetch results from DDC).

diff --git a/Engine/Source/Runtime/Engine/Private/TextureDerivedData.cpp b/Engine/Source/Runtime/Engine/Private/TextureDerivedData.cpp
--- a/Engine/Source/Runtime/Engine/Private/TextureDerivedData.cpp
+++ b/Engine/Source/Runtime/Engine/Private/TextureDerivedData.cpp
@@ -3866,28 +3866,23 @@ void UTexture::ForceRebuildPlatformData(uint8 InEncodeSpeedOverride /* =255 ETex
 			EncodeSpeed = GetDesiredEncodeSpeed();
 		}
 
-		TArray<FTextureBuildSettings> BuildSettingsFetch;
-		TArray<FTextureBuildSettings> BuildSettingsFetchOrBuild;
-		TArray<FTexturePlatformData::FTextureEncodeResultMetadata> ResultMetadataFetch;
-		TArray<FTexturePlatformData::FTextureEncodeResultMetadata> ResultMetadataFetchOrBuild;
-
+		// Since we are forcing a rebuild, build what is desired rather than what is available
 		if (EncodeSpeed == ETextureEncodeSpeed::FinalIfAvailable)
 		{
-			GetBuildSettingsForRunningPlatform(*this, ETextureEncodeSpeed::Final, BuildSettingsFetch, &ResultMetadataFetch);
-			GetBuildSettingsForRunningPlatform(*this, ETextureEncodeSpeed::Fast, BuildSettingsFetchOrBuild, &ResultMetadataFetchOrBuild);
-		}
-		else
-		{
-			GetBuildSettingsForRunningPlatform(*this, EncodeSpeed, BuildSettingsFetchOrBuild, &ResultMetadataFetchOrBuild);
+			EncodeSpeed = ETextureEncodeSpeed::Final;
 		}
+
+		TArray<FTextureBuildSettings> BuildSettingsFetchOrBuild;
+		TArray<FTexturePlatformData::FTextureEncodeResultMetadata> ResultMetadataFetchOrBuild;
+		GetBuildSettingsForRunningPlatform(*this, EncodeSpeed, BuildSettingsFetchOrBuild, &ResultMetadataFetchOrBuild);
 		
 		check(BuildSettingsFetchOrBuild.Num() == Source.GetNumLayers());
 
 		PlatformDataLink->Cache(
 			*this,
-			BuildSettingsFetch.GetData(),
+			nullptr,
 			BuildSettingsFetchOrBuild.GetData(),
-			ResultMetadataFetch.GetData(),
+			nullptr,
 			ResultMetadataFetchOrBuild.GetData(),
 			uint32(ETextureCacheFlags::ForceRebuild),
 			nullptr
