Add env. preconvolution for Substrate's specular profile.

diff --git a/Engine/Shaders/Private/SpecularProfile.usf b/Engine/Shaders/Private/SpecularProfile.usf
--- a/Engine/Shaders/Private/SpecularProfile.usf
+++ b/Engine/Shaders/Private/SpecularProfile.usf
@@ -1,6 +1,12 @@
 // Copyright Epic Games, Inc. All Rights Reserved.
 
 #include "Common.ush"
+#include "/Engine/Shared/SubstrateDefinitions.h"
+#include "SpecularProfileCommon.ush"
+
+///////////////////////////////////////////////////////////////////////////////////////////////////
+
+#ifdef SHADER_COPY
 
 int2 SourceResolution;
 uint SourceMipCount;
@@ -24,14 +30,94 @@ void MainCS(uint3 DispatchThreadId : SV_DispatchThreadID)
 		{
 			const float3 C0 = ViewColorBuffer[PixelPos.x].xyz;
 			const float3 C1 = LightColorBuffer[PixelPos.y].xyz;
-			TargetTexture[uint3(PixelPos,TargetIndex)] = float4(C0 * C1, 1);
+			TargetTexture[uint3(PixelPos,TargetIndex * SUBSTRATE_SPECULAR_PROFILE_ENTRY_COUNT + SUBSTRATE_SPECULAR_PROFILE_ENTRY_LIGHT)] = float4(C0 * C1, 1);
 		}
 		#else
 		{
 			const float SourceMip = clamp(log2(SourceResolution.x) - log2(TargetResolution.x), 0.f, SourceMipCount-1);
 			const float4 Value = SourceTexture.SampleLevel(SourceSampler, TexCoord, SourceMip);
-			TargetTexture[uint3(PixelPos,TargetIndex)] = Value;
+			TargetTexture[uint3(PixelPos,TargetIndex * SUBSTRATE_SPECULAR_PROFILE_ENTRY_COUNT + SUBSTRATE_SPECULAR_PROFILE_ENTRY_LIGHT)] = Value;
 		}
 		#endif
 	}
-}
\ No newline at end of file
+}
+#endif
+
+///////////////////////////////////////////////////////////////////////////////////////////////////
+
+#ifdef SHADER_CONVOLVE
+
+#include "MonteCarlo.ush"
+#include "BRDF.ush"
+
+uint SampleCount;
+int2 Resolution;
+uint ProfileCount;
+SamplerState TargetSampler;
+RWTexture2DArray<float4> TargetTexture;
+StructuredBuffer<float> ProfileSignBuffer;
+
+float3 EvaluateSpecularProfile(float NoV, float NoL, float VoH, float NoH, float SpecularProfileId)
+{
+	const float3 ProfileCoord = GetSpecularProfileCoord(SpecularProfileId, NoV, NoL, VoH, NoH);
+	uint3 FetchCoord = uint3(ProfileCoord.xy * (Resolution - 1), ProfileCoord.z);
+	FetchCoord.x = clamp(FetchCoord.x, 0, Resolution.x - 1);
+	FetchCoord.y = clamp(FetchCoord.y, 0, Resolution.y - 1);
+	return TargetTexture[FetchCoord].xyz;
+}
+
+float3 Convolve(float3 V, float3 N, float Roughness, float SpecularProfileId)
+{
+	float3 OutValue = 0;
+	float OutWeight = 0;
+
+	LOOP
+	for (uint SampleIt = 0; SampleIt < SampleCount; ++SampleIt)
+	{
+		const float2 E = Hammersley(SampleIt, SampleCount, 0);
+
+		const float4 HAndPDF = ImportanceSampleVisibleGGX(E, Pow2(Roughness), V);
+		const float3 H = HAndPDF.xyz;
+		const float PDF = HAndPDF.w;
+
+		// Sample generation
+		const float3 L = 2 * dot(V, H) * H - V;
+		if (L.z > 0)
+		{
+			const float NoV = saturate(dot(N, V));
+			const float NoL = saturate(dot(N, L));
+			const float NoH = saturate(dot(N, H));
+			const float VoH = saturate(dot(V, H));
+
+			if (PDF > 0)
+			{
+				const float3 Value = EvaluateSpecularProfile(NoV, NoL, VoH, NoH, SpecularProfileId);
+				OutValue += Value;
+				OutWeight+= 1;
+			}
+		}
+	}
+	OutValue /= max(0.001f, OutWeight);
+	return OutValue;
+}
+
+[numthreads(8, 8, 1)]
+void MainCS(uint3 DispatchThreadId : SV_DispatchThreadID)
+{
+	const uint2 PixelPos = DispatchThreadId.xy;
+	const uint ProfileIndex = DispatchThreadId.z;
+	if (all(PixelPos < uint2(Resolution)) && ProfileIndex < ProfileCount)
+	{
+		const float NoV = PixelPos.x / float(Resolution.x-1);
+		const float Roughness = max(PixelPos.y / float(Resolution.y-1), 0.02f);
+		const float SpecularProfileId = ProfileSignBuffer[ProfileIndex] * ProfileIndex; // 	The sign encodes the parameterization (NoV / NoH)
+
+		const float3 N = float3(0, 0, 1);
+		const float3 V = float3(sqrt(1-NoV*NoV), 0, NoV);
+		const float3 ConvolvedColor = Convolve(V, N, Roughness, SpecularProfileId);
+
+		const uint3 OutCoord = uint3(PixelPos, ProfileIndex * SUBSTRATE_SPECULAR_PROFILE_ENTRY_COUNT + SUBSTRATE_SPECULAR_PROFILE_ENTRY_ENV);
+		TargetTexture[OutCoord] = float4(ConvolvedColor, 1);		
+	}
+}
+#endif
diff --git a/Engine/Shaders/Private/SpecularProfileCommon.ush b/Engine/Shaders/Private/SpecularProfileCommon.ush
new file mode 100644
--- /dev/null
+++ b/Engine/Shaders/Private/SpecularProfileCommon.ush
@@ -0,0 +1,34 @@
+// Copyright Epic Games, Inc. All Rights Reserved.
+
+#pragma once
+
+uint GetSpecularProfileId(float In)
+{
+	return abs(In);
+}
+
+// Return 0: View/Light angles, 1: Half angles
+uint GetSpecularProfileParameterization(float In)
+{
+	return In > 0.f ? 0u : 1u;
+}
+
+float2 InternalGetSpecularProfileCoord(float InSpecularProfileId, float NoV, float NoL, float VoH, float NoH)
+{
+	const bool bUseHalfVector = GetSpecularProfileParameterization(InSpecularProfileId) == 1;
+	return bUseHalfVector ? float2(VoH, NoH) : float2(NoV, NoL);
+}
+
+float3 GetSpecularProfileCoord(float InSpecularProfileId, float NoV, float NoL, float VoH, float NoH)
+{
+	const uint ProfileID = GetSpecularProfileId(InSpecularProfileId);
+	const uint SliceIndex = ProfileID * SUBSTRATE_SPECULAR_PROFILE_ENTRY_COUNT + SUBSTRATE_SPECULAR_PROFILE_ENTRY_LIGHT;
+	return float3(InternalGetSpecularProfileCoord(NoV, NoL, VoH, NoH, InSpecularProfileId), SliceIndex);
+}
+
+float3 GetSpecularProfileEnvCoord(float InSpecularProfileId, float NoV, float InRoughness)
+{
+	const uint ProfileID = GetSpecularProfileId(InSpecularProfileId);
+	const uint SliceIndex = ProfileID * SUBSTRATE_SPECULAR_PROFILE_ENTRY_COUNT + SUBSTRATE_SPECULAR_PROFILE_ENTRY_ENV;
+	return float3(NoV, InRoughness, SliceIndex);
+}
\ No newline at end of file
diff --git a/Engine/Shaders/Private/Substrate/Substrate.ush b/Engine/Shaders/Private/Substrate/Substrate.ush
--- a/Engine/Shaders/Private/Substrate/Substrate.ush
+++ b/Engine/Shaders/Private/Substrate/Substrate.ush
@@ -25,6 +25,7 @@
 #include "/Engine/Shared/SubstrateDefinitions.h"
 #include "/Engine/Private/Substrate/SubstrateStatisticalOperators.ush"
 #include "/Engine/Private/Substrate/SubstrateSubsurface.ush"
+#include "../SpecularProfileCommon.ush"
 
 // DXC has issue with unrolling certain part of the code. To address temporarly the warning-as-error issues, 
 // we redefine the unroll function and fall back to loop on DXC
@@ -871,22 +872,15 @@ void UnpackSpecularProfile(uint In, inout float OutProfileId)
 	OutProfileId = asfloat(In);
 }
 
-uint GetSpecularProfileId(float In)
-{
-	return abs(In);
-}
-
-// Return 0: View/Light angles, 1: Half angles
-uint GetSpecularProfileParameterization(float In)
+float3 EvaluateSpecularProfile(float InSpecularProfileId, float NoV, float NoL, float VoH, float NoH)
 {
-	return In > 0.f ? 0u : 1u;
+	const float3 ProfileCoord = GetSpecularProfileCoord(InSpecularProfileId, NoV, NoL, VoH, NoH);
+	return View.SpecularProfileTexture.SampleLevel(View.SpecularProfileSampler, ProfileCoord, 0).xyz;
 }
 
-float3 EvaluateSpecularProfile(float InSpecularProfileId, float NoV, float NoL, float VoH, float NoH)
+float3 EvaluateSpecularProfileEnv(float InSpecularProfileId, float NoV, float InRoughness)
 {
-	const uint ProfileID = GetSpecularProfileId(InSpecularProfileId);
-	const bool bUseHalfVector = GetSpecularProfileParameterization(InSpecularProfileId) == 1;
-	const float3 ProfileCoord = bUseHalfVector ? float3(VoH, NoH, ProfileID) : float3(NoV, NoL, ProfileID);
+	const float3 ProfileCoord = GetSpecularProfileEnvCoord(InSpecularProfileId, NoV, InRoughness);
 	return View.SpecularProfileTexture.SampleLevel(View.SpecularProfileSampler, ProfileCoord, 0).xyz;
 }
 
diff --git a/Engine/Shaders/Private/Substrate/SubstrateEvaluation.ush b/Engine/Shaders/Private/Substrate/SubstrateEvaluation.ush
--- a/Engine/Shaders/Private/Substrate/SubstrateEvaluation.ush
+++ b/Engine/Shaders/Private/Substrate/SubstrateEvaluation.ush
@@ -1568,8 +1568,7 @@ FSubstrateEnvLightResult SubstrateEvaluateForEnvLight(FSubstrateBSDFContext BSDF
 				BRANCH
 				if (BSDF_GETHASSPECPROFILE(BSDFContext.BSDF))
 				{
-					const float NoR = saturate(dot(SubstrateEnvLightResult.SpecularDirection, BSDFContext.N));
-					SubstrateEnvLightResult.SpecularWeight *= EvaluateSpecularProfile(SLAB_SPECPROFILEID(BSDFContext.BSDF), BSDFContext.SatNoV, NoR, BSDFContext.SatNoV, NoR);
+					SubstrateEnvLightResult.SpecularWeight *= EvaluateSpecularProfileEnv(SLAB_SPECPROFILEID(BSDFContext.BSDF), BSDFContext.SatNoV, SubstrateEnvLightResult.SpecularSafeRoughness);
 				}
 			#endif // SUBSTRATE_COMPLEXSPECIALPATH
 
@@ -1632,8 +1631,7 @@ FSubstrateEnvLightResult SubstrateEvaluateForEnvLight(FSubstrateBSDFContext BSDF
 					BRANCH
 					if (BSDF_GETHASSPECPROFILE(BSDFContext.BSDF))
 					{
-						const float NoR = saturate(dot(SubstrateEnvLightResult.SpecularHazeDirection, BSDFContext.N));
-						SubstrateEnvLightResult.SpecularHazeWeight *= EvaluateSpecularProfile(SLAB_SPECPROFILEID(BSDFContext.BSDF), BSDFContext.SatNoV, NoR, BSDFContext.SatNoV, NoR);
+						SubstrateEnvLightResult.SpecularHazeWeight *= EvaluateSpecularProfileEnv(SLAB_SPECPROFILEID(BSDFContext.BSDF), BSDFContext.SatNoV, SubstrateEnvLightResult.SpecularHazeSafeRoughness);
 					}
 				#endif // SUBSTRATE_COMPLEXSPECIALPATH
 
diff --git a/Engine/Shaders/Shared/SubstrateDefinitions.h b/Engine/Shaders/Shared/SubstrateDefinitions.h
--- a/Engine/Shaders/Shared/SubstrateDefinitions.h
+++ b/Engine/Shaders/Shared/SubstrateDefinitions.h
@@ -79,6 +79,11 @@
 #define SUBSTRATE_TILE_ENCODING_16BITS 					0
 #define SUBSTRATE_TILE_ENCODING_8BITS  					1
 
+// Specular profile
+#define SUBSTRATE_SPECULAR_PROFILE_ENTRY_LIGHT			0u
+#define SUBSTRATE_SPECULAR_PROFILE_ENTRY_ENV			1u
+#define SUBSTRATE_SPECULAR_PROFILE_ENTRY_COUNT			2u
+
 // sizeof(FRHIDrawIndirectParameters) = 4 uints = 16 bytes
 #define GetSubstrateTileTypeDrawIndirectArgOffset_Byte(x)  (x * 16)
 #define GetSubstrateTileTypeDrawIndirectArgOffset_DWord(x) (x * 4)
diff --git a/Engine/Source/Runtime/Engine/Private/Rendering/SpecularProfile.cpp b/Engine/Source/Runtime/Engine/Private/Rendering/SpecularProfile.cpp
--- a/Engine/Source/Runtime/Engine/Private/Rendering/SpecularProfile.cpp
+++ b/Engine/Source/Runtime/Engine/Private/Rendering/SpecularProfile.cpp
@@ -12,6 +12,8 @@
 #include "Rendering/Texture2DResource.h"
 #include "RenderGraphResources.h"
 #include "RenderGraphBuilder.h"
+#include "SubstrateDefinitions.h"
+#include "ShaderCompilerCore.h"
 
 #include UE_INLINE_GENERATED_CPP_BY_NAME(SpecularProfile)
 
@@ -71,10 +73,48 @@ public:
 	{
 		return true;
 	}
+	static void ModifyCompilationEnvironment(const FGlobalShaderPermutationParameters& Parameters, FShaderCompilerEnvironment& OutEnvironment)
+	{
+		FGlobalShader::ModifyCompilationEnvironment(Parameters, OutEnvironment);
+		OutEnvironment.SetDefine(TEXT("SHADER_COPY"), 1);
+	}
 };
 
 IMPLEMENT_GLOBAL_SHADER(FSpecularProfileCopyCS, "/Engine/Private/SpecularProfile.usf", "MainCS", SF_Compute);
 
+///////////////////////////////////////////////////////////////////////////////////////////////////
+
+class FSpecularProfileConvolveCS : public FGlobalShader
+{
+public:
+	DECLARE_GLOBAL_SHADER(FSpecularProfileConvolveCS);
+	SHADER_USE_PARAMETER_STRUCT(FSpecularProfileConvolveCS, FGlobalShader);
+
+	using FPermutationDomain = TShaderPermutationDomain<>;
+
+	BEGIN_SHADER_PARAMETER_STRUCT(FParameters, )
+		SHADER_PARAMETER(uint32, SampleCount)
+		SHADER_PARAMETER(FIntPoint, Resolution)
+		SHADER_PARAMETER(uint32, ProfileCount)
+		SHADER_PARAMETER_SAMPLER(SamplerState, TargetSampler)
+		SHADER_PARAMETER_RDG_TEXTURE_UAV(RWTextureArray2D, TargetTexture)
+		SHADER_PARAMETER_RDG_BUFFER_SRV(StructuredBuffer, ProfileSignBuffer)
+	END_SHADER_PARAMETER_STRUCT()
+
+	static bool ShouldCompilePermutation(const FGlobalShaderPermutationParameters& Parameters)
+	{
+		return true;
+	}
+	static void ModifyCompilationEnvironment(const FGlobalShaderPermutationParameters& Parameters, FShaderCompilerEnvironment& OutEnvironment)
+	{
+		FGlobalShader::ModifyCompilationEnvironment(Parameters, OutEnvironment);
+		OutEnvironment.SetDefine(TEXT("SHADER_CONVOLVE"), 1);
+		OutEnvironment.CompilerFlags.Add(CFLAG_AllowTypedUAVLoads);		
+	}
+};
+
+IMPLEMENT_GLOBAL_SHADER(FSpecularProfileConvolveCS, "/Engine/Private/SpecularProfile.usf", "MainCS", SF_Compute);
+
 ///////////////////////////////////////////////////////////////////////////////////////////////////
 // FSpecularProfileTextureManager
 
@@ -323,14 +363,16 @@ IPooledRenderTarget* FSpecularProfileTextureManager::GetAtlasTexture(FRDGBuilder
 		const uint32 Resolution = CVarSpecularProfileResolution.GetValueOnRenderThread();
 		check(LayerCount);
 
+		RDG_EVENT_SCOPE(GraphBuilder, "SpecularProfile");
+
 		// 1. Create atlas texture
-		FPooledRenderTargetDesc Desc(FPooledRenderTargetDesc::Create2DArrayDesc(FIntPoint(Resolution, Resolution), PF_FloatR11G11B10, FClearValueBinding::None, TexCreate_None, TexCreate_UAV | TexCreate_ShaderResource, false, LayerCount));
+		FPooledRenderTargetDesc Desc(FPooledRenderTargetDesc::Create2DArrayDesc(FIntPoint(Resolution, Resolution), PF_FloatR11G11B10, FClearValueBinding::None, TexCreate_None, TexCreate_UAV | TexCreate_ShaderResource, false, LayerCount * SUBSTRATE_SPECULAR_PROFILE_ENTRY_COUNT));
 		GRenderTargetPool.FindFreeElement(RHICmdList, Desc, GSpecularProfileTextureAtlas, TEXT("SpecularProfileTexture"));
 
 		// 2. Fill in profiles
 		const auto GlobalShaderMap = GetGlobalShaderMap(ShaderPlatform);
 		FRDGTextureRef SpecularProfileTexture = GraphBuilder.RegisterExternalTexture(GSpecularProfileTextureAtlas, TEXT("SpecularProfileTexture"));
-		FRDGTextureUAVRef SpecularProfileUAV = GraphBuilder.CreateUAV(SpecularProfileTexture);
+		FRDGTextureUAVRef SpecularProfileUAVSkipBarrier = GraphBuilder.CreateUAV(SpecularProfileTexture, ERDGUnorderedAccessViewFlags::SkipBarrier);
 		for (uint32 LayerIt = 0; LayerIt < LayerCount; ++LayerIt)
 		{
 			const FSpecularProfileStruct Data = SpecularProfileEntries[LayerIt].Settings;
@@ -342,7 +384,7 @@ IPooledRenderTarget* FSpecularProfileTextureManager::GetAtlasTexture(FRDGBuilder
 			PassParameters->SourceSampler = TStaticSamplerState<SF_Point, AM_Clamp, AM_Clamp, AM_Clamp>::GetRHI();
 			PassParameters->SourceResolution = SpecularProfileTexture->Desc.Extent;
 			PassParameters->TargetResolution = SpecularProfileTexture->Desc.Extent;
-			PassParameters->TargetTexture = SpecularProfileUAV;
+			PassParameters->TargetTexture = SpecularProfileUAVSkipBarrier;
 			PassParameters->TargetIndex = LayerIt;
 
 			struct FLinearColor16 
@@ -377,14 +419,42 @@ IPooledRenderTarget* FSpecularProfileTextureManager::GetAtlasTexture(FRDGBuilder
 			PermutationVector.Set<FSpecularProfileCopyCS::FProcedural>(Data.IsProcedural());
 			TShaderMapRef<FSpecularProfileCopyCS> Shader(GlobalShaderMap, PermutationVector);
 			FComputeShaderUtils::AddPass(GraphBuilder, 
-										RDG_EVENT_NAME("SpecularProfile::CopyTexture"), 
-										Shader, 
-										PassParameters, 
-										FIntVector(FMath::DivideAndRoundUp(PassParameters->TargetResolution.X, 8), FMath::DivideAndRoundUp(PassParameters->TargetResolution.Y, 8), 1));
+				RDG_EVENT_NAME("SpecularProfile::CopyTexture"), 
+				Shader, 
+				PassParameters, 
+				FIntVector(FMath::DivideAndRoundUp(PassParameters->TargetResolution.X, 8), FMath::DivideAndRoundUp(PassParameters->TargetResolution.Y, 8), 1));
 
 			SpecularProfileEntries[LayerIt].CachedResolution = Texture ? Texture->TextureReferenceRHI->GetDesc().Extent : FIntPoint(0,0);
 		}
 
+		// Convolve all the profiles
+		if (LayerCount > 0)
+		{
+			TArray<float> ProfileSign;
+			ProfileSign.SetNum(LayerCount);
+			for (uint32 LayerIt=0; LayerIt < LayerCount; ++LayerIt)
+			{
+				ProfileSign[LayerIt] = GetParameterizationSign(LayerIt);
+			}
+			FRDGBufferRef ProfileSignBuffer = GraphBuilder.CreateBuffer(FRDGBufferDesc::CreateStructuredDesc(sizeof(float), LayerCount), TEXT("SpecularProfileData"), ERDGBufferFlags::MultiFrame);
+			GraphBuilder.QueueBufferUpload(ProfileSignBuffer, ProfileSign.GetData(), sizeof(float)*ProfileSign.Num(), ERDGInitialDataFlags::None);
+
+			FSpecularProfileConvolveCS::FParameters* PassParameters = GraphBuilder.AllocParameters<FSpecularProfileConvolveCS::FParameters>();
+			PassParameters->SampleCount = 64u;
+			PassParameters->ProfileCount = LayerCount;
+			PassParameters->Resolution = SpecularProfileTexture->Desc.Extent;
+			PassParameters->ProfileSignBuffer = GraphBuilder.CreateSRV(ProfileSignBuffer);
+			PassParameters->TargetSampler = TStaticSamplerState<SF_Point, AM_Clamp, AM_Clamp, AM_Clamp>::GetRHI();;
+			PassParameters->TargetTexture = GraphBuilder.CreateUAV(SpecularProfileTexture);
+
+			TShaderMapRef<FSpecularProfileConvolveCS> Shader(GlobalShaderMap);
+			FComputeShaderUtils::AddPass(GraphBuilder, 
+				RDG_EVENT_NAME("SpecularProfile::Convolve"), 
+				Shader, 
+				PassParameters, 
+				FIntVector(FMath::DivideAndRoundUp(PassParameters->Resolution.X, 8), FMath::DivideAndRoundUp(PassParameters->Resolution.Y, 8), LayerCount));
+		}
+
 		// Transit texture to SRV for letting the non-RDG resource to be bound later by the various passes
 		FRDGExternalAccessQueue ExternalAccessQueue;
 		ExternalAccessQueue.Add(SpecularProfileTexture);
