Fix CIS warning about shader parameters type

diff --git a/Engine/Source/Runtime/Engine/Private/Rendering/SpecularProfile.cpp b/Engine/Source/Runtime/Engine/Private/Rendering/SpecularProfile.cpp
--- a/Engine/Source/Runtime/Engine/Private/Rendering/SpecularProfile.cpp
+++ b/Engine/Source/Runtime/Engine/Private/Rendering/SpecularProfile.cpp
@@ -97,7 +97,7 @@ public:
 		SHADER_PARAMETER(FIntPoint, Resolution)
 		SHADER_PARAMETER(uint32, ProfileCount)
 		SHADER_PARAMETER_SAMPLER(SamplerState, TargetSampler)
-		SHADER_PARAMETER_RDG_TEXTURE_UAV(RWTextureArray2D, TargetTexture)
+		SHADER_PARAMETER_RDG_TEXTURE_UAV(RWTexture2DArray, TargetTexture)
 		SHADER_PARAMETER_RDG_BUFFER_SRV(StructuredBuffer, ProfileSignBuffer)
 	END_SHADER_PARAMETER_STRUCT()
 
