[NetPhysics] Replicated AngVel is in radians for Predictive Interpolation and Resimulation while in Degrees for Default replication.
Fixup to adhere to this in replication, before bringing all replication into using the same value.

diff --git a/Engine/Source/Runtime/Engine/Private/PhysicsEngine/PhysicsReplication.cpp b/Engine/Source/Runtime/Engine/Private/PhysicsEngine/PhysicsReplication.cpp
--- a/Engine/Source/Runtime/Engine/Private/PhysicsEngine/PhysicsReplication.cpp
+++ b/Engine/Source/Runtime/Engine/Private/PhysicsEngine/PhysicsReplication.cpp
@@ -168,7 +168,7 @@ namespace PhysicsReplicationCVars
 		static FAutoConsoleVariableRef CVarErrorAccLinVelMaxLimit(TEXT("np2.PredictiveInterpolation.ErrorAccLinVelMaxLimit"), ErrorAccLinVelMaxLimit, TEXT("If target velocity is below this limit we check for desync to trigger softsnap and accumulate time to build up to a hardsnap."));
 		
 		float ErrorAccAngVelMaxLimit = 1.5f;
-		static FAutoConsoleVariableRef CVarErrorAccAngVelMaxLimit(TEXT("np2.PredictiveInterpolation.ErrorAccAngVelMaxLimit"), ErrorAccAngVelMaxLimit, TEXT("If target angular velocity (in degrees) is below this limit we check for desync to trigger softsnap and accumulate time to build up to a hardsnap."));
+		static FAutoConsoleVariableRef CVarErrorAccAngVelMaxLimit(TEXT("np2.PredictiveInterpolation.ErrorAccAngVelMaxLimit"), ErrorAccAngVelMaxLimit, TEXT("If target angular velocity (in radians) is below this limit we check for desync to trigger softsnap and accumulate time to build up to a hardsnap."));
 		
 		float SoftSnapPosStrength = 0.5f;
 		static FAutoConsoleVariableRef CVarSoftSnapPosStrength(TEXT("np2.PredictiveInterpolation.SoftSnapPosStrength"), SoftSnapPosStrength, TEXT("Value in percent between 0.0 - 1.0 representing how much to softsnap each tick of the remaining distance."));
@@ -1561,13 +1561,13 @@ bool FPhysicsReplicationAsync::PredictiveInterpolation(Chaos::FPBDRigidParticleH
 	CurrentState.Position = Handle->GetX();
 	CurrentState.Quaternion = Handle->GetR();
 	CurrentState.LinVel = Handle->GetV();
-	CurrentState.AngVel = Handle->GetW(); // Note: Current angular velocity is in Radians
+	CurrentState.AngVel = Handle->GetW(); // Radians
 
 	// NewState
 	const FVector TargetPos = FVector(Target.TargetState.Position);
 	const FQuat TargetRot = Target.TargetState.Quaternion;
 	const FVector TargetLinVel = FVector(Target.TargetState.LinVel);
-	const FVector TargetAngVel = FVector(Target.TargetState.AngVel); // Note: Target angular velocity is in Degrees
+	const FVector TargetAngVel = FVector(Target.TargetState.AngVel); // Radians
 
 	/** --- Reconciliation ---
 	* If target velocities are low enough, check the traveled direction and distance from previous frame and compare with replicated linear velocity.
@@ -1626,7 +1626,7 @@ bool FPhysicsReplicationAsync::PredictiveInterpolation(Chaos::FPBDRigidParticleH
 			const bool bCorrectConnectedBodies = PhysicsReplicationCVars::PredictiveInterpolationCVars::bCorrectConnectedBodies;
 			RigidsSolver->GetEvolution()->ApplyParticleTransformCorrection(Handle, Target.PrevPosTarget, Target.PrevRotTarget, bCorrectConnectedBodies, /*bInRecalculateFrictionOnConnectedBodies*/ true);
 			Handle->SetV(TargetLinVel);
-			Handle->SetW(FMath::DegreesToRadians(TargetAngVel));
+			Handle->SetW(TargetAngVel);
 		}
 
 		// Cache data for next replication
@@ -1716,7 +1716,7 @@ bool FPhysicsReplicationAsync::PredictiveInterpolation(Chaos::FPBDRigidParticleH
 		{	// --- Angular Velocity Replication ---
 
 			// Get AngVelDiff by adding inverted CurrentState.AngVel to TargetAngVel
-			const FVector AngVelDiff = -CurrentState.AngVel + FVector::DegreesToRadians(TargetAngVel);
+			const FVector AngVelDiff = -CurrentState.AngVel + TargetAngVel;
 
 			// Calculate velocity blend amount for this tick as an alpha value
 			const float VelocityAlpha = FMath::Clamp(DeltaSeconds / InterpolationTime, 0.0f, 1.0f);
@@ -1799,7 +1799,6 @@ void FPhysicsReplicationAsync::ExtrapolateTarget(FReplicatedPhysicsTargetAsync&
 	float TargetAngVelSize;
 	FVector TargetAngVelAxis;
 	Target.TargetState.AngVel.FVector::ToDirectionAndLength(TargetAngVelAxis, TargetAngVelSize);
-	TargetAngVelSize = FMath::DegreesToRadians(TargetAngVelSize);
 	const FQuat TargetRotExtrapDelta = FQuat(TargetAngVelAxis, TargetAngVelSize * ExtrapolationTime);
 	Target.TargetState.Quaternion = TargetRotExtrapDelta * Target.TargetState.Quaternion;
 }
diff --git a/Engine/Source/Runtime/Engine/Public/Physics/NetworkPhysicsComponent.h b/Engine/Source/Runtime/Engine/Public/Physics/NetworkPhysicsComponent.h
--- a/Engine/Source/Runtime/Engine/Public/Physics/NetworkPhysicsComponent.h
+++ b/Engine/Source/Runtime/Engine/Public/Physics/NetworkPhysicsComponent.h
@@ -431,6 +431,8 @@ struct FNetworkPhysicsData
 
 	/** Define how to compare client and server data for the same frame, returning false means the data differ enough to trigger a resimulation.
 	* @param PredictedData is data predicted on the client to compare with the current data received from the server.
+	* NOTE: To use this function, CVars np2.Resim.CompareStateToTriggerRewind and/or np2.Resim.CompareInputToTriggerRewind needs to be set to true
+	* or the equivalent settings overridden on the actor via UNetworkPhysicsSettingsComponent. 
 	*/
 	virtual bool CompareData(const FNetworkPhysicsData& PredictedData) { return true; }
 
