[World Partition] Added wp.Runtime.OverrideRuntimeHashSetLoadingRange console command to override loading ranges of specific partitions.
Not available in shipping builds.

diff --git a/Engine/Source/Runtime/Engine/Private/WorldPartition/RuntimeHashSet/WorldPartitionRuntimeHashSet.cpp b/Engine/Source/Runtime/Engine/Private/WorldPartition/RuntimeHashSet/WorldPartitionRuntimeHashSet.cpp
--- a/Engine/Source/Runtime/Engine/Private/WorldPartition/RuntimeHashSet/WorldPartitionRuntimeHashSet.cpp
+++ b/Engine/Source/Runtime/Engine/Private/WorldPartition/RuntimeHashSet/WorldPartitionRuntimeHashSet.cpp
@@ -10,8 +10,64 @@
 #include "WorldPartition/DataLayer/DataLayersID.h"
 #include "Algo/RemoveIf.h"
 #include "Algo/Transform.h"
+#include "Misc/HashBuilder.h"
 #include "Misc/ArchiveMD5.h"
 
+#if !UE_BUILD_SHIPPING
+#include "Engine/Engine.h"
+#include "WorldPartition/WorldPartitionHelpers.h"
+#include "WorldPartition/WorldPartitionSubsystem.h"
+
+TMap<FName, int32> FRuntimePartitionStreamingData::OverriddenLoadingRanges;
+static const TCHAR* GOverrideLoadingRangeCommandName = TEXT("wp.Runtime.OverrideRuntimeHashSetLoadingRange");
+static FDelegateHandle OnWorldPartitionSubsystemDeinitializedFDelegateHandle;
+FAutoConsoleCommand FRuntimePartitionStreamingData::OverrideLoadingRangeCommand(
+	GOverrideLoadingRangeCommandName,
+	TEXT("Sets runtime loading range. Args -partition=[Name] -range=[Range]"),
+	FConsoleCommandWithArgsDelegate::CreateLambda([](const TArray<FString>& InArgs)
+	{
+		FString ArgString = FString::Join(InArgs, TEXT(" "));
+		FName OverrideGridName;
+		int32 OverrideLoadingRange = -1;
+		FParse::Value(*ArgString, TEXT("partition="), OverrideGridName);
+		FParse::Value(*ArgString, TEXT("range="), OverrideLoadingRange);
+
+		if (!OnWorldPartitionSubsystemDeinitializedFDelegateHandle.IsValid())
+		{
+			OnWorldPartitionSubsystemDeinitializedFDelegateHandle = UWorldPartitionSubsystem::OnWorldPartitionSubsystemDeinitialized.AddLambda([](UWorldPartitionSubsystem* InWorldPartitionSubsystem, UWorld* InWorld)
+			{
+				if (InWorld && InWorld->IsGameWorld())
+				{
+					OverriddenLoadingRanges.Reset();
+				}
+			});
+		}
+
+		for (const FWorldContext& Context : GEngine->GetWorldContexts())
+		{
+			UWorld* World = Context.World();
+			if (World && World->IsGameWorld())
+			{
+				if (UWorld::HasSubsystem<UWorldPartitionSubsystem>(World))
+				{
+					FWorldPartitionHelpers::ServerExecConsoleCommand(World, GOverrideLoadingRangeCommandName, InArgs);
+
+					if (OverrideLoadingRange >= 0)
+					{
+						OverriddenLoadingRanges.Add(OverrideGridName, OverrideLoadingRange);
+					}
+					else
+					{
+						OverriddenLoadingRanges.Remove(OverrideGridName);
+					}
+					break;
+				}
+			}
+		}
+	})
+);
+#endif
+
 void FRuntimePartitionStreamingData::CreatePartitionsSpatialIndex() const
 {
 	if (!SpatialIndex)
@@ -47,6 +103,18 @@ void FRuntimePartitionStreamingData::DestroyPartitionsSpatialIndex() const
 	SpatialIndex2D.Reset();
 }
 
+int32 FRuntimePartitionStreamingData::GetLoadingRange() const
+{
+#if !UE_BUILD_SHIPPING
+	if (int32* OverriddenLoadingRange = OverriddenLoadingRanges.Find(Name))
+	{
+		return *OverriddenLoadingRange;
+	}
+#endif
+
+	return LoadingRange;
+}
+
 void URuntimeHashSetExternalStreamingObject::CreatePartitionsSpatialIndex() const
 {
 	for (const FRuntimePartitionStreamingData& StreamingData : RuntimeStreamingData)
@@ -444,7 +512,7 @@ void UWorldPartitionRuntimeHashSet::ForEachStreamingCellsQuery(const FWorldParti
 
 	ForEachStreamingData([&QuerySource, &ForEachSpatiallyLoadedCells, &ForEachNonSpatiallyLoadedCells](const FRuntimePartitionStreamingData& StreamingData)
 	{
-		return ForEachSpatiallyLoadedCells(StreamingData.SpatialIndex.Get(), StreamingData.LoadingRange, StreamingData.Name) && ForEachNonSpatiallyLoadedCells(StreamingData.NonSpatiallyLoadedCells);
+		return ForEachSpatiallyLoadedCells(StreamingData.SpatialIndex.Get(), StreamingData.GetLoadingRange(), StreamingData.Name) && ForEachNonSpatiallyLoadedCells(StreamingData.NonSpatiallyLoadedCells);
 	});
 }
 
@@ -510,7 +578,7 @@ void UWorldPartitionRuntimeHashSet::ForEachStreamingCellsSources(const TArray<FW
 
 				for (const FRuntimePartitionStreamingData* StreamingData : *StreamingDataList)
 				{
-					Source.ForEachShape(StreamingData->LoadingRange, false, [this, &Source, StreamingData, DataLayersStateEpoch, &Func](const FSphericalSector& Shape)
+					Source.ForEachShape(StreamingData->GetLoadingRange(), false, [this, &Source, StreamingData, DataLayersStateEpoch, &Func](const FSphericalSector& Shape)
 					{
 						const FSphere ShapeSphere(Shape.GetCenter(), Shape.GetRadius());
 
@@ -547,6 +615,17 @@ void UWorldPartitionRuntimeHashSet::ForEachStreamingCellsSources(const TArray<FW
 	}
 }
 
+uint32 UWorldPartitionRuntimeHashSet::ComputeUpdateStreamingHash() const
+{
+	FHashBuilder HashBuilder(Super::ComputeUpdateStreamingHash());
+
+#if !UE_BUILD_SHIPPING
+	HashBuilder << FRuntimePartitionStreamingData::OverriddenLoadingRanges;
+#endif
+
+	return HashBuilder.GetHash();
+}
+
 #if WITH_EDITOR
 void UWorldPartitionRuntimeHashSet::PostEditChangeChainProperty(FPropertyChangedChainEvent& PropertyChangedEvent)
 {
diff --git a/Engine/Source/Runtime/Engine/Private/WorldPartition/RuntimeHashSet/WorldPartitionRuntimeHashSetDebugDraw.cpp b/Engine/Source/Runtime/Engine/Private/WorldPartition/RuntimeHashSet/WorldPartitionRuntimeHashSetDebugDraw.cpp
--- a/Engine/Source/Runtime/Engine/Private/WorldPartition/RuntimeHashSet/WorldPartitionRuntimeHashSetDebugDraw.cpp
+++ b/Engine/Source/Runtime/Engine/Private/WorldPartition/RuntimeHashSet/WorldPartitionRuntimeHashSetDebugDraw.cpp
@@ -216,7 +216,7 @@ bool UWorldPartitionRuntimeHashSet::Draw2D(FWorldPartitionDraw2DContext& DrawCon
 		for (const FWorldPartitionStreamingSource& Source : Sources)
 		{
 			const FColor Color = Source.GetDebugColor();
-			Source.ForEachShape(StreamingData->LoadingRange, StreamingData->Name, true, [&Color, &WorldToScreen, &GridScreenBounds, &DrawContext, this](const FSphericalSector& Shape)
+			Source.ForEachShape(StreamingData->GetLoadingRange(), StreamingData->Name, true, [&Color, &WorldToScreen, &GridScreenBounds, &DrawContext, this](const FSphericalSector& Shape)
 			{
 				check(!Shape.IsNearlyZero())
 
@@ -256,7 +256,7 @@ bool UWorldPartitionRuntimeHashSet::Draw2D(FWorldPartitionDraw2DContext& DrawCon
 			if (!Velocity2D.IsNearlyZero())
 			{
 				const FVector2D Center2D = FVector2D(Source.Location);
-				DrawContext.PushDrawSegment(GridScreenBounds, WorldToScreen(Center2D), WorldToScreen(Center2D + Velocity2D * StreamingData->LoadingRange * 0.5f), Color, 1);
+				DrawContext.PushDrawSegment(GridScreenBounds, WorldToScreen(Center2D), WorldToScreen(Center2D + Velocity2D * StreamingData->GetLoadingRange() * 0.5f), Color, 1);
 			}
 		}
 	};
@@ -268,7 +268,7 @@ bool UWorldPartitionRuntimeHashSet::Draw2D(FWorldPartitionDraw2DContext& DrawCon
 	{
 		for (const FWorldPartitionStreamingSource& Source : Sources)
 		{
-			Source.ForEachShape(StreamingDataList[0]->LoadingRange, Name, true, [&GridsShapeBounds](const FSphericalSector& Shape) { GridsShapeBounds += Shape.CalcBounds(); });
+			Source.ForEachShape(StreamingDataList[0]->GetLoadingRange(), Name, true, [&GridsShapeBounds](const FSphericalSector& Shape) { GridsShapeBounds += Shape.CalcBounds(); });
 		}
 
 		FVector2D GridReferenceWorldPos;
@@ -307,7 +307,7 @@ bool UWorldPartitionRuntimeHashSet::Draw2D(FWorldPartitionDraw2DContext& DrawCon
 			FVector2D GridInfoPos = GridScreenOffset - GridScreenHalfExtent;
 			FWorldPartitionCanvasMultiLineText MultiLineText;
 			MultiLineText.Emplace(UWorld::RemovePIEPrefix(FPaths::GetBaseFilename(WorldPartition->GetPackage()->GetName())), FLinearColor::White);
-			FString GridInfoText = FString::Printf(TEXT("%s | %d m"), *Name.ToString(), int32(StreamingDataList[0]->LoadingRange * 0.01f));
+			FString GridInfoText = FString::Printf(TEXT("%s | %d m"), *Name.ToString(), int32(StreamingDataList[0]->GetLoadingRange() * 0.01f));
 			MultiLineText.Emplace(GridInfoText, FLinearColor::Yellow);
 			FWorldPartitionCanvasMultiLineTextItem Item(GridInfoPos, MultiLineText);
 			DrawContext.PushDrawText(Item);
@@ -350,7 +350,7 @@ void UWorldPartitionRuntimeHashSet::Draw3D(const TArray<FWorldPartitionStreaming
 	{
 		for (const FWorldPartitionStreamingSource& Source : Sources)
 		{
-			Source.ForEachShape(StreamingData.LoadingRange, StreamingData.Name, false, [this, &StreamingData, VisualizeMode, &DataLayerDebugColors, ContentBundleManager, OwningWorld, &WorldPartitionTransform](const FSphericalSector& Shape)
+			Source.ForEachShape(StreamingData.GetLoadingRange(), StreamingData.Name, false, [this, &StreamingData, VisualizeMode, &DataLayerDebugColors, ContentBundleManager, OwningWorld, &WorldPartitionTransform](const FSphericalSector& Shape)
 			{
 				const FSphere ShapeSphere(Shape.GetCenter(), Shape.GetRadius());
 
diff --git a/Engine/Source/Runtime/Engine/Public/WorldPartition/RuntimeHashSet/WorldPartitionRuntimeHashSet.h b/Engine/Source/Runtime/Engine/Public/WorldPartition/RuntimeHashSet/WorldPartitionRuntimeHashSet.h
--- a/Engine/Source/Runtime/Engine/Public/WorldPartition/RuntimeHashSet/WorldPartitionRuntimeHashSet.h
+++ b/Engine/Source/Runtime/Engine/Public/WorldPartition/RuntimeHashSet/WorldPartitionRuntimeHashSet.h
@@ -70,9 +70,15 @@ struct FRuntimePartitionStreamingData
 {
 	GENERATED_USTRUCT_BODY()
 
+	friend class UWorldPartitionRuntimeHashSet;
+	friend class URuntimeHashSetExternalStreamingObject;
+
 	void CreatePartitionsSpatialIndex() const;
 	void DestroyPartitionsSpatialIndex() const;
 
+	int32 GetLoadingRange() const;
+
+protected:
 	/** Name of the runtime partition, currently maps to target grids. */
 	UPROPERTY()
 	FName Name;
@@ -89,6 +95,11 @@ struct FRuntimePartitionStreamingData
 	// Transient
 	mutable TUniquePtr<FStaticSpatialIndexType> SpatialIndex;
 	mutable TUniquePtr<FStaticSpatialIndexType> SpatialIndex2D;
+
+#if !UE_BUILD_SHIPPING
+	static TMap<FName, int32> OverriddenLoadingRanges;
+	static class FAutoConsoleCommand OverrideLoadingRangeCommand;
+#endif
 };
 
 template<>
@@ -164,6 +175,7 @@ public:
 	ENGINE_API virtual void ForEachStreamingCells(TFunctionRef<bool(const UWorldPartitionRuntimeCell*)> Func) const;
 	ENGINE_API virtual void ForEachStreamingCellsQuery(const FWorldPartitionStreamingQuerySource& QuerySource, TFunctionRef<bool(const UWorldPartitionRuntimeCell*)> Func, FWorldPartitionQueryCache* QueryCache) const override;
 	ENGINE_API virtual void ForEachStreamingCellsSources(const TArray<FWorldPartitionStreamingSource>& Sources, TFunctionRef<bool(const UWorldPartitionRuntimeCell*, EStreamingSourceTargetState)> Func) const override;
+	ENGINE_API virtual uint32 ComputeUpdateStreamingHash() const override;
 
 private:
 	ENGINE_API virtual void OnBeginPlay() override;
