Fix unity build.

diff --git a/Engine/Source/Runtime/Engine/Private/WorldPartition/RuntimeHashSet/WorldPartitionRuntimeHashSet.cpp b/Engine/Source/Runtime/Engine/Private/WorldPartition/RuntimeHashSet/WorldPartitionRuntimeHashSet.cpp
--- a/Engine/Source/Runtime/Engine/Private/WorldPartition/RuntimeHashSet/WorldPartitionRuntimeHashSet.cpp
+++ b/Engine/Source/Runtime/Engine/Private/WorldPartition/RuntimeHashSet/WorldPartitionRuntimeHashSet.cpp
@@ -19,10 +19,10 @@
 #include "WorldPartition/WorldPartitionSubsystem.h"
 
 TMap<FName, int32> FRuntimePartitionStreamingData::OverriddenLoadingRanges;
-static const TCHAR* GOverrideLoadingRangeCommandName = TEXT("wp.Runtime.OverrideRuntimeHashSetLoadingRange");
+static const TCHAR* GOverrideHashSetLoadingRangeCommandName = TEXT("wp.Runtime.OverrideRuntimeHashSetLoadingRange");
 static FDelegateHandle OnWorldPartitionSubsystemDeinitializedFDelegateHandle;
 FAutoConsoleCommand FRuntimePartitionStreamingData::OverrideLoadingRangeCommand(
-	GOverrideLoadingRangeCommandName,
+	GOverrideHashSetLoadingRangeCommandName,
 	TEXT("Sets runtime loading range. Args -partition=[Name] -range=[Range]"),
 	FConsoleCommandWithArgsDelegate::CreateLambda([](const TArray<FString>& InArgs)
 	{
@@ -50,7 +50,7 @@ FAutoConsoleCommand FRuntimePartitionStreamingData::OverrideLoadingRangeCommand(
 			{
 				if (UWorld::HasSubsystem<UWorldPartitionSubsystem>(World))
 				{
-					FWorldPartitionHelpers::ServerExecConsoleCommand(World, GOverrideLoadingRangeCommandName, InArgs);
+					FWorldPartitionHelpers::ServerExecConsoleCommand(World, GOverrideHashSetLoadingRangeCommandName, InArgs);
 
 					if (OverrideLoadingRange >= 0)
 					{
