Add missing ENGINE_API to FRuntimePartitionStreamingData::GetLoadingRange.

diff --git a/Engine/Source/Runtime/Engine/Public/WorldPartition/RuntimeHashSet/WorldPartitionRuntimeHashSet.h b/Engine/Source/Runtime/Engine/Public/WorldPartition/RuntimeHashSet/WorldPartitionRuntimeHashSet.h
--- a/Engine/Source/Runtime/Engine/Public/WorldPartition/RuntimeHashSet/WorldPartitionRuntimeHashSet.h
+++ b/Engine/Source/Runtime/Engine/Public/WorldPartition/RuntimeHashSet/WorldPartitionRuntimeHashSet.h
@@ -77,7 +77,7 @@ struct FRuntimePartitionStreamingData
 	void CreatePartitionsSpatialIndex() const;
 	void DestroyPartitionsSpatialIndex() const;
 
-	int32 GetLoadingRange() const;
+	ENGINE_API int32 GetLoadingRange() const;
 
 protected:
 	/** Name of the runtime partition, currently maps to target grids. */
