Chopped SVT frame DDC data into 2MB buckets to work around bad performance when requesting subranges of a DDC chunk. This vastly improves SVT streaming performance in editor.

#rnx

diff --git a/Engine/Source/Runtime/Engine/Classes/SparseVolumeTexture/SparseVolumeTexture.h b/Engine/Source/Runtime/Engine/Classes/SparseVolumeTexture/SparseVolumeTexture.h
--- a/Engine/Source/Runtime/Engine/Classes/SparseVolumeTexture/SparseVolumeTexture.h
+++ b/Engine/Source/Runtime/Engine/Classes/SparseVolumeTexture/SparseVolumeTexture.h
@@ -209,7 +209,8 @@ public:
 #if WITH_EDITORONLY_DATA
 	FString ResourceName;
 	FIoHash DDCKeyHash;
-	FIoHash DDCRawHash;
+	TArray<TStaticArray<uint8, 12>> DDCChunkIds; // 12-byte hashes used to reconstruct the FValueId required to look up the chunks.
+	TArray<uint32> DDCChunkMaxTileIndices; // Maximum tile index of each chunk. Used to look up the chunk a given tile belongs to.
 #endif
 
 	// Called when serializing to/from DDC buffers and when serializing the owning USparseVolumeTextureFrame.
@@ -240,6 +241,7 @@ private:
 	};
 	TDontCopy<TPimplPtr<UE::DerivedData::FRequestOwner>> DDCRequestOwner;
 	std::atomic<EDDCRebuildState> DDCRebuildState;
+	std::atomic_int DDCRebuildNumFinishedRequests;
 	void BeginRebuildBulkDataFromCache(const UObject* Owner);
 	void EndRebuildBulkDataFromCache();
 #endif
diff --git a/Engine/Source/Runtime/Engine/Private/SparseVolumeTexture/SparseVolumeTexture.cpp b/Engine/Source/Runtime/Engine/Private/SparseVolumeTexture/SparseVolumeTexture.cpp
--- a/Engine/Source/Runtime/Engine/Private/SparseVolumeTexture/SparseVolumeTexture.cpp
+++ b/Engine/Source/Runtime/Engine/Private/SparseVolumeTexture/SparseVolumeTexture.cpp
@@ -64,6 +64,14 @@ static FAutoConsoleVariableRef CVarSVTStreamingRequestMipBias(
 	ECVF_Default
 );
 
+static int32 GSVTStreamingDDCChunkSize = 2;
+static FAutoConsoleVariableRef CVarSVTStreamingDDCChunkSize(
+	TEXT("r.SparseVolumeTexture.Streaming.DDCChunkSize"),
+	GSVTStreamingDDCChunkSize,
+	TEXT("Size of DDC chunks the streaming data is split into (in MiB). A smaller size leads to more requests but can improve streaming performance. Default: 2 MiB"),
+	ECVF_Default | ECVF_ReadOnly
+);
+
 ////////////////////////////////////////////////////////////////////////////////////////////////
 
 FArchive& operator<<(FArchive& Ar, UE::SVT::FMipLevelStreamingInfo& MipLevelStreamingInfo)
@@ -126,7 +134,7 @@ static int32 ComputeNumMipLevels(const FIntVector3& InVirtualVolumeMin, const FI
 
 static const FString& GetDerivedDataVersion()
 {
-	static FString CachedVersionString = TEXT("2113A905-4C31-4EFD-B55C-82375C6A7B56");	// Bump this if you want to ignore all cached data so far.
+	static FString CachedVersionString = TEXT("49DD2D7C-C346-4A02-963A-D3F81E7E1601");	// Bump this if you want to ignore all cached data so far.
 	return CachedVersionString;
 }
 
@@ -278,6 +286,15 @@ void FResources::Serialize(FArchive& Ar, UObject* Owner, bool bCooked)
 
 		Topology.Serialize(Ar);
 
+#if WITH_EDITORONLY_DATA
+		// These members are only needed when streaming from DDC
+		if (!bCooked)
+		{
+			Ar << DDCChunkIds;
+			Ar << DDCChunkMaxTileIndices;
+		}
+#endif
+
 		// StreamableMipLevels is only serialized in cooked builds and when caching to DDC failed in editor builds.
 		// If the data was successfully cached to DDC, we just query it from DDC on the next run or recreate it if that failed.
 		if (StoredResourceFlags & EResourceFlag_StreamingDataInDDC)
@@ -382,7 +399,8 @@ bool FResources::Build(USparseVolumeTextureFrame* Owner, UE::Serialization::FEdi
 		ResourceFlags = 0;
 		ResourceName.Reset();
 		DDCKeyHash.Reset();
-		DDCRawHash.Reset();
+		DDCChunkIds.Reset();
+		DDCChunkMaxTileIndices.Reset();
 		DDCRebuildState.store(EDDCRebuildState::Initial);
 
 		// Build page topology
@@ -431,8 +449,8 @@ void FResources::Cache(USparseVolumeTextureFrame* Owner, UE::Serialization::FEdi
 	using namespace UE::DerivedData;
 
 	static const FValueId SVTDataId = FValueId::FromName("SparseVolumeTextureData");
-	static const FValueId SVTStreamingDataId = FValueId::FromName("SparseVolumeTextureStreamingData");
-	const FString KeySuffix = SourceData.GetIdentifier().ToString() + FString::Format(TEXT("{0}_{1}_{2}_{3}"), { Owner->GetNumMipLevels(), Owner->GetTextureAddressX(), Owner->GetTextureAddressY(), Owner->GetTextureAddressZ() });
+	const int32 DDCChunkSizeInMiB = FMath::Max(GSVTStreamingDDCChunkSize, 1);
+	const FString KeySuffix = SourceData.GetIdentifier().ToString() + FString::Format(TEXT("{0}_{1}_{2}_{3}_{4}"), { Owner->GetNumMipLevels(), Owner->GetTextureAddressX(), Owner->GetTextureAddressY(), Owner->GetTextureAddressZ(), DDCChunkSizeInMiB });
 	FString DerivedDataKey = FDerivedDataCacheInterface::BuildCacheKey(TEXT("SPARSEVOLUMETEXTURE"), *GetDerivedDataVersion(), *KeySuffix);
 
 	FCacheKey CacheKey;
@@ -450,12 +468,11 @@ void FResources::Cache(USparseVolumeTextureFrame* Owner, UE::Serialization::FEdi
 
 	// Check if the data already exists in DDC
 	FSharedBuffer ResourcesDataBuffer;
-	FIoHash SVTStreamingDataHash;
 	{
 		TRACE_CPUPROFILER_EVENT_SCOPE(SVT::FResources::Cache::CheckDDC);
 
-		FCacheRecordPolicyBuilder PolicyBuilder(DefaultCachePolicy | ECachePolicy::KeepAlive);
-		PolicyBuilder.AddValuePolicy(SVTStreamingDataId, DefaultCachePolicy | ECachePolicy::SkipData);
+		FCacheRecordPolicyBuilder PolicyBuilder(DefaultCachePolicy | ECachePolicy::KeepAlive | ECachePolicy::SkipData);
+		PolicyBuilder.AddValuePolicy(SVTDataId, DefaultCachePolicy);
 
 		FCacheGetRequest Request;
 		Request.Name = Owner->GetPathName();
@@ -464,14 +481,12 @@ void FResources::Cache(USparseVolumeTextureFrame* Owner, UE::Serialization::FEdi
 
 		FRequestOwner RequestOwner(EPriority::Blocking);
 		GetCache().Get(MakeArrayView(&Request, 1), RequestOwner,
-			[&ResourcesDataBuffer, &SVTStreamingDataHash](FCacheGetResponse&& Response)
+			[&ResourcesDataBuffer](FCacheGetResponse&& Response)
 			{
 				if (Response.Status == EStatus::Ok)
 				{
 					const FCompressedBuffer& CompressedBuffer = Response.Record.GetValue(SVTDataId).GetData();
 					ResourcesDataBuffer = CompressedBuffer.Decompress();
-
-					SVTStreamingDataHash = Response.Record.GetValue(SVTStreamingDataId).GetRawHash();
 				}
 			});
 		RequestOwner.Wait();
@@ -490,7 +505,6 @@ void FResources::Cache(USparseVolumeTextureFrame* Owner, UE::Serialization::FEdi
 		if (ResourceFlags & EResourceFlag_StreamingDataInDDC)
 		{
 			DDCKeyHash = CacheKey.Hash;
-			DDCRawHash = SVTStreamingDataHash;
 		}
 	}
 	else
@@ -503,13 +517,46 @@ void FResources::Cache(USparseVolumeTextureFrame* Owner, UE::Serialization::FEdi
 		if (HasStreamingData())
 		{
 			FByteBulkData& BulkData = StreamableMipLevels;
+			const uint8* SrcPtr = static_cast<const uint8*>(BulkData.LockReadOnly());
+			const int64 SrcSize = BulkData.GetBulkDataSize();
+			const int64 TargetChunkSize = DDCChunkSizeInMiB * 1024LL * 1024LL;
+			
+			check(DDCChunkMaxTileIndices.IsEmpty());
+			check(DDCChunkIds.IsEmpty());
+			static_assert(sizeof(FValueId::ByteArray) == 12);
+
+			uint32 FirstTileIndexInChunk = StreamingMetaData.FirstStreamingTileIndex;
+			for (uint32 TileIndex = StreamingMetaData.FirstStreamingTileIndex; TileIndex < StreamingMetaData.GetNumTiles(); ++TileIndex)
+			{
+				bool bCreateChunk = !StreamingMetaData.TileDataOffsets.IsValidIndex(TileIndex + 2); // Create a chunk if this is the last tile.
+				if (!bCreateChunk)
+				{
+					const int64 RangeSizeIncludingNextTile = StreamingMetaData.TileDataOffsets[TileIndex + 2] - StreamingMetaData.TileDataOffsets[FirstTileIndexInChunk];
+					bCreateChunk = RangeSizeIncludingNextTile > TargetChunkSize; // Create a chunk if the next tile would exceed the target chunk size
+				}
+
+				if (bCreateChunk)
+				{
+					const int64 ChunkOffset = StreamingMetaData.TileDataOffsets[FirstTileIndexInChunk] - StreamingMetaData.GetRootTileSize();
+					const int64 ChunkSize = StreamingMetaData.TileDataOffsets[TileIndex + 1] - StreamingMetaData.TileDataOffsets[FirstTileIndexInChunk];
+					const uint8* ChunkPtr = SrcPtr + ChunkOffset;
+					const FValue Value = FValue::Compress(FSharedBuffer::MakeView(ChunkPtr, ChunkSize));
+
+					const FString ChunkName = FString::Printf(TEXT("SparseVolumeTextureDataChunk%i"), DDCChunkMaxTileIndices.Num());
+					const FValueId ChunkId = FValueId::FromName(ChunkName.GetCharArray());
+					const uint8* ChunkIdBytes = ChunkId.GetBytes();
+
+					RecordBuilder.AddValue(ChunkId, Value);
+
+					DDCChunkMaxTileIndices.Add(TileIndex);
+					FMemory::Memcpy(DDCChunkIds.AddDefaulted_GetRef().GetData(), ChunkIdBytes, 12);
+					FirstTileIndexInChunk = TileIndex + 1;
+				}
+			}
 
-			FValue Value = FValue::Compress(FSharedBuffer::MakeView(BulkData.LockReadOnly(), BulkData.GetBulkDataSize()));
-			RecordBuilder.AddValue(SVTStreamingDataId, Value);
 			BulkData.Unlock();
 			ResourceFlags |= EResourceFlag_StreamingDataInDDC;
 			DDCKeyHash = CacheKey.Hash;
-			DDCRawHash = Value.GetRawHash();
 		}
 
 		// Serialize to a buffer and store into DDC.
@@ -562,7 +609,8 @@ void FResources::SetDefault(EPixelFormat FormatA, EPixelFormat FormatB, const FV
 	StreamableMipLevels.RemoveBulkData();
 	ResourceName.Reset();
 	DDCKeyHash.Reset();
-	DDCRawHash.Reset();
+	DDCChunkIds.Reset();
+	DDCChunkMaxTileIndices.Reset();
 	DDCRebuildState.store(EDDCRebuildState::Initial);
 }
 
@@ -572,36 +620,55 @@ void FResources::SetDefault(EPixelFormat FormatA, EPixelFormat FormatB, const FV
 
 void FResources::BeginRebuildBulkDataFromCache(const UObject* Owner)
 {
+	using namespace UE::DerivedData;
+
 	check(DDCRebuildState.load() == EDDCRebuildState::Initial);
 	if (!HasStreamingData() || (ResourceFlags & EResourceFlag_StreamingDataInDDC) == 0u)
 	{
 		return;
 	}
-	using namespace UE::DerivedData;
-	FCacheKey Key;
-	Key.Bucket = FCacheBucket(TEXT("SparseVolumeTexture"));
-	Key.Hash = DDCKeyHash;
+	
 	check(!DDCKeyHash.IsZero());
-	FCacheGetChunkRequest Request;
-	Request.Name = Owner->GetPathName();
-	Request.Id = FValueId::FromName("SparseVolumeTextureStreamingData");
-	Request.Key = Key;
-	Request.RawHash = DDCRawHash;
-	check(!DDCRawHash.IsZero());
-	FSharedBuffer SharedBuffer;
 	*DDCRequestOwner = MakePimpl<FRequestOwner>(EPriority::Normal);
 	DDCRebuildState.store(EDDCRebuildState::Pending);
-	GetCache().GetChunks(MakeArrayView(&Request, 1), **DDCRequestOwner,
-		[this](FCacheGetChunkResponse&& Response)
+	DDCRebuildNumFinishedRequests.store(0);
+
+	// Lock and realloc bulk data
+	StreamableMipLevels.Lock(LOCK_READ_WRITE);
+	const int64 StreamableBulkDataSize = StreamingMetaData.TileDataOffsets.Last() - StreamingMetaData.GetRootTileSize();
+	uint8* BulkDataPtr = (uint8*)StreamableMipLevels.Realloc(StreamableBulkDataSize);
+
+	// Generate requests
+	const int32 NumChunks = DDCChunkIds.Num();
+	check(NumChunks > 0);
+	TArray<FCacheGetChunkRequest> DDCRequests;
+	DDCRequests.Reserve(NumChunks);
+	uint32 FirstTileIndexInChunk = StreamingMetaData.FirstStreamingTileIndex;
+	for (int32 ChunkIndex = 0; ChunkIndex < NumChunks; ++ChunkIndex)
+	{
+		FCacheGetChunkRequest& Request = DDCRequests.AddDefaulted_GetRef();
+		Request.Name = Owner->GetPathName();
+		Request.Key.Bucket = FCacheBucket(TEXT("SparseVolumeTexture"));
+		Request.Key.Hash = DDCKeyHash;
+		Request.Id = FValueId(FMemoryView(DDCChunkIds[ChunkIndex].GetData(), 12));
+		Request.UserData = StreamingMetaData.TileDataOffsets[FirstTileIndexInChunk] - StreamingMetaData.GetRootTileSize(); // Store write offset in UserData
+		FirstTileIndexInChunk = DDCChunkMaxTileIndices[ChunkIndex] + 1;
+	}
+
+	// Issue requests
+	GetCache().GetChunks(DDCRequests, **DDCRequestOwner,
+		[this, BulkDataPtr, NumChunks](FCacheGetChunkResponse&& Response)
 		{
 			if (Response.Status == EStatus::Ok)
 			{
-				StreamableMipLevels.Lock(LOCK_READ_WRITE);
-				uint8* Ptr = (uint8*)StreamableMipLevels.Realloc(Response.RawData.GetSize());
-				FMemory::Memcpy(Ptr, Response.RawData.GetData(), Response.RawData.GetSize());
-				StreamableMipLevels.Unlock();
-				StreamableMipLevels.SetBulkDataFlags(BULKDATA_Force_NOT_InlinePayload);
-				DDCRebuildState.store(EDDCRebuildState::Succeeded);
+				FMemory::Memcpy(BulkDataPtr + Response.UserData, Response.RawData.GetData(), Response.RawData.GetSize());
+				
+				// The last request to finish sets the Succeeded flag
+				const int32 NumFinishedRequests = DDCRebuildNumFinishedRequests.fetch_add(1) + 1;
+				if (NumFinishedRequests == NumChunks)
+				{
+					DDCRebuildState.store(EDDCRebuildState::Succeeded);
+				}
 			}
 			else
 			{
@@ -612,6 +679,8 @@ void FResources::BeginRebuildBulkDataFromCache(const UObject* Owner)
 
 void FResources::EndRebuildBulkDataFromCache()
 {
+	StreamableMipLevels.Unlock();
+	StreamableMipLevels.SetBulkDataFlags(BULKDATA_Force_NOT_InlinePayload);
 	if (*DDCRequestOwner)
 	{
 		(*DDCRequestOwner)->Wait();
diff --git a/Engine/Source/Runtime/Engine/Private/SparseVolumeTexture/SparseVolumeTextureStreamingManager.cpp b/Engine/Source/Runtime/Engine/Private/SparseVolumeTexture/SparseVolumeTextureStreamingManager.cpp
--- a/Engine/Source/Runtime/Engine/Private/SparseVolumeTexture/SparseVolumeTextureStreamingManager.cpp
+++ b/Engine/Source/Runtime/Engine/Private/SparseVolumeTexture/SparseVolumeTextureStreamingManager.cpp
@@ -63,7 +63,7 @@ static FAutoConsoleVariableRef CVarSVTStreamingLogVerbosity(
 	ECVF_RenderThreadSafe
 );
 
-static int32 GSVTStreamingMaxPendingRequests = 256;
+static int32 GSVTStreamingMaxPendingRequests = 256 * 32;
 static FAutoConsoleVariableRef CVarSVTStreamingMaxPendingRequests(
 	TEXT("r.SparseVolumeTexture.Streaming.MaxPendingRequests"),
 	GSVTStreamingMaxPendingRequests,
@@ -1148,8 +1148,14 @@ void FStreamingManager::IssueRequests()
 
 		// Allocate tiles
 		uint32 NumSuccessfulAllocations = 0;
+		bool bOutOfRequestSlots = false;
 		for (uint32 i = 0; i < TileRange.TileCount; ++i)
 		{
+			if (bOutOfRequestSlots)
+			{
+				break;
+			}
+
 			const uint32 TileIndex = TileRange.TileOffset + i;
 			check(FrameInfo.TileAllocations[TileIndex] == 0);
 
@@ -1220,22 +1226,20 @@ void FStreamingManager::IssueRequests()
 
 		// Create request
 		{
-			const int32 PendingRequestIndex = NextPendingRequestIndex;
-			FPendingRequest& PendingRequest = PendingRequests[PendingRequestIndex];
-			LOCK_PENDING_REQUEST(PendingRequest);
-			PendingRequest.Reset();
-			PendingRequest.SVTHandle = TileRange.SVTHandle;
-			PendingRequest.FrameIndex = TileRange.FrameIndex;
-			PendingRequest.TileOffset = TileRange.TileOffset;
-			PendingRequest.TileCount = TileRange.TileCount;
-			PendingRequest.IssuedInFrame = UpdateIndex;
-			PendingRequest.bBlocking = GSVTStreamingForceBlockingRequests || (TileRange.Priority == uint8(-1));
-
-			uint32 ReadOffset = 0;
-			uint32 ReadSize = 0;
-			Resources->StreamingMetaData.GetTileRangeMemoryOffsetSize(PendingRequest.TileOffset, PendingRequest.TileCount, ReadOffset, ReadSize);
+			auto HandleRequestAllocationFailure = [&](uint32 FirstTileIndex, uint32 LastTileIndex)
+			{
+				// Free allocated tile data texture slots of all the tiles for which we haven't issued a request yet.
+				for (uint32 TileIndexToFree = FirstTileIndex; TileIndexToFree < LastTileIndex; ++TileIndexToFree)
+				{
+					SVTInfo->TileAllocator.Free(FrameInfo.TileAllocations[TileIndexToFree]);
+					FrameInfo.TileAllocations[TileIndexToFree] = 0;
+				}
+				bOutOfRequestSlots = true;
+				UE_LOG(LogSparseVolumeTextureStreamingManager, Error, TEXT("Ran out of SparseVolumeTexture IO request slots (%i)! r.SparseVolumeTexture.Streaming.MaxPendingRequests must be increased to fix this issue."), MaxPendingRequests);
+			};
 
 			const FByteBulkData& BulkData = Resources->StreamableMipLevels;
+			const bool bBlockingRequest = GSVTStreamingForceBlockingRequests || (TileRange.Priority == uint8(-1));
 #if WITH_EDITORONLY_DATA
 			const bool bDiskRequest = (!(Resources->ResourceFlags & EResourceFlag_StreamingDataInDDC) && !BulkData.IsBulkDataLoaded());
 #else
@@ -1243,51 +1247,105 @@ void FStreamingManager::IssueRequests()
 #endif
 
 #if WITH_EDITORONLY_DATA
-			if (!bDiskRequest)
+			// When streaming from DDC, we have a slightly more complicated setup where the data is spread over multiple DDC chunks, so we likely need to issue multiple requests
+			if (!bDiskRequest && (Resources->ResourceFlags & EResourceFlag_StreamingDataInDDC))
 			{
-				if (Resources->ResourceFlags & EResourceFlag_StreamingDataInDDC)
+				// Iterate over all chunks to find the range of chunks storing the requested range of tiles and issue a DDC request for each chunk.
+				uint32 FirstTileIndexInChunk = Resources->StreamingMetaData.FirstStreamingTileIndex;
+				uint32 FirstTileIndexToRead = TileRange.TileOffset;
+				uint32 NumRemainingTiles = TileRange.TileCount;
+
+				const int32 NumChunks = Resources->DDCChunkMaxTileIndices.Num();
+				for (int32 ChunkIndex = 0; ChunkIndex < NumChunks && NumRemainingTiles > 0; ++ChunkIndex)
 				{
-					UE::DerivedData::FCacheGetChunkRequest DDCRequest = BuildDDCRequest(*Resources, ReadOffset, ReadSize, NextPendingRequestIndex);
-					if (PendingRequest.bBlocking)
+					// Does this chunk overlap the requested tile range?
+					const uint32 LastTileIndexInChunkPlusOne = Resources->DDCChunkMaxTileIndices[ChunkIndex] + 1;
+					if (FirstTileIndexToRead >= FirstTileIndexInChunk && FirstTileIndexToRead < LastTileIndexInChunkPlusOne)
 					{
-						DDCRequestsBlocking.Add(DDCRequest);
-					}
-					else
-					{
-						DDCRequests.Add(DDCRequest);
+						// We can only read within a single chunk per request, so we need to limit the number of tiles in this request.
+						const uint32 ChunkNumReadTiles = FMath::Min(FirstTileIndexToRead + NumRemainingTiles, LastTileIndexInChunkPlusOne) - FirstTileIndexToRead;
+
+						// Allocate and fill out a FPendingRequest.
+						const int32 PendingRequestIndex = AllocatePendingRequestIndex();
+						if (!ensure(PendingRequestIndex != INDEX_NONE)) // Handle allocation failure
+						{
+							HandleRequestAllocationFailure(FirstTileIndexToRead, FirstTileIndexToRead + NumRemainingTiles);
+							break;
+						}
+						FPendingRequest& PendingRequest = PendingRequests[PendingRequestIndex];
+						LOCK_PENDING_REQUEST(PendingRequest);
+						PendingRequest.Reset();
+						PendingRequest.Set(TileRange.SVTHandle, TileRange.FrameIndex, FirstTileIndexToRead, ChunkNumReadTiles, UpdateIndex, bBlockingRequest);
+						PendingRequest.State = FPendingRequest::EState::DDC_Pending;
+
+						// Create the DDC request.
+						const FCacheGetChunkRequest DDCRequest = BuildDDCRequest(*Resources, PendingRequest.TileOffset, PendingRequest.TileCount, PendingRequestIndex, ChunkIndex);
+						new (PendingRequest.bBlocking ? DDCRequestsBlocking : DDCRequests) FCacheGetChunkRequest(DDCRequest);
+
+						check(ChunkNumReadTiles <= NumRemainingTiles);
+						NumRemainingTiles -= ChunkNumReadTiles;
+						FirstTileIndexToRead += ChunkNumReadTiles;
+
+						// Link the streaming tiles with the request.
+						for (uint32 TileIndex = PendingRequest.TileOffset; TileIndex < (PendingRequest.TileOffset + PendingRequest.TileCount); ++TileIndex)
+						{
+							FrameInfo.TileIndexToPendingRequestIndex.Add(TileIndex, (uint32)PendingRequestIndex);
+						}
+
+						// Mark tiles as streaming. Once they've actually been loaded into GPU memory, they'll also be marked as resident.
+						FrameInfo.StreamingTiles.SetRange(PendingRequest.TileOffset, PendingRequest.TileCount, true);
 					}
-					PendingRequest.State = FPendingRequest::EState::DDC_Pending;
-				}
-				else
-				{
-					PendingRequest.State = FPendingRequest::EState::Memory;
+					FirstTileIndexInChunk = LastTileIndexInChunkPlusOne;
 				}
 			}
+			// Otherwise we can use a single request to fetch all the tiles.
 			else
 #endif
 			{
-				PendingRequest.RequestBuffer = FIoBuffer(ReadSize); // SVT_TODO: Use FIoBuffer::Wrap with preallocated memory
-				const EAsyncIOPriorityAndFlags Priority = PendingRequest.bBlocking ? AIOP_CriticalPath : AIOP_Low;
-				Batch.Read(BulkData, ReadOffset, ReadSize, Priority, PendingRequest.RequestBuffer, PendingRequest.Request);
-				bIssueIOBatch = true;
+				// Allocate and fill out a FPendingRequest.
+				const int32 PendingRequestIndex = AllocatePendingRequestIndex();
+				if (!ensure(PendingRequestIndex != INDEX_NONE)) // Handle allocation failure
+				{
+					HandleRequestAllocationFailure(TileRange.TileOffset, TileRange.TileOffset + TileRange.TileCount);
+					continue;
+				}
+				FPendingRequest& PendingRequest = PendingRequests[PendingRequestIndex];
+				LOCK_PENDING_REQUEST(PendingRequest);
+				PendingRequest.Reset();
+				PendingRequest.Set(TileRange.SVTHandle, TileRange.FrameIndex, TileRange.TileOffset, TileRange.TileCount, UpdateIndex, bBlockingRequest);
 
 #if WITH_EDITORONLY_DATA
-				PendingRequest.State = FPendingRequest::EState::Disk;
+				if (!bDiskRequest)
+				{
+					PendingRequest.State = FPendingRequest::EState::Memory;
+				}
+				else
 #endif
-			}
+				{
+					uint32 ReadOffset = 0;
+					uint32 ReadSize = 0;
+					Resources->StreamingMetaData.GetTileRangeMemoryOffsetSize(TileRange.TileOffset, TileRange.TileCount, ReadOffset, ReadSize);
 
-			for (uint32 TileIndex = TileRange.TileOffset; TileIndex < (TileRange.TileOffset + TileRange.TileCount); ++TileIndex)
-			{
-				FrameInfo.TileIndexToPendingRequestIndex.Add(TileIndex, (uint32)NextPendingRequestIndex);
-			}
+					PendingRequest.RequestBuffer = FIoBuffer(ReadSize); // SVT_TODO: Use FIoBuffer::Wrap with preallocated memory
+					const EAsyncIOPriorityAndFlags Priority = PendingRequest.bBlocking ? AIOP_CriticalPath : AIOP_Low;
+					Batch.Read(BulkData, ReadOffset, ReadSize, Priority, PendingRequest.RequestBuffer, PendingRequest.Request);
+					bIssueIOBatch = true;
 
-			NextPendingRequestIndex = (NextPendingRequestIndex + 1) % MaxPendingRequests;
-			check(NumPendingRequests < MaxPendingRequests);
-			++NumPendingRequests;
-		}
+#if WITH_EDITORONLY_DATA
+					PendingRequest.State = FPendingRequest::EState::Disk;
+#endif
+				}
 
-		// Mark tiles as streaming. Once they've actually been loaded into GPU memory, they'll also be marked as resident.
-		FrameInfo.StreamingTiles.SetRange(TileRange.TileOffset, TileRange.TileCount, true);
+				// Link the streaming tiles with the request.
+				for (uint32 TileIndex = PendingRequest.TileOffset; TileIndex < (PendingRequest.TileOffset + PendingRequest.TileCount); ++TileIndex)
+				{
+					FrameInfo.TileIndexToPendingRequestIndex.Add(TileIndex, (uint32)PendingRequestIndex);
+				}
+
+				// Mark tiles as streaming. Once they've actually been loaded into GPU memory, they'll also be marked as resident.
+				FrameInfo.StreamingTiles.SetRange(PendingRequest.TileOffset, PendingRequest.TileCount, true);
+			}
+		}
 	}
 
 	// Now we can finally issue the requests
@@ -1365,13 +1423,22 @@ int32 FStreamingManager::DetermineReadyRequests()
 					*SVTInfo->SVTName.ToString(), PendingRequest.FrameIndex, PendingRequest.TileOffset, PendingRequest.TileCount);
 			}
 
-			uint32 ReadOffset = 0;
-			uint32 ReadSize = 0;
-			Resources->StreamingMetaData.GetTileRangeMemoryOffsetSize(PendingRequest.TileOffset, PendingRequest.TileCount, ReadOffset, ReadSize);
-			FCacheGetChunkRequest Request = BuildDDCRequest(*Resources, ReadOffset, ReadSize, PendingRequestIndex);
-			const bool bBlocking = GSVTStreamingForceBlockingRequests || PendingRequest.bBlocking;
-			RequestDDCData(MakeArrayView(&Request, 1), bBlocking);
-
+			const int32 NumChunks = Resources->DDCChunkMaxTileIndices.Num();
+			uint32 FirstTileIndexInChunk = Resources->StreamingMetaData.FirstStreamingTileIndex;
+			bool bFoundChunk = false;
+			for (int32 ChunkIndex = 0; ChunkIndex < NumChunks; ++ChunkIndex)
+			{
+				const uint32 LastTileIndexInChunkPlusOne = Resources->DDCChunkMaxTileIndices[ChunkIndex] + 1;
+				if (PendingRequest.TileOffset >= FirstTileIndexInChunk && (PendingRequest.TileOffset + PendingRequest.TileCount) <= LastTileIndexInChunkPlusOne)
+				{
+					bFoundChunk = true;
+					const FCacheGetChunkRequest Request = BuildDDCRequest(*Resources, PendingRequest.TileOffset, PendingRequest.TileCount, PendingRequestIndex, ChunkIndex);
+					const bool bBlocking = GSVTStreamingForceBlockingRequests || PendingRequest.bBlocking;
+					RequestDDCData(MakeArrayView(&Request, 1), bBlocking);
+					break;
+				}
+			}
+			check(bFoundChunk);
 			++PendingRequest.RetryCount;
 			break;
 		}
@@ -1770,23 +1837,44 @@ FStreamingManager::FStreamingInfo* FStreamingManager::FindStreamingInfo(UStreama
 	return nullptr;
 }
 
+int32 FStreamingManager::AllocatePendingRequestIndex()
+{
+	if (NumPendingRequests < MaxPendingRequests)
+	{
+		int32 Result = NextPendingRequestIndex;
+		NextPendingRequestIndex = (NextPendingRequestIndex + 1) % MaxPendingRequests;
+		++NumPendingRequests;
+		return Result;
+	}
+	else
+	{
+		return INDEX_NONE;
+	}
+}
+
 #if WITH_EDITORONLY_DATA
 
-UE::DerivedData::FCacheGetChunkRequest FStreamingManager::BuildDDCRequest(const FResources& Resources, uint64 ReadOffset, uint64 ReadSize, uint32 PendingRequestIndex)
+UE::DerivedData::FCacheGetChunkRequest FStreamingManager::BuildDDCRequest(const FResources& Resources, uint32 FirstTileIndex, uint32 NumTiles, uint32 PendingRequestIndex, int32 ChunkIndex)
 {
 	using namespace UE::DerivedData;
 
-	FCacheKey Key;
-	Key.Bucket = FCacheBucket(TEXT("SparseVolumeTexture"));
-	Key.Hash = Resources.DDCKeyHash;
-	check(!Resources.DDCRawHash.IsZero());
+	const uint32 FirstTileIndexInChunk = ChunkIndex > 0 ? (Resources.DDCChunkMaxTileIndices[ChunkIndex - 1] + 1) : Resources.StreamingMetaData.FirstStreamingTileIndex;
+	const uint32 LastTileIndexInChunkPlusOne = Resources.DDCChunkMaxTileIndices[ChunkIndex] + 1;
+	check(FirstTileIndex >= FirstTileIndexInChunk);
+	check((FirstTileIndex + NumTiles) <= LastTileIndexInChunkPlusOne);
+	const uint32 ReadOffsetInChunk = Resources.StreamingMetaData.TileDataOffsets[FirstTileIndex] - Resources.StreamingMetaData.TileDataOffsets[FirstTileIndexInChunk];
+	const uint32 ReadSizeInChunk = Resources.StreamingMetaData.TileDataOffsets[FirstTileIndex + NumTiles] - Resources.StreamingMetaData.TileDataOffsets[FirstTileIndex];
+	const uint32 ChunkTotalSize = Resources.StreamingMetaData.TileDataOffsets[LastTileIndexInChunkPlusOne] - Resources.StreamingMetaData.TileDataOffsets[FirstTileIndexInChunk];
 
 	FCacheGetChunkRequest Request;
-	Request.Id = FValueId::FromName("SparseVolumeTextureStreamingData");
-	Request.Key = Key;
-	Request.RawOffset = ReadOffset;
-	Request.RawSize = ReadSize;
-	Request.RawHash = Resources.DDCRawHash;
+	Request.Id = FValueId(FMemoryView(Resources.DDCChunkIds[ChunkIndex].GetData(), 12));
+	Request.Key.Bucket = FCacheBucket(TEXT("SparseVolumeTexture"));
+	Request.Key.Hash = Resources.DDCKeyHash;
+	if (ReadOffsetInChunk != 0 || ReadSizeInChunk != ChunkTotalSize)
+	{
+		Request.RawOffset = ReadOffsetInChunk;
+		Request.RawSize = ReadSizeInChunk;
+	}
 	Request.UserData = (((uint64)PendingRequestIndex) << uint64(32)) | (uint64)PendingRequests[PendingRequestIndex].RequestVersion;
 	return Request;
 }
diff --git a/Engine/Source/Runtime/Engine/Private/SparseVolumeTexture/SparseVolumeTextureStreamingManager.h b/Engine/Source/Runtime/Engine/Private/SparseVolumeTexture/SparseVolumeTextureStreamingManager.h
--- a/Engine/Source/Runtime/Engine/Private/SparseVolumeTexture/SparseVolumeTextureStreamingManager.h
+++ b/Engine/Source/Runtime/Engine/Private/SparseVolumeTexture/SparseVolumeTextureStreamingManager.h
@@ -182,6 +182,16 @@ private:
 
 		bool IsValid() const { return SVTHandle != uint16(INDEX_NONE); }
 
+		void Set(uint16 InSVTHandle, uint16 InFrameIndex, uint32 InTileOffset, uint32 InTileCount, uint32 InIssuedInFrame, bool bInBlocking)
+		{
+			SVTHandle = InSVTHandle;
+			FrameIndex = InFrameIndex;
+			TileOffset = InTileOffset;
+			TileCount = InTileCount;
+			IssuedInFrame = InIssuedInFrame;
+			bBlocking = bInBlocking;
+		}
+
 		void Reset()
 		{
 			SVTHandle = INDEX_NONE;
@@ -323,9 +333,10 @@ private:
 	void PatchPageTable(FRDGBuilder& GraphBuilder); // Patches the page table to reflect streamed in/out pages and to ensure non-resident mip levels fall back to coarser mip level tile data
 	FStreamingInfo* FindStreamingInfo(uint16 SparseVolumeTextureHandle); // Returns nullptr if the key can't be found
 	FStreamingInfo* FindStreamingInfo(UStreamableSparseVolumeTexture* SparseVolumeTexture); // Returns nullptr if the key can't be found
+	int32 AllocatePendingRequestIndex(); // Allocates a request in the PendingRequests array. Returns INDEX_NONE when out of slots.
 
 #if WITH_EDITORONLY_DATA
-	UE::DerivedData::FCacheGetChunkRequest BuildDDCRequest(const FResources& Resources, uint64 ReadOffset, uint64 ReadSize, uint32 PendingMipLevelIndex);
+	UE::DerivedData::FCacheGetChunkRequest BuildDDCRequest(const FResources& Resources, uint32 FirstTileIndex, uint32 NumTiles, uint32 PendingRequestIndex, int32 ChunkIndex);
 	void RequestDDCData(TConstArrayView<UE::DerivedData::FCacheGetChunkRequest> DDCRequests, bool bBlocking);
 #endif
 };
