Modified FActorInstanceHandle to force re-resolving when loaded, which addressed the issue with hit results passed from client to server (due to InstancedActors not supporting client-side insta-hydration).

diff --git a/Engine/Source/Runtime/Engine/Private/Engine/ActorInstanceHandle.cpp b/Engine/Source/Runtime/Engine/Private/Engine/ActorInstanceHandle.cpp
--- a/Engine/Source/Runtime/Engine/Private/Engine/ActorInstanceHandle.cpp
+++ b/Engine/Source/Runtime/Engine/Private/Engine/ActorInstanceHandle.cpp
@@ -490,6 +490,15 @@ FArchive& operator<<(FArchive& Ar, FActorInstanceHandle& Handle)
 	if (Ar.IsLoading())
 	{
 		Handle.ManagerInterface = FActorInstanceManagerInterface(WeakManagerObject.Get());
+		IActorInstanceManagerInterface* AsManager = Handle.ManagerInterface.Get();
+
+		if (AsManager && Handle.InstanceIndex != INDEX_NONE
+			&& Handle.ManagerInterface.GetObject() == Handle.ReferenceObject)
+		{
+			Handle.ResolutionStatus = FActorInstanceHandle::EResolutionStatus::Invalid;
+			Handle.ReferenceObject = nullptr;
+			Handle.ReferenceObject = AsManager->FindActor(Handle);
+		}
 		Handle.ResolutionStatus = FActorInstanceHandle::EResolutionStatus::Resolved;
 	}
 
