[World Partition] Update Streaming State Optimizations:
- Avoid normalizing already normalized values when constructing a FSphericalSector as a sphere.
- Fix streaming sources hash computations when not using quantized values.
- Fix 2D debug display that wasn't normalizing the velocity normal when rendering velocity vector.
- Added new world partition stats category and added some cycle counters.

diff --git a/Engine/Source/Runtime/Engine/Private/WorldPartition/RuntimeHashSet/WorldPartitionRuntimeHashSetDebugDraw.cpp b/Engine/Source/Runtime/Engine/Private/WorldPartition/RuntimeHashSet/WorldPartitionRuntimeHashSetDebugDraw.cpp
--- a/Engine/Source/Runtime/Engine/Private/WorldPartition/RuntimeHashSet/WorldPartitionRuntimeHashSetDebugDraw.cpp
+++ b/Engine/Source/Runtime/Engine/Private/WorldPartition/RuntimeHashSet/WorldPartitionRuntimeHashSetDebugDraw.cpp
@@ -256,7 +256,7 @@ bool UWorldPartitionRuntimeHashSet::Draw2D(FWorldPartitionDraw2DContext& DrawCon
 			if (!Velocity2D.IsNearlyZero())
 			{
 				const FVector2D Center2D = FVector2D(Source.Location);
-				DrawContext.PushDrawSegment(GridScreenBounds, WorldToScreen(Center2D), WorldToScreen(Center2D + Velocity2D * StreamingData->GetLoadingRange() * 0.5f), Color, 1);
+				DrawContext.PushDrawSegment(GridScreenBounds, WorldToScreen(Center2D), WorldToScreen(Center2D + Velocity2D.GetSafeNormal() * StreamingData->GetLoadingRange() * 0.5f), Color, 1);
 			}
 		}
 	};
diff --git a/Engine/Source/Runtime/Engine/Private/WorldPartition/WorldPartitionStreamingSource.cpp b/Engine/Source/Runtime/Engine/Private/WorldPartition/WorldPartitionStreamingSource.cpp
--- a/Engine/Source/Runtime/Engine/Private/WorldPartition/WorldPartitionStreamingSource.cpp
+++ b/Engine/Source/Runtime/Engine/Private/WorldPartition/WorldPartitionStreamingSource.cpp
@@ -99,43 +99,66 @@ FString FWorldPartitionStreamingSource::ToString() const
 void FWorldPartitionStreamingSource::UpdateHash()
 {
 	// Update old values when they are changing enough, to avoid the case where we are on the edge of a quantization unit.
-	if (!UWorldPartitionStreamingPolicy::IsUpdateStreamingOptimEnabled() || 
-		(FVector::Dist(Location, OldLocation) > LocationQuantization))
+	if (FVector::Dist(Location, QuantizedLocation) > LocationQuantization)
 	{
-		OldLocation = Location;
+		QuantizedLocation = Location;
 	}
 	
-	if (!UWorldPartitionStreamingPolicy::IsUpdateStreamingOptimEnabled() ||
-		(FMath::Abs(Rotation.Pitch - OldRotation.Pitch) > RotationQuantization) ||
-		(FMath::Abs(Rotation.Yaw - OldRotation.Yaw) > RotationQuantization) ||
-		(FMath::Abs(Rotation.Roll - OldRotation.Roll) > RotationQuantization))
+	if (Rotation.GetManhattanDistance(QuantizedRotation) > RotationQuantization)
 	{
-		OldRotation = Rotation;
+		QuantizedRotation = Rotation;
 	}
 
 	FHashBuilder HashBuilder;
-	HashBuilder	<< Name << TargetState << bBlockOnSlowLoading << bReplay << bRemote << Priority << TargetBehavior << TargetGrids << Shapes << ExtraRadius  << ExtraAngle;
+	const uint8 ValuesMask = (bBlockOnSlowLoading ? 1 : 0) | (bUseVelocityContributionToCellsSorting ? 2 : 0) | (bReplay ? 4 : 0) | (bRemote ? 8 : 0) | (bForce2D ? 16 : 0);
+	HashBuilder	<< Name << TargetState << ValuesMask << Priority << TargetBehavior << TargetGrids << Shapes;
 
-	if (LocationQuantization)
+	if (ExtraRadius > 0.f)
+	{
+		HashBuilder << ExtraRadius;
+	}
+
+	if (ExtraAngle > 0.f)
+	{
+		HashBuilder << ExtraAngle;
+	}
+
+	if (LocationQuantization > 0)
+	{
+		HashBuilder << FMath::FloorToInt(QuantizedLocation.X / LocationQuantization) << FMath::FloorToInt(QuantizedLocation.Y / LocationQuantization);
+	}
+	else
 	{
-		HashBuilder << FMath::FloorToInt(OldLocation.X / LocationQuantization) << FMath::FloorToInt(OldLocation.Y / LocationQuantization);
+		HashBuilder << Location.X << Location.Y;
 	}
 
 	if (RotationQuantization > 0)
 	{
-		HashBuilder << FMath::FloorToInt(OldRotation.Yaw / RotationQuantization);
+		HashBuilder << FMath::FloorToInt(Rotation.Yaw / RotationQuantization);
+	}
+	else
+	{
+		HashBuilder << Rotation.Yaw;
 	}
 
 	Hash2D = HashBuilder.GetHash();
 
-	if (LocationQuantization)
+	if (LocationQuantization > 0)
+	{
+		HashBuilder << FMath::FloorToInt(QuantizedLocation.Z / LocationQuantization);
+	}
+	else
 	{
-		HashBuilder << FMath::FloorToInt(OldLocation.Z / LocationQuantization);
+		HashBuilder << Location.Z;
 	}
 
-	if (RotationQuantization)
+	if (RotationQuantization > 0)
+	{
+		HashBuilder << FMath::FloorToInt(QuantizedRotation.Pitch / RotationQuantization) << FMath::FloorToInt(QuantizedRotation.Roll / RotationQuantization);
+	}
+	else
 	{
-		HashBuilder << FMath::FloorToInt(OldRotation.Pitch / RotationQuantization) << FMath::FloorToInt(OldRotation.Roll / RotationQuantization);
+		HashBuilder << QuantizedRotation.Pitch << QuantizedRotation.Roll;
 	}
 
 	Hash3D = HashBuilder.GetHash();
diff --git a/Engine/Source/Runtime/Engine/Private/WorldPartition/WorldPartitionSubsystem.cpp b/Engine/Source/Runtime/Engine/Private/WorldPartition/WorldPartitionSubsystem.cpp
--- a/Engine/Source/Runtime/Engine/Private/WorldPartition/WorldPartitionSubsystem.cpp
+++ b/Engine/Source/Runtime/Engine/Private/WorldPartition/WorldPartitionSubsystem.cpp
@@ -45,7 +45,8 @@
 #include UE_INLINE_GENERATED_CPP_BY_NAME(WorldPartitionSubsystem)
 
 DECLARE_STATS_GROUP(TEXT("World Partition"), STATGROUP_WorldPartition, STATCAT_Advanced);
-DECLARE_CYCLE_STAT(TEXT("UpdateStreamingStateExternal"), STAT_UpdateStreamingStateExternal, STATGROUP_WorldPartition);
+DECLARE_CYCLE_STAT(TEXT("Update Streaming State"), STAT_WorldPartitionUpdateStreamingState, STATGROUP_WorldPartition);
+DECLARE_CYCLE_STAT(TEXT("Update Streaming State (External)"), STAT_WorldPartitionUpdateStreamingStateExternal, STATGROUP_WorldPartition);
 CSV_DEFINE_CATEGORY(WorldPartition, (!UE_BUILD_SHIPPING));
 
 extern int32 GBlockOnSlowStreaming;
@@ -192,9 +193,8 @@ TMulticastDelegate<void(UWorldPartitionSubsystem*, UWorld*)> UWorldPartitionSubs
 TMulticastDelegate<void(UWorldPartitionSubsystem*, UWorld*)> UWorldPartitionSubsystem::OnWorldPartitionSubsystemDeinitialized;
 
 UWorldPartitionSubsystem::UWorldPartitionSubsystem()
-: StreamingSourcesHash(0)
-, NumWorldPartitionServerStreamingEnabled(0)
-, ServerClientsVisibleLevelsHash(0)
+	: NumWorldPartitionServerStreamingEnabled(0)
+	, ServerClientsVisibleLevelsHash(0)
 {}
 
 UWorldPartition* UWorldPartitionSubsystem::GetWorldPartition()
@@ -906,7 +906,7 @@ void UWorldPartitionSubsystem::UpdateStreamingSources()
 
 			{
 				// This will include game-sepcific code called from GetStreamingSourceProviders if IsStreamingSourceProviderFiltered is bound, and also GetStreamingSources calls.
-				SCOPE_CYCLE_COUNTER(STAT_UpdateStreamingStateExternal);
+				SCOPE_CYCLE_COUNTER(STAT_WorldPartitionUpdateStreamingStateExternal);
 				CSV_SCOPED_TIMING_STAT(WorldPartition, GetStreamingSourceProviders);
 
 				TArray<FWorldPartitionStreamingSource> ProviderStreamingSources;
@@ -1008,12 +1008,11 @@ void UWorldPartitionSubsystem::GetStreamingSources(const UWorldPartition* InWorl
 		{
 			StreamingSource.Location = WorldToLocal.TransformPosition(StreamingSource.Location);
 			StreamingSource.Rotation = WorldToLocal.TransformRotation(StreamingSource.Rotation.Quaternion()).Rotator();
+			StreamingSource.Velocity = WorldToLocal.TransformVector(StreamingSource.Velocity);
 		}
 	}
 }
 
-DECLARE_CYCLE_STAT(TEXT("World Partition Update Streaming"), STAT_WorldPartitionUpdateStreaming, STATGROUP_Engine);
-
 void UWorldPartitionSubsystem::OnWorldBeginPlay(UWorld& InWorld)
 {
 	Super::OnWorldBeginPlay(InWorld);
@@ -1026,7 +1025,7 @@ void UWorldPartitionSubsystem::OnWorldBeginPlay(UWorld& InWorld)
 
 void UWorldPartitionSubsystem::UpdateStreamingState()
 {
-	SCOPE_CYCLE_COUNTER(STAT_WorldPartitionUpdateStreaming);
+	SCOPE_CYCLE_COUNTER(STAT_WorldPartitionUpdateStreamingState);
 
 	UWorldPartitionSubsystem::UpdateStreamingStateInternal(GetWorld());
 }
diff --git a/Engine/Source/Runtime/Engine/Public/WorldPartition/WorldPartitionStreamingSource.h b/Engine/Source/Runtime/Engine/Public/WorldPartition/WorldPartitionStreamingSource.h
--- a/Engine/Source/Runtime/Engine/Public/WorldPartition/WorldPartitionStreamingSource.h
+++ b/Engine/Source/Runtime/Engine/Public/WorldPartition/WorldPartitionStreamingSource.h
@@ -17,17 +17,26 @@ public:
 		: Center(ForceInit)
 		, Radius(0.0f)
 		, Axis(ForceInit)
-	{
-		SetAsSphere();
-	}
+		, Angle(360.0f)
+	{}
 
 	/** Creates and initializes a spherical sector using given parameters. */
-	FSphericalSector(FVector InCenter, FReal InRadius, FVector InAxis = FVector::ForwardVector, FReal InAngle = 0)
+	FSphericalSector(const FVector& InCenter, FReal InRadius)
 		: Center(InCenter)
 		, Radius(InRadius)
+		, Axis(FVector::ForwardVector)
+		, Angle(360.0f)
+	{}
+
+	/** Creates and initializes a spherical sector using given parameters. */
+	FSphericalSector(const FVector& InCenter, FReal InRadius, const FVector& InAxis, FReal InAngle)
+		: Center(InCenter)
+		, Radius(InRadius)
+		, Axis(InAxis)
+		, Angle(InAngle)
 	{
-		SetAngle(InAngle);
-		SetAxis(InAxis);
+		check(InAxis.IsNormalized());
+		check(InAngle >= 0 && InAngle <= 360);
 	}
 
 	void SetCenter(const FVector& InCenter) { Center = InCenter; }
@@ -43,7 +52,7 @@ public:
 	FVector GetAxis() const { return Axis; }
 	FVector GetScaledAxis() const { return Axis * Radius; }
 
-	void SetAsSphere() { SetAngle(360.0f); }
+	void SetAsSphere() { Angle = 360.0f; }
 	bool IsSphere() const { return FMath::IsNearlyEqual(Angle, (FReal)360.0); }
 
 	bool IsNearlyZero() const { return FMath::IsNearlyZero(Radius) || Axis.IsNearlyZero() || FMath::IsNearlyZero(Angle); }
@@ -58,8 +67,7 @@ public:
 	/** Get result of Transforming spherical sector with transform. */
 	FSphericalSector TransformBy(const FTransform& M) const
 	{
-		FSphericalSector Result(M.TransformPosition(Center), M.GetMaximumAxisScale() * Radius, M.TransformVector(Axis), Angle);
-		return Result;
+		return FSphericalSector(M.TransformPosition(Center), M.GetMaximumAxisScale() * Radius, M.TransformVector(Axis), Angle);
 	}
 
 	/** Helper method that builds a list of debug display segments */
@@ -175,9 +183,11 @@ public:
 		const FTransform Transform(bInProjectIn2D ? FRotator(0, InRotation.Yaw, 0) : InRotation, InLocation);
 		if (InShapes.IsEmpty())
 		{
-			FSphericalSector LocalShape(FVector::ZeroVector, InDefaultRadius + InExtraRadius);
-			if (LocalShape.IsValid())
+			if (InDefaultRadius + InExtraRadius)
 			{
+				const FSphericalSector LocalShape(FVector::ZeroVector, InDefaultRadius + InExtraRadius);
+				check(LocalShape.IsValid()); // Radius, axis and angle should be all valid here
+
 				InOperation(LocalShape.TransformBy(Transform));
 			}
 		}
@@ -186,9 +196,11 @@ public:
 			for (const FStreamingSourceShape& Shape : InShapes)
 			{
 				const FVector::FReal ShapeRadius = (Shape.bUseGridLoadingRange ? (InGridLoadingRange * Shape.LoadingRangeScale) : Shape.Radius) + InExtraRadius;
-				const FVector::FReal ShapeAngle = Shape.bIsSector ? (Shape.SectorAngle + InExtraAngle) : 360.0f;
-				const FVector ShapeAxis = bInProjectIn2D ? FRotator(0, Shape.Rotation.Yaw, 0).Vector() : Shape.Rotation.Vector();
-				FSphericalSector LocalShape(bInProjectIn2D ? FVector(Shape.Location.X, Shape.Location.Y, 0) : Shape.Location, ShapeRadius, ShapeAxis, ShapeAngle);
+				const FVector::FReal ShapeAngle = Shape.bIsSector ? FMath::Min(Shape.SectorAngle + InExtraAngle, 360.0f) : 360.0f;
+				const FRotator ShapeRotation = bInProjectIn2D ? FRotator(0, Shape.Rotation.Yaw, 0) : Shape.Rotation;
+				const FVector ShapeAxis = ShapeRotation.IsNearlyZero() ? FVector::ForwardVector : ShapeRotation.Vector();
+				const FSphericalSector LocalShape(bInProjectIn2D ? FVector(Shape.Location.X, Shape.Location.Y, 0) : Shape.Location, ShapeRadius, ShapeAxis, ShapeAngle);
+				
 				if (LocalShape.IsValid())
 				{
 					InOperation(LocalShape.TransformBy(Transform));
@@ -337,8 +349,8 @@ struct FWorldPartitionStreamingSource
 		, bForce2D(false)
 		, Hash2D(0)
 		, Hash3D(0)
-		, OldLocation(FVector::ZeroVector)
-		, OldRotation(FRotator::ZeroRotator)
+		, QuantizedLocation(FVector::ZeroVector)
+		, QuantizedRotation(FRotator::ZeroRotator)
 		, ExtraRadius(0)
 		, ExtraAngle(0)
 	{}
@@ -359,8 +371,8 @@ struct FWorldPartitionStreamingSource
 		, bForce2D(false)
 		, Hash2D(0)
 		, Hash3D(0)
-		, OldLocation(InLocation)
-		, OldRotation(InRotation)
+		, QuantizedLocation(InLocation)
+		, QuantizedRotation(InRotation)
 		, ExtraRadius(0)
 		, ExtraAngle(0)
 	{}
@@ -476,8 +488,8 @@ private:
 	uint32 Hash3D;
 
 	/** Source values used for hash computations. */
-	FVector OldLocation;
-	FRotator OldRotation;
+	FVector QuantizedLocation;
+	FRotator QuantizedRotation;
 
 	/** Used internally for server streaming */
 	float ExtraRadius;
