Fix implicit conversion floating-point number into bool warning.

diff --git a/Engine/Source/Runtime/Engine/Public/WorldPartition/WorldPartitionStreamingSource.h b/Engine/Source/Runtime/Engine/Public/WorldPartition/WorldPartitionStreamingSource.h
--- a/Engine/Source/Runtime/Engine/Public/WorldPartition/WorldPartitionStreamingSource.h
+++ b/Engine/Source/Runtime/Engine/Public/WorldPartition/WorldPartitionStreamingSource.h
@@ -183,7 +183,7 @@ public:
 		const FTransform Transform(bInProjectIn2D ? FRotator(0, InRotation.Yaw, 0) : InRotation, InLocation);
 		if (InShapes.IsEmpty())
 		{
-			if (InDefaultRadius + InExtraRadius)
+			if ((InDefaultRadius + InExtraRadius) > 0)
 			{
 				const FSphericalSector LocalShape(FVector::ZeroVector, InDefaultRadius + InExtraRadius);
 				check(LocalShape.IsValid()); // Radius, axis and angle should be all valid here
