Fix mip streaming to use GetValue rather than GetChunks so that data can propagate between the caching layers. Currently if data doesn't exist locally, but does remotely, we will consistenly miss every time locally, and have to pull data from the remote since GetChunks does not currently support propogating data down to other cache layers.

This removes thousands of cache misses that can happen every time you load the editor with a cache less populated than the remote.

The codepath is now clearly delineated between DDC1 and DDC2 which is unfortunate, but once chunk propagation is supported we can reduce this codepath to what it was previously if desired.

The debug name provided for mip streaming now also clarifies which DDC requests are initiated by the prefetcher by tagging them with "[Prefetch]" similar to "[Fast]/[Final]" tags

diff --git a/Engine/Source/Runtime/Engine/Private/Streaming/Texture2DStreamIn_DDC.cpp b/Engine/Source/Runtime/Engine/Private/Streaming/Texture2DStreamIn_DDC.cpp
--- a/Engine/Source/Runtime/Engine/Private/Streaming/Texture2DStreamIn_DDC.cpp
+++ b/Engine/Source/Runtime/Engine/Private/Streaming/Texture2DStreamIn_DDC.cpp
@@ -53,13 +53,16 @@ void FTexture2DStreamIn_DDC::DoCreateAsyncDDCRequests(const FContext& Context)
 		const int32 LODBias = static_cast<int32>(Context.MipsView.GetData() - PlatformData->Mips.GetData());
 
 		using namespace UE::DerivedData;
-		TArray<FCacheGetChunkRequest> MipKeys;
+		TArray<FCacheGetValueRequest> DDC1MipKeys;
+		TArray<FCacheGetChunkRequest> DDC2MipKeys;
 
 		TStringBuilder<256> MipNameBuilder;
 		Context.Texture->GetPathName(nullptr, MipNameBuilder);
 		const int32 TextureNameLen = MipNameBuilder.Len();
 
-		if (PlatformData->DerivedDataKey.IsType<FString>())
+		const bool bUsingDDC1 = PlatformData->DerivedDataKey.IsType<FString>();
+		const bool bUsingDDC2 = PlatformData->DerivedDataKey.IsType<UE::DerivedData::FCacheKeyProxy>();
+		if (bUsingDDC1)
 		{
 			for (int32 MipIndex = PendingFirstLODIdx; MipIndex < CurrentFirstLODIdx && !IsCancelled(); ++MipIndex)
 			{
@@ -73,7 +76,7 @@ void FTexture2DStreamIn_DDC::DoCreateAsyncDDCRequests(const FContext& Context)
 				}
 				else if (!Status.bRequestIssued && !Status.Buffer)
 				{
-					FCacheGetChunkRequest& Request = MipKeys.AddDefaulted_GetRef();
+					FCacheGetValueRequest& Request = DDC1MipKeys.AddDefaulted_GetRef();
 					MipNameBuilder.Appendf(TEXT(" [MIP %d]"), MipIndex + LODBias);
 					Request.Name = MipNameBuilder;
 					Request.Key = ConvertLegacyCacheKey(PlatformData->GetDerivedDataMipKeyString(MipIndex + LODBias, MipMap));
@@ -83,7 +86,7 @@ void FTexture2DStreamIn_DDC::DoCreateAsyncDDCRequests(const FContext& Context)
 				}
 			}
 		}
-		else if (PlatformData->DerivedDataKey.IsType<UE::DerivedData::FCacheKeyProxy>())
+		else if (bUsingDDC2)
 		{
 			const FCacheKey& Key = *PlatformData->DerivedDataKey.Get<UE::DerivedData::FCacheKeyProxy>().AsCacheKey();
 			for (int32 MipIndex = PendingFirstLODIdx; MipIndex < CurrentFirstLODIdx && !IsCancelled(); ++MipIndex)
@@ -92,7 +95,7 @@ void FTexture2DStreamIn_DDC::DoCreateAsyncDDCRequests(const FContext& Context)
 				FMipRequestStatus& Status = DDCMipRequestStatus[MipIndex];
 				if (MipMap.IsPagedToDerivedData() && !Status.bRequestIssued && !Status.Buffer)
 				{
-					FCacheGetChunkRequest& Request = MipKeys.AddDefaulted_GetRef();
+					FCacheGetChunkRequest& Request = DDC2MipKeys.AddDefaulted_GetRef();
 					MipNameBuilder.Appendf(TEXT(" [MIP %d]"), MipIndex + LODBias);
 					Request.Name = MipNameBuilder;
 					Request.Key = Key;
@@ -110,33 +113,66 @@ void FTexture2DStreamIn_DDC::DoCreateAsyncDDCRequests(const FContext& Context)
 			MarkAsCancelled();
 		}
 
-		if (MipKeys.Num())
+		if (DDC1MipKeys.Num() || DDC2MipKeys.Num())
 		{
 		#if !UE_BUILD_SHIPPING
 			// On some platforms the IO is too fast to test cancellation requests timing issues.
 			if (FRenderAssetStreamingSettings::ExtraIOLatency > 0 && TaskSynchronization.GetValue() == 0)
 			{
-				FPlatformProcess::Sleep(MipKeys.Num() * FRenderAssetStreamingSettings::ExtraIOLatency * 0.001f); // Slow down the streaming.
+				const int32 NumMipKeys = FMath::Max(DDC1MipKeys.Num(), DDC2MipKeys.Num());
+				FPlatformProcess::Sleep(NumMipKeys * FRenderAssetStreamingSettings::ExtraIOLatency * 0.001f); // Slow down the streaming.
 			}
 		#endif
 
 			FRequestBarrier Barrier(DDCRequestOwner);
-			GetCache().GetChunks(MipKeys, DDCRequestOwner, [this](FCacheGetChunkResponse&& Response)
+
+			if (bUsingDDC1)
 			{
-				if (Response.Status == EStatus::Ok)
-				{
-					const int32 MipIndex = int32(Response.UserData);
-					FMipRequestStatus& Status = DDCMipRequestStatus[MipIndex];
-					check(!Status.Buffer);
-					Status.Buffer = MoveTemp(Response.RawData);
-					check(Status.bRequestIssued);
-					Status.bRequestIssued = false;
-				}
-				else if (Response.Status == EStatus::Error)
-				{
-					FTextureCompilingManager::Get().ForceDeferredTextureRebuildAnyThread({Texture});
-				}
-			});
+				GetCache().GetValue(DDC1MipKeys, DDCRequestOwner, [this](FCacheGetValueResponse&& Response)
+					{
+						bool bOk = Response.Status == EStatus::Ok;
+						if (bOk)
+						{
+							if (FSharedBuffer&& MipBuffer = Response.Value.GetData().Decompress())
+							{
+								const int32 MipIndex = int32(Response.UserData);
+								FMipRequestStatus& Status = DDCMipRequestStatus[MipIndex];
+								check(!Status.Buffer);
+								Status.Buffer = MipBuffer;
+								check(Status.bRequestIssued);
+								Status.bRequestIssued = false;
+							}
+							else
+							{
+								bOk = false;
+							}
+						}
+
+						if (!bOk)
+						{
+							FTextureCompilingManager::Get().ForceDeferredTextureRebuildAnyThread({ Texture });
+						}
+					});
+			}
+			else if (bUsingDDC2)
+			{
+				GetCache().GetChunks(DDC2MipKeys, DDCRequestOwner, [this](FCacheGetChunkResponse&& Response)
+					{
+						if (Response.Status == EStatus::Ok)
+						{
+							const int32 MipIndex = int32(Response.UserData);
+							FMipRequestStatus& Status = DDCMipRequestStatus[MipIndex];
+							check(!Status.Buffer);
+							Status.Buffer = MoveTemp(Response.RawData);
+							check(Status.bRequestIssued);
+							Status.bRequestIssued = false;
+						}
+						else if (Response.Status == EStatus::Error)
+						{
+							FTextureCompilingManager::Get().ForceDeferredTextureRebuildAnyThread({ Texture });
+						}
+					});
+			}
 		}
 	}
 	else
diff --git a/Engine/Source/Runtime/Engine/Private/TextureDerivedDataTask.cpp b/Engine/Source/Runtime/Engine/Private/TextureDerivedDataTask.cpp
--- a/Engine/Source/Runtime/Engine/Private/TextureDerivedDataTask.cpp
+++ b/Engine/Source/Runtime/Engine/Private/TextureDerivedDataTask.cpp
@@ -1404,7 +1404,7 @@ static bool TryCacheStreamingMips(const FString& TexturePathName, int32 FirstMip
 
 	const int32 LowestMipIndexToPrefetchOrLoad = FMath::Min(FirstMipToPrefetch, FirstMipToLoad);
 	const int32 NumMips = DerivedData->Mips.Num();
-	const FSharedString Name(TexturePathName);
+	const FSharedString Name(WriteToString<256>(TexturePathName, TEXTVIEW(" [Prefetch]")));
 	for (int32 MipIndex = LowestMipIndexToPrefetchOrLoad; MipIndex < NumMips; ++MipIndex)
 	{
 		const FTexture2DMipMap& Mip = DerivedData->Mips[MipIndex];
