Re-enable fatal shader compiler errors and have the ODSC server behave like the editor where we offer a message box to recompile

diff --git a/Engine/Source/Runtime/Engine/Private/ShaderCompiler/ShaderCompiler.cpp b/Engine/Source/Runtime/Engine/Private/ShaderCompiler/ShaderCompiler.cpp
--- a/Engine/Source/Runtime/Engine/Private/ShaderCompiler/ShaderCompiler.cpp
+++ b/Engine/Source/Runtime/Engine/Private/ShaderCompiler/ShaderCompiler.cpp
@@ -228,6 +228,14 @@ static TAutoConsoleVariable<bool> CVarAreShaderErrorsFatal(
 	TEXT("Default: true"),
 	ECVF_RenderThreadSafe);
 
+// Used for remote compile by COTF: Allow to prompt the message box to the user in case of a shader error
+static bool GRetryShaderCompileOnErrorForRemote = false;
+
+static inline bool IsLogShadersActive()
+{
+	return UE_LOG_ACTIVE(LogShaders, Log) || GRetryShaderCompileOnErrorForRemote;
+}
+
 bool AreShaderErrorsFatal()
 {
 	return CVarAreShaderErrorsFatal.GetValueOnAnyThread();
@@ -3330,7 +3338,7 @@ static int32 AddAndProcessErrorsForFailedJobFiltered(FShaderCompileJob& CurrentJ
 		FString CurrentErrorString = CurrentError.GetErrorString();
 
 		// Include warnings if LogShaders is unsuppressed, otherwise only include filtered messages
-		if (UE_LOG_ACTIVE(LogShaders, Log) || FilterMessage == nullptr || CurrentError.StrippedErrorMessage.Contains(FilterMessage))
+		if (IsLogShadersActive() || FilterMessage == nullptr || CurrentError.StrippedErrorMessage.Contains(FilterMessage))
 		{
 			// Extract source location from error message if the shader backend doesn't provide it separated from the stripped message
 			CurrentError.ExtractSourceLocation();
@@ -7143,7 +7151,7 @@ static bool GatherUniqueErrors(const TArray<FShaderCommonCompileJobPtr>& Complet
 
 bool FShaderCompilingManager::HandlePotentialRetryOnError(TMap<int32, FShaderMapFinalizeResults>& CompletedShaderMaps)
 {
-	if (FApp::IsUnattended())
+	if (FApp::IsUnattended() && !GRetryShaderCompileOnErrorForRemote)
 	{
 		return false;
 	}
@@ -7172,7 +7180,7 @@ bool FShaderCompilingManager::HandlePotentialRetryOnError(TMap<int32, FShaderMap
 				}
 			}
 
-			if (UE_LOG_ACTIVE(LogShaders, Log) 
+			if (IsLogShadersActive()
 				// Always log detailed errors when a special engine material or global shader fails to compile, as those will be fatal errors
 				|| bSpecialEngineMaterial 
 				|| It.Key() == GlobalShaderMapId)
@@ -7194,7 +7202,7 @@ bool FShaderCompilingManager::HandlePotentialRetryOnError(TMap<int32, FShaderMap
 
 				BuildErrorStringAndReport(ShaderErrorInfo, ErrorString);
 
-				if (UE_LOG_ACTIVE(LogShaders, Log) && (bAnyErrorLikelyToBeCodeError || bPromptToRetryFailedShaderCompiles || bSpecialEngineMaterial))
+				if (IsLogShadersActive() && (bAnyErrorLikelyToBeCodeError || bPromptToRetryFailedShaderCompiles || bSpecialEngineMaterial))
 				{
 					// Use debug break in debug with the debugger attached, otherwise message box
 					if (bDebugBreakOnPromptToRetryShaderCompile && FPlatformMisc::IsDebuggerPresent())
@@ -10429,8 +10437,12 @@ void RecompileShadersForRemote(
 {
 	TRACE_CPUPROFILER_EVENT_SCOPE(RecompileShadersForRemote);
 
-	CVarAreShaderErrorsFatal->Set(0, ECVF_SetByCode);
-	ON_SCOPE_EXIT{ CVarAreShaderErrorsFatal->Unset(ECVF_SetByCode); };
+	GRetryShaderCompileOnErrorForRemote = true;
+
+	ON_SCOPE_EXIT
+	{ 
+		GRetryShaderCompileOnErrorForRemote = false;
+	};
 
 	// figure out what shader platforms to recompile
 	ITargetPlatformManagerModule* TPM = GetTargetPlatformManager();
