[AnimRuntime] Added an optional interpolation override when sampling animations.
This makes it possible to force sampling animations in Step interpolation mode even though the Anim Sequence assets are set to Linear.

diff --git a/Engine/Plugins/Animation/ACLPlugin/Source/ACLPluginEditor/Private/ACLStatsDumpCommandlet.cpp b/Engine/Plugins/Animation/ACLPlugin/Source/ACLPluginEditor/Private/ACLStatsDumpCommandlet.cpp
--- a/Engine/Plugins/Animation/ACLPlugin/Source/ACLPluginEditor/Private/ACLStatsDumpCommandlet.cpp
+++ b/Engine/Plugins/Animation/ACLPlugin/Source/ACLPluginEditor/Private/ACLStatsDumpCommandlet.cpp
@@ -246,7 +246,7 @@ static void SampleUE4Clip(const acl::track_array_qvvf& Tracks, USkeleton* UE4Ske
 		FTransform BoneTransform;
 		if (BoneTreeIndex != INDEX_NONE && Model->IsValidBoneTrackName(BoneName))
 		{
-			UE4Clip->GetBoneTransform(BoneTransform, FSkeletonPoseBoneIndex(BoneTreeIndex), double(SampleTime), false);
+			UE4Clip->GetBoneTransform(BoneTransform, FSkeletonPoseBoneIndex(BoneTreeIndex), double(SampleTime), false, TOptional<EAnimInterpolationType>());
 		}
 		else
 		{
diff --git a/Engine/Source/Runtime/Engine/Classes/Animation/AnimSequence.h b/Engine/Source/Runtime/Engine/Classes/Animation/AnimSequence.h
--- a/Engine/Source/Runtime/Engine/Classes/Animation/AnimSequence.h
+++ b/Engine/Source/Runtime/Engine/Classes/Animation/AnimSequence.h
@@ -541,10 +541,12 @@ public:
 	 * @param	TrackIndex		Index of track to interpolate.
 	 * @param	Time			Time on track to interpolate to.
 	 * @param	bUseRawData		If true, use raw animation data instead of compressed data.
+	 * @param   InterpolationOverride The optional animation interpolation type to use. If not set, it will use the AnimSequence interpolation type.
 	 */
+	ENGINE_API void GetBoneTransform(FTransform& OutAtom, FSkeletonPoseBoneIndex BoneIndex, double Time, bool bUseRawData, TOptional<EAnimInterpolationType> InterpolationOverride=TOptional<EAnimInterpolationType>()) const;
+
 	UE_DEPRECATED(5.1, "Use other GetBoneTransform signature using double and skeleton index")
 	void GetBoneTransform(FTransform& OutAtom, int32 TrackIndex, float Time, bool bUseRawData) const {}
-	ENGINE_API void GetBoneTransform(FTransform& OutAtom, FSkeletonPoseBoneIndex BoneIndex, double Time, bool bUseRawData) const;
 
 	/**
 	 * Get Bone Transform of the Time given, relative to Parent for the Track Given
diff --git a/Engine/Source/Runtime/Engine/Classes/Animation/AnimSingleNodeInstance.h b/Engine/Source/Runtime/Engine/Classes/Animation/AnimSingleNodeInstance.h
--- a/Engine/Source/Runtime/Engine/Classes/Animation/AnimSingleNodeInstance.h
+++ b/Engine/Source/Runtime/Engine/Classes/Animation/AnimSingleNodeInstance.h
@@ -10,7 +10,9 @@
 
 #include "CoreMinimal.h"
 #include "UObject/ObjectMacros.h"
+#include "Misc/Optional.h"
 #include "Animation/AnimInstance.h"
+#include "Animation/AnimTypes.h"
 #include "AnimSingleNodeInstance.generated.h"
 
 struct FAnimInstanceProxy;
@@ -77,6 +79,7 @@ public:
 	/** Set pose value */
  	UFUNCTION(BlueprintCallable, Category = "Animation")
  	ENGINE_API void SetPreviewCurveOverride(const FName& PoseName, float Value, bool bRemoveIfZero);
+
 public:
 	/** Gets the current state of any BlendSpace */
 	ENGINE_API void GetBlendSpaceState(FVector& OutPosition, FVector& OutFilteredPosition) const;
@@ -118,6 +121,14 @@ public:
 
 	/** Get the last filter output */
 	ENGINE_API FVector GetFilterLastOutput();
+
+	/** Set animation interpolation type override. If not set, it will not override. */
+	ENGINE_API void SetInterpolationOverride(TOptional<EAnimInterpolationType> InterpolationType);
+
+	/** Get animation interpolation type override. If not set, it will not override. */
+ 	ENGINE_API TOptional<EAnimInterpolationType> GetInterpolationOverride() const;
+
+
 protected:
 	ENGINE_API virtual FAnimInstanceProxy* CreateAnimInstanceProxy() override;
 };
diff --git a/Engine/Source/Runtime/Engine/Classes/Animation/AnimationAsset.h b/Engine/Source/Runtime/Engine/Classes/Animation/AnimationAsset.h
--- a/Engine/Source/Runtime/Engine/Classes/Animation/AnimationAsset.h
+++ b/Engine/Source/Runtime/Engine/Classes/Animation/AnimationAsset.h
@@ -282,6 +282,14 @@ struct FAnimExtractContext
 	 */
 	TArray<bool> BonesRequired;
 
+	/** 
+	 * The optional interpolation mode override.
+	 * If not set, it will simply use the interpolation mode provided by the asset.
+	 * One example where this could be used is if you want to force sampling the animation with Step interpolation
+	 * even when the animation sequence asset is set to Linear interpolation.
+	 */
+	TOptional<EAnimInterpolationType> InterpolationOverride;
+
 #if WITH_EDITOR
 	bool bIgnoreRootLock;
 #endif 
@@ -294,6 +302,7 @@ struct FAnimExtractContext
 		, bLooping(InbLooping)
 		, PoseCurves()
 		, BonesRequired()
+		, InterpolationOverride()
 #if WITH_EDITOR
 		, bIgnoreRootLock(false)
 #endif 
@@ -307,6 +316,7 @@ struct FAnimExtractContext
 		, bLooping(InbLooping)
 		, PoseCurves()
 		, BonesRequired()
+		, InterpolationOverride()
 #if WITH_EDITOR
 		, bIgnoreRootLock(false)
 #endif 
diff --git a/Engine/Source/Runtime/Engine/Private/Animation/AnimInstanceProxy.cpp b/Engine/Source/Runtime/Engine/Private/Animation/AnimInstanceProxy.cpp
--- a/Engine/Source/Runtime/Engine/Private/Animation/AnimInstanceProxy.cpp
+++ b/Engine/Source/Runtime/Engine/Private/Animation/AnimInstanceProxy.cpp
@@ -1500,7 +1500,6 @@ void FAnimInstanceProxy::SlotEvaluatePose(const FName& SlotNodeName, const FComp
 	SlotEvaluatePose(SlotNodeName, SourceAnimationPoseData, InSourceWeight, BlendedAnimationPoseData, InBlendWeight, InTotalNodeWeight);
 }
 
-
 void FAnimInstanceProxy::SlotEvaluatePoseWithBlendProfiles(const FName& SlotNodeName, const FAnimationPoseData& SourceAnimationPoseData, float InSourceWeight, FAnimationPoseData& OutBlendedAnimationPoseData, float InBlendWeight)
 {
 	const FCompactPose& SourcePose = SourceAnimationPoseData.GetPose();
@@ -1581,6 +1580,7 @@ void FAnimInstanceProxy::SlotEvaluatePoseWithBlendProfiles(const FName& SlotNode
 
 			// Extract pose from Track.
 			FAnimExtractContext ExtractionContext(static_cast<double>(EvalState.MontagePosition), Montage->HasRootMotion() && RootMotionMode != ERootMotionMode::NoRootMotionExtraction, EvalState.DeltaTimeRecord);
+			ExtractionContext.InterpolationOverride = InterpolationOverride;
 			FAnimationPoseData NewAnimationPoseData(NewPose);
 			AnimTrack->GetAnimationPose(NewAnimationPoseData, ExtractionContext);
 
@@ -1920,6 +1920,7 @@ void FAnimInstanceProxy::SlotEvaluatePose(const FName& SlotNodeName, const FAnim
 
 			// Extract pose from Track
 			FAnimExtractContext ExtractionContext(static_cast<double>(EvalState.MontagePosition), Montage->HasRootMotion() && RootMotionMode != ERootMotionMode::NoRootMotionExtraction, EvalState.DeltaTimeRecord);
+			ExtractionContext.InterpolationOverride = InterpolationOverride;
 
 			FAnimationPoseData NewAnimationPoseData(NewPose);
 			AnimTrack->GetAnimationPose(NewAnimationPoseData, ExtractionContext);
@@ -3694,5 +3695,16 @@ void FAnimInstanceProxy::ForEachStateMachine(const TFunctionRef<void(FAnimNode_S
 		}
 	}
 }
+
+void FAnimInstanceProxy::SetInterpolationOverride(TOptional<EAnimInterpolationType> InterpolationType)
+{
+	InterpolationOverride = InterpolationType;
+}
+
+TOptional<EAnimInterpolationType> FAnimInstanceProxy::GetInterpolationOverride() const
+{
+	return InterpolationOverride;
+}
+
 #undef LOCTEXT_NAMESPACE
 
diff --git a/Engine/Source/Runtime/Engine/Private/Animation/AnimSequence.cpp b/Engine/Source/Runtime/Engine/Private/Animation/AnimSequence.cpp
--- a/Engine/Source/Runtime/Engine/Private/Animation/AnimSequence.cpp
+++ b/Engine/Source/Runtime/Engine/Private/Animation/AnimSequence.cpp
@@ -1259,7 +1259,8 @@ static FAutoConsoleVariableRef CVarForceRawData(
 	ECVF_Cheat);
 #endif // WITH_EDITOR
 
-void UAnimSequence::GetBoneTransform(FTransform& OutAtom, FSkeletonPoseBoneIndex BoneIndex, double Time, bool bUseRawData) const
+
+void UAnimSequence::GetBoneTransform(FTransform& OutAtom, FSkeletonPoseBoneIndex BoneIndex, double Time, bool bUseRawData, TOptional<EAnimInterpolationType> InterpolationOverride) const
 {
 	// If the caller didn't request that raw animation data be used . . .
 	const bool bEvaluateCompressedData =
@@ -1268,9 +1269,10 @@ void UAnimSequence::GetBoneTransform(FTransform& OutAtom, FSkeletonPoseBoneIndex
 #endif 
 		(!bUseRawData || !CanEvaluateRawAnimationData());
 
+	const EAnimInterpolationType InterpolationType = InterpolationOverride.Get(Interpolation);
 	if ( bEvaluateCompressedData && IsCompressedDataValid())
 	{
-		FAnimSequenceDecompressionContext DecompContext(PlatformTargetFrameRate.Default, PlatformTargetFrameRate.Default.AsFrameTime(GetPlayLength()).RoundToFrame().Value, Interpolation, GetRetargetTransformsSourceName(), *CompressedData.CompressedDataStructure, GetSkeleton()->GetRefLocalPoses(), CompressedData.CompressedTrackToSkeletonMapTable, GetSkeleton(), IsValidAdditive(), AdditiveAnimType);
+		FAnimSequenceDecompressionContext DecompContext(PlatformTargetFrameRate.Default, PlatformTargetFrameRate.Default.AsFrameTime(GetPlayLength()).RoundToFrame().Value, InterpolationType, GetRetargetTransformsSourceName(), *CompressedData.CompressedDataStructure, GetSkeleton()->GetRefLocalPoses(), CompressedData.CompressedTrackToSkeletonMapTable, GetSkeleton(), IsValidAdditive(), AdditiveAnimType);
 		DecompContext.Seek(Time);
 		if (CompressedData.BoneCompressionCodec)
 		{
@@ -1286,7 +1288,7 @@ void UAnimSequence::GetBoneTransform(FTransform& OutAtom, FSkeletonPoseBoneIndex
 #if WITH_EDITOR
 		ValidateModel();
 		const FName BoneName = GetSkeleton()->GetReferenceSkeleton().GetBoneName(BoneIndex.GetInt());
-		OutAtom = DataModelInterface->EvaluateBoneTrackTransform(BoneName, DataModelInterface->GetFrameRate().AsFrameTime(Time), Interpolation);
+		OutAtom = DataModelInterface->EvaluateBoneTrackTransform(BoneName, DataModelInterface->GetFrameRate().AsFrameTime(Time), InterpolationType);
 
 		const FAnimationCurveIdentifier TransformCurveId(BoneName, ERawCurveTrackTypes::RCT_Transform);
 		if (const FTransformCurve* TransformCurvePtr = DataModelInterface->FindTransformCurve(TransformCurveId))
@@ -1678,7 +1680,8 @@ void UAnimSequence::GetBonePose(FAnimationPoseData& OutAnimationPoseData, const
 	if (bUseRawDataForPoseExtraction)
 	{
 		{
-			const UE::Anim::DataModel::FEvaluationContext EvaluationContext(ExtractionContext.CurrentTime, DataModelInterface->GetFrameRate(), GetRetargetTransformsSourceName(), GetRetargetTransforms(), Interpolation);
+			const EAnimInterpolationType InterpolationType = ExtractionContext.InterpolationOverride.Get(Interpolation);
+			const UE::Anim::DataModel::FEvaluationContext EvaluationContext(ExtractionContext.CurrentTime, DataModelInterface->GetFrameRate(), GetRetargetTransformsSourceName(), GetRetargetTransforms(), InterpolationType);
 			DataModelInterface->Evaluate(OutAnimationPoseData, EvaluationContext);
 		}
 
@@ -1695,7 +1698,8 @@ void UAnimSequence::GetBonePose(FAnimationPoseData& OutAnimationPoseData, const
 	if (NumTracks != 0)
 	{
 		// Evaluate compressed bone data
-		FAnimSequenceDecompressionContext DecompContext(PlatformTargetFrameRate.Default, PlatformTargetFrameRate.Default.AsFrameTime(GetPlayLength()).RoundToFrame().Value, Interpolation, GetRetargetTransformsSourceName(), *CompressedData.CompressedDataStructure, GetSkeleton()->GetRefLocalPoses(), CompressedData.CompressedTrackToSkeletonMapTable, GetSkeleton(), IsValidAdditive(), AdditiveAnimType);
+		const EAnimInterpolationType InterpolationType = ExtractionContext.InterpolationOverride.Get(Interpolation);
+		FAnimSequenceDecompressionContext DecompContext(PlatformTargetFrameRate.Default, PlatformTargetFrameRate.Default.AsFrameTime(GetPlayLength()).RoundToFrame().Value, InterpolationType, GetRetargetTransformsSourceName(), *CompressedData.CompressedDataStructure, GetSkeleton()->GetRefLocalPoses(), CompressedData.CompressedTrackToSkeletonMapTable, GetSkeleton(), IsValidAdditive(), AdditiveAnimType);
 		UE::Anim::Decompression::DecompressPose(OutPose, CompressedData, ExtractionContext, DecompContext, GetRetargetTransforms(), RootMotionReset);
 	}
 
diff --git a/Engine/Source/Runtime/Engine/Private/Animation/AnimSingleNodeInstance.cpp b/Engine/Source/Runtime/Engine/Private/Animation/AnimSingleNodeInstance.cpp
--- a/Engine/Source/Runtime/Engine/Private/Animation/AnimSingleNodeInstance.cpp
+++ b/Engine/Source/Runtime/Engine/Private/Animation/AnimSingleNodeInstance.cpp
@@ -401,6 +401,16 @@ void UAnimSingleNodeInstance::SetBlendSpacePosition(const FVector& InPosition)
 	GetProxyOnGameThread<FAnimSingleNodeInstanceProxy>().SetBlendSpacePosition(InPosition);
 }
 
+void UAnimSingleNodeInstance::SetInterpolationOverride(TOptional<EAnimInterpolationType> InterpolationType)
+{
+	GetProxyOnGameThread<FAnimSingleNodeInstanceProxy>().SetInterpolationOverride(InterpolationType);
+}
+
+TOptional<EAnimInterpolationType> UAnimSingleNodeInstance::GetInterpolationOverride() const
+{
+	return GetProxyOnGameThread<FAnimSingleNodeInstanceProxy>().GetInterpolationOverride();
+}
+
 void UAnimSingleNodeInstance::GetBlendSpaceState(FVector& OutPosition, FVector& OutFilteredPosition) const
 {
 	GetProxyOnGameThread<FAnimSingleNodeInstanceProxy>().GetBlendSpaceState(OutPosition, OutFilteredPosition);
diff --git a/Engine/Source/Runtime/Engine/Private/Animation/AnimSingleNodeInstanceProxy.cpp b/Engine/Source/Runtime/Engine/Private/Animation/AnimSingleNodeInstanceProxy.cpp
--- a/Engine/Source/Runtime/Engine/Private/Animation/AnimSingleNodeInstanceProxy.cpp
+++ b/Engine/Source/Runtime/Engine/Private/Animation/AnimSingleNodeInstanceProxy.cpp
@@ -118,6 +118,8 @@ void FAnimSingleNodeInstanceProxy::InitializeObjects(UAnimInstance* InAnimInstan
 
 	UAnimSingleNodeInstance* AnimSingleNodeInstance = CastChecked<UAnimSingleNodeInstance>(InAnimInstance);
 	CurrentAsset = AnimSingleNodeInstance->CurrentAsset;
+
+	InterpolationOverride = AnimSingleNodeInstance->GetInterpolationOverride();
 }
 
 void FAnimSingleNodeInstanceProxy::ClearObjects()
@@ -169,6 +171,7 @@ void FAnimSingleNodeInstanceProxy::InternalBlendSpaceEvaluatePose(class UBlendSp
 	FAnimationPoseData AnimationPoseData = { OutContext.Pose, OutContext.Curve, OutContext.CustomAttributes };
 
 	FAnimExtractContext ExtractionContext(static_cast<double>(CurrentTime), ShouldExtractRootMotion(), DeltaTimeRecord, bLooping);
+	ExtractionContext.InterpolationOverride = InterpolationOverride;
 
 	if (BlendSpace->IsValidAdditive())
 	{
@@ -182,7 +185,9 @@ void FAnimSingleNodeInstanceProxy::InternalBlendSpaceEvaluatePose(class UBlendSp
 #if WITH_EDITORONLY_DATA
 		if (BlendSpace->PreviewBasePose)
 		{
-			BlendSpace->PreviewBasePose->GetBonePose(AnimationPoseData, FAnimExtractContext(static_cast<double>(PreviewPoseCurrentTime)));
+			FAnimExtractContext Context(static_cast<double>(PreviewPoseCurrentTime));
+			Context.InterpolationOverride = InterpolationOverride;
+			BlendSpace->PreviewBasePose->GetBonePose(AnimationPoseData, Context);
 		}
 		else
 #endif // WITH_EDITORONLY_DATA
@@ -320,6 +325,7 @@ void FAnimNode_SingleNode::Evaluate_AnyThread(FPoseContext& Output)
 			ExtractionContext.bIgnoreRootLock = Proxy->bIgnoreRootLock;
 			ExtractionContext.bExtractRootMotion = Proxy->bIgnoreRootMotion ? false : ExtractionContext.bExtractRootMotion;
 #endif
+			ExtractionContext.InterpolationOverride = Proxy->InterpolationOverride;
 
 			if (Sequence->IsValidAdditive())
 			{
@@ -378,12 +384,15 @@ void FAnimNode_SingleNode::Evaluate_AnyThread(FPoseContext& Output)
 			else*/
 			{
 				// if SkeletalMesh isn't there, we'll need to use skeleton
-				Streamable->GetAnimationPose(OutputAnimationPoseData, FAnimExtractContext(static_cast<double>(Proxy->CurrentTime), Streamable->bEnableRootMotion, Proxy->DeltaTimeRecord, Proxy->bLooping));
+				FAnimExtractContext Context(static_cast<double>(Proxy->CurrentTime), Streamable->bEnableRootMotion, Proxy->DeltaTimeRecord, Proxy->bLooping);
+				Context.InterpolationOverride = Proxy->InterpolationOverride;
+				Streamable->GetAnimationPose(OutputAnimationPoseData, Context);
 			}
 		}
 		else if (UAnimComposite* Composite = Cast<UAnimComposite>(Proxy->CurrentAsset))
 		{
 			FAnimExtractContext ExtractionContext(static_cast<double>(Proxy->CurrentTime), Proxy->ShouldExtractRootMotion(), Proxy->DeltaTimeRecord, Proxy->bLooping);
+			ExtractionContext.InterpolationOverride = Proxy->InterpolationOverride;
 			const FAnimTrack& AnimTrack = Composite->AnimationTrack;
 
 			// find out if this is additive animation
@@ -441,7 +450,9 @@ void FAnimNode_SingleNode::Evaluate_AnyThread(FPoseContext& Output)
 					if (bCanProcessAdditiveAnimationsLocal && Montage->PreviewBasePose && Montage->GetPlayLength() > 0.f)
 					{
 						FAnimationPoseData LocalAnimationPoseData = { LocalSourcePose, LocalSourceCurve, LocalSourceAttributes };
-						Montage->PreviewBasePose->GetBonePose(LocalAnimationPoseData, FAnimExtractContext(static_cast<double>(Proxy->CurrentTime)));
+						FAnimExtractContext Context(static_cast<double>(Proxy->CurrentTime));
+						Context.InterpolationOverride = Proxy->InterpolationOverride;
+						Montage->PreviewBasePose->GetBonePose(LocalAnimationPoseData, Context);
 					}
 					else
 #endif // WITH_EDITORONLY_DATA
diff --git a/Engine/Source/Runtime/Engine/Public/Animation/AnimInstanceProxy.h b/Engine/Source/Runtime/Engine/Public/Animation/AnimInstanceProxy.h
--- a/Engine/Source/Runtime/Engine/Public/Animation/AnimInstanceProxy.h
+++ b/Engine/Source/Runtime/Engine/Public/Animation/AnimInstanceProxy.h
@@ -559,6 +559,12 @@ public:
 	 **/
 	ENGINE_API void AddAnimNotifies(const TArray<FAnimNotifyEventReference>& NewNotifies, const float InstanceWeight);
 
+	/** Set the optional animation interpolation type override. When not set, it will use the interpolation type set inside the animation asset. */
+	ENGINE_API void SetInterpolationOverride(TOptional<EAnimInterpolationType> InterpolationType);
+
+	/** Get the interpolation type override. If not set, it will use the iterpolation type set inside the animation asset. */
+	ENGINE_API TOptional<EAnimInterpolationType> GetInterpolationOverride() const;
+
 	/** Only restricted classes can access the protected interface */
 	friend class UAnimInstance;
 	friend class UAnimSingleNodeInstance;
@@ -1089,6 +1095,9 @@ protected:
 	uint64 FrameCounterForUpdate;
 	uint64 FrameCounterForNodeUpdate;
 
+	/** Optional override of the animation interpolation type. If not set, it is ignored. */
+	TOptional<EAnimInterpolationType> InterpolationOverride;
+
 private:
 	// Root motion extracted from animation since the last time ConsumeExtractedRootMotion was called
 	FRootMotionMovementParams ExtractedRootMotion;
@@ -1152,12 +1161,10 @@ private:
 	bool bUpdatingRoot;
 
 protected:
-
 	/** When RequiredBones mapping has changed, AnimNodes need to update their bones caches. */
 	uint8 bBoneCachesInvalidated : 1;
 
 private:
-
 	// Diplicate of bool result of ShouldExtractRootMotion()
 	uint8 bShouldExtractRootMotion : 1;
 
