Add SetTickableTickType to FTickableGameObject to allow enabling or disabling tick based on events.
Change UTickableWorldSubsystem to use the new method which blocks it from trying to tick after shutdown.
Refactor FTickableObjectBase so other subclasses like FTickableEditorObject can use the same registration logic.
Fix some outdated threading code and change tickable objects to not be copyable due to their automatic registration.
Deprecate IsAllowedToTick because IsTickable or SetTickableTickType are better options.

diff --git a/Engine/Source/Editor/UnrealEd/Private/TickableEditorObject.cpp b/Engine/Source/Editor/UnrealEd/Private/TickableEditorObject.cpp
--- a/Engine/Source/Editor/UnrealEd/Private/TickableEditorObject.cpp
+++ b/Engine/Source/Editor/UnrealEd/Private/TickableEditorObject.cpp
@@ -3,30 +3,18 @@
 #include "TickableEditorObject.h"
 #include "Tickable.h"
 
-TArray<FTickableObjectBase::FTickableObjectEntry>& FTickableEditorObject::GetTickableObjects()
+FTickableObjectBase::FTickableStatics& FTickableEditorObject::GetStatics()
 {
-	static TTickableObjectsCollection TickableObjects;
-	return TickableObjects;
+	static FTickableStatics Singleton;
+	return Singleton;
 }
 
-TArray<FTickableEditorObject*>& FTickableEditorObject::GetPendingTickableObjects()
+FTickableObjectBase::FTickableStatics& FTickableCookObject::GetStatics()
 {
-	static TArray<FTickableEditorObject*> PendingTickableObjects;
-	return PendingTickableObjects;
+	static FTickableStatics Singleton;
+	return Singleton;
 }
 
-TArray<FTickableObjectBase::FTickableObjectEntry>& FTickableCookObject::GetTickableObjects()
-{
-	static TTickableObjectsCollection TickableObjects;
-	return TickableObjects;
-}
-
-TArray<FTickableCookObject*>& FTickableCookObject::GetPendingTickableObjects()
-{
-	static TArray<FTickableCookObject*> PendingTickableObjects;
-	return PendingTickableObjects;
-}
+FTickableObjectBase* FTickableEditorObject::ObjectBeingTicked = nullptr;
+FTickableObjectBase* FTickableCookObject::ObjectBeingTicked = nullptr;
 
-bool FTickableCookObject::bCollectionIntact = true;
-bool FTickableCookObject::bIsTickingObjects = false;
-FTickableCookObject* FTickableCookObject::ObjectBeingTicked = nullptr;
diff --git a/Engine/Source/Editor/UnrealEd/Private/UnrealEdMisc.cpp b/Engine/Source/Editor/UnrealEd/Private/UnrealEdMisc.cpp
--- a/Engine/Source/Editor/UnrealEd/Private/UnrealEdMisc.cpp
+++ b/Engine/Source/Editor/UnrealEd/Private/UnrealEdMisc.cpp
@@ -104,11 +104,6 @@
 
 DEFINE_LOG_CATEGORY_STATIC(LogUnrealEdMisc, Log, All);
 
-bool FTickableEditorObject::bCollectionIntact = true;
-bool FTickableEditorObject::bIsTickingObjects = false;
-FTickableObjectBase* FTickableEditorObject::ObjectBeingTicked = nullptr;
-
-
 namespace
 {
 	static const FName LevelEditorName("LevelEditor");
diff --git a/Engine/Source/Editor/UnrealEd/Public/TickableEditorObject.h b/Engine/Source/Editor/UnrealEd/Public/TickableEditorObject.h
--- a/Engine/Source/Editor/UnrealEd/Public/TickableEditorObject.h
+++ b/Engine/Source/Editor/UnrealEd/Public/TickableEditorObject.h
@@ -20,109 +20,46 @@
 class FTickableEditorObject : public FTickableObjectBase
 {
 public:
+	UE_NONCOPYABLE(FTickableEditorObject);
 
+	/** Calls Tick on every tickable editor object with tick enabled */
 	static void TickObjects(const float DeltaSeconds)
 	{
-		TArray<FTickableEditorObject*>& PendingTickableObjects = GetPendingTickableObjects();
-		TArray<FTickableObjectEntry>& TickableObjects = GetTickableObjects();
+		FTickableStatics& Statics = GetStatics();
 
-		for (FTickableEditorObject* PendingTickable : PendingTickableObjects)
+		SimpleTickObjects(Statics, [DeltaSeconds](FTickableObjectBase* TickableObject)
 		{
-			AddTickableObject(TickableObjects, PendingTickable);
-		}
-		PendingTickableObjects.Empty();
-
-		if (TickableObjects.Num() > 0)
-		{
-			check(!bIsTickingObjects);
-			bIsTickingObjects = true;
-
-			bool bNeedsCleanup = false;
-
-			for (const FTickableObjectEntry& TickableEntry : TickableObjects)
-			{
-				if (FTickableObjectBase* TickableObject = TickableEntry.TickableObject)
-				{
-					if ((TickableEntry.TickType == ETickableTickType::Always) || TickableObject->IsTickable())
-					{
-						ObjectBeingTicked = TickableObject;
-						TickableObject->Tick(DeltaSeconds);
-						ObjectBeingTicked = nullptr;
-					}
-
-					// In case it was removed during tick
-					if (TickableEntry.TickableObject == nullptr)
-					{
-						bNeedsCleanup = true;
-					}
-				}
-				else
-				{
-					bNeedsCleanup = true;
-				}
-			}
-
-			if (bNeedsCleanup)
-			{
-				TickableObjects.RemoveAll([](const FTickableObjectEntry& Entry) { return (Entry.TickableObject == nullptr); });
-			}
-
-			bIsTickingObjects = false;
-		}
+			ObjectBeingTicked = TickableObject;
+			TickableObject->Tick(DeltaSeconds);
+			ObjectBeingTicked = nullptr;
+		});
 	}
 
 	/** Registers this instance with the static array of tickable objects. */
 	FTickableEditorObject()
 	{
 		ensure(IsInGameThread() || IsInAsyncLoadingThread());
-		check(!GetPendingTickableObjects().Contains(this));
-		check(!GetTickableObjects().Contains(this));
-		GetPendingTickableObjects().Add(this);
+
+		FTickableStatics& Statics = GetStatics();
+		Statics.QueueTickableObjectForAdd(this);
 	}
 
 	/** Removes this instance from the static array of tickable objects. */
 	virtual ~FTickableEditorObject()
 	{
 		ensureMsgf(ObjectBeingTicked != this, TEXT("Detected possible memory stomp. We are in the Tickable objects Tick function but hit its deconstructor, the 'this' pointer for the Object will now be invalid"));
-
 		ensure(IsInGameThread() || IsInAsyncLoadingThread());
-		if (bCollectionIntact && GetPendingTickableObjects().Remove(this) == 0)
-		{
-			RemoveTickableObject(GetTickableObjects(), this, bIsTickingObjects);
-		}
+		
+		FTickableStatics& Statics = GetStatics();
+		Statics.SetTickTypeForTickableObject(this, ETickableTickType::Never);
 	}
 
 private:
-
-	/**
-	 * Class that avoids crashes when unregistering a tickable editor object too late.
-	 *
-	 * Some tickable objects can outlive the collection
-	 * (global/static destructor order is unpredictable).
-	 */
-	class TTickableObjectsCollection : public TArray<FTickableObjectEntry>
-	{
-	public:
-		~TTickableObjectsCollection()
-		{
-			FTickableEditorObject::bCollectionIntact = false;
-		}
-	};
-
-	friend class TTickableObjectsCollection;
-
-	/** True if collection of tickable objects is still intact. */
-	static UNREALED_API bool bCollectionIntact;
-	/** True if currently ticking of tickable editor objects. */
-	static UNREALED_API bool bIsTickingObjects;
-
 	/** Set if we are in the Tick function for an editor tickable object */
 	static UNREALED_API FTickableObjectBase* ObjectBeingTicked;
 
-
-	static UNREALED_API TArray<FTickableObjectEntry>& GetTickableObjects();
-
-	static UNREALED_API TArray<FTickableEditorObject*>& GetPendingTickableObjects();
+	/** Returns the tracking struct for this type */
+	static UNREALED_API FTickableStatics& GetStatics();
 };
 
 /**
@@ -130,115 +67,53 @@ private:
  * cooking.
  * If a system needs to be cooked both during cook commandlet and in editor without the cook commandlet,
  * it should dual-inherit from both FTickableCookObject and FTickableEditorObject.
- * TODO: Reduce duplication between  FTickableCookObject and FTickableEditorObject.
  */
 class FTickableCookObject : public FTickableObjectBase
 {
 public:
+	UE_NONCOPYABLE(FTickableCookObject);
 
+	/** Calls TickCook on every enabled tickable object */
 	static void TickObjects(const float DeltaSeconds, bool bCookComplete)
 	{
-		TArray<FTickableCookObject*>& PendingTickableObjects = GetPendingTickableObjects();
-		TArray<FTickableObjectEntry>& TickableObjects = GetTickableObjects();
+		FTickableStatics& Statics = GetStatics();
 
-		for (FTickableCookObject* PendingTickable : PendingTickableObjects)
+		SimpleTickObjects(Statics, [DeltaSeconds, bCookComplete](FTickableObjectBase* TickableObject)
 		{
-			AddTickableObject(TickableObjects, PendingTickable);
-		}
-		PendingTickableObjects.Empty();
-
-		if (TickableObjects.Num() > 0)
-		{
-			check(!bIsTickingObjects);
-			bIsTickingObjects = true;
-
-			bool bNeedsCleanup = false;
-
-			for (const FTickableObjectEntry& TickableEntry : TickableObjects)
-			{
-				if (FTickableObjectBase* TickableObject = TickableEntry.TickableObject)
-				{
-					if ((TickableEntry.TickType == ETickableTickType::Always) || TickableObject->IsTickable())
-					{
-						FTickableCookObject* CookTickableObject = static_cast<FTickableCookObject*>(TickableObject);
-						ObjectBeingTicked = CookTickableObject;
-						CookTickableObject->TickCook(DeltaSeconds, bCookComplete);
-						ObjectBeingTicked = nullptr;
-					}
-
-					// In case it was removed during tick
-					if (TickableEntry.TickableObject == nullptr)
-					{
-						bNeedsCleanup = true;
-					}
-				}
-				else
-				{
-					bNeedsCleanup = true;
-				}
-			}
-
-			if (bNeedsCleanup)
-			{
-				TickableObjects.RemoveAll([](const FTickableObjectEntry& Entry) { return (Entry.TickableObject == nullptr); });
-			}
-
-			bIsTickingObjects = false;
-		}
+			FTickableCookObject* CookTickableObject = static_cast<FTickableCookObject*>(TickableObject);
+			ObjectBeingTicked = TickableObject;
+			CookTickableObject->TickCook(DeltaSeconds, bCookComplete);
+			ObjectBeingTicked = nullptr;
+		});
 	}
 
 	/** Registers this instance with the static array of tickable objects. */
 	FTickableCookObject()
 	{
 		ensure(IsInGameThread() || IsInAsyncLoadingThread());
-		check(!GetPendingTickableObjects().Contains(this));
-		check(!GetTickableObjects().Contains(this));
-		GetPendingTickableObjects().Add(this);
+
+		FTickableStatics& Statics = GetStatics();
+		Statics.QueueTickableObjectForAdd(this);
 	}
 
 	/** Removes this instance from the static array of tickable objects. */
 	virtual ~FTickableCookObject()
 	{
 		ensureMsgf(ObjectBeingTicked != this, TEXT("Detected possible memory stomp. We are in the Tickable objects Tick function but hit its deconstructor, the 'this' pointer for the Object will now be invalid"));
-
 		ensure(IsInGameThread() || IsInAsyncLoadingThread());
-		if (bCollectionIntact && GetPendingTickableObjects().Remove(this) == 0)
-		{
-			RemoveTickableObject(GetTickableObjects(), this, bIsTickingObjects);
-		}
+
+		FTickableStatics& Statics = GetStatics();
+		Statics.SetTickTypeForTickableObject(this, ETickableTickType::Never);
 	}
 
+	/** Cook tick virtual, must be implemented in subclass */
 	virtual void TickCook(float DeltaTime, bool bCookCompete) = 0;
 
 private:
 
-	/**
-	 * Class that avoids crashes when unregistering a tickable editor object too late.
-	 *
-	 * Some tickable objects can outlive the collection
-	 * (global/static destructor order is unpredictable).
-	 */
-	class TTickableObjectsCollection : public TArray<FTickableObjectEntry>
-	{
-	public:
-		~TTickableObjectsCollection()
-		{
-			FTickableCookObject::bCollectionIntact = false;
-		}
-	};
-
-	friend class TTickableObjectsCollection;
-
-	/** True if collection of tickable objects is still intact. */
-	static UNREALED_API bool bCollectionIntact;
-	/** True if currently ticking of tickable editor objects. */
-	static UNREALED_API bool bIsTickingObjects;
-
 	/** Set if we are in the Tick function for an editor tickable object */
-	static UNREALED_API FTickableCookObject* ObjectBeingTicked;
-
-
-	static UNREALED_API TArray<FTickableObjectEntry>& GetTickableObjects();
+	static UNREALED_API FTickableObjectBase* ObjectBeingTicked;
 
-	static UNREALED_API TArray<FTickableCookObject*>& GetPendingTickableObjects();
+	/** Returns the tracking struct for this type */
+	static UNREALED_API FTickableStatics& GetStatics();
 };
diff --git a/Engine/Source/Runtime/Engine/Private/Subsystems/WorldSubsystem.cpp b/Engine/Source/Runtime/Engine/Private/Subsystems/WorldSubsystem.cpp
--- a/Engine/Source/Runtime/Engine/Private/Subsystems/WorldSubsystem.cpp
+++ b/Engine/Source/Runtime/Engine/Private/Subsystems/WorldSubsystem.cpp
@@ -49,35 +49,58 @@ UTickableWorldSubsystem::UTickableWorldSubsystem()
 
 }
 
+UWorld* UTickableWorldSubsystem::GetTickableGameObjectWorld() const
+{
+	return GetWorld();
+}
+
 ETickableTickType UTickableWorldSubsystem::GetTickableTickType() const 
-{ 
-	// By default (if the child class doesn't override GetTickableTickType), don't let CDOs ever tick: 
-	return IsTemplate() ? ETickableTickType::Never : FTickableGameObject::GetTickableTickType(); 
+{
+	// If this is a template or has not been initialized yet, set to never tick and it will be enabled when it is initialized
+	if (IsTemplate() || !bInitialized)
+	{
+		return ETickableTickType::Never;
+	}
+
+	// Otherwise default to conditional
+	return ETickableTickType::Conditional;
 }
 
 bool UTickableWorldSubsystem::IsAllowedToTick() const
 {
-	// No matter what IsTickable says, don't let CDOs or uninitialized world subsystems tick :
-	// Note: even if GetTickableTickType was overridden by the child class and returns something else than ETickableTickType::Never for CDOs, 
-	//  it's probably a mistake, so by default, don't allow ticking. If the child class really intends its CDO to tick, it can always override IsAllowedToTick...
-	// NOTE: `bInitialized` must be checked first as `IsTemplate()` might access a dangling `Outer` if we are awaiting GC but `Outer` has already been deleted.
-	return bInitialized && !IsTemplate();
+	// This function is now deprecated and subclasses should implement IsTickable instead
+	// This should never be false because Initialize should always be called before the first tick and Deinitialize cancels the tick
+	ensureMsgf(bInitialized, TEXT("Tickable subsystem %s tried to tick when not initialized! Check for missing Super call"), *GetFullName());
+
+	return bInitialized;
 }
 
 void UTickableWorldSubsystem::Tick(float DeltaTime)
 {
-	checkf(IsInitialized(), TEXT("Ticking should have been disabled for an uninitialized subsystem : remember to call IsInitialized in the subsystem's IsTickable, IsTickableInEditor and/or IsTickableWhenPaused implementation"));
+	checkf(IsInitialized(), TEXT("Ticking should have been disabled for an uninitialized subsystem!"));
 }
 
 void UTickableWorldSubsystem::Initialize(FSubsystemCollectionBase& Collection)
 {
 	check(!bInitialized);
 	bInitialized = true;
+
+	// Refresh the tick type after initialization
+	SetTickableTickType(GetTickableTickType());
 }
 
 void UTickableWorldSubsystem::Deinitialize()
 {
 	check(bInitialized);
 	bInitialized = false;
+
+	// Always cancel tick as this is about to be destroyed
+	SetTickableTickType(ETickableTickType::Never);
 }
 
+void UTickableWorldSubsystem::BeginDestroy()
+{
+	Super::BeginDestroy();
+
+	ensureMsgf(!bInitialized, TEXT("Tickable subsystem %s was destroyed while still initialized! Check for missing Super call"), *GetFullName());
+}
diff --git a/Engine/Source/Runtime/Engine/Private/Tickable.cpp b/Engine/Source/Runtime/Engine/Private/Tickable.cpp
--- a/Engine/Source/Runtime/Engine/Private/Tickable.cpp
+++ b/Engine/Source/Runtime/Engine/Private/Tickable.cpp
@@ -6,130 +6,174 @@
 
 DECLARE_CYCLE_STAT(TEXT("TickableGameObjects Time"), STAT_TickableGameObjectsTime, STATGROUP_Game);
 
-struct FTickableStatics
+void FTickableObjectBase::FTickableStatics::QueueTickableObjectForAdd(FTickableObjectBase* InTickable)
 {
-	FCriticalSection TickableObjectsCritical;
-	TArray<FTickableObjectBase::FTickableObjectEntry> TickableObjects;
-
-	FCriticalSection NewTickableObjectsCritical;
-	TSet<FTickableGameObject*> NewTickableObjects;
-
-	bool bIsTickingObjects = false;
-
-	void QueueTickableObjectForAdd(FTickableGameObject* InTickable)
-	{
-		FScopeLock NewTickableObjectsLock(&NewTickableObjectsCritical);
-		NewTickableObjects.Add(InTickable);
-	}
+	// This only needs to lock the new object queue
+	FScopeLock NewTickableObjectsLock(&NewTickableObjectsCritical);
+	NewTickableObjects.Add(InTickable, ETickableTickType::NewObject);
+}
 
-	void RemoveTickableObjectFromNewObjectsQueue(FTickableGameObject* InTickable)
+void FTickableObjectBase::FTickableStatics::SetTickTypeForTickableObject(FTickableObjectBase* InTickable, ETickableTickType NewTickType)
+{
+	FScopeLock TickableObjectsLock(&TickableObjectsCritical);
+	FScopeLock NewTickableObjectsLock(&NewTickableObjectsCritical);
+	
+	// Existing entries should never be set to new object
+	check(NewTickType != ETickableTickType::NewObject);
+
+	const int32 Pos = TickableObjects.IndexOfByKey(InTickable);
+	if (NewTickType == ETickableTickType::Never)
 	{
-		FScopeLock NewTickableObjectsLock(&NewTickableObjectsCritical);
+		// Remove from pending list if it hasn't been registered yet
 		NewTickableObjects.Remove(InTickable);
-	}
 
-	static FTickableStatics& Get()
+		// The item may be missing depending on destruction order during shutdown
+		if (Pos != INDEX_NONE)
+		{
+			if (bIsTickingObjects)
+			{
+				// During ticking it is not safe to modify the array so null and mark for later
+				TickableObjects[Pos].TickableObject = nullptr;
+				bNeedsCleanup = true;
+			}
+			else
+			{
+				TickableObjects.RemoveAt(Pos);
+			}
+		}
+	}
+	else
 	{
-		static FTickableStatics Singleton;
-		return Singleton;
+		if (Pos != INDEX_NONE)
+		{
+			// If this is registered, it was removed from the new list in BeginTicking
+			check(!NewTickableObjects.Contains(InTickable));
+
+			// This will modify behavior for the current frame if it has not ticked yet
+			TickableObjects[Pos].TickType = NewTickType;
+		}
+		else
+		{
+			// Add to the pending list (which could override previous request), this will apply it next frame
+			NewTickableObjects.Add(InTickable, NewTickType);
+		}
 	}
-};
+}
 
-void FTickableObjectBase::AddTickableObject(TArray<FTickableObjectEntry>& TickableObjects, FTickableObjectBase* TickableObject)
+void FTickableObjectBase::FTickableStatics::StartTicking()
 {
-	check(!TickableObjects.Contains(TickableObject));
-	const ETickableTickType TickType = TickableObject->GetTickableTickType();
-	if (TickType != ETickableTickType::Never)
+	check(!bIsTickingObjects);
+
+	FScopeLock NewTickableObjectsLock(&NewTickableObjectsCritical);
+
+	for (TPair<FTickableObjectBase*, ETickableTickType> Pair : NewTickableObjects)
 	{
-		TickableObjects.Add({ TickableObject, TickType });
+		FTickableObjectBase* NewTickableObject = Pair.Key;
+		ETickableTickType TickType = Pair.Value;
+
+		// SetTickTypeForTickableObject will not add to new list if it already exists in TickableObjects
+		check(!TickableObjects.Contains(NewTickableObject));
+
+		if (TickType == ETickableTickType::NewObject)
+		{
+			// Query object if unknown
+			TickType = NewTickableObject->GetTickableTickType();
+		}
+
+		if (TickType != ETickableTickType::Never)
+		{
+			TickableObjects.Add({ NewTickableObject, TickType });
+		}
 	}
+	NewTickableObjects.Empty();
+
+	bIsTickingObjects = true;
 }
 
-void FTickableObjectBase::RemoveTickableObject(TArray<FTickableObjectEntry>& TickableObjects, FTickableObjectBase* TickableObject, const bool bIsTickingObjects)
+void FTickableObjectBase::FTickableStatics::FinishTicking()
 {
-	const int32 Pos = TickableObjects.IndexOfByKey(TickableObject);
-
-#if 0 // virtual from destructor doesn't work ... need to rethink how to do warning
-	// ensure that GetTickableTickType did not change over time
-	switch (TickableObject->GetTickableTickType())
+	check(bIsTickingObjects);
+	if (bNeedsCleanup)
 	{
-	case ETickableTickType::Always:
-		ensureMsgf(Pos != INDEX_NONE && TickableObjects[Pos].TickType == ETickableTickType::Always, TEXT("TickType has changed since object was created. Result of GetTickableTickType must be invariant for a given object."));
-		break;
+		TickableObjects.RemoveAll([](const FTickableObjectEntry& Entry) { return Entry.TickableObject == nullptr; });
+		bNeedsCleanup = false;
+	}
 
-	case ETickableTickType::Conditional:
-		ensureMsgf(Pos != INDEX_NONE && TickableObjects[Pos].TickType == ETickableTickType::Conditional, TEXT("TickType has changed since object was created. Result of GetTickableTickType must be invariant for a given object."));
-		break;
+	bIsTickingObjects = false;
+}
 
-	case ETickableTickType::Never:
-		ensureMsgf(Pos == INDEX_NONE, TEXT("TickType has changed since object was created. Result of GetTickableTickType must be invariant for a given object."));
-		break;
-	}
-#endif
 
-	if (Pos != INDEX_NONE)
+void FTickableObjectBase::SimpleTickObjects(FTickableStatics& Statics, TFunctionRef<void(FTickableObjectBase*)> TickFunc)
+{
+	FScopeLock LockTickableObjects(&Statics.TickableObjectsCritical);
+
+	Statics.StartTicking();
+
+	for (const FTickableObjectEntry& TickableEntry : Statics.TickableObjects)
 	{
-		if (bIsTickingObjects)
-		{
-			TickableObjects[Pos].TickableObject = nullptr;
-		}
-		else
+		if (FTickableObjectBase* TickableObject = TickableEntry.TickableObject)
 		{
-			TickableObjects.RemoveAt(Pos);
+			// NOTE: This deliberately does not call IsAllowedToTick as it is deprecated and was not called in code this is replacing
+			if ((TickableEntry.TickType == ETickableTickType::Always) || TickableObject->IsTickable())
+			{
+				TickFunc(TickableObject);
+			}
 		}
 	}
+
+	Statics.FinishTicking();
 }
 
+
+// FTickableGameObject implementation
+
 FTickableGameObject::FTickableGameObject()
 {
-	FTickableStatics& Statics = FTickableStatics::Get();
+	FTickableStatics& Statics = GetStatics();
 
-	if (UObjectInitialized())
-	{
-		Statics.QueueTickableObjectForAdd(this);
-	}
-	else
+	// Queue for creation, this can get called very early in startup
+	Statics.QueueTickableObjectForAdd(this);
+}
+
+FTickableGameObject::~FTickableGameObject()
+{
+	// Depending on destruction order this could create a new statics object during shutdown, but the removal request will be ignored
+	FTickableStatics& Statics = GetStatics();
+	
+	// This won't do anything if it was already set to Never
+	Statics.SetTickTypeForTickableObject(this, ETickableTickType::Never);
+}
+
+void FTickableGameObject::SetTickableTickType(ETickableTickType NewTickType)
+{
+	if (ensure(NewTickType != ETickableTickType::NewObject))
 	{
-		AddTickableObject(Statics.TickableObjects, this);
+		FTickableStatics& Statics = GetStatics();
+		Statics.SetTickTypeForTickableObject(this, NewTickType);
 	}
 }
 
-FTickableGameObject::~FTickableGameObject()
+FTickableObjectBase::FTickableStatics& FTickableGameObject::GetStatics()
 {
-	FTickableStatics& Statics = FTickableStatics::Get();	
-	Statics.RemoveTickableObjectFromNewObjectsQueue(this);	
-	FScopeLock LockTickableObjects(&Statics.TickableObjectsCritical);
-	RemoveTickableObject(Statics.TickableObjects, this, Statics.bIsTickingObjects);
+	static FTickableStatics Singleton;
+	return Singleton;
 }
 
-void FTickableGameObject::TickObjects(UWorld* World, const int32 InTickType, const bool bIsPaused, const float DeltaSeconds)
+PRAGMA_DISABLE_DEPRECATION_WARNINGS
+void FTickableGameObject::TickObjects(UWorld* World, const ELevelTick LevelTickType, const bool bIsPaused, const float DeltaSeconds)
 {
 	SCOPE_CYCLE_COUNTER(STAT_TickableGameObjectsTime);
 	CSV_SCOPED_TIMING_STAT_EXCLUSIVE(Tickables);
 
-	FTickableStatics& Statics = FTickableStatics::Get();
+	FTickableStatics& Statics = GetStatics();
 
 	check(IsInGameThread());
 
-	// It's a long lock but it's ok, the only thing we can block here is the GC worker thread that destroys UObjects
-	FScopeLock LockTickableObjects(&Statics.TickableObjectsCritical);
-
-	{
-		FScopeLock NewTickableObjectsLock(&Statics.NewTickableObjectsCritical);
-		for (FTickableGameObject* NewTickableObject : Statics.NewTickableObjects)
-		{
-			AddTickableObject(Statics.TickableObjects, NewTickableObject);
-		}
-		Statics.NewTickableObjects.Empty();
-	}
-
-	if (Statics.TickableObjects.Num() > 0)
 	{
-		check(!Statics.bIsTickingObjects);
-		Statics.bIsTickingObjects = true;
+		// It's a long lock but it's ok, the only thing we can block here is the GC worker thread that destroys UObjects
+		FScopeLock LockTickableObjects(&Statics.TickableObjectsCritical);
 
-		bool bNeedsCleanup = false;
-		const ELevelTick TickType = (ELevelTick)InTickType;
+		Statics.StartTicking();
 
 		for (const FTickableObjectEntry& TickableEntry : Statics.TickableObjects)
 		{
@@ -140,34 +184,22 @@ void FTickableGameObject::TickObjects(UWorld* World, const int32 InTickType, con
 					&& ((TickableEntry.TickType == ETickableTickType::Always) || TickableObject->IsTickable())
 					&& (TickableObject->GetTickableGameObjectWorld() == World))
 				{
-					const bool bIsGameWorld = InTickType == LEVELTICK_All || (World && World->IsGameWorld());
+					// If tick type is All because at least one game world ticked, this will treat the null world as a game world
+					const bool bIsGameWorld = LevelTickType == LEVELTICK_All || (World && World->IsGameWorld());
+
 					// If we are in editor and it is editor tickable, always tick
 					// If this is a game world then tick if we are not doing a time only (paused) update and we are not paused or the object is tickable when paused
 					if ((GIsEditor && TickableObject->IsTickableInEditor()) ||
-						(bIsGameWorld && ((!bIsPaused && TickType != LEVELTICK_TimeOnly) || (bIsPaused && TickableObject->IsTickableWhenPaused()))))
+						(bIsGameWorld && ((!bIsPaused && LevelTickType != LEVELTICK_TimeOnly) || (bIsPaused && TickableObject->IsTickableWhenPaused()))))
 					{
 						FScopeCycleCounter Context(TickableObject->GetStatId());
 						TickableObject->Tick(DeltaSeconds);
-
-						// In case it was removed during tick
-						if (TickableEntry.TickableObject == nullptr)
-						{
-							bNeedsCleanup = true;
-						}
 					}
 				}
 			}
-			else
-			{
-				bNeedsCleanup = true;
-			}
-		}
-
-		if (bNeedsCleanup)
-		{
-			Statics.TickableObjects.RemoveAll([](const FTickableObjectEntry& Entry) { return Entry.TickableObject == nullptr; });
 		}
 
-		Statics.bIsTickingObjects = false;
+		Statics.FinishTicking();
 	}
 }
+PRAGMA_ENABLE_DEPRECATION_WARNINGS
\ No newline at end of file
diff --git a/Engine/Source/Runtime/Engine/Public/Subsystems/WorldSubsystem.h b/Engine/Source/Runtime/Engine/Public/Subsystems/WorldSubsystem.h
--- a/Engine/Source/Runtime/Engine/Public/Subsystems/WorldSubsystem.h
+++ b/Engine/Source/Runtime/Engine/Public/Subsystems/WorldSubsystem.h
@@ -47,7 +47,10 @@ protected:
 
 /**
  * UTickableWorldSubsystem
- * Base class for auto instanced and initialized systems that share the lifetime of a UWorld and are ticking along with it
+ * Base class for auto instanced and initialized systems that share the lifetime of a UWorld and are ticking along with it.
+ * With the default implementation, it will start ticking after Initialize and stop during Deinitialize,
+ * and it will call IsTickable every frame (defaults to true) before calling Tick.
+ * Subclasses must forward calls to the Initialize/Deinitialize functions to correctly enable ticking.
  */
 UCLASS(Abstract, MinimalAPI)
 class UTickableWorldSubsystem : public UWorldSubsystem, public FTickableGameObject
@@ -58,7 +61,7 @@ public:
 	ENGINE_API UTickableWorldSubsystem();
 
 	// FTickableGameObject implementation Begin
-	virtual UWorld* GetTickableGameObjectWorld() const override { return GetWorld(); }
+	ENGINE_API UWorld* GetTickableGameObjectWorld() const override;
 	ENGINE_API virtual ETickableTickType GetTickableTickType() const override;
 	ENGINE_API virtual bool IsAllowedToTick() const override final;
 	ENGINE_API virtual void Tick(float DeltaTime) override;
@@ -68,8 +71,10 @@ public:
 	// USubsystem implementation Begin
 	ENGINE_API virtual void Initialize(FSubsystemCollectionBase& Collection) override;
 	ENGINE_API virtual void Deinitialize() override;
+	ENGINE_API virtual void BeginDestroy() override;
 	// USubsystem implementation End
 
+	/** Returns true if Initialize has been called but Deinitialize has not */
 	bool IsInitialized() const { return bInitialized; }
 
 private:
diff --git a/Engine/Source/Runtime/Engine/Public/Tickable.h b/Engine/Source/Runtime/Engine/Public/Tickable.h
--- a/Engine/Source/Runtime/Engine/Public/Tickable.h
+++ b/Engine/Source/Runtime/Engine/Public/Tickable.h
@@ -10,9 +10,10 @@
 #include "CoreMinimal.h"
 #include "Stats/Stats.h"
 #include "Misc/ScopeLock.h"
+#include "Engine/EngineBaseTypes.h"
 
 /**
- * Enum used to convey whether a tickable object will always tick, conditionally tick, or never tick.
+ * Enum used to determine the current ticking rules for an object, this can change after creation.
  */
 enum class ETickableTickType : uint8
 {
@@ -23,74 +24,103 @@ enum class ETickableTickType : uint8
 	Always,
 
 	/** Never tick the object, do not add to tickables array */
-	Never        
-};
+	Never,
 
-struct FTickableStatics;
+	/** Unknown state, true for newly registered objects that have not yet called GetTickableTickType */
+	NewObject
+};
 
 /**
  * Base class for tickable objects
  */
 class FTickableObjectBase
 {
-	friend struct FTickableStatics;
-
-protected:
+protected:	
+	/** Implementation struct for an individual tickable object */
 	struct FTickableObjectEntry
 	{
-		FTickableObjectBase* TickableObject;
-		ETickableTickType TickType;
+		FTickableObjectBase* TickableObject = nullptr;
+		ETickableTickType TickType = ETickableTickType::NewObject;
 
 		bool operator==(FTickableObjectBase* OtherObject) const { return TickableObject == OtherObject; }
 	};
 
-	static ENGINE_API void AddTickableObject(TArray<FTickableObjectEntry>& TickableObjects, FTickableObjectBase* TickableObject);
+	/** Implementation struct for internals of ticking, there should be one instance of this for each direct subclass */
+	struct FTickableStatics
+	{
+		/** This critical section should be locked during entire tick process */
+		FCriticalSection TickableObjectsCritical;
+
+		/** List of objects that are fully ticking */
+		TArray<FTickableObjectBase::FTickableObjectEntry> TickableObjects;
+
+		/** Lock for modifying new list, this is automatically acquired by functions below */
+		FCriticalSection NewTickableObjectsCritical;
+
+		/** Set of objects that have not yet been queried for tick type */
+		TMap<FTickableObjectBase*, ETickableTickType> NewTickableObjects;
+
+		/** True if any of the tick arrays are being iterated */
+		bool bIsTickingObjects = false;
+
+		/** True if any objects were deleted and this needs cleanup after tick */
+		bool bNeedsCleanup = false;
+
+
+		/** Call from anywhere to lock and add to the new objects list */
+		ENGINE_API void QueueTickableObjectForAdd(FTickableObjectBase* InTickable);
+
+		/** Call from anywhere to lock critical sections and set the tick type for an object which could modify both arrays */
+		ENGINE_API void SetTickTypeForTickableObject(FTickableObjectBase* TickableObject, ETickableTickType NewTickType);
+
+
+		/** Call after locking TickableObjectsCritical to process the NewTickableObjects array and prepare for tick */
+		ENGINE_API void StartTicking();
+
+		/** Finishes ticking and handle cleanup for entries that were invalidated */
+		ENGINE_API void FinishTicking();
+	};
 
-	static ENGINE_API void RemoveTickableObject(TArray<FTickableObjectEntry>& TickableObjects, FTickableObjectBase* TickableObject, const bool bIsTickingObjects);
+	/** Perform a simple tick using a class-specific statics struct and a function ref */
+	static ENGINE_API void SimpleTickObjects(FTickableStatics& Statics, TFunctionRef<void(FTickableObjectBase*)> TickFunc);
 
 public:
 	/**
 	 * Pure virtual that must be overloaded by the inheriting class. It will
-	 * be called from within LevelTick.cpp after ticking all actors or from
-	 * the rendering thread (depending on bIsRenderingThreadObject)
+	 * be called at different times in the frame depending on the subclass.
 	 *
 	 * @param DeltaTime	Game time passed since the last call.
 	 */
 	virtual void Tick( float DeltaTime ) = 0;
 
 	/**
-	* Virtual that can be overloaded by the inheriting class. It is
-	* used to determine whether an object will ever be able to tick, and if not,
-	* it will not get added to the tickable objects array. If the tickable tick type
-	* is Conditional then the virtual IsTickable will be called to determine whether
-	* to tick the object on each given frame
-	*
-	* @return	true if object will ever want to be ticked, false otherwise.
-	*/
+	 * Virtual that can be overloaded by the inheriting class and is called before first tick.
+	 * It is used to determine whether an object can possibly tick, and if not,
+	 * it will not get added to the tickable objects array. If the tickable tick type
+	 * is Conditional then the virtual IsTickable will be called to determine whether
+	 * to tick the object on each given frame.
+	 *
+	 * @return an enum defining the rules for ticking this object.
+	 */
 	virtual ETickableTickType GetTickableTickType() const { return ETickableTickType::Conditional; }
 
 	/**
 	 * Virtual that can be overloaded by the inheriting class. It is
-	 * used to determine whether an object is ready to be ticked. This is 
-	 * required for example for all UObject derived classes as they might be
-	 * loaded async and therefore won't be ready immediately.
+	 * used to determine whether an object should be conditionally ticked.
 	 *
 	 * @return	true if object is ready to be ticked, false otherwise.
 	 */
 	virtual bool IsTickable() const { return true; }
 
 	/**
-	 * Virtual that can be overloaded by the inheriting class. It is
-	 * used to determine whether an object is allowed to tick. This is different than IsTickable
-	 * for keeping existing tickable objects working while providing an extra-safety for classes that decide to 
-	 * implement it (e.g. tickable world subsystems that have been deinitialized are not allowed to tick, regardless 
-	 * of what IsTickable says)
-	 *
+	 * Function called before IsTickable, used to change rules without breaking existing API.
+	 * This was used to allow checking for tickable safety before calling GetTickableGameObjectWorld.
 	 * @return	true if object is allowed to be ticked, false otherwise.
 	 */
+	UE_DEPRECATED(5.5, "Use IsTickable for conditional ticks and SetTickableTickType to change the tick type after initial registration")
 	virtual bool IsAllowedToTick() const { return true; }
 
-	/** return the stat id to use for this tickable **/
+	/** Return the stat id used to track the performance of this object */
 	virtual TStatId GetStatId() const = 0;
 };
 
@@ -102,10 +132,15 @@ public:
  */
 class FTickableGameObject : public FTickableObjectBase
 {
+	/** Returns the tracking struct for this type */
+	static ENGINE_API FTickableStatics& GetStatics();
+
 public:
+	/** Tickable objects cannot be copied safely due to the auto registration */
+	UE_NONCOPYABLE(FTickableGameObject);
+
 	/**
 	 * Registers this instance with the static array of tickable objects.	
-	 *
 	 */
 	ENGINE_API FTickableGameObject();
 
@@ -126,20 +161,43 @@ public:
 	}
 
 	/**
-	 * Used to determine whether the object should be ticked in the editor.  Defaults to false since
-	 * that is the previous behavior.
+	 * Used to determine whether the object should be ticked in the editor when there is no gameplay world.
+	 * Objects will still be ticked in Play in Editor if they are associated with a PIE world.
 	 *
-	 * @return	true if this tickable object can be ticked in the editor
+	 * @return true if this tickable object should always be ticked in the editor, even if there is no gameplay world
 	 */
 	virtual bool IsTickableInEditor() const
 	{
 		return false;
 	}
 
+	/**
+	 * Used to determine the specific world this object is associated with.
+	 * If this returns a valid world it will tick during that world's level tick.
+	 * If this returns null, it will tick during the general engine tick after all world ticks.
+	 *
+	 * @return the world this object is associated with, or nullptr
+	 */
 	virtual UWorld* GetTickableGameObjectWorld() const 
 	{ 
 		return nullptr;
 	}
 
-	static ENGINE_API void TickObjects(UWorld* World, int32 TickType, bool bIsPaused, float DeltaSeconds);
+	/**
+	 * Call to modify the tickable type of this instance. 
+	 * This can be used to enable or disable tick even if GetTickableTickType has already been called.
+	 * 
+	 * @param NewTickType	The new tick type of this instance
+	 */
+	ENGINE_API void SetTickableTickType(ETickableTickType NewTickType);
+
+	/**
+	 * Tick all FTickableGameObject instances that match the parameters.
+	 *
+	 * @param World			Specific world that is ticking, must match GetTickableGameObjectWorld to tick.
+	 * @param LevelTickType	The type of tick where LEVELTICK_All is treated like a gameplay tick.
+	 * @param bIsPaused		True if the gameplay world is paused, if this is false IsTickableWhenPaused must return true to tick.
+	 * @param DeltaTime		Game time passed since the last call.
+	 */
+	static ENGINE_API void TickObjects(UWorld* World, ELevelTick LevelTickType, bool bIsPaused, float DeltaSeconds);
 };
