[Backout] - CL32538931
Original CL Desc
-----------------------------------------------------------------
Re-enable fatal shader compiler errors and have the ODSC server behave like the editor where we offer a message box to recompile

diff --git a/Engine/Source/Runtime/Engine/Private/ShaderCompiler/ShaderCompiler.cpp b/Engine/Source/Runtime/Engine/Private/ShaderCompiler/ShaderCompiler.cpp
--- a/Engine/Source/Runtime/Engine/Private/ShaderCompiler/ShaderCompiler.cpp
+++ b/Engine/Source/Runtime/Engine/Private/ShaderCompiler/ShaderCompiler.cpp
@@ -228,14 +228,6 @@ static TAutoConsoleVariable<bool> CVarAreShaderErrorsFatal(
 	TEXT("Default: true"),
 	ECVF_RenderThreadSafe);
 
-// Used for remote compile by COTF: Allow to prompt the message box to the user in case of a shader error
-static bool GRetryShaderCompileOnErrorForRemote = false;
-
-static inline bool IsLogShadersActive()
-{
-	return UE_LOG_ACTIVE(LogShaders, Log) || GRetryShaderCompileOnErrorForRemote;
-}
-
 bool AreShaderErrorsFatal()
 {
 	return CVarAreShaderErrorsFatal.GetValueOnAnyThread();
@@ -3348,7 +3340,7 @@ static int32 AddAndProcessErrorsForFailedJobFiltered(FShaderCompileJob& CurrentJ
 		FString CurrentErrorString = CurrentError.GetErrorString();
 
 		// Include warnings if LogShaders is unsuppressed, otherwise only include filtered messages
-		if (IsLogShadersActive() || FilterMessage == nullptr || CurrentError.StrippedErrorMessage.Contains(FilterMessage))
+		if (UE_LOG_ACTIVE(LogShaders, Log) || FilterMessage == nullptr || CurrentError.StrippedErrorMessage.Contains(FilterMessage))
 		{
 			// Extract source location from error message if the shader backend doesn't provide it separated from the stripped message
 			CurrentError.ExtractSourceLocation();
@@ -7161,7 +7153,7 @@ static bool GatherUniqueErrors(const TArray<FShaderCommonCompileJobPtr>& Complet
 
 bool FShaderCompilingManager::HandlePotentialRetryOnError(TMap<int32, FShaderMapFinalizeResults>& CompletedShaderMaps)
 {
-	if (FApp::IsUnattended() && !GRetryShaderCompileOnErrorForRemote)
+	if (FApp::IsUnattended())
 	{
 		return false;
 	}
@@ -7190,7 +7182,7 @@ bool FShaderCompilingManager::HandlePotentialRetryOnError(TMap<int32, FShaderMap
 				}
 			}
 
-			if (IsLogShadersActive()
+			if (UE_LOG_ACTIVE(LogShaders, Log) 
 				// Always log detailed errors when a special engine material or global shader fails to compile, as those will be fatal errors
 				|| bSpecialEngineMaterial 
 				|| It.Key() == GlobalShaderMapId)
@@ -7212,7 +7204,7 @@ bool FShaderCompilingManager::HandlePotentialRetryOnError(TMap<int32, FShaderMap
 
 				BuildErrorStringAndReport(ShaderErrorInfo, ErrorString);
 
-				if (IsLogShadersActive() && (bAnyErrorLikelyToBeCodeError || bPromptToRetryFailedShaderCompiles || bSpecialEngineMaterial))
+				if (UE_LOG_ACTIVE(LogShaders, Log) && (bAnyErrorLikelyToBeCodeError || bPromptToRetryFailedShaderCompiles || bSpecialEngineMaterial))
 				{
 					// Use debug break in debug with the debugger attached, otherwise message box
 					if (bDebugBreakOnPromptToRetryShaderCompile && FPlatformMisc::IsDebuggerPresent())
@@ -10447,12 +10439,8 @@ void RecompileShadersForRemote(
 {
 	TRACE_CPUPROFILER_EVENT_SCOPE(RecompileShadersForRemote);
 
-	GRetryShaderCompileOnErrorForRemote = true;
-
-	ON_SCOPE_EXIT
-	{ 
-		GRetryShaderCompileOnErrorForRemote = false;
-	};
+	CVarAreShaderErrorsFatal->Set(0, ECVF_SetByCode);
+	ON_SCOPE_EXIT{ CVarAreShaderErrorsFatal->Unset(ECVF_SetByCode); };
 
 	// figure out what shader platforms to recompile
 	ITargetPlatformManagerModule* TPM = GetTargetPlatformManager();
