Fixed crash when a cooked game streaming SVTs needs to cancel a currently running IO request.
The crash was caused by requests created with FBulkDataBatchRequest::FBatchBuilder not supporting Wait() operations. Fixed by switching to FBulkDataBatchRequest::FScatterGatherBuilder instead. An unfortunate downside of this fix is that FScatterGatherBuilder only supports cancelling and waiting on all requests in the batch, so for now we're forced to use a single batch for every request. For up to 8-16 requests, there's supposedly not really any performance difference between batching and not batching, so in most cases this change probably does not affect streaming performance negatively.
This issue should be revisited in the future to try and change the streamer logic such that we can also supported batched requests.

#rnx

diff --git a/Engine/Source/Runtime/Engine/Private/SparseVolumeTexture/SparseVolumeTextureStreamingManager.cpp b/Engine/Source/Runtime/Engine/Private/SparseVolumeTexture/SparseVolumeTextureStreamingManager.cpp
--- a/Engine/Source/Runtime/Engine/Private/SparseVolumeTexture/SparseVolumeTextureStreamingManager.cpp
+++ b/Engine/Source/Runtime/Engine/Private/SparseVolumeTexture/SparseVolumeTextureStreamingManager.cpp
@@ -1134,9 +1134,6 @@ void FStreamingManager::IssueRequests()
 	DDCRequestsBlocking.Reserve(TileRangesToStream.Num());
 #endif
 
-	FBulkDataBatchRequest::FBatchBuilder Batch = FBulkDataBatchRequest::NewBatch(TileRangesToStream.Num());
-	bool bIssueIOBatch = false;
-
 	// Process all tile ranges selected for streaming, allocate a slot in the tile data texture for every tile and finally create IO requests for every range.
 	for (FTileRange& TileRange : TileRangesToStream)
 	{
@@ -1328,8 +1325,11 @@ void FStreamingManager::IssueRequests()
 
 					PendingRequest.RequestBuffer = FIoBuffer(ReadSize); // SVT_TODO: Use FIoBuffer::Wrap with preallocated memory
 					const EAsyncIOPriorityAndFlags Priority = PendingRequest.bBlocking ? AIOP_CriticalPath : AIOP_Low;
-					Batch.Read(BulkData, ReadOffset, ReadSize, Priority, PendingRequest.RequestBuffer, PendingRequest.Request);
-					bIssueIOBatch = true;
+					// SVT_TODO: We're currently using a single batch per request so we can individually cancel and wait on requests.
+					// This isn't ideal and should be revisited in the future.
+					FBulkDataBatchRequest::FScatterGatherBuilder Batch = FBulkDataBatchRequest::ScatterGather(1);
+					Batch.Read(BulkData, ReadOffset, ReadSize);
+					Batch.Issue(PendingRequest.RequestBuffer, Priority, [](FBulkDataRequest::EStatus){}, PendingRequest.Request);
 
 #if WITH_EDITORONLY_DATA
 					PendingRequest.State = FPendingRequest::EState::Disk;
@@ -1362,11 +1362,6 @@ void FStreamingManager::IssueRequests()
 			DDCRequestsBlocking.Empty();
 		}
 #endif
-
-		if (bIssueIOBatch)
-		{
-			(void)Batch.Issue();
-		}
 	}
 }
 
@@ -1466,9 +1461,9 @@ int32 FStreamingManager::DetermineReadyRequests()
 							*SVTInfo->SVTName.ToString(), PendingRequest.FrameIndex, PendingRequest.TileOffset, PendingRequest.TileCount, ReadOffset, ReadSize);
 					}
 					
-					FBulkDataBatchRequest::FBatchBuilder Batch = FBulkDataBatchRequest::NewBatch(1);
-					Batch.Read(Resources->StreamableMipLevels, ReadOffset, ReadSize, AIOP_Low, PendingRequest.RequestBuffer, PendingRequest.Request);
-					(void)Batch.Issue();
+					FBulkDataBatchRequest::FScatterGatherBuilder Batch = FBulkDataBatchRequest::ScatterGather(1);
+					Batch.Read(Resources->StreamableMipLevels, ReadOffset, ReadSize);
+					Batch.Issue(PendingRequest.RequestBuffer, AIOP_Low, [](FBulkDataRequest::EStatus) {}, PendingRequest.Request);
 					break;
 				}
 			}
