AActor::FixupDataLayers : Fixes cases where an actor changing from EDL to no EDL didn't cleanup old EDL child DLs

diff --git a/Engine/Source/Runtime/Engine/Private/ActorEditor.cpp b/Engine/Source/Runtime/Engine/Private/ActorEditor.cpp
--- a/Engine/Source/Runtime/Engine/Private/ActorEditor.cpp
+++ b/Engine/Source/Runtime/Engine/Private/ActorEditor.cpp
@@ -1744,7 +1744,7 @@ void AActor::FixupDataLayers(bool bRevertChangesOnLockedDataLayer /*= false*/)
 	// Remove any data layer with a Root external data layer different from the actor's external data layer
 	// This case can happen after duplicating an actor using the actor editor context with a different external data layer
 	TSet<TSoftObjectPtr<const UDataLayerAsset>> InvalidDataLayerAssets;
-	if (ExternalDataLayerAsset && !DataLayerAssets.IsEmpty())
+	if (!DataLayerAssets.IsEmpty())
 	{
 		for (const UDataLayerInstance* DataLayerInstance : GetDataLayerInstances())
 		{
