FExternalDataLayerHelper::MoveActorsToExternalDataLayer now removes old Content Bundle Guid when succeeds.

diff --git a/Engine/Source/Runtime/Engine/Classes/GameFramework/Actor.h b/Engine/Source/Runtime/Engine/Classes/GameFramework/Actor.h
--- a/Engine/Source/Runtime/Engine/Classes/GameFramework/Actor.h
+++ b/Engine/Source/Runtime/Engine/Classes/GameFramework/Actor.h
@@ -4535,6 +4535,7 @@ private:
 		InActor->ContentBundleGuid = InContentBundleGuid;
 	}
 	friend class FContentBundleEditor;
+	friend class FExternalDataLayerHelper;
 	friend class UGameFeatureActionConvertContentBundleWorldPartitionBuilder;
 };
 
diff --git a/Engine/Source/Runtime/Engine/Private/WorldPartition/DataLayer/ExternalDataLayerHelper.cpp b/Engine/Source/Runtime/Engine/Private/WorldPartition/DataLayer/ExternalDataLayerHelper.cpp
--- a/Engine/Source/Runtime/Engine/Private/WorldPartition/DataLayer/ExternalDataLayerHelper.cpp
+++ b/Engine/Source/Runtime/Engine/Private/WorldPartition/DataLayer/ExternalDataLayerHelper.cpp
@@ -246,6 +246,9 @@ bool FExternalDataLayerHelper::MoveActorsToExternalDataLayer(const TArray<AActor
 		bool bLevelPackageWasDirty = InActor->GetLevel()->GetPackage()->IsDirty();
 		InActor->SetPackageExternal(false, bShouldDirty);
 
+		// Clear Content Bundle Guid
+		FSetActorContentBundleGuid(InActor, FGuid());
+
 		// If set, remove EDL from actor
 		const UExternalDataLayerAsset* OldExternalDataLayerAsset = InActor->GetExternalDataLayerAsset();
 		if (OldExternalDataLayerAsset)
