[Content Bundle] Changed verbosity on several spamming logs that affects cooking times.

diff --git a/Engine/Source/Runtime/Engine/Private/WorldPartition/ContentBundle/ContentBundleBase.cpp b/Engine/Source/Runtime/Engine/Private/WorldPartition/ContentBundle/ContentBundleBase.cpp
--- a/Engine/Source/Runtime/Engine/Private/WorldPartition/ContentBundle/ContentBundleBase.cpp
+++ b/Engine/Source/Runtime/Engine/Private/WorldPartition/ContentBundle/ContentBundleBase.cpp
@@ -131,7 +131,8 @@ void FContentBundleBase::SetStatus(EContentBundleStatus NewStatus)
 {
 	check(NewStatus != Status);
 
-	UE_LOG(LogContentBundle, Log, TEXT("%s State changing from %s to %s"), 
+	UE_LOG(LogContentBundle, Verbose, TEXT("%s State changing from %s to %s"), 
 		*ContentBundle::Log::MakeDebugInfoString(*this), *UEnum::GetDisplayValueAsText(Status).ToString(), *UEnum::GetDisplayValueAsText(NewStatus).ToString());
+
 	Status = NewStatus;
 }
diff --git a/Engine/Source/Runtime/Engine/Private/WorldPartition/ContentBundle/ContentBundleClient.cpp b/Engine/Source/Runtime/Engine/Private/WorldPartition/ContentBundle/ContentBundleClient.cpp
--- a/Engine/Source/Runtime/Engine/Private/WorldPartition/ContentBundle/ContentBundleClient.cpp
+++ b/Engine/Source/Runtime/Engine/Private/WorldPartition/ContentBundle/ContentBundleClient.cpp
@@ -20,7 +20,7 @@ FContentBundleClient::FContentBundleClient(const UContentBundleDescriptor* InCon
 	, DisplayName(InDisplayName)
 	, State(EContentBundleClientState::Unregistered)
 {
-	UE_LOG(LogContentBundle, Log, TEXT("[CB: %s] Client Created"), *GetDescriptor()->GetDisplayName());
+	UE_LOG(LogContentBundle, Verbose, TEXT("[CB: %s] Client Created"), *GetDescriptor()->GetDisplayName());
 	SetState(EContentBundleClientState::Registered);
 }
 
@@ -126,12 +126,12 @@ void FContentBundleClient::SetWorldContentState(UWorld* World, EWorldContentStat
 {
 	if (EWorldContentState* OldState = WorldContentStates.Find(World))
 	{
-		UE_LOG(LogContentBundle, Log, TEXT("%s Client WorldState changing from %s to %s"), *ContentBundle::Log::MakeDebugInfoString(*this, World), *UEnum::GetDisplayValueAsText(*OldState).ToString(), *UEnum::GetDisplayValueAsText(NewState).ToString());
+		UE_LOG(LogContentBundle, Verbose, TEXT("%s Client WorldState changing from %s to %s"), *ContentBundle::Log::MakeDebugInfoString(*this, World), *UEnum::GetDisplayValueAsText(*OldState).ToString(), *UEnum::GetDisplayValueAsText(NewState).ToString());
 		*OldState = NewState;
 		return;
 	}
 	
-	UE_LOG(LogContentBundle, Log, TEXT("%s Client WorldState changing to %s"), *ContentBundle::Log::MakeDebugInfoString(*this, World), *UEnum::GetDisplayValueAsText(NewState).ToString());
+	UE_LOG(LogContentBundle, Verbose, TEXT("%s Client WorldState changing to %s"), *ContentBundle::Log::MakeDebugInfoString(*this, World), *UEnum::GetDisplayValueAsText(NewState).ToString());
 	WorldContentStates.Add(World, NewState);
 }
 
@@ -163,6 +163,8 @@ void FContentBundleClient::SetState(EContentBundleClientState NewState)
 {
 	check(NewState != State);
 
-	UE_LOG(LogContentBundle, Log, TEXT("%s Client State changing from %s to %s"), *ContentBundle::Log::MakeDebugInfoString(*this), *UEnum::GetDisplayValueAsText(State).ToString(), *UEnum::GetDisplayValueAsText(NewState).ToString());
+	UE_LOG(LogContentBundle, Verbose, TEXT("%s Client State changing from %s to %s"), 
+		*ContentBundle::Log::MakeDebugInfoString(*this), *UEnum::GetDisplayValueAsText(State).ToString(), *UEnum::GetDisplayValueAsText(NewState).ToString());
+
 	State = NewState;
 }
\ No newline at end of file
diff --git a/Engine/Source/Runtime/Engine/Private/WorldPartition/ContentBundle/ContentBundleContainer.cpp b/Engine/Source/Runtime/Engine/Private/WorldPartition/ContentBundle/ContentBundleContainer.cpp
--- a/Engine/Source/Runtime/Engine/Private/WorldPartition/ContentBundle/ContentBundleContainer.cpp
+++ b/Engine/Source/Runtime/Engine/Private/WorldPartition/ContentBundle/ContentBundleContainer.cpp
@@ -35,7 +35,7 @@ UWorld* FContentBundleContainer::GetInjectedWorld() const
 void FContentBundleContainer::Initialize()
 {
 	TRACE_CPUPROFILER_EVENT_SCOPE(FContentBundleContainer::Initialize);
-	UE_LOG(LogContentBundle, Log, TEXT("%s Creating new container."), *ContentBundle::Log::MakeDebugInfoString(*this));
+	UE_LOG(LogContentBundle, Verbose, TEXT("%s Creating new container."), *ContentBundle::Log::MakeDebugInfoString(*this));
 
 #if WITH_EDITOR
 	if (UseEditorContentBundle())
@@ -68,7 +68,7 @@ void FContentBundleContainer::Deinitialize()
 
 	if (GetInjectedWorld())
 	{
-		UE_LOG(LogContentBundle, Log, TEXT("%s Deleting container."), *ContentBundle::Log::MakeDebugInfoString(*this));
+		UE_LOG(LogContentBundle, Verbose, TEXT("%s Deleting container."), *ContentBundle::Log::MakeDebugInfoString(*this));
 
 #if WITH_EDITOR
 		if (UseEditorContentBundle())
@@ -235,7 +235,7 @@ FContentBundleBase& FContentBundleContainer::InitializeContentBundle(TSharedPtr<
 
 	FContentBundleBase* ContentBundle = nullptr;
 
-	UE_LOG(LogContentBundle, Log, TEXT("%s Creating new content bundle from client %s with client state %s."), 
+	UE_LOG(LogContentBundle, Verbose, TEXT("%s Creating new content bundle from client %s with client state %s."), 
 		*ContentBundle::Log::MakeDebugInfoString(*ContentBundleClient , GetInjectedWorld()), *ContentBundleClient->GetDisplayName(), *UEnum::GetDisplayValueAsText(ContentBundleClient->GetState()).ToString());
 
 #if WITH_EDITOR
@@ -380,7 +380,7 @@ void FContentBundleContainer::InitializeContentBundlesForegisteredClients()
 
 	if (!ContentBundleClients.IsEmpty())
 	{
-		UE_LOG(LogContentBundle, Log, TEXT("%s Begin initializing ContentBundles from %u registered clients."), *ContentBundle::Log::MakeDebugInfoString(*this), ContentBundleClients.Num());
+		UE_LOG(LogContentBundle, Verbose, TEXT("%s Begin initializing ContentBundles from %u registered clients."), *ContentBundle::Log::MakeDebugInfoString(*this), ContentBundleClients.Num());
 
 		for (TSharedPtr<FContentBundleClient>& ContentBundleClient : ContentBundleClients)
 		{
@@ -388,7 +388,7 @@ void FContentBundleContainer::InitializeContentBundlesForegisteredClients()
 			InjectContentBundle(ContentBundle);
 		}
 
-		UE_LOG(LogContentBundle, Log, TEXT("%s End initializing ContentBundles."), *ContentBundle::Log::MakeDebugInfoString(*this));
+		UE_LOG(LogContentBundle, Verbose, TEXT("%s End initializing ContentBundles."), *ContentBundle::Log::MakeDebugInfoString(*this));
 	}
 }
 
diff --git a/Engine/Source/Runtime/Engine/Private/WorldPartition/ContentBundle/ContentBundleEditor.cpp b/Engine/Source/Runtime/Engine/Private/WorldPartition/ContentBundle/ContentBundleEditor.cpp
--- a/Engine/Source/Runtime/Engine/Private/WorldPartition/ContentBundle/ContentBundleEditor.cpp
+++ b/Engine/Source/Runtime/Engine/Private/WorldPartition/ContentBundle/ContentBundleEditor.cpp
@@ -89,7 +89,7 @@ void FContentBundleEditor::DoInjectContent()
 		ActorDescContainerInstance = WorldPartition->RegisterActorDescContainerInstance(InitParams);
 		if (ActorDescContainerInstance.IsValid() && ActorDescContainerInstance->IsInitialized())
 		{
-			UE_LOG(LogContentBundle, Log, TEXT("%s ExternalActors in %s found. %u actors were injected"), *ContentBundle::Log::MakeDebugInfoString(*this), *ActorDescContainerInstance->GetExternalActorPath(), ActorDescContainerInstance->GetActorDescInstanceCount());
+			UE_LOG(LogContentBundle, Verbose, TEXT("%s ExternalActors in %s found. %u actors were injected"), *ContentBundle::Log::MakeDebugInfoString(*this), *ActorDescContainerInstance->GetExternalActorPath(), ActorDescContainerInstance->GetActorDescInstanceCount());
 
 			check(GetDescriptor()->GetGuid().IsValid());
 
