Added a local position material expression, replacing the local position material function. The native version reduces the number of LWC operations, improving precision and performance.

#virtualized

diff --git a/Engine/Shaders/Private/MaterialTemplate.ush b/Engine/Shaders/Private/MaterialTemplate.ush
--- a/Engine/Shaders/Private/MaterialTemplate.ush
+++ b/Engine/Shaders/Private/MaterialTemplate.ush
@@ -970,17 +970,17 @@ FWSVector3 GetPrevParticleWorldPosition(FMaterialPixelParameters Parameters)
 
 // Local position
 
-float3 GetPrimitiveSpacePosition(FMaterialPixelParameters Parameters)
+float3 GetPositionPrimitiveSpace(FMaterialPixelParameters Parameters)
 {
 	return DFFastMultiplyDemote(Parameters.AbsoluteWorldPosition, GetPrimitiveData(Parameters).WorldToLocal);
 }
 
-float3 GetPrimitiveSpacePosition(FMaterialVertexParameters Parameters)
+float3 GetPositionPrimitiveSpace(FMaterialVertexParameters Parameters)
 {
 	return Parameters.PositionPrimitiveSpace;
 }
 
-float3 GetPrevPrimitiveSpacePosition(FMaterialVertexParameters Parameters)
+float3 GetPrevPositionPrimitiveSpace(FMaterialVertexParameters Parameters)
 {
 	// Previous position and current position are sharing the
 	// same attribute in Parameters, because in BasePassVertexShader.usf
@@ -989,22 +989,22 @@ float3 GetPrevPrimitiveSpacePosition(FMaterialVertexParameters Parameters)
 	return Parameters.PositionPrimitiveSpace;
 }
 
-float3 GetPrimitiveSpacePosition_NoMaterialOffsets(FMaterialPixelParameters Parameters)
+float3 GetPositionPrimitiveSpace_NoMaterialOffsets(FMaterialPixelParameters Parameters)
 {
 	return DFFastMultiplyDemote(Parameters.WorldPosition_NoOffsets, GetPrimitiveData(Parameters).WorldToLocal);
 }
 
-float3 GetInstanceSpacePosition(FMaterialPixelParameters Parameters)
+float3 GetPositionInstanceSpace(FMaterialPixelParameters Parameters)
 {
 	return DFFastMultiplyDemote(Parameters.AbsoluteWorldPosition, GetWorldToInstanceDF(Parameters));
 }
 
-float3 GetInstanceSpacePosition(FMaterialVertexParameters Parameters)
+float3 GetPositionInstanceSpace(FMaterialVertexParameters Parameters)
 {
 	return Parameters.PositionInstanceSpace;
 }
 
-float3 GetPrevInstanceSpacePosition(FMaterialVertexParameters Parameters)
+float3 GetPrevPositionInstanceSpace(FMaterialVertexParameters Parameters)
 {
 	// Previous position and current position are sharing the
 	// same attribute in Parameters, because in BasePassVertexShader.usf
@@ -1013,7 +1013,7 @@ float3 GetPrevInstanceSpacePosition(FMaterialVertexParameters Parameters)
 	return Parameters.PositionInstanceSpace;
 }
 
-float3 GetInstanceSpacePosition_NoMaterialOffsets(FMaterialPixelParameters Parameters)
+float3 GetPositionInstanceSpace_NoMaterialOffsets(FMaterialPixelParameters Parameters)
 {
 	return DFFastMultiplyDemote(Parameters.WorldPosition_NoOffsets, GetWorldToInstanceDF(Parameters));
 }
diff --git a/Engine/Source/Runtime/Core/Private/UObject/DevObjectVersion.cpp b/Engine/Source/Runtime/Core/Private/UObject/DevObjectVersion.cpp
--- a/Engine/Source/Runtime/Core/Private/UObject/DevObjectVersion.cpp
+++ b/Engine/Source/Runtime/Core/Private/UObject/DevObjectVersion.cpp
@@ -135,7 +135,7 @@ FDevSystemGuids::FDevSystemGuids()
 	, POSESEARCHDB_DERIVEDDATA_VER(0x389117E4, 0x807A4CC0, 0x9F37C2E6, 0xD808A78D)
 	, SkeletalMeshDerivedDataVersion(0x9B5F4544, 0x76D7481C, 0x9AD3F614, 0xA6C07904)
 	, STATICMESH_DERIVEDDATA_VER(0x2C6C400C, 0x2EDF47B3, 0x9BD36689, 0x7B77B208)
-	, MaterialTranslationDDCVersion(0x59E1F296, 0x27064757, 0x8CA7EB9B, 0xFC366BD2)
+	, MaterialTranslationDDCVersion(0xB1C6F9A0, 0x038E48A3, 0xA6D467A6, 0xC81A5849)
 {
 }
 
diff --git a/Engine/Source/Runtime/Core/Private/UObject/FortniteMainBranchObjectVersion.cpp b/Engine/Source/Runtime/Core/Private/UObject/FortniteMainBranchObjectVersion.cpp
--- a/Engine/Source/Runtime/Core/Private/UObject/FortniteMainBranchObjectVersion.cpp
+++ b/Engine/Source/Runtime/Core/Private/UObject/FortniteMainBranchObjectVersion.cpp
@@ -14,6 +14,6 @@ TMap<FGuid, FGuid> FFortniteMainBranchObjectVersion::GetSystemGuids()
 	SystemGuids.Add(DevGuids.Niagara_LatestScriptCompileVersion, FGuid("E267F600682A4E94885E8AEC5CE1DC91"));
 	SystemGuids.Add(DevGuids.SkeletalMeshDerivedDataVersion, FGuid("B9AEAA3EE7AC4FCFB3860E8F3E06DD75"));
 	SystemGuids.Add(DevGuids.STATICMESH_DERIVEDDATA_VER, FGuid("543843B103794E0E9BA4BE80FB602F79"));
-	SystemGuids.Add(DevGuids.MaterialTranslationDDCVersion, FGuid("BBD650804D314CE7A44ADF9BB4F3F5F9"));
+	SystemGuids.Add(DevGuids.MaterialTranslationDDCVersion, FGuid("B1C6F9A0038E48A3A6D467A6C81A5849"));
 	return SystemGuids;
 }
diff --git a/Engine/Source/Runtime/Engine/Classes/Materials/MaterialExpressionLocalPosition.h b/Engine/Source/Runtime/Engine/Classes/Materials/MaterialExpressionLocalPosition.h
new file mode 100644
--- /dev/null
+++ b/Engine/Source/Runtime/Engine/Classes/Materials/MaterialExpressionLocalPosition.h
@@ -0,0 +1,55 @@
+// Copyright Epic Games, Inc. All Rights Reserved.
+
+#pragma once
+
+#include "CoreTypes.h"
+#include "UObject/ObjectMacros.h"
+#include "Materials/MaterialExpression.h"
+#include "MaterialExpressionLocalPosition.generated.h"
+
+/** Specifies which shader generated offsets should included in the position (displacement/WPO etc.) */
+UENUM()
+enum class EPositionIncludedOffsets
+{
+	/** Position with all material shader offsets applied */
+	IncludeOffsets UMETA(DisplayName="Include Material Shader Offsets"),
+
+	/** Position with no material shader offsets applied */
+	ExcludeOffsets UMETA(DisplayName="Exclude Material Shader Offsets")
+};
+
+/** Specifies the reference point of the local position. This can be different in some cases, e.g. for instanced meshes. */
+UENUM()
+enum class ELocalPositionOrigin
+{
+	/** Position relative to instance */
+	Instance UMETA(DisplayName="Instance"),
+
+	/** Position relative to instance, before skinning is applied. */
+	/** Returns the instance position for non-skeletal meshes. Incompatible with GPU skin cache feature. */
+	InstancePreSkinning UMETA(DisplayName="Instance (pre-skinning)"),
+
+	/** Position relative to primitive actor component */
+	Primitive UMETA(DisplayName="Component")
+};
+
+UCLASS(collapsecategories, hidecategories=Object, MinimalAPI)
+class UMaterialExpressionLocalPosition : public UMaterialExpression
+{
+	GENERATED_UCLASS_BODY()
+
+	UPROPERTY(EditAnywhere, Category=MaterialExpressionLocalPosition, meta=(DisplayName = "Shader Offsets", ShowAsInputPin = "Advanced"))
+	EPositionIncludedOffsets IncludedOffsets;
+
+	UPROPERTY(EditAnywhere, Category=MaterialExpressionLocalPosition, meta=(DisplayName = "Local Origin", ShowAsInputPin = "Advanced"))
+	ELocalPositionOrigin LocalOrigin;
+	
+	//~ Begin UMaterialExpression Interface
+#if WITH_EDITOR
+	bool GenerateHLSLExpression(FMaterialHLSLGenerator& Generator, UE::HLSLTree::FScope& Scope, int32 OutputIndex, UE::HLSLTree::FExpression const*& OutExpression) const;
+	virtual int32 Compile(class FMaterialCompiler* Compiler, int32 OutputIndex) override;
+	virtual void GetCaption(TArray<FString>& OutCaptions) const override;
+	virtual FText GetKeywords() const override {return FText::FromString(TEXT("position preskinned local instance primitive"));}
+#endif
+	//~ End UMaterialExpression Interface
+};
diff --git a/Engine/Source/Runtime/Engine/Private/Materials/HLSLMaterialTranslator.cpp b/Engine/Source/Runtime/Engine/Private/Materials/HLSLMaterialTranslator.cpp
--- a/Engine/Source/Runtime/Engine/Private/Materials/HLSLMaterialTranslator.cpp
+++ b/Engine/Source/Runtime/Engine/Private/Materials/HLSLMaterialTranslator.cpp
@@ -5935,20 +5935,96 @@ int32 FHLSLMaterialTranslator::ParticleSpriteRotation()
 	return AddInlinedCodeChunkZeroDeriv(MCT_Float2, TEXT("float2(Parameters.Particle.SpriteRotation, Parameters.Particle.SpriteRotation * 57.2957795131f)"));
 }
 
+int32 FHLSLMaterialTranslator::LocalPosition(EPositionIncludedOffsets IncludedOffsets, ELocalPositionOrigin OriginType)
+{
+	if (OriginType == ELocalPositionOrigin::InstancePreSkinning)
+	{
+		if (ShaderFrequency != SF_Vertex)
+		{
+			return Errorf(TEXT("Pre-skinned position is only available in the vertex shader, pass through custom interpolators if needed."));
+		}
+
+		return AddInlinedCodeChunk(MCT_Float3,TEXT("Parameters.PreSkinnedPosition"));
+	}
+
+	// If compiling for the previous frame in the vertex shader
+	bool bPreviousIsAvailable = ShaderFrequency == SF_Vertex;
+	const TCHAR* PrevPart = (bCompilingPreviousFrame && bPreviousIsAvailable) ? TEXT("Prev") : TEXT("");
+
+	const TCHAR* OriginPart = nullptr;
+	switch (OriginType)
+	{
+		case ELocalPositionOrigin::Instance:
+		{
+			OriginPart = TEXT("Instance");
+			break;
+		}
+		case ELocalPositionOrigin::Primitive:
+		{
+			OriginPart = TEXT("Primitive");
+			break;
+		}
+		default: checkNoEntry();
+	}
+
+	const TCHAR* NoMaterialOffsetsPart = nullptr;
+	switch (IncludedOffsets)
+	{
+		case EPositionIncludedOffsets::IncludeOffsets:
+		{
+			NoMaterialOffsetsPart = TEXT("");
+			break;
+		}
+		case EPositionIncludedOffsets::ExcludeOffsets:
+		{
+			// No material offset only available in the pixel shader.
+			if (ShaderFrequency == SF_Pixel)
+			{
+				NoMaterialOffsetsPart = TEXT("_NoMaterialOffsets");
+			}
+			else
+			{
+				NoMaterialOffsetsPart = TEXT("");
+			}
+			break;
+		}
+		default: checkNoEntry();
+	}
+
+	bUsesInstanceWorldToLocalPS |= (ShaderFrequency == SF_Pixel && OriginType == ELocalPositionOrigin::Instance);
+
+	TArray<FStringFormatArg> FormatArgs =
+	{
+		PrevPart,
+		OriginPart,
+		NoMaterialOffsetsPart,
+	};
+	FString FiniteCode = FString::Format(TEXT("Get{0}Position{1}Space{2}(Parameters)"), FormatArgs);
+
+	return AddInlinedCodeChunk(MCT_Float3, *FiniteCode);
+}
+
 int32 FHLSLMaterialTranslator::WorldPosition(EWorldPositionIncludedOffsets WorldPositionIncludedOffsets)
 {
-	FString FunctionNamePattern;
-	EMaterialValueType Type = (EMaterialValueType)0;
+	// If compiling for the previous frame in the vertex shader
+	bool bPreviousIsAvailable = ShaderFrequency == SF_Vertex;
+	const TCHAR* PrevPart = (bCompilingPreviousFrame && bPreviousIsAvailable) ? TEXT("Prev") : TEXT("");
+
+	bool bNoMaterialOffsetsIsAvailable = ShaderFrequency == SF_Pixel;
 
 	// If this material has no expressions for world position offset or world displacement, the non-offset world position will
 	// be exactly the same as the offset one, so there is no point bringing in the extra code.
 	// Also, we can't access the full offset world position in anything other than the pixel shader, because it won't have
 	// been calculated yet
+	const TCHAR* OriginPart = nullptr;
+	const TCHAR* NoMaterialOffsetsPart = nullptr;
+	EMaterialValueType Type = (EMaterialValueType)0;
 	switch (WorldPositionIncludedOffsets)
 	{
 	case WPT_Default:
 		{
-			FunctionNamePattern = TEXT("Get<PREV>WorldPosition");
+			OriginPart = TEXT("World");
+			NoMaterialOffsetsPart = TEXT("");
 			Type = MCT_LWCVector3;
 			break;
 		}
@@ -5956,14 +6032,16 @@ int32 FHLSLMaterialTranslator::WorldPosition(EWorldPositionIncludedOffsets World
 	case WPT_ExcludeAllShaderOffsets:
 		{
 			bNeedsWorldPositionExcludingShaderOffsets |= (ShaderFrequency == SF_Pixel);
-			FunctionNamePattern = TEXT("Get<PREV>WorldPosition<NO_MATERIAL_OFFSETS>");
+			OriginPart = TEXT("World");
+			NoMaterialOffsetsPart = bNoMaterialOffsetsIsAvailable ? TEXT("_NoMaterialOffsets") : TEXT("");
 			Type = MCT_LWCVector3;
 			break;
 		}
 
 	case WPT_CameraRelative:
 		{
-			FunctionNamePattern = TEXT("Get<PREV>TranslatedWorldPosition");
+			OriginPart = TEXT("TranslatedWorld");
+			NoMaterialOffsetsPart = TEXT("");
 			Type = MCT_Float3;
 			break;
 		}
@@ -5971,7 +6049,8 @@ int32 FHLSLMaterialTranslator::WorldPosition(EWorldPositionIncludedOffsets World
 	case WPT_CameraRelativeNoOffsets:
 		{
 			bNeedsWorldPositionExcludingShaderOffsets |= (ShaderFrequency == SF_Pixel);
-			FunctionNamePattern = TEXT("Get<PREV>TranslatedWorldPosition<NO_MATERIAL_OFFSETS>");
+			OriginPart = TEXT("TranslatedWorld");
+			NoMaterialOffsetsPart = bNoMaterialOffsetsIsAvailable ? TEXT("_NoMaterialOffsets") : TEXT("");
 			Type = MCT_Float3;
 			break;
 		}
@@ -5982,23 +6061,17 @@ int32 FHLSLMaterialTranslator::WorldPosition(EWorldPositionIncludedOffsets World
 			return INDEX_NONE;
 		}
 	}
+	
+	bUsesVertexPosition = true;
 
-	// If compiling for the previous frame in the vertex shader
-	FunctionNamePattern.ReplaceInline(TEXT("<PREV>"), bCompilingPreviousFrame && ShaderFrequency == SF_Vertex ? TEXT("Prev") : TEXT(""));
-		
-	if (ShaderFrequency == SF_Pixel)
+	TArray<FStringFormatArg> FormatArgs =
 	{
-		// No material offset only available in the vertex shader.
-		FunctionNamePattern.ReplaceInline(TEXT("<NO_MATERIAL_OFFSETS>"), TEXT("_NoMaterialOffsets"));
-	}
-	else
-	{
-		FunctionNamePattern.ReplaceInline(TEXT("<NO_MATERIAL_OFFSETS>"), TEXT(""));
-	}
-
-	bUsesVertexPosition = true;
+		PrevPart,
+		OriginPart,
+		NoMaterialOffsetsPart,
+	};
+	FString FiniteCode = FString::Format(TEXT("Get{0}{1}Position{2}(Parameters)"), FormatArgs);
 
-	FString FiniteCode = FString::Printf(TEXT("%s(Parameters)"), *FunctionNamePattern);
 	int32 Result = INDEX_NONE;
 	if (IsAnalyticDerivEnabled())
 	{
diff --git a/Engine/Source/Runtime/Engine/Private/Materials/HLSLMaterialTranslator.h b/Engine/Source/Runtime/Engine/Private/Materials/HLSLMaterialTranslator.h
--- a/Engine/Source/Runtime/Engine/Private/Materials/HLSLMaterialTranslator.h
+++ b/Engine/Source/Runtime/Engine/Private/Materials/HLSLMaterialTranslator.h
@@ -962,6 +962,7 @@ protected:
 	virtual int32 ParticleSize() override;
 	virtual int32 ParticleSpriteRotation() override;
 
+	virtual int32 LocalPosition(EPositionIncludedOffsets IncludedOffsets, ELocalPositionOrigin OriginType) override;
 	virtual int32 WorldPosition(EWorldPositionIncludedOffsets WorldPositionIncludedOffsets) override;
 
 	virtual int32 ObjectWorldPosition(EPositionOrigin OriginType) override;
diff --git a/Engine/Source/Runtime/Engine/Private/Materials/MaterialExpressionHLSL.cpp b/Engine/Source/Runtime/Engine/Private/Materials/MaterialExpressionHLSL.cpp
--- a/Engine/Source/Runtime/Engine/Private/Materials/MaterialExpressionHLSL.cpp
+++ b/Engine/Source/Runtime/Engine/Private/Materials/MaterialExpressionHLSL.cpp
@@ -107,6 +107,7 @@
 #include "Materials/MaterialExpressionLightVector.h"
 #include "Materials/MaterialExpressionLinearInterpolate.h"
 #include "Materials/MaterialExpressionLightmassReplace.h"
+#include "Materials/MaterialExpressionLocalPosition.h"
 #include "Materials/MaterialExpressionLogarithm.h"
 #include "Materials/MaterialExpressionLogarithm10.h"
 #include "Materials/MaterialExpressionLogarithm2.h"
@@ -1017,6 +1018,43 @@ bool UMaterialExpressionPixelDepth::GenerateHLSLExpression(FMaterialHLSLGenerato
 	return true;
 }
 
+bool UMaterialExpressionLocalPosition::GenerateHLSLExpression(FMaterialHLSLGenerator& Generator, UE::HLSLTree::FScope& Scope, int32 OutputIndex, UE::HLSLTree::FExpression const*& OutExpression) const
+{
+	using namespace UE::HLSLTree::Material;
+
+	if (LocalOrigin == ELocalPositionOrigin::InstancePreSkinning)
+	{
+		OutExpression = Generator.NewExternalInput(EExternalInput::PreSkinnedPosition);
+		return true;
+	}
+
+	if (IncludedOffsets == EPositionIncludedOffsets::IncludeOffsets
+		&& LocalOrigin == ELocalPositionOrigin::Instance)
+	{
+		OutExpression = Generator.NewExternalInput(EExternalInput::PositionInstanceSpace);
+	}
+	else if (IncludedOffsets == EPositionIncludedOffsets::IncludeOffsets
+		&& LocalOrigin == ELocalPositionOrigin::Primitive)
+	{
+		OutExpression = Generator.NewExternalInput(EExternalInput::PositionPrimitiveSpace);
+	}
+	else if (IncludedOffsets == EPositionIncludedOffsets::ExcludeOffsets
+		&& LocalOrigin == ELocalPositionOrigin::Instance)
+	{
+		OutExpression = Generator.NewExternalInput(EExternalInput::PositionInstanceSpace_NoOffsets);
+	}
+	else if (IncludedOffsets == EPositionIncludedOffsets::ExcludeOffsets
+		&& LocalOrigin == ELocalPositionOrigin::Primitive)
+	{
+		OutExpression = Generator.NewExternalInput(EExternalInput::PositionPrimitiveSpace_NoOffsets);
+	}
+	else
+	{
+		checkNoEntry();
+	}
+	return true;
+}
+
 bool UMaterialExpressionWorldPosition::GenerateHLSLExpression(FMaterialHLSLGenerator& Generator, UE::HLSLTree::FScope& Scope, int32 OutputIndex, UE::HLSLTree::FExpression const*& OutExpression) const
 {
 	using namespace UE::HLSLTree::Material;
diff --git a/Engine/Source/Runtime/Engine/Private/Materials/MaterialExpressions.cpp b/Engine/Source/Runtime/Engine/Private/Materials/MaterialExpressions.cpp
--- a/Engine/Source/Runtime/Engine/Private/Materials/MaterialExpressions.cpp
+++ b/Engine/Source/Runtime/Engine/Private/Materials/MaterialExpressions.cpp
@@ -133,6 +133,7 @@
 #include "Materials/MaterialExpressionExponential.h"
 #include "Materials/MaterialExpressionExponential2.h"
 #include "Materials/MaterialExpressionLength.h"
+#include "Materials/MaterialExpressionLocalPosition.h"
 #include "Materials/MaterialExpressionLogarithm.h"
 #include "Materials/MaterialExpressionLogarithm2.h"
 #include "Materials/MaterialExpressionLogarithm10.h"
@@ -13212,6 +13213,75 @@ void UMaterialExpressionFontSampleParameter::SetEditableName(const FString& NewN
 }
 #endif
 
+///////////////////////////////////////////////////////////////////////////////
+// UMaterialExpressionLocalPosition
+///////////////////////////////////////////////////////////////////////////////
+
+UMaterialExpressionLocalPosition::UMaterialExpressionLocalPosition(const FObjectInitializer& ObjectInitializer)
+	: Super(ObjectInitializer)
+{
+#if WITH_EDITORONLY_DATA
+	// Structure to hold one-time initialization
+	struct FConstructorStatics
+	{
+		FText NAME_Coordinates;
+		FConstructorStatics()
+		: NAME_Coordinates(LOCTEXT( "Coordinates", "Coordinates" ))
+		{
+		}
+	};
+	static FConstructorStatics ConstructorStatics;
+
+	MenuCategories.Add(ConstructorStatics.NAME_Coordinates);
+
+	Outputs.Reset();
+	Outputs.Add(FExpressionOutput(TEXT("XYZ"), 1, 1, 1, 1, 0));
+	Outputs.Add(FExpressionOutput(TEXT("XY"), 1, 1, 1, 0, 0));
+	Outputs.Add(FExpressionOutput(TEXT("Z"), 1, 0, 0, 1, 0));
+
+	bShowOutputNameOnPin = true;
+	bShaderInputData = true;
+#endif
+
+	IncludedOffsets = EPositionIncludedOffsets::IncludeOffsets;
+	LocalOrigin = ELocalPositionOrigin::Instance;
+}
+
+#if WITH_EDITOR
+int32 UMaterialExpressionLocalPosition::Compile(class FMaterialCompiler* Compiler, int32 OutputIndex)
+{
+	return Compiler->LocalPosition(IncludedOffsets, LocalOrigin);
+}
+
+void UMaterialExpressionLocalPosition::GetCaption(TArray<FString>& OutCaptions) const
+{
+	if (LocalOrigin == ELocalPositionOrigin::InstancePreSkinning)
+	{
+		OutCaptions.Add(NSLOCTEXT("MaterialExpressions", "LocalPositionInstanceIncludingOffsetsText", "Pre-Skinned Local Position").ToString());
+	}
+	else if (IncludedOffsets == EPositionIncludedOffsets::IncludeOffsets && LocalOrigin == ELocalPositionOrigin::Instance)
+	{
+		OutCaptions.Add(NSLOCTEXT("MaterialExpressions", "LocalPositionInstanceIncludingOffsetsText", "Local Position").ToString());
+	}
+	else if (IncludedOffsets == EPositionIncludedOffsets::ExcludeOffsets && LocalOrigin == ELocalPositionOrigin::Instance)
+	{
+		OutCaptions.Add(NSLOCTEXT("MaterialExpressions", "LocalPositionInstanceExcludingOffsetsText", "Local Position (Excluding Material Offsets)").ToString());
+	}
+	else if (IncludedOffsets == EPositionIncludedOffsets::IncludeOffsets && LocalOrigin == ELocalPositionOrigin::Primitive)
+	{
+		OutCaptions.Add(NSLOCTEXT("MaterialExpressions", "LocalPositionComponentIncludingOffsetsText", "Component Local Position").ToString());
+	}
+	else if (IncludedOffsets == EPositionIncludedOffsets::ExcludeOffsets && LocalOrigin == ELocalPositionOrigin::Primitive)
+	{
+		OutCaptions.Add(NSLOCTEXT("MaterialExpressions", "LocalPositionComponentExcludingOffsetsText", "Component Local Position (Excluding Material Offsets)").ToString());
+	}
+	else
+	{
+		checkNoEntry();
+	}
+}
+#endif // WITH_EDITOR
+
 ///////////////////////////////////////////////////////////////////////////////
 // UMaterialExpressionWorldPosition
 ///////////////////////////////////////////////////////////////////////////////
@@ -22710,13 +22780,14 @@ int32 UMaterialExpressionPreSkinnedPosition::Compile(class FMaterialCompiler* Co
 
 void UMaterialExpressionPreSkinnedPosition::GetCaption(TArray<FString>& OutCaptions) const
 {
-	OutCaptions.Add(TEXT("Pre-Skinned Local Position"));
+	OutCaptions.Add(TEXT("Pre-Skinned Local Position (deprecated)"));
 }
 
 void UMaterialExpressionPreSkinnedPosition::GetExpressionToolTip(TArray<FString>& OutToolTip) 
 {
-	ConvertToMultilineToolTip(TEXT("Returns pre-skinned local position for skeletal meshes, usable in vertex shader only."
-		"Returns the local position for non-skeletal meshes. Incompatible with GPU skin cache feature."), 40, OutToolTip);
+	ConvertToMultilineToolTip(TEXT("Deprecated, has been merged into 'Local Position'."
+	                               "Returns pre-skinned local position for skeletal meshes, usable in vertex shader only."
+	                               "Returns the local position for non-skeletal meshes. Incompatible with GPU skin cache feature."), 40, OutToolTip);
 }
 #endif // WITH_EDITOR
 
diff --git a/Engine/Source/Runtime/Engine/Private/Materials/MaterialHLSLEmitter.cpp b/Engine/Source/Runtime/Engine/Private/Materials/MaterialHLSLEmitter.cpp
--- a/Engine/Source/Runtime/Engine/Private/Materials/MaterialHLSLEmitter.cpp
+++ b/Engine/Source/Runtime/Engine/Private/Materials/MaterialHLSLEmitter.cpp
@@ -403,6 +403,17 @@ static void GetMaterialEnvironment(EShaderPlatform InPlatform,
 		OutEnvironment.SetDefine(TEXT("NEEDS_WORLD_POSITION_EXCLUDING_SHADER_OFFSETS"), TEXT("1"));
 	}
 
+	const bool bNeedsInstanceWorldToLocalPS = 
+		EmitMaterialData.IsExternalInputUsed(SF_Pixel, Material::EExternalInput::PositionInstanceSpace) ||
+		EmitMaterialData.IsExternalInputUsed(SF_Pixel, Material::EExternalInput::PositionInstanceSpace_NoOffsets) ||
+		EmitMaterialData.IsExternalInputUsed(SF_Pixel, Material::EExternalInput::PrevPositionInstanceSpace) ||
+		EmitMaterialData.IsExternalInputUsed(SF_Pixel, Material::EExternalInput::PrevPositionInstanceSpace_NoOffsets);
+
+	if (bNeedsInstanceWorldToLocalPS)
+	{
+		OutEnvironment.SetDefine(TEXT("NEEDS_INSTANCE_WORLD_TO_LOCAL_PS"), TEXT("1"));
+	}
+
 	const bool bNeedsParticleSize = EmitMaterialData.IsExternalInputUsed(Material::EExternalInput::ParticleSize);
 
 	if (bNeedsParticleSize)
diff --git a/Engine/Source/Runtime/Engine/Private/Materials/MaterialHLSLTree.cpp b/Engine/Source/Runtime/Engine/Private/Materials/MaterialHLSLTree.cpp
--- a/Engine/Source/Runtime/Engine/Private/Materials/MaterialHLSLTree.cpp
+++ b/Engine/Source/Runtime/Engine/Private/Materials/MaterialHLSLTree.cpp
@@ -62,6 +62,16 @@ FExternalInputDescription GetExternalInputDescription(EExternalInput Input)
 	case EExternalInput::VertexColor_Ddx: return FExternalInputDescription(TEXT("VertexColor_Ddx"), Shader::EValueType::Float4);
 	case EExternalInput::VertexColor_Ddy: return FExternalInputDescription(TEXT("VertexColor_Ddy"), Shader::EValueType::Float4);
 
+	case EExternalInput::PositionInstanceSpace: return FExternalInputDescription(TEXT("PositionInstanceSpace"), Shader::EValueType::Float3, EExternalInput::None, EExternalInput::None, EExternalInput::PrevPositionInstanceSpace);
+	case EExternalInput::PositionInstanceSpace_NoOffsets: return FExternalInputDescription(TEXT("PositionInstanceSpace_NoOffsets"), Shader::EValueType::Float3, EExternalInput::None, EExternalInput::None, EExternalInput::PrevPositionInstanceSpace_NoOffsets);
+	case EExternalInput::PositionPrimitiveSpace: return FExternalInputDescription(TEXT("PositionPrimitiveSpace"), Shader::EValueType::Float3, EExternalInput::None, EExternalInput::None, EExternalInput::PrevPositionPrimitiveSpace);
+	case EExternalInput::PositionPrimitiveSpace_NoOffsets: return FExternalInputDescription(TEXT("PositionPrimitiveSpace_NoOffsets"), Shader::EValueType::Float3, EExternalInput::None, EExternalInput::PrevPositionPrimitiveSpace_NoOffsets);
+
+	case EExternalInput::PrevPositionInstanceSpace: return FExternalInputDescription(TEXT("PrevPositionInstanceSpace"), Shader::EValueType::Float3);
+	case EExternalInput::PrevPositionInstanceSpace_NoOffsets: return FExternalInputDescription(TEXT("PrevPositionInstanceSpace_NoOffsets"), Shader::EValueType::Float3);
+	case EExternalInput::PrevPositionPrimitiveSpace: return FExternalInputDescription(TEXT("PrevPositionPrimitiveSpace"), Shader::EValueType::Float3);
+	case EExternalInput::PrevPositionPrimitiveSpace_NoOffsets: return FExternalInputDescription(TEXT("PrevPositionPrimitiveSpace_NoOffsets"), Shader::EValueType::Float3);
+
 	case EExternalInput::WorldPosition: return FExternalInputDescription(TEXT("WorldPosition"), Shader::EValueType::Double3, EExternalInput::WorldPosition_Ddx, EExternalInput::WorldPosition_Ddy, EExternalInput::PrevWorldPosition);
 	case EExternalInput::WorldPosition_NoOffsets: return FExternalInputDescription(TEXT("WorldPosition_NoOffsets"), Shader::EValueType::Double3, EExternalInput::WorldPosition_Ddx, EExternalInput::WorldPosition_Ddy, EExternalInput::PrevWorldPosition_NoOffsets);
 	case EExternalInput::TranslatedWorldPosition: return FExternalInputDescription(TEXT("TranslatedWorldPosition"), Shader::EValueType::Float3, EExternalInput::WorldPosition_Ddx, EExternalInput::WorldPosition_Ddy, EExternalInput::PrevTranslatedWorldPosition);
@@ -235,6 +245,18 @@ EExternalInput FExpressionExternalInput::GetResolvedInputType(EShaderFrequency S
 	{
 		switch (Result)
 		{
+		case EExternalInput::PrevPositionInstanceSpace:
+			Result = EExternalInput::PositionInstanceSpace;
+			break;
+		case EExternalInput::PrevPositionInstanceSpace_NoOffsets:
+			Result = EExternalInput::PositionInstanceSpace_NoOffsets;
+			break;
+		case EExternalInput::PrevPositionPrimitiveSpace:
+			Result = EExternalInput::PositionPrimitiveSpace;
+			break;
+		case EExternalInput::PrevPositionPrimitiveSpace_NoOffsets:
+			Result = EExternalInput::PositionPrimitiveSpace_NoOffsets;
+			break;
 		case EExternalInput::PrevWorldPosition:
 			Result = EExternalInput::WorldPosition;
 			break;
@@ -255,6 +277,18 @@ EExternalInput FExpressionExternalInput::GetResolvedInputType(EShaderFrequency S
 	{
 		switch (Result)
 		{
+		case EExternalInput::PositionInstanceSpace_NoOffsets:
+			Result = EExternalInput::PositionInstanceSpace;
+			break;
+		case EExternalInput::PositionPrimitiveSpace_NoOffsets:
+			Result = EExternalInput::PositionPrimitiveSpace;
+			break;
+		case EExternalInput::PrevPositionInstanceSpace_NoOffsets:
+			Result = EExternalInput::PrevPositionInstanceSpace;
+			break;
+		case EExternalInput::PrevPositionPrimitiveSpace_NoOffsets:
+			Result = EExternalInput::PrevPositionPrimitiveSpace;
+			break;
 		case EExternalInput::WorldPosition_NoOffsets:
 			Result = EExternalInput::WorldPosition;
 			break;
@@ -353,6 +387,16 @@ void FExpressionExternalInput::EmitValueShader(FEmitContext& Context, FEmitScope
 		case EExternalInput::VertexColor: Code = TEXT("Parameters.VertexColor"); break;
 		case EExternalInput::VertexColor_Ddx: Code = TEXT("Parameters.VertexColor_DDX"); break;
 		case EExternalInput::VertexColor_Ddy: Code = TEXT("Parameters.VertexColor_DDY"); break;
+		
+		case EExternalInput::PositionInstanceSpace:					Code = TEXT("GetPositionInstanceSpace(Parameters)"); break;
+		case EExternalInput::PositionInstanceSpace_NoOffsets: 		Code = TEXT("GetPositionInstanceSpace_NoMaterialOffsets(Parameters)"); break;
+		case EExternalInput::PositionPrimitiveSpace:				Code = TEXT("GetPositionPrimitiveSpace(Parameters)"); break;
+		case EExternalInput::PositionPrimitiveSpace_NoOffsets:		Code = TEXT("GetPositionPrimitiveSpace_NoMaterialOffsets(Parameters)"); break;
+		case EExternalInput::PrevPositionInstanceSpace:				Code = TEXT("GetPrevPositionInstanceSpace(Parameters)"); break;
+		case EExternalInput::PrevPositionInstanceSpace_NoOffsets:	checkNoEntry(); break;
+		case EExternalInput::PrevPositionPrimitiveSpace:			Code = TEXT("GetPrevPositionPrimitiveSpace(Parameters)"); break;
+		case EExternalInput::PrevPositionPrimitiveSpace_NoOffsets:	checkNoEntry(); break;
+		
 		case EExternalInput::WorldPosition: Code = TEXT("GetWorldPosition(Parameters)"); break;
 		case EExternalInput::WorldPosition_NoOffsets: Code = TEXT("GetWorldPosition_NoMaterialOffsets(Parameters)"); break;
 		case EExternalInput::TranslatedWorldPosition: Code = TEXT("GetTranslatedWorldPosition(Parameters)"); break;
diff --git a/Engine/Source/Runtime/Engine/Public/MaterialCompiler.h b/Engine/Source/Runtime/Engine/Public/MaterialCompiler.h
--- a/Engine/Source/Runtime/Engine/Public/MaterialCompiler.h
+++ b/Engine/Source/Runtime/Engine/Public/MaterialCompiler.h
@@ -11,6 +11,7 @@
 #include "MaterialShared.h"
 #include "Materials/MaterialExpressionSpeedTree.h"
 #include "Materials/MaterialExpressionTextureSample.h"
+#include "Materials/MaterialExpressionLocalPosition.h"
 #include "Materials/MaterialExpressionWorldPosition.h"
 #include "Materials/MaterialExpressionTextureProperty.h"
 #include "Materials/MaterialExpressionViewProperty.h"
@@ -235,6 +236,7 @@ public:
 
 	virtual int32 GetViewportUV() = 0;
 	virtual int32 GetPixelPosition() = 0;
+	virtual int32 LocalPosition(EPositionIncludedOffsets IncludedOffsets, ELocalPositionOrigin OriginType) = 0;
 	virtual int32 WorldPosition(EWorldPositionIncludedOffsets WorldPositionIncludedOffsets) = 0;
 	virtual int32 ObjectWorldPosition(EPositionOrigin OriginType) = 0;
 	UE_DEPRECATED(5.4, "Use ObjectWorldPosition(EPositionOrigin) instead")
@@ -842,6 +844,7 @@ public:
 
 	virtual int32 GetViewportUV() override { return Compiler->GetViewportUV(); }
 	virtual int32 GetPixelPosition() override { return Compiler->GetPixelPosition(); }
+	virtual int32 LocalPosition(EPositionIncludedOffsets IncludedOffsets, ELocalPositionOrigin OriginType) { return Compiler->LocalPosition(IncludedOffsets, OriginType); }
 	virtual int32 WorldPosition(EWorldPositionIncludedOffsets WorldPositionIncludedOffsets) override { return Compiler->WorldPosition(WorldPositionIncludedOffsets); }
 	virtual int32 ObjectWorldPosition(EPositionOrigin OriginType) override { return Compiler->ObjectWorldPosition(OriginType); }
 	virtual int32 ObjectRadius() override { return Compiler->ObjectRadius(); }
diff --git a/Engine/Source/Runtime/Engine/Public/MaterialHLSLTree.h b/Engine/Source/Runtime/Engine/Public/MaterialHLSLTree.h
--- a/Engine/Source/Runtime/Engine/Public/MaterialHLSLTree.h
+++ b/Engine/Source/Runtime/Engine/Public/MaterialHLSLTree.h
@@ -97,6 +97,16 @@ enum class EExternalInput : uint8
 	VertexColor_Ddx,
 	VertexColor_Ddy,
 
+	PositionInstanceSpace,
+	PositionInstanceSpace_NoOffsets,
+	PositionPrimitiveSpace,
+	PositionPrimitiveSpace_NoOffsets,
+
+	PrevPositionInstanceSpace,
+	PrevPositionInstanceSpace_NoOffsets,
+	PrevPositionPrimitiveSpace,
+	PrevPositionPrimitiveSpace_NoOffsets,
+
 	WorldPosition,
 	WorldPosition_NoOffsets,
 	TranslatedWorldPosition,
