Heterogeneous Volume Component: Add SetVolumeResolution

Only reset SVT-related parameters if a change was detected.

diff --git a/Engine/Source/Runtime/Engine/Classes/Components/HeterogeneousVolumeComponent.h b/Engine/Source/Runtime/Engine/Classes/Components/HeterogeneousVolumeComponent.h
--- a/Engine/Source/Runtime/Engine/Classes/Components/HeterogeneousVolumeComponent.h
+++ b/Engine/Source/Runtime/Engine/Classes/Components/HeterogeneousVolumeComponent.h
@@ -71,6 +71,9 @@ class UHeterogeneousVolumeComponent : public UMeshComponent
 	UPROPERTY(Transient)
 	TObjectPtr<UMaterialInstanceDynamic> MaterialInstanceDynamic;
 
+	UFUNCTION(BlueprintCallable, Category = "Volume")
+	ENGINE_API void SetVolumeResolution(FIntVector NewValue);
+
 	UFUNCTION(BlueprintCallable, Category = "Animation")
 	ENGINE_API void SetFrame(float NewValue);
 
diff --git a/Engine/Source/Runtime/Engine/Private/Components/HeterogeneousVolumeComponent.cpp b/Engine/Source/Runtime/Engine/Private/Components/HeterogeneousVolumeComponent.cpp
--- a/Engine/Source/Runtime/Engine/Private/Components/HeterogeneousVolumeComponent.cpp
+++ b/Engine/Source/Runtime/Engine/Private/Components/HeterogeneousVolumeComponent.cpp
@@ -270,6 +270,16 @@ void UHeterogeneousVolumeComponent::SetStreamingMipBias(int32 NewValue)
 	}
 }
 
+void UHeterogeneousVolumeComponent::SetVolumeResolution(FIntVector NewValue)
+{
+	if (AreDynamicDataChangesAllowed()
+		&& VolumeResolution != NewValue)
+	{
+		VolumeResolution = NewValue;
+		MarkRenderStateDirty();
+	}
+}
+
 void UHeterogeneousVolumeComponent::SetFrame(float NewValue)
 {
 	if (AreDynamicDataChangesAllowed()
@@ -424,23 +434,26 @@ UMaterialInstanceDynamic* UHeterogeneousVolumeComponent::CreateOrCastToMID(UMate
 
 void UHeterogeneousVolumeComponent::OnSparseVolumeTextureChanged(const USparseVolumeTexture* SparseVolumeTexture)
 {
-	if (SparseVolumeTexture)
-	{
-		VolumeResolution = SparseVolumeTexture->GetVolumeResolution();
-		StartFrame = 0;
-		EndFrame = SparseVolumeTexture->GetNumFrames() - 1;
-		Frame = FMath::Clamp(Frame, StartFrame, EndFrame);
-	}
-	else
+	if (SparseVolumeTexture != PreviousSVT)
 	{
-		VolumeResolution = FIntVector(128);
-		Frame = 0.0f;
-		StartFrame = 0.0f;
-		EndFrame = 0.0f;
-	}
+		if (SparseVolumeTexture)
+		{
+			VolumeResolution = SparseVolumeTexture->GetVolumeResolution();
+			StartFrame = 0;
+			EndFrame = SparseVolumeTexture->GetNumFrames() - 1;
+			Frame = FMath::Clamp(Frame, StartFrame, EndFrame);
+		}
+		else
+		{
+			VolumeResolution = FIntVector(128);
+			Frame = 0.0f;
+			StartFrame = 0.0f;
+			EndFrame = 0.0f;
+		}
 
-	PreviousSVT = SparseVolumeTexture;
-	MarkRenderStateDirty();
+		PreviousSVT = SparseVolumeTexture;
+		MarkRenderStateDirty();
+	}
 }
 
 UMaterialInterface* UHeterogeneousVolumeComponent::GetHeterogeneousVolumeMaterial() const
