Temporarily roll back DDC non-determinism fix in BuildStaticMeshDerivedDataKeySuffix until the next time we bump the DDC key for SM

diff --git a/Engine/Source/Runtime/Engine/Private/StaticMesh.cpp b/Engine/Source/Runtime/Engine/Private/StaticMesh.cpp
--- a/Engine/Source/Runtime/Engine/Private/StaticMesh.cpp
+++ b/Engine/Source/Runtime/Engine/Private/StaticMesh.cpp
@@ -2704,7 +2704,8 @@ static FString BuildStaticMeshDerivedDataKeySuffix(const ITargetPlatform* Target
 	}
 #endif
 
-	if (Mesh->IsNaniteEnabled())
+	//if (Mesh->IsNaniteEnabled()) // TODO: Re-introduce on next SM key bump
+	if (DoesTargetPlatformSupportNanite(TargetPlatform))
 	{
 		TempBytes.Reset();
 		FMemoryWriter Ar(TempBytes, /*bIsPersistent=*/ true);
