Fix typo in NSLOCTEXT for the local position expression

diff --git a/Engine/Source/Runtime/Engine/Private/Materials/MaterialExpressions.cpp b/Engine/Source/Runtime/Engine/Private/Materials/MaterialExpressions.cpp
--- a/Engine/Source/Runtime/Engine/Private/Materials/MaterialExpressions.cpp
+++ b/Engine/Source/Runtime/Engine/Private/Materials/MaterialExpressions.cpp
@@ -13257,7 +13257,7 @@ void UMaterialExpressionLocalPosition::GetCaption(TArray<FString>& OutCaptions)
 {
 	if (LocalOrigin == ELocalPositionOrigin::InstancePreSkinning)
 	{
-		OutCaptions.Add(NSLOCTEXT("MaterialExpressions", "LocalPositionInstanceIncludingOffsetsText", "Pre-Skinned Local Position").ToString());
+		OutCaptions.Add(NSLOCTEXT("MaterialExpressions", "LocalPositionPreSkinnedText", "Pre-Skinned Local Position").ToString());
 	}
 	else if (IncludedOffsets == EPositionIncludedOffsets::IncludeOffsets && LocalOrigin == ELocalPositionOrigin::Instance)
 	{
