Allow RestorePreAdditiveRootMotionVelocity to be overriden by child classes and customize the behavior of additive root motion
#rnx

diff --git a/Engine/Source/Runtime/Engine/Classes/GameFramework/CharacterMovementComponent.h b/Engine/Source/Runtime/Engine/Classes/GameFramework/CharacterMovementComponent.h
--- a/Engine/Source/Runtime/Engine/Classes/GameFramework/CharacterMovementComponent.h
+++ b/Engine/Source/Runtime/Engine/Classes/GameFramework/CharacterMovementComponent.h
@@ -2713,7 +2713,7 @@ public:
 
 protected:
 	/** Restores Velocity to LastPreAdditiveVelocity during Root Motion Phys*() function calls */
-	ENGINE_API void RestorePreAdditiveRootMotionVelocity();
+	ENGINE_API virtual void RestorePreAdditiveRootMotionVelocity();
 
 	/** Applies root motion from root motion sources to velocity (override and additive) */
 	ENGINE_API virtual void ApplyRootMotionToVelocity(float deltaTime);
