#UE Add some NaN checking to track down the source of NaNs in static mesh data

diff --git a/Engine/Source/Runtime/Engine/Private/StaticMesh.cpp b/Engine/Source/Runtime/Engine/Private/StaticMesh.cpp
--- a/Engine/Source/Runtime/Engine/Private/StaticMesh.cpp
+++ b/Engine/Source/Runtime/Engine/Private/StaticMesh.cpp
@@ -1649,6 +1649,13 @@ void FStaticMeshRenderData::SerializeInlineDataRepresentations(FArchive& Ar, USt
 					LOD.CardRepresentationData = new FCardRepresentationData();
 				}
 
+#if !UE_BUILD_SHIPPING
+				if (LOD.CardRepresentationData->ContainsNaN())
+				{
+					UE_LOG(LogStaticMesh, Display, TEXT("FStaticMeshRenderData::SerializeInlineDataRepresentations found NaN in CardRepresentationData of '%s'. LOD:%d"), *GetPathNameSafe(Owner), ResourceIndex);
+				}
+#endif // UE_BUILD_SHIPPING
+
 				Ar << *(LOD.CardRepresentationData);
 			}
 		}
@@ -1814,6 +1821,12 @@ void FStaticMeshRenderData::Serialize(FArchive& Ar, UStaticMesh* Owner, bool bCo
 							LOD.DistanceFieldData->AssetName = Owner->GetFName();
 						}
 
+#if !UE_BUILD_SHIPPING
+						if (LOD.DistanceFieldData->LocalSpaceMeshBounds.ContainsNaN())
+						{
+							UE_LOG(LogStaticMesh, Display, TEXT("FStaticMeshRenderData::Serialize found NaN in DistanceFieldData of '%s'. LOD:%d"), *GetPathNameSafe(Owner), ResourceIndex);
+						}
+#endif // UE_BUILD_SHIPPING
 						LOD.DistanceFieldData->Serialize(Ar, Owner);
 					}
 				}
diff --git a/Engine/Source/Runtime/Engine/Public/MeshCardBuild.h b/Engine/Source/Runtime/Engine/Public/MeshCardBuild.h
--- a/Engine/Source/Runtime/Engine/Public/MeshCardBuild.h
+++ b/Engine/Source/Runtime/Engine/Public/MeshCardBuild.h
@@ -20,6 +20,11 @@ public:
 	FLumenCardOBBf OBB;
 	uint8 AxisAlignedDirectionIndex;
 
+	bool ContainsNaN() const
+	{
+		return OBB.ContainsNaN();
+	}
+
 	friend FArchive& operator<<(FArchive& Ar, FLumenCardBuildData& Data)
 	{
 		// Note: this is derived data, no need for versioning (bump the DDC guid)
@@ -89,6 +94,11 @@ public:
 	// Temporary debug visualization data, don't serialize
 	FLumenCardBuildDebugData DebugData;
 
+	bool ContainsNaN() const
+	{
+		return Bounds.ContainsNaN() || CardBuildData.ContainsByPredicate([](const FLumenCardBuildData& Data) { return Data.ContainsNaN(); });
+	}
+
 	friend FArchive& operator<<(FArchive& Ar, FMeshCardsBuildData& Data)
 	{
 		// Note: this is derived data, no need for versioning (bump the DDC guid)
@@ -149,6 +159,11 @@ public:
 		return ResSize.GetTotalMemoryBytes();
 	}
 
+	bool ContainsNaN() const
+	{
+		return MeshCardsBuildData.ContainsNaN();
+	}
+
 #if WITH_EDITORONLY_DATA
 
 	void CacheDerivedData(const FString& InDDCKey, const ITargetPlatform* TargetPlatform, UStaticMesh* Mesh, UStaticMesh* GenerateSource, int32 MaxLumenMeshCards, bool bGenerateDistanceFieldAsIfTwoSided, FSourceMeshDataForDerivedDataTask* OptionalSourceMeshData);
diff --git a/Engine/Source/Runtime/Engine/Public/MeshCardRepresentation.h b/Engine/Source/Runtime/Engine/Public/MeshCardRepresentation.h
--- a/Engine/Source/Runtime/Engine/Public/MeshCardRepresentation.h
+++ b/Engine/Source/Runtime/Engine/Public/MeshCardRepresentation.h
@@ -152,6 +152,11 @@ public:
 		return UE::Math::TBox<T>(BoxMin, BoxMax);
 	}
 
+	bool ContainsNaN() const
+	{
+		return Origin.ContainsNaN() || AxisX.ContainsNaN() || AxisY.ContainsNaN() || AxisZ.ContainsNaN() || Extent.ContainsNaN();
+	}
+
 	friend FArchive& operator<<(FArchive& Ar, TLumenCardOBB<T>& Data)
 	{
 		Ar << Data.AxisX;
