Splitscreen fixes

#rnx

diff --git a/Engine/Source/Runtime/Engine/Private/Net/Iris/ReplicationSystem/ReplicationSystemUtil.cpp b/Engine/Source/Runtime/Engine/Private/Net/Iris/ReplicationSystem/ReplicationSystemUtil.cpp
--- a/Engine/Source/Runtime/Engine/Private/Net/Iris/ReplicationSystem/ReplicationSystemUtil.cpp
+++ b/Engine/Source/Runtime/Engine/Private/Net/Iris/ReplicationSystem/ReplicationSystemUtil.cpp
@@ -542,7 +542,7 @@ void FReplicationSystemUtil::UpdateSubObjectGroupMemberships(const APlayerContro
 	{
 		if (UReplicationSystem* ReplicationSystem = Conn->GetDriver() ? Conn->GetDriver()->GetReplicationSystem() : nullptr)
 		{
-			const uint32 ConnId = Conn->GetConnectionId();
+			const uint32 ConnId = Conn->GetParentConnectionId();
 			for (const FName NetGroup : PC->GetNetConditionGroups())
 			{
 				if (!IsSpecialNetConditionGroup(NetGroup))
