Fixed UWorld::Rename to not move UDeletedObjectPlaceholder object in the newly renamed actor packages. Also disabled UDeletedObjectPlaceholder for any running commandlet (not just cook) as it's only necessary when running inside the editor.

diff --git a/Engine/Source/Runtime/Engine/Private/DeletedObjectPlaceholder.cpp b/Engine/Source/Runtime/Engine/Private/DeletedObjectPlaceholder.cpp
--- a/Engine/Source/Runtime/Engine/Private/DeletedObjectPlaceholder.cpp
+++ b/Engine/Source/Runtime/Engine/Private/DeletedObjectPlaceholder.cpp
@@ -18,7 +18,7 @@ bool UDeletedObjectPlaceholder::IsAsset() const
 
 UDeletedObjectPlaceholder* UDeletedObjectPlaceholder::Create(ULevel* InLevel, UPackage* InPackage, const UObject* InOriginalObject)
 {
-	if (IsRunningCookCommandlet())
+	if (IsRunningCommandlet())
 	{
 		return nullptr;
 	}
@@ -39,7 +39,7 @@ UDeletedObjectPlaceholder* UDeletedObjectPlaceholder::Create(ULevel* InLevel, UP
 
 UDeletedObjectPlaceholder* UDeletedObjectPlaceholder::RemoveFromPackage(UPackage* InPackage)
 {
-	if (IsRunningCookCommandlet())
+	if (IsRunningCommandlet())
 	{
 		return nullptr;
 	}
@@ -58,7 +58,7 @@ UDeletedObjectPlaceholder* UDeletedObjectPlaceholder::RemoveFromPackage(UPackage
 
 UDeletedObjectPlaceholder* UDeletedObjectPlaceholder::FindInPackage(UPackage* InPackage)
 {
-	if (IsRunningCookCommandlet())
+	if (IsRunningCommandlet())
 	{
 		return nullptr;
 	}
diff --git a/Engine/Source/Runtime/Engine/Private/World.cpp b/Engine/Source/Runtime/Engine/Private/World.cpp
--- a/Engine/Source/Runtime/Engine/Private/World.cpp
+++ b/Engine/Source/Runtime/Engine/Private/World.cpp
@@ -96,6 +96,7 @@
 	#include "ActorFolder.h"
 	#include "ActorDeferredScriptManager.h"
 	#include "AssetCompilingManager.h"
+	#include "DeletedObjectPlaceholder.h"
 #endif
 
 
@@ -947,7 +948,7 @@ bool UWorld::Rename(const TCHAR* InName, UObject* NewOuter, ERenameFlags Flags)
 			TArray<UObject*> DependantObjects;
 			ForEachObjectWithPackage(ExternalPackage, [&DependantObjects](UObject* Object)
 			{
-				if (!Cast<UMetaData>(Object))
+				if (!Cast<UMetaData>(Object) && !Cast<UDeletedObjectPlaceholder>(Object))
 				{
 					DependantObjects.Add(Object);
 				}
