Fix for non pow 2 textures with high LOD bias getting in a situation where they try to remove all mips due to incorrect Ceil usage during mip calculation.

diff --git a/Engine/Source/Runtime/Engine/Private/TextureLODSettings.cpp b/Engine/Source/Runtime/Engine/Private/TextureLODSettings.cpp
--- a/Engine/Source/Runtime/Engine/Private/TextureLODSettings.cpp
+++ b/Engine/Source/Runtime/Engine/Private/TextureLODSettings.cpp
@@ -3,6 +3,7 @@
 #include "Engine/TextureLODSettings.h"
 #include "Engine/TextureCube.h"
 #include "HAL/IConsoleManager.h"
+#include "ImageCoreUtils.h"
 #include "Interfaces/ITargetPlatform.h"
 
 #include UE_INLINE_GENERATED_CPP_BY_NAME(TextureLODSettings)
@@ -176,7 +177,10 @@ int32 UTextureLODSettings::CalculateLODBias(int32 Width, int32 Height, int32 Max
 		Width = FMath::Min(Width, MaxSize);
 		Height = FMath::Min(Height, MaxSize);
 	}
-	int32 TextureMaxLOD	= FMath::CeilLogTwo( FMath::Max( Width, Height ) );
+
+	// Don't let the LOD remove ALL mips
+	int32 TextureMaxLOD = FImageCoreUtils::GetMipCountFromDimensions(Width, Height, 0, false) - 1;
+
 
 	// Calculate LOD bias.
 	int32 UsedLODBias = 0;
