Fixed: Vulkan validation realated to buffer size mismatch for vkCmdCopyBuffer command

diff --git a/Engine/Source/Runtime/Engine/Private/Particles/ParticleGpuSimulation.cpp b/Engine/Source/Runtime/Engine/Private/Particles/ParticleGpuSimulation.cpp
--- a/Engine/Source/Runtime/Engine/Private/Particles/ParticleGpuSimulation.cpp
+++ b/Engine/Source/Runtime/Engine/Private/Particles/ParticleGpuSimulation.cpp
@@ -1894,7 +1894,9 @@ void InjectNewParticles(FRHICommandList& RHICmdList, FGraphicsPipelineStateIniti
 		// Copy new particles in to the vertex buffer.
 		const int32 ParticlesThisDrawCall = FMath::Min<int32>( ParticleCount, MaxParticlesPerDrawCall );
 		const void* Src = NewParticles.GetData() + FirstParticle;
-		void* Dest = RHICmdList.LockBuffer( ScratchVertexBufferRHI, 0, ParticlesThisDrawCall * sizeof(FNewParticle), RLM_WriteOnly );
+		// Make sure Lock size is a multiple of ScratchVertexBuffer stride
+		const uint32 AlignedLockSize = AlignArbitrary(ParticlesThisDrawCall * sizeof(FNewParticle), GParticleScratchVertexBuffer.Stride);
+		void* Dest = RHICmdList.LockBuffer( ScratchVertexBufferRHI, 0, AlignedLockSize, RLM_WriteOnly);
 		FMemory::Memcpy( Dest, Src, ParticlesThisDrawCall * sizeof(FNewParticle) );
 		RHICmdList.UnlockBuffer( ScratchVertexBufferRHI );
 		ParticleCount -= ParticlesThisDrawCall;
diff --git a/Engine/Source/Runtime/Engine/Private/Particles/ParticleResources.cpp b/Engine/Source/Runtime/Engine/Private/Particles/ParticleResources.cpp
--- a/Engine/Source/Runtime/Engine/Private/Particles/ParticleResources.cpp
+++ b/Engine/Source/Runtime/Engine/Private/Particles/ParticleResources.cpp
@@ -144,7 +144,7 @@ void FParticleScratchVertexBuffer::InitRHI(FRHICommandListBase& RHICmdList)
 
 	FRHIResourceCreateInfo CreateInfo(TEXT("FParticleScratchVertexBuffer"));
 	VertexBufferRHI = RHICmdList.CreateVertexBuffer(GParticleScratchVertexBufferSize, Flags, CreateInfo);
-	VertexBufferSRV_A32B32G32R32F = RHICmdList.CreateShaderResourceView(VertexBufferRHI, /*Stride=*/ sizeof(FVector4f), PF_A32B32G32R32F);
+	VertexBufferSRV_A32B32G32R32F = RHICmdList.CreateShaderResourceView(VertexBufferRHI, Stride, PF_A32B32G32R32F);
 }
 
 FParticleShaderParamRef FParticleScratchVertexBuffer::GetShaderParam()
diff --git a/Engine/Source/Runtime/Engine/Public/ParticleResources.h b/Engine/Source/Runtime/Engine/Public/ParticleResources.h
--- a/Engine/Source/Runtime/Engine/Public/ParticleResources.h
+++ b/Engine/Source/Runtime/Engine/Public/ParticleResources.h
@@ -72,6 +72,8 @@ public:
 	virtual void InitRHI(FRHICommandListBase& RHICmdList) override;
 	virtual void ReleaseRHI() override;
 
+	static const uint32 Stride = sizeof(FVector4f);
+
 private:
 
 	/** SRV in to the buffer as an array of FVector4f values. */
diff --git a/Engine/Source/Runtime/VulkanRHI/Private/VulkanBuffer.cpp b/Engine/Source/Runtime/VulkanRHI/Private/VulkanBuffer.cpp
--- a/Engine/Source/Runtime/VulkanRHI/Private/VulkanBuffer.cpp
+++ b/Engine/Source/Runtime/VulkanRHI/Private/VulkanBuffer.cpp
@@ -292,9 +292,6 @@ void* FVulkanResourceMultiBuffer::Lock(FVulkanCommandListContext& Context, EReso
 	void* Data = nullptr;
 	uint32 DataOffset = 0;
 
-	// High level code might create SRVs on an allocated buffer, make sure size of the buffer is at least multiple of float4 (16bytes)
-	LockSize = Align(LockSize, 16u);
-
 	const bool bVolatile = EnumHasAnyFlags(GetUsage(), BUF_Volatile);
 	check(LockStatus == ELockStatus::Unlocked);
 
