Make sure Skeleton is (pre-)loaded before the Animation DataModel is

diff --git a/Engine/Source/Runtime/Engine/Private/Animation/AnimSequenceBase.cpp b/Engine/Source/Runtime/Engine/Private/Animation/AnimSequenceBase.cpp
--- a/Engine/Source/Runtime/Engine/Private/Animation/AnimSequenceBase.cpp
+++ b/Engine/Source/Runtime/Engine/Private/Animation/AnimSequenceBase.cpp
@@ -137,6 +137,8 @@ void UAnimSequenceBase::PostLoad()
 			    CreateModel();
 		    }
 
+			PreloadSkeleton();
+			
 		    ValidateModel();
 			if (UObject* DataModelObject = DataModelInterface.GetObject())
 			{
@@ -144,9 +146,7 @@ void UAnimSequenceBase::PostLoad()
 				DataModelObject->ConditionalPostLoadSubobjects();
 			}
 		    GetController();
-		    BindToModelModificationEvent();
-
-		    PreloadSkeleton();
+		    BindToModelModificationEvent();		    
 
 		    if (bRequiresModelPopulation || bRequiresModelCreation)
 		    {
