Changed the display name for bEnableRootMotion to not be camel case and remove redundant display name for bUseNormalizedRootMotionScale.

diff --git a/Engine/Source/Runtime/Engine/Classes/Animation/AnimSequence.h b/Engine/Source/Runtime/Engine/Classes/Animation/AnimSequence.h
--- a/Engine/Source/Runtime/Engine/Classes/Animation/AnimSequence.h
+++ b/Engine/Source/Runtime/Engine/Classes/Animation/AnimSequence.h
@@ -337,7 +337,7 @@ public:
 	EAnimInterpolationType Interpolation;
 	
 	/** If this is on, it will allow extracting of root motion **/
-	UPROPERTY(EditAnywhere, AssetRegistrySearchable, Category = RootMotion, meta = (DisplayName = "EnableRootMotion"))
+	UPROPERTY(EditAnywhere, AssetRegistrySearchable, Category = RootMotion)
 	bool bEnableRootMotion;
 
 	/** Root Bone will be locked to that position when extracting root motion.**/
@@ -349,7 +349,7 @@ public:
 	bool bForceRootLock;
 
 	/** If this is on, it will use a normalized scale value for the root motion extracted: FVector(1.0, 1.0, 1.0) **/
-	UPROPERTY(EditAnywhere, AssetRegistrySearchable, Category = RootMotion, meta = (DisplayName = "Use Normalized Root Motion Scale"))
+	UPROPERTY(EditAnywhere, AssetRegistrySearchable, Category = RootMotion)
 	bool bUseNormalizedRootMotionScale;
 
 	/** Have we copied root motion settings from an owning montage */
