Add AActor::CanReplaceSelectedActor virtual (like CanDeleteSelectedActor) needed for upcoming change
- Get rid of todo and move specific code to proper class overrides (ABrush / AWorldSettings)

#rnx

diff --git a/Engine/Source/Editor/LevelEditor/Private/LevelEditorCreateActorMenu.cpp b/Engine/Source/Editor/LevelEditor/Private/LevelEditorCreateActorMenu.cpp
--- a/Engine/Source/Editor/LevelEditor/Private/LevelEditorCreateActorMenu.cpp
+++ b/Engine/Source/Editor/LevelEditor/Private/LevelEditorCreateActorMenu.cpp
@@ -218,13 +218,14 @@ class SAssetMenuEntry : public SCompoundWidget
 static bool CanReplaceActors()
 {
 	bool bCanReplace = false;
-
+	FText OutReason;
 	for (FSelectionIterator SelectionIter = GEditor->GetSelectedActorIterator(); SelectionIter; ++SelectionIter)
 	{
 		if (AActor* Actor = Cast<AActor>(*SelectionIter))
 		{ 
 			bCanReplace = true;
-			if(!Actor->IsUserManaged() || FActorEditorUtils::IsABuilderBrush(Actor))
+			
+			if(!Actor->CanReplaceSelectedActor(OutReason))
 			{
 				bCanReplace = false;
 				break;
diff --git a/Engine/Source/Editor/UnrealEd/Private/EditorActor.cpp b/Engine/Source/Editor/UnrealEd/Private/EditorActor.cpp
--- a/Engine/Source/Editor/UnrealEd/Private/EditorActor.cpp
+++ b/Engine/Source/Editor/UnrealEd/Private/EditorActor.cpp
@@ -720,33 +720,6 @@ bool UUnrealEdEngine::CanDeleteActor(const AActor* InActor, FText* OutReason) co
 		}
 	}
 
-	// TODO: The Brush and WorldSettings logic below should be moved to use the CanDeleteSelectedActor virtual
-
-	if (const ABrush* Brush = Cast<ABrush>(InActor))
-	{
-		if (FActorEditorUtils::IsABuilderBrush(Brush))
-		{
-			if (OutReason)
-			{
-				*OutReason = LOCTEXT("CanDeleteActor_Error_BuilderBrush", "Can't delete a builder brush.");
-			}
-			return false;
-		}
-	}
-
-	if (const AWorldSettings* WorldSettings = Cast<AWorldSettings>(InActor))
-	{
-		const ULevel* OwnerLevel = WorldSettings->GetLevel();
-		if (OwnerLevel && OwnerLevel->GetWorldSettings() == WorldSettings)
-		{
-			if (OutReason)
-			{
-				*OutReason = LOCTEXT("CanDeleteActor_Error_WorldSettings", "Can't delete a level's world settings.");
-			}
-			return false;
-		}
-	}
-
 	return true;
 }
 
diff --git a/Engine/Source/Runtime/Engine/Classes/Engine/Brush.h b/Engine/Source/Runtime/Engine/Classes/Engine/Brush.h
--- a/Engine/Source/Runtime/Engine/Classes/Engine/Brush.h
+++ b/Engine/Source/Runtime/Engine/Classes/Engine/Brush.h
@@ -181,6 +181,8 @@ public:
 	ENGINE_API virtual bool SetIsHiddenEdLayer(bool bIsHiddenEdLayer) override;
 	ENGINE_API virtual bool SupportsLayers() const override;
 	ENGINE_API virtual bool SupportsExternalPackaging() const override;
+	ENGINE_API virtual bool CanDeleteSelectedActor(FText& OutReason) const override;
+	ENGINE_API virtual bool CanReplaceSelectedActor(FText& OutReason) const override;
 public:
 
 	ENGINE_API virtual void InitPosRotScale();
diff --git a/Engine/Source/Runtime/Engine/Classes/GameFramework/Actor.h b/Engine/Source/Runtime/Engine/Classes/GameFramework/Actor.h
--- a/Engine/Source/Runtime/Engine/Classes/GameFramework/Actor.h
+++ b/Engine/Source/Runtime/Engine/Classes/GameFramework/Actor.h
@@ -2280,6 +2280,9 @@ public:
 	/** When selected can this actor be deleted? */
 	ENGINE_API virtual bool CanDeleteSelectedActor(FText& OutReason) const;
 
+	/** When selected can this actor be replaced */
+	ENGINE_API virtual bool CanReplaceSelectedActor(FText& OutReason) const;
+
 	/** Does this actor supports external packaging? */
 	ENGINE_API virtual bool SupportsExternalPackaging() const;
 #endif
diff --git a/Engine/Source/Runtime/Engine/Classes/GameFramework/WorldSettings.h b/Engine/Source/Runtime/Engine/Classes/GameFramework/WorldSettings.h
--- a/Engine/Source/Runtime/Engine/Classes/GameFramework/WorldSettings.h
+++ b/Engine/Source/Runtime/Engine/Classes/GameFramework/WorldSettings.h
@@ -895,6 +895,8 @@ public:
 	ENGINE_API virtual void CheckForErrors() override;
 	virtual bool IsSelectable() const override { return false; }
 	virtual bool SupportsExternalPackaging() const override { return false; }
+	ENGINE_API virtual bool CanDeleteSelectedActor(FText& OutReason) const override;
+	ENGINE_API virtual bool CanReplaceSelectedActor(FText& OutReason) const override;
 #endif // WITH_EDITOR
 	ENGINE_API virtual void PostInitProperties() override;
 	ENGINE_API virtual void PreInitializeComponents() override;
diff --git a/Engine/Source/Runtime/Engine/Private/Actor.cpp b/Engine/Source/Runtime/Engine/Private/Actor.cpp
--- a/Engine/Source/Runtime/Engine/Private/Actor.cpp
+++ b/Engine/Source/Runtime/Engine/Private/Actor.cpp
@@ -120,6 +120,19 @@ void AActor::FDuplicationSeedInterface::AddEntry(UObject* Source, UObject* Desti
 	DuplicationSeed.Emplace(Source, Destination);
 }
 
+namespace ActorUtils
+{
+	bool CanDeleteOrReplaceCommon(const AActor* InActor, FText& OutReason)
+	{
+		if (!InActor->IsUserManaged())
+		{
+			OutReason = LOCTEXT("CanDeleteOrReplace_Error_UserManaged", "Actor is not user managed");
+			return false;
+		}
+
+		return true;
+	}
+}
 
 #endif
 
@@ -3615,15 +3628,13 @@ static void DispatchOnComponentsCreated(AActor* NewActor)
 
 bool AActor::CanDeleteSelectedActor(FText& OutReason) const
 {
-	if (!IsUserManaged())
-	{
-		OutReason = LOCTEXT("UserManaged", "Actor is not user managed");
-		return false;
-	}
-
-	return true;
+	return ActorUtils::CanDeleteOrReplaceCommon(this, OutReason);
 }
 
+bool AActor::CanReplaceSelectedActor(FText& OutReason) const
+{
+	return ActorUtils::CanDeleteOrReplaceCommon(this, OutReason);
+}
 
 void AActor::PostEditImport()
 {
diff --git a/Engine/Source/Runtime/Engine/Private/Brush.cpp b/Engine/Source/Runtime/Engine/Private/Brush.cpp
--- a/Engine/Source/Runtime/Engine/Private/Brush.cpp
+++ b/Engine/Source/Runtime/Engine/Private/Brush.cpp
@@ -34,6 +34,21 @@ bool ABrush::bSuppressBSPRegeneration = false;
 
 // Debug purposes only; an attempt to catch the cause of UE-36265
 const TCHAR* ABrush::GGeometryRebuildCause = nullptr;
+
+namespace BrushUtils
+{
+	bool CanDeleteOrReplaceCommon(const ABrush* InActor, FText& OutReason)
+	{
+		if (FActorEditorUtils::IsABuilderBrush(InActor))
+		{
+			OutReason = NSLOCTEXT("Brush", "CanDeleteOrReplace_Error_BuilderBrush", "Can't delete or replace a builder brush.");
+			return false;
+		}
+
+		return true;
+	}
+}
+
 #endif
 
 DEFINE_LOG_CATEGORY_STATIC(LogBrush, Log, All);
@@ -250,6 +265,26 @@ bool ABrush::SupportsExternalPackaging() const
 	return GetClass() != ABrush::StaticClass() && Super::SupportsExternalPackaging();
 }
 
+bool ABrush::CanDeleteSelectedActor(FText& OutReason) const
+{
+	if (!Super::CanDeleteSelectedActor(OutReason))
+	{
+		return false;
+	}
+
+	return BrushUtils::CanDeleteOrReplaceCommon(this, OutReason);
+}
+
+bool ABrush::CanReplaceSelectedActor(FText& OutReason) const
+{
+	if (!Super::CanReplaceSelectedActor(OutReason))
+	{
+		return false;
+	}
+
+	return BrushUtils::CanDeleteOrReplaceCommon(this, OutReason);
+}
+
 void ABrush::PostLoad()
 {
 	Super::PostLoad();
diff --git a/Engine/Source/Runtime/Engine/Private/WorldSettings.cpp b/Engine/Source/Runtime/Engine/Private/WorldSettings.cpp
--- a/Engine/Source/Runtime/Engine/Private/WorldSettings.cpp
+++ b/Engine/Source/Runtime/Engine/Private/WorldSettings.cpp
@@ -50,6 +50,21 @@ ENGINE_API float GNewWorldToMetersScale = 0.0f;
 #if WITH_EDITOR
 AWorldSettings::FOnBookmarkClassChanged AWorldSettings::OnBookmarkClassChanged;
 AWorldSettings::FOnNumberOfBookmarksChanged AWorldSettings::OnNumberOfBoomarksChanged;
+
+namespace WorldSettingsUtils
+{
+	bool CanDeleteOrReplaceCommon(const AWorldSettings* InActor, FText& OutReason)
+	{
+		const ULevel* OwnerLevel = InActor->GetLevel();
+		if (OwnerLevel && OwnerLevel->GetWorldSettings() == InActor)
+		{
+			OutReason = NSLOCTEXT("WorldSettings", "CanDeleteOrReplace_Error_WorldSettings", "Can't delete or replace a level's world settings.");
+			return false;
+		}
+
+		return true;
+	}
+}
 #endif
 AWorldSettings::FOnNaniteSettingsChanged AWorldSettings::OnNaniteSettingsChanged;
 
@@ -671,6 +686,27 @@ void AWorldSettings::CheckForErrors()
 	}
 }
 
+bool AWorldSettings::CanDeleteSelectedActor(FText& OutReason) const
+{
+	if (!Super::CanDeleteSelectedActor(OutReason))
+	{
+		return false;
+	}
+
+	return WorldSettingsUtils::CanDeleteOrReplaceCommon(this, OutReason);
+}
+
+bool AWorldSettings::CanReplaceSelectedActor(FText& OutReason) const
+{
+	if (!Super::CanReplaceSelectedActor(OutReason))
+	{
+		return false;
+	}
+
+	return WorldSettingsUtils::CanDeleteOrReplaceCommon(this, OutReason);
+}
+
+
 bool AWorldSettings::CanEditChange(const FProperty* InProperty) const
 {
 	if (InProperty)
