[shaders] don't bother serializing source dir mappings to scw input if no jobs require preprocessing, they aren't needed for anything else

diff --git a/Engine/Source/Runtime/Engine/Private/ShaderCompiler/ShaderCompiler.cpp b/Engine/Source/Runtime/Engine/Private/ShaderCompiler/ShaderCompiler.cpp
--- a/Engine/Source/Runtime/Engine/Private/ShaderCompiler/ShaderCompiler.cpp
+++ b/Engine/Source/Runtime/Engine/Private/ShaderCompiler/ShaderCompiler.cpp
@@ -2546,24 +2546,32 @@ static int32 GetNumTotalJobs(const TArray<FShaderCommonCompileJobPtr>& Jobs)
 	return NumJobs;
 }
 
-static void SplitJobsByType(const TArray<FShaderCommonCompileJobPtr>& QueuedJobs, TArray<FShaderCompileJob*>& OutQueuedSingleJobs, TArray<FShaderPipelineCompileJob*>& OutQueuedPipelineJobs)
+static bool SplitJobsByType(const TArray<FShaderCommonCompileJobPtr>& QueuedJobs, TArray<FShaderCompileJob*>& OutQueuedSingleJobs, TArray<FShaderPipelineCompileJob*>& OutQueuedPipelineJobs)
 {
+	bool bAnyPreprocessingNeeded = false;
 	for (int32 Index = 0; Index < QueuedJobs.Num(); ++Index)
 	{
 		FShaderCommonCompileJobPtr CommonJob = QueuedJobs[Index];
 		if (FShaderCompileJob* SingleJob = CommonJob->GetSingleShaderJob())
 		{
 			OutQueuedSingleJobs.Add(SingleJob);
+			bAnyPreprocessingNeeded |= !SingleJob->Input.bCachePreprocessed;
+
 		}
 		else if (FShaderPipelineCompileJob* PipelineJob = CommonJob->GetShaderPipelineJob())
 		{
 			OutQueuedPipelineJobs.Add(PipelineJob);
+			PipelineJob->ForEachSingleShaderJob([&bAnyPreprocessingNeeded](const FShaderCompileJob& SingleJob)
+				{
+					bAnyPreprocessingNeeded |= !SingleJob.Input.bCachePreprocessed;
+				});
 		}
 		else
 		{
 			checkf(0, TEXT("FShaderCommonCompileJob::Type=%d is not a valid type for a shader compile job"), (int32)CommonJob->Type);
 		}
 	}
+	return bAnyPreprocessingNeeded;
 }
 
 bool DoWriteTasksInner(const TArray<FShaderCommonCompileJobPtr>& QueuedJobs, FArchive& InTransferFile, IDistributedBuildController* BuildDistributionController, bool bUseRelativePaths, bool bCompressTaskFile)
@@ -2585,8 +2593,18 @@ bool DoWriteTasksInner(const TArray<FShaderCommonCompileJobPtr>& QueuedJobs, FAr
 
 	TransferFile << FormatVersionMap;
 
+	TArray<FShaderCompileJob*> QueuedSingleJobs;
+	TArray<FShaderPipelineCompileJob*> QueuedPipelineJobs;
+	bool bAnyPreprocessingNeeded = SplitJobsByType(QueuedJobs, QueuedSingleJobs, QueuedPipelineJobs);
+
+	TMap<FString, FString> ShaderSourceDirectoryMappings;
+	// Only serialize source directory mappings if any jobs need to preprocess; these are only used for include handling
+	if (bAnyPreprocessingNeeded)
+	{
+		ShaderSourceDirectoryMappings = AllShaderSourceDirectoryMappings();
+	}
+	
 	// Convert all the source directory paths to absolute, since SCW might be in a different directory to the editor executable
-	TMap<FString, FString> ShaderSourceDirectoryMappings = AllShaderSourceDirectoryMappings();
 	for(TPair<FString, FString>& Pair : ShaderSourceDirectoryMappings)
 	{
 		// Remap/enforce relative paths when bUseRelativePaths=true
@@ -2611,10 +2629,6 @@ bool DoWriteTasksInner(const TArray<FShaderCommonCompileJobPtr>& QueuedJobs, FAr
 	}
 	TransferFile << ShaderSourceDirectoryMappings;
 
-	TArray<FShaderCompileJob*> QueuedSingleJobs;
-	TArray<FShaderPipelineCompileJob*> QueuedPipelineJobs;
-	SplitJobsByType(QueuedJobs, QueuedSingleJobs, QueuedPipelineJobs);
-
 	TArray<TRefCountPtr<FSharedShaderCompilerEnvironment>> SharedEnvironments;
 	TArray<const FShaderParametersMetadata*> RequestShaderParameterStructures;
 
