TextureDerivedDataTask : NOP small cleanups

diff --git a/Engine/Source/Runtime/Engine/Private/TextureDerivedDataTask.cpp b/Engine/Source/Runtime/Engine/Private/TextureDerivedDataTask.cpp
--- a/Engine/Source/Runtime/Engine/Private/TextureDerivedDataTask.cpp
+++ b/Engine/Source/Runtime/Engine/Private/TextureDerivedDataTask.cpp
@@ -317,10 +317,14 @@ void FTextureSourceData::GetSourceMips(FTextureSource& Source, IImageWrapperModu
 			return;
 		}
 
-		// If we didn't get this from the texture source. As time goes on this will get hit less and less.
+		// If we didn't get ChannelMinMax from the texture source, then compute it now. As time goes on this will get hit less and less.
 		if (LayerChannelMinMax.Num() != Layers.Num())
 		{
 			TRACE_CPUPROFILER_EVENT_SCOPE(FTextureSourceData::GetSourceMips_ChannelMinMax);
+
+			// @@ can we do this on the Texture instead so it's cached for next time?
+			//	I think you can just call Source.UpdateChannelMinMaxFromIncomingTextureData(ScopedMipData) here
+
 			LayerChannelMinMax.Reset();
 			for (int32 LayerIndex = 0; LayerIndex < Layers.Num(); ++LayerIndex)
 			{
@@ -358,35 +362,41 @@ void FTextureSourceData::GetSourceMips(FTextureSource& Source, IImageWrapperModu
 		{
 			TRACE_CPUPROFILER_EVENT_SCOPE(FTextureSourceData::GetSourceMips_CopyMips);
 
+			check( Blocks.Num() == Source.GetNumBlocks() );
+
 			for (int32 BlockIndex = 0; BlockIndex < Blocks.Num(); ++BlockIndex)
 			{
 				FTextureSourceBlock SourceBlock;
 				Source.GetBlock(BlockIndex, SourceBlock);
 
 				FTextureSourceBlockData& BlockData = Blocks[BlockIndex];
+				check( BlockData.MipsPerLayer.Num() == Source.GetNumLayers() );
+				check( Layers.Num() == Source.GetNumLayers() );
+
 				for (int32 LayerIndex = 0; LayerIndex < Layers.Num(); ++LayerIndex)
 				{
 					const FTextureSourceLayerData& LayerData = Layers[LayerIndex];
-					if (!BlockData.MipsPerLayer[LayerIndex].Num()) // If we already got valid data, nothing to do.
+					if ( BlockData.MipsPerLayer[LayerIndex].Num() == 0 ) // If we already got valid data, nothing to do. (@@ this is weird, should detect this and early out sooner, before GetMipData)
 					{
+						BlockData.MipsPerLayer[LayerIndex].SetNum( BlockData.NumMips );
+
 						for (int32 MipIndex = 0; MipIndex < BlockData.NumMips; ++MipIndex)
 						{
 							FImageInfo MipImageInfo;
 							FSharedBuffer MipData = ScopedMipData.GetMipDataWithInfo(BlockIndex, LayerIndex, MipIndex, MipImageInfo);
-
-							FImage& SourceMip = BlockData.MipsPerLayer[LayerIndex].Emplace_GetRef(
-								MipImageInfo.SizeX, MipImageInfo.SizeY, MipImageInfo.NumSlices,
-								MipImageInfo.Format,
-								MipImageInfo.GammaSpace);
+							
 							check(MipImageInfo.GammaSpace == LayerData.SourceGammaSpace);
 							check(MipImageInfo.Format == LayerData.ImageFormat);
 
-							// @todo Oodle : FIX ME : this is slow; instead use FImageView for MipsPerLayer?
-							// could avoid this if we had an FImage that held a SharedBuffer
-							//	OR perhaps better/easier, make these just use FImageView and point them at the mips
-							//	and keep the ScopedMipData ref alove in BlockData
-							SourceMip.RawData.Reset(MipData.GetSize());
-							SourceMip.RawData.Append((const uint8*)MipData.GetData(), MipData.GetSize());
+							FImageView MipView( MipImageInfo, const_cast<void *>(MipData.GetData()) );
+							
+							// allocates the destination FImage and copies into it :
+							MipView.CopyTo( BlockData.MipsPerLayer[LayerIndex][MipIndex] );
+
+							// this copy takes a while, and potentially we could just point at the FSharedBuffer from the ScopedMipData
+							//   (like MipView here does)
+							// at the moment that's not easy because all the code around TextureCompressorModule expects FImage, not FImageView
+							//	perhaps ideally we'd have an FImage variant that's COW
 						}
 					}
 				}
