Ensure that custom gravity direction for character movement component is normalized
#rnx

diff --git a/Engine/Source/Runtime/Engine/Private/Components/CharacterMovementComponent.cpp b/Engine/Source/Runtime/Engine/Private/Components/CharacterMovementComponent.cpp
--- a/Engine/Source/Runtime/Engine/Private/Components/CharacterMovementComponent.cpp
+++ b/Engine/Source/Runtime/Engine/Private/Components/CharacterMovementComponent.cpp
@@ -4622,7 +4622,7 @@ FVector UCharacterMovementComponent::GetAirControl(float DeltaTime, float TickAi
 
 void UCharacterMovementComponent::SetGravityDirection(const FVector& InNewGravityDir)
 {
-	FVector NewGravityDir = InNewGravityDir;
+	FVector NewGravityDir = InNewGravityDir.GetSafeNormal();
 	if (ensure(!NewGravityDir.IsNearlyZero()))
 	{
 		if (!GravityDirection.Equals(NewGravityDir))
