Bring back deprecation warnings disabling macros in some code that's needed for another bramch integration.

diff --git a/Engine/Source/Runtime/Engine/Public/WorldPartition/ErrorHandling/WorldPartitionStreamingGenerationErrorHandler.h b/Engine/Source/Runtime/Engine/Public/WorldPartition/ErrorHandling/WorldPartitionStreamingGenerationErrorHandler.h
--- a/Engine/Source/Runtime/Engine/Public/WorldPartition/ErrorHandling/WorldPartitionStreamingGenerationErrorHandler.h
+++ b/Engine/Source/Runtime/Engine/Public/WorldPartition/ErrorHandling/WorldPartitionStreamingGenerationErrorHandler.h
@@ -127,6 +127,8 @@ public:
 	// Helpers
 	static ENGINE_API FString GetActorName(const IWorldPartitionActorDescInstanceView& ActorDescView);
 	static ENGINE_API FString GetFullActorName(const IWorldPartitionActorDescInstanceView& ActorDescView);
+
+	PRAGMA_DISABLE_DEPRECATION_WARNINGS
 	
 	UE_DEPRECATED(5.4, "Use OnInvalidReferenceDataLayers with EDataLayerInvalidReason instead.")
 	virtual void OnInvalidReferenceDataLayers(const IWorldPartitionActorDescInstanceView& ActorDescView, const IWorldPartitionActorDescInstanceView& ReferenceActorDescView) final {}
@@ -178,5 +180,7 @@ public:
 
 	UE_DEPRECATED(5.4, "Use IWorldPartitionActorDescInstanceView version instead")
 	static ENGINE_API FString GetFullActorName(const FWorldPartitionActorDescView& ActorDescView) { return FString(); }
+
+	PRAGMA_ENABLE_DEPRECATION_WARNINGS
 };
 #endif
