Remove obsolete comments in VirtualTextureDataBuilder.

diff --git a/Engine/Source/Runtime/Engine/Private/VT/VirtualTextureDataBuilder.cpp b/Engine/Source/Runtime/Engine/Private/VT/VirtualTextureDataBuilder.cpp
--- a/Engine/Source/Runtime/Engine/Private/VT/VirtualTextureDataBuilder.cpp
+++ b/Engine/Source/Runtime/Engine/Private/VT/VirtualTextureDataBuilder.cpp
@@ -1026,7 +1026,6 @@ void FVirtualTextureDataBuilder::BuildLayerBlocks(FSlowTask& BuildTask, uint32 L
 		}
 		else
 		{
-			// @@ I think this is wrong for non-square blocks in UDIM (see "dimensionstress")
 			const uint32 BlockSizeInTilesXY = FMath::DivideAndRoundUp<uint32>(BlockSizeXY, TileSize);
 			const uint32 MaxMipInBlockXY = FMath::CeilLogTwo(BlockSizeInTilesXY);
 			BlockData.NumMips = FMath::Min<int32>(CompressedMips.Num(), MaxMipInBlockXY + 1);
@@ -1099,15 +1098,9 @@ void FVirtualTextureDataBuilder::BuildLayerBlocks(FSlowTask& BuildTask, uint32 L
 		BlockData.SizeInBlocksY = DerivedInfo.SizeInBlocksY;
 		BlockData.SizeX = FMath::Max(MipInputSizeX >> 1, 1u);
 		BlockData.SizeY = FMath::Max(MipInputSizeY >> 1, 1u);
-		BlockData.NumMips = OutData.NumMips - MaxMipInBlock - 1;//   FMath::CeilLogTwo(MipInputSize); // Don't add 1, since 'MipInputSize' is one mip larger
+		BlockData.NumMips = OutData.NumMips - MaxMipInBlock - 1;
 		BlockData.NumSlices = 1; // TODO?
 		BlockData.MipBias = MaxMipInBlock + 1;
-		// @@ BlockData.NumMips is negative for non-square blocks in UDIM (see "dimensionstress")
-		check(BlockData.NumMips > 0);
-
-		// Total number of mips should be equal to number of mips per block plus number of miptail mips
-		check(MaxMipInBlock + BlockData.NumMips + 1 == OutData.NumMips);
-
 
 		const FTextureBuildSettings& BuildSettingsForLayer = SettingsPerLayer[LayerIndex];
 
