ColorManagement: Further luminance cleanup since global factors update at CL 32327634.

* We use sRGB/Rec.709 luminance factors in Material UI domains, if possible.
* Cleaned up  unnecessary scene luminance function in ColorSpace.ush.

diff --git a/Engine/Shaders/Private/ColorSpace.ush b/Engine/Shaders/Private/ColorSpace.ush
--- a/Engine/Shaders/Private/ColorSpace.ush
+++ b/Engine/Shaders/Private/ColorSpace.ush
@@ -69,23 +69,6 @@ float3 XYZ_2_LinearRGB( float3 XYZ )
 #endif
 }
 
-#if WORKING_COLOR_SPACE_IS_SRGB
-static const float3 WorkingColorSpaceLuminanceFactors = Linear_sRGB_2_XYZ_MAT._m10_m11_m12;
-#else
-static const float3 WorkingColorSpaceLuminanceFactors = WORKING_COLOR_SPACE_RGB_TO_XYZ_MAT._m10_m11_m12;
-#endif
-
-float Luminance( float3 LinearColor, float3 LuminanceFactors )
-{
-	return dot( LinearColor, LuminanceFactors );
-}
-
-// Working color space-accurate luminance function for scene-referred linear colors.
-float SceneLuminance( float3 LinearColor )
-{
-	return Luminance( LinearColor, WorkingColorSpaceLuminanceFactors );
-}
-
 //------------------------------------------------------------------
 // CIE LAB
 // https://en.wikipedia.org/wiki/Lab_color_space
diff --git a/Engine/Shaders/Private/Common.ush b/Engine/Shaders/Private/Common.ush
--- a/Engine/Shaders/Private/Common.ush
+++ b/Engine/Shaders/Private/Common.ush
@@ -606,7 +606,7 @@ float3 LuminanceFactors()
 #if UE_LEGACY_LUMINANCE_FACTORS
 	return float3(0.3, 0.59, 0.11);
 #else
-	#if WORKING_COLOR_SPACE_IS_SRGB
+	#if WORKING_COLOR_SPACE_IS_SRGB || MATERIAL_DOMAIN_UI
 		return float3(0.2126390059, 0.7151686788, 0.0721923154);
 	#else
 		return float3(WORKING_COLOR_SPACE_RGB_TO_XYZ_MAT._m10_m11_m12);
diff --git a/Engine/Shaders/Private/PostProcessVisualizeHDR.usf b/Engine/Shaders/Private/PostProcessVisualizeHDR.usf
--- a/Engine/Shaders/Private/PostProcessVisualizeHDR.usf
+++ b/Engine/Shaders/Private/PostProcessVisualizeHDR.usf
@@ -240,7 +240,7 @@ void MainPS(noperspective float4 UVAndScreenPos : TEXCOORD0, float4 SvPosition :
 				OutColor = Texture2DSample(SceneColorTexture, SceneColorSampler, UV); // Influenced area default scene color
 			}
 
-			float LuminanceAvg = SceneLuminance(HDRSceneColorAvg);
+			float LuminanceAvg = Luminance(HDRSceneColorAvg);
 
 			// Luminance
 			int2 TopLeft = OutputViewportCenter + int2(11, 11);
@@ -325,7 +325,7 @@ void MainPS(noperspective float4 UVAndScreenPos : TEXCOORD0, float4 SvPosition :
 			}
 
 			float3 Illuminance = HDRSceneColorAvg * PI;
-			float MeanIlluminance = SceneLuminance(Illuminance);
+			float MeanIlluminance = Luminance(Illuminance);
 
 			int2 TopLeft = IlluminanceMeterCenter + int2(11, 11);
 			// Mean illuminance
diff --git a/Engine/Source/Runtime/Engine/Classes/Materials/MaterialExpressionDesaturation.h b/Engine/Source/Runtime/Engine/Classes/Materials/MaterialExpressionDesaturation.h
--- a/Engine/Source/Runtime/Engine/Classes/Materials/MaterialExpressionDesaturation.h
+++ b/Engine/Source/Runtime/Engine/Classes/Materials/MaterialExpressionDesaturation.h
@@ -21,8 +21,16 @@ class UMaterialExpressionDesaturation : public UMaterialExpression
 	UPROPERTY()
 	FExpressionInput Fraction;
 
+	/* 
+	 * Luminance factors for converting a color to greyscale.
+	 * 
+	 * The default luminance factors values are now derived from the working color space. For uses cases
+	 * outside scene rendering, users are responsible for updating these factors accordingly. For example,
+	 * factors derived from an AP1 working color space would not be applicable to UI domain materials that
+	 * remain in sRGB/Rec.709 and thus should instead use approximately [0.2126, 0.7152, 0.0722].
+	 */
 	UPROPERTY(EditAnywhere, Category=MaterialExpressionDesaturation, Meta = (ShowAsInputPin = "Advanced"))
-	FLinearColor LuminanceFactors;    // Color component factors for converting a color to greyscale.
+	FLinearColor LuminanceFactors;
 
 
 	//~ Begin UMaterialExpression Interface
diff --git a/Engine/Source/Runtime/Engine/Private/Materials/MaterialExpressions.cpp b/Engine/Source/Runtime/Engine/Private/Materials/MaterialExpressions.cpp
--- a/Engine/Source/Runtime/Engine/Private/Materials/MaterialExpressions.cpp
+++ b/Engine/Source/Runtime/Engine/Private/Materials/MaterialExpressions.cpp
@@ -8593,7 +8593,6 @@ UMaterialExpressionDesaturation::UMaterialExpressionDesaturation(const FObjectIn
 	}
 	else
 	{
-		//TODO: Handle MATERIAL_DOMAIN_UI which should remain in sRGB/Rec.709
 		LuminanceFactors = UE::Color::FColorSpace::GetWorking().GetLuminanceFactors();
 	}
 
