Removing duplicated code

#rnx

diff --git a/Engine/Source/Runtime/Engine/Classes/Sound/SoundSubmix.h b/Engine/Source/Runtime/Engine/Classes/Sound/SoundSubmix.h
--- a/Engine/Source/Runtime/Engine/Classes/Sound/SoundSubmix.h
+++ b/Engine/Source/Runtime/Engine/Classes/Sound/SoundSubmix.h
@@ -489,6 +489,9 @@ protected:
 	// Send submix modulation changes to audio devices
 	void PushModulationChanges();
 
+	/** Return the audio device that corresponds to the provided world context object. */
+	static FAudioDevice* GetAudioDeviceFrom(const UObject* WorldContextObject);
+
 	// State handling for bouncing output.
 	TUniquePtr<Audio::FAudioRecordingData> RecordingData;
 
diff --git a/Engine/Source/Runtime/Engine/Private/SoundSubmix.cpp b/Engine/Source/Runtime/Engine/Private/SoundSubmix.cpp
--- a/Engine/Source/Runtime/Engine/Private/SoundSubmix.cpp
+++ b/Engine/Source/Runtime/Engine/Private/SoundSubmix.cpp
@@ -104,18 +104,9 @@ USoundfieldEndpointSubmix::USoundfieldEndpointSubmix(const FObjectInitializer& O
 
 void USoundSubmix::StartRecordingOutput(const UObject* WorldContextObject, float ExpectedDuration)
 {
-	if (!GEngine)
+	if (FAudioDevice* AudioDevice = GetAudioDeviceFrom(WorldContextObject))
 	{
-		return;
-	}
-
-	// Find device for this specific audio recording thing.
-	if (UWorld* ThisWorld = GEngine->GetWorldFromContextObject(WorldContextObject, EGetWorldErrorMode::LogAndReturnNull))
-	{
-		if (FAudioDevice* AudioDevice = ThisWorld->GetAudioDeviceRaw())
-		{
-			StartRecordingOutput(AudioDevice, ExpectedDuration);
-		}
+		StartRecordingOutput(AudioDevice, ExpectedDuration);
 	}
 }
 
@@ -129,18 +120,9 @@ void USoundSubmix::StartRecordingOutput(FAudioDevice* InDevice, float ExpectedDu
 
 void USoundSubmix::StopRecordingOutput(const UObject* WorldContextObject, EAudioRecordingExportType ExportType, const FString& Name, FString Path, USoundWave* ExistingSoundWaveToOverwrite)
 {
-	if (!GEngine)
+	if (FAudioDevice* AudioDevice = GetAudioDeviceFrom(WorldContextObject))
 	{
-		return;
-	}
-
-	// Find device for this specific audio recording thing.
-	if (UWorld* ThisWorld = GEngine->GetWorldFromContextObject(WorldContextObject, EGetWorldErrorMode::LogAndReturnNull))
-	{
-		if (FAudioDevice* AudioDevice = ThisWorld->GetAudioDeviceRaw())
-		{
-			StopRecordingOutput(AudioDevice, ExportType, Name, Path, ExistingSoundWaveToOverwrite);
-		}
+		StopRecordingOutput(AudioDevice, ExportType, Name, Path, ExistingSoundWaveToOverwrite);
 	}
 }
 
@@ -212,18 +194,9 @@ void USoundSubmix::StopRecordingOutput(FAudioDevice* InDevice, EAudioRecordingEx
 
 void USoundSubmix::StartEnvelopeFollowing(const UObject* WorldContextObject)
 {
-	if (!GEngine)
+	if (FAudioDevice* AudioDevice = GetAudioDeviceFrom(WorldContextObject))
 	{
-		return;
-	}
-
-	// Find device for this specific audio recording thing.
-	if (UWorld* ThisWorld = GEngine->GetWorldFromContextObject(WorldContextObject, EGetWorldErrorMode::LogAndReturnNull))
-	{
-		if (FAudioDevice* AudioDevice = ThisWorld->GetAudioDeviceRaw())
-		{
-			StartEnvelopeFollowing(AudioDevice);
-		}
+		StartEnvelopeFollowing(AudioDevice);
 	}
 }
 
@@ -237,18 +210,9 @@ void USoundSubmix::StartEnvelopeFollowing(FAudioDevice* InAudioDevice)
 
 void USoundSubmix::StopEnvelopeFollowing(const UObject* WorldContextObject)
 {
-	if (!GEngine)
-	{
-		return;
-	}
-
-	// Find device for this specific audio recording thing.
-	if (UWorld* ThisWorld = GEngine->GetWorldFromContextObject(WorldContextObject, EGetWorldErrorMode::LogAndReturnNull))
+	if (FAudioDevice* AudioDevice = GetAudioDeviceFrom(WorldContextObject))
 	{
-		if (FAudioDevice* AudioDevice = ThisWorld->GetAudioDeviceRaw())
-		{
-			StopEnvelopeFollowing(AudioDevice);
-		}
+		StopEnvelopeFollowing(AudioDevice);
 	}
 }
 
@@ -262,100 +226,47 @@ void USoundSubmix::StopEnvelopeFollowing(FAudioDevice* InAudioDevice)
 
 void USoundSubmix::AddEnvelopeFollowerDelegate(const UObject* WorldContextObject, const FOnSubmixEnvelopeBP& OnSubmixEnvelopeBP)
 {
-	if (!GEngine)
+	if (FAudioDevice* AudioDevice = GetAudioDeviceFrom(WorldContextObject))
 	{
-		return;
-	}
-
-	// Find device for this specific audio recording thing.
-	if (UWorld* ThisWorld = GEngine->GetWorldFromContextObject(WorldContextObject, EGetWorldErrorMode::LogAndReturnNull))
-	{
-		if (FAudioDevice* AudioDevice = ThisWorld->GetAudioDeviceRaw())
-		{
-			AudioDevice->AddEnvelopeFollowerDelegate(this, OnSubmixEnvelopeBP);
-		}
+		AudioDevice->AddEnvelopeFollowerDelegate(this, OnSubmixEnvelopeBP);
 	}
 }
 
 void USoundSubmix::RemoveEnvelopeFollowerDelegate(const UObject* WorldContextObject, const FOnSubmixEnvelopeBP& OnSubmixEnvelopeBP)
 {
-	if (!GEngine)
+	if (FAudioDevice* AudioDevice = GetAudioDeviceFrom(WorldContextObject))
 	{
-		return;
-	}
-
-	// Find device for this specific audio recording thing.
-	if (UWorld* ThisWorld = GEngine->GetWorldFromContextObject(WorldContextObject, EGetWorldErrorMode::LogAndReturnNull))
-	{
-		if (FAudioDevice* AudioDevice = ThisWorld->GetAudioDeviceRaw())
-		{
-			AudioDevice->RemoveEnvelopeFollowerDelegate(this, OnSubmixEnvelopeBP);
-		}
+		AudioDevice->RemoveEnvelopeFollowerDelegate(this, OnSubmixEnvelopeBP);
 	}
 }
 
 void USoundSubmix::AddSpectralAnalysisDelegate(const UObject* WorldContextObject, const TArray<FSoundSubmixSpectralAnalysisBandSettings>& InBandSettings, const FOnSubmixSpectralAnalysisBP& OnSubmixSpectralAnalysisBP, float UpdateRate,  float DecibelNoiseFloor, bool bDoNormalize, bool bDoAutoRange, float AutoRangeAttackTime, float AutoRangeReleaseTime)
 {
-
-	if (!GEngine)
+	if (FAudioDevice* AudioDevice = GetAudioDeviceFrom(WorldContextObject))
 	{
-		return;
-	}
-
-	// Find device for this specific audio recording thing.
-	if (UWorld* ThisWorld = GEngine->GetWorldFromContextObject(WorldContextObject, EGetWorldErrorMode::LogAndReturnNull))
-	{
-		if (FAudioDevice* AudioDevice = ThisWorld->GetAudioDeviceRaw())
-		{
-			FSoundSpectrumAnalyzerDelegateSettings DelegateSettings = USoundSubmix::GetSpectrumAnalysisDelegateSettings(InBandSettings, UpdateRate, DecibelNoiseFloor, bDoNormalize, bDoAutoRange, AutoRangeAttackTime, AutoRangeReleaseTime);
-
-
-			AudioDevice->AddSpectralAnalysisDelegate(this, DelegateSettings, OnSubmixSpectralAnalysisBP);
-		}
+		FSoundSpectrumAnalyzerDelegateSettings DelegateSettings = USoundSubmix::GetSpectrumAnalysisDelegateSettings(InBandSettings, UpdateRate, DecibelNoiseFloor, bDoNormalize, bDoAutoRange, AutoRangeAttackTime, AutoRangeReleaseTime);
+		AudioDevice->AddSpectralAnalysisDelegate(this, DelegateSettings, OnSubmixSpectralAnalysisBP);
 	}
 }
 
 void USoundSubmix::RemoveSpectralAnalysisDelegate(const UObject* WorldContextObject, const FOnSubmixSpectralAnalysisBP& OnSubmixSpectralAnalysisBP)
 {
-	if (!GEngine)
+	if (FAudioDevice* AudioDevice = GetAudioDeviceFrom(WorldContextObject))
 	{
-		return;
-	}
-
-	// Find device for this specific audio recording thing.
-	if (UWorld* ThisWorld = GEngine->GetWorldFromContextObject(WorldContextObject, EGetWorldErrorMode::LogAndReturnNull))
-	{
-		if (FAudioDevice* AudioDevice = ThisWorld->GetAudioDeviceRaw())
-		{
-			AudioDevice->RemoveSpectralAnalysisDelegate(this, OnSubmixSpectralAnalysisBP);
-		}
+		AudioDevice->RemoveSpectralAnalysisDelegate(this, OnSubmixSpectralAnalysisBP);
 	}
 }
 
 void USoundSubmix::StartSpectralAnalysis(const UObject* WorldContextObject, EFFTSize FFTSize, EFFTPeakInterpolationMethod InterpolationMethod, EFFTWindowType WindowType, float HopSize, EAudioSpectrumType SpectrumType)
 {
-	if (!GEngine)
+	if (FAudioDevice* AudioDevice = GetAudioDeviceFrom(WorldContextObject))
 	{
-		return;
-	}
-
-	// Find device for this specific audio recording thing.
-	if (UWorld* ThisWorld = GEngine->GetWorldFromContextObject(WorldContextObject, EGetWorldErrorMode::LogAndReturnNull))
-	{	
-		if (FAudioDevice* AudioDevice = ThisWorld->GetAudioDeviceRaw())
-		{
-			StartSpectralAnalysis(AudioDevice, FFTSize, InterpolationMethod, WindowType, HopSize, SpectrumType);
-		}
+		StartSpectralAnalysis(AudioDevice, FFTSize, InterpolationMethod, WindowType, HopSize, SpectrumType);
 	}
 }
 
 void USoundSubmix::StartSpectralAnalysis(FAudioDevice* InAudioDevice, EFFTSize FFTSize, EFFTPeakInterpolationMethod InterpolationMethod, EFFTWindowType WindowType, float HopSize, EAudioSpectrumType SpectrumType)
 {
-	if (!GEngine)
-	{
-		return;
-	}
-
 	if (InAudioDevice)
 	{
 		FSoundSpectrumAnalyzerSettings Settings = USoundSubmix::GetSpectrumAnalyzerSettings(FFTSize, InterpolationMethod, WindowType, HopSize, SpectrumType);
@@ -365,18 +276,9 @@ void USoundSubmix::StartSpectralAnalysis(FAudioDevice* InAudioDevice, EFFTSize F
 
 void USoundSubmix::StopSpectralAnalysis(const UObject* WorldContextObject)
 {
-	if (!GEngine)
+	if (FAudioDevice* AudioDevice = GetAudioDeviceFrom(WorldContextObject))
 	{
-		return;
-	}
-
-	// Find device for this specific audio recording thing.
-	if (UWorld* ThisWorld = GEngine->GetWorldFromContextObject(WorldContextObject, EGetWorldErrorMode::LogAndReturnNull))
-	{
-		if (FAudioDevice* AudioDevice = ThisWorld->GetAudioDeviceRaw())
-		{
-			AudioDevice->StopSpectrumAnalysis(this);
-		}
+		AudioDevice->StopSpectrumAnalysis(this);
 	}
 }
 
@@ -390,49 +292,25 @@ void USoundSubmix::StopSpectralAnalysis(FAudioDevice* InAudioDevice)
 
 void USoundSubmix::SetSubmixOutputVolume(const UObject* WorldContextObject, float InOutputVolume)
 {
-	if (!GEngine)
+	if (FAudioDevice* AudioDevice = GetAudioDeviceFrom(WorldContextObject))
 	{
-		return;
-	}
-
-	if (UWorld* ThisWorld = GEngine->GetWorldFromContextObject(WorldContextObject, EGetWorldErrorMode::LogAndReturnNull))
-	{
-		if (FAudioDevice* AudioDevice = ThisWorld->GetAudioDeviceRaw())
-		{
-			AudioDevice->SetSubmixOutputVolume(this, InOutputVolume);
-		}
+		AudioDevice->SetSubmixOutputVolume(this, InOutputVolume);
 	}
 }
 
 void USoundSubmix::SetSubmixWetLevel(const UObject* WorldContextObject, float InWetLevel)
 {
-	if (!GEngine)
+	if (FAudioDevice* AudioDevice = GetAudioDeviceFrom(WorldContextObject))
 	{
-		return;
-	}
-
-	if (UWorld* ThisWorld = GEngine->GetWorldFromContextObject(WorldContextObject, EGetWorldErrorMode::LogAndReturnNull))
-	{
-		if (FAudioDevice* AudioDevice = ThisWorld->GetAudioDeviceRaw())
-		{
-			AudioDevice->SetSubmixWetLevel(this, InWetLevel);
-		}
+		AudioDevice->SetSubmixWetLevel(this, InWetLevel);
 	}
 }
 
 void USoundSubmix::SetSubmixDryLevel(const UObject* WorldContextObject, float InDryLevel)
 {
-	if (!GEngine)
-	{
-		return;
-	}
-
-	if (UWorld* ThisWorld = GEngine->GetWorldFromContextObject(WorldContextObject, EGetWorldErrorMode::LogAndReturnNull))
+	if (FAudioDevice* AudioDevice = GetAudioDeviceFrom(WorldContextObject))
 	{
-		if (FAudioDevice* AudioDevice = ThisWorld->GetAudioDeviceRaw())
-		{
-			AudioDevice->SetSubmixDryLevel(this, InDryLevel);
-		}
+		AudioDevice->SetSubmixDryLevel(this, InDryLevel);
 	}
 }
 
@@ -531,6 +409,19 @@ void USoundSubmix::PushModulationChanges()
 	}
 }
 
+FAudioDevice* USoundSubmix::GetAudioDeviceFrom(const UObject* WorldContextObject) 
+{
+	if (GEngine)
+	{
+		if (UWorld* ThisWorld = GEngine->GetWorldFromContextObject(WorldContextObject, EGetWorldErrorMode::LogAndReturnNull))
+		{
+			return ThisWorld->GetAudioDeviceRaw();
+		}
+	}
+
+	return nullptr;
+}
+
 FString USoundSubmixBase::GetDesc()
 {
 	return FString(TEXT("Sound Submix"));
