[WorldPartition - Cook] Adapt world partition cooking to support the case where ICookPackageSplitter::PopulateGeneratedPackage can be called on an ExternalStreamingObject (Content Bundle generated package) before ICookPackageSplitter::PopulateGeneratorPackage is called on the partition world (generator package) and also before ICookPackageSplitter::PopulateGeneratedPackage is called on all generated packages of the cells referenced by the external streaming object.

diff --git a/Engine/Source/Runtime/Engine/Private/WorldPartition/ContentBundle/ContentBundleEditor.cpp b/Engine/Source/Runtime/Engine/Private/WorldPartition/ContentBundle/ContentBundleEditor.cpp
--- a/Engine/Source/Runtime/Engine/Private/WorldPartition/ContentBundle/ContentBundleEditor.cpp
+++ b/Engine/Source/Runtime/Engine/Private/WorldPartition/ContentBundle/ContentBundleEditor.cpp
@@ -414,7 +414,7 @@ bool FContentBundleEditor::GatherPackagesToCook(class IWorldPartitionCookPackage
 	{
 		ExternalStreamingObject->ForEachStreamingCells([this, &bIsSuccess, &CookContext](UWorldPartitionRuntimeCell& RuntimeCell)
 		{
-			if (const FWorldPartitionCookPackage* CookPackage = CookContext.AddLevelStreamingPackageToGenerate(this, ContentBundlePaths::GetCookedContentBundleLevelFolder(*this), RuntimeCell.GetPackageNameToCreate()))
+			if (const FWorldPartitionCookPackage* CookPackage = CookContext.AddPackageToGenerate(this, &RuntimeCell, ContentBundlePaths::GetCookedContentBundleLevelFolder(*this), RuntimeCell.GetPackageNameToCreate()))
 			{
 				CookPackageIdsToCell.Emplace(CookPackage->PackageId, &RuntimeCell);
 			}
@@ -425,7 +425,7 @@ bool FContentBundleEditor::GatherPackagesToCook(class IWorldPartitionCookPackage
 			}
 		});
 		
-		const FWorldPartitionCookPackage* CookPackage = CookContext.AddGenericPackageToGenerate(this, ContentBundlePaths::GetCookedContentBundleLevelFolder(*this), GetExternalStreamingObjectPackageName());
+		const FWorldPartitionCookPackage* CookPackage = CookContext.AddPackageToGenerate(this, ExternalStreamingObject, ContentBundlePaths::GetCookedContentBundleLevelFolder(*this), GetExternalStreamingObjectPackageName());
 		if (CookPackage == nullptr)
 		{
 			UE_LOG(LogContentBundle, Error, TEXT("%s[Cook] Failed to add streaming object package in cook context."), *ContentBundle::Log::MakeDebugInfoString(*this));
@@ -455,7 +455,7 @@ bool FContentBundleEditor::PopulateGeneratorPackageForCook(class IWorldPartition
 					// Make sure the cell outer is set to the  ExternalStreamingObject so it will be saved in the right package at the end of the cook.
 					check(Cell->GetOuter() == ExternalStreamingObject);
 
-					if (!Cell->OnPopulateGeneratorPackageForCook(CookPackage->GetPackage()))
+					if (!Cell->OnPopulateGeneratorPackageForCook(CookContext, CookPackage->GetPackage()))
 					{
 						UE_LOG(LogContentBundle, Error, TEXT("%s[Cook] Failed to prepare cell with package %s for cook."), *ContentBundle::Log::MakeDebugInfoString(*this), *CookPackage->RelativePath);
 						bIsSuccess = false;
@@ -463,13 +463,13 @@ bool FContentBundleEditor::PopulateGeneratorPackageForCook(class IWorldPartition
 				}
 				else
 				{
-					UE_LOG(LogContentBundle, Error, TEXT("%s[Cook] Could not find cell for package %s while populating generator pacakges."), *ContentBundle::Log::MakeDebugInfoString(*this), *CookPackage->RelativePath);
+					UE_LOG(LogContentBundle, Error, TEXT("%s[Cook] Could not find cell for package %s while populating generator packages."), *ContentBundle::Log::MakeDebugInfoString(*this), *CookPackage->RelativePath);
 					bIsSuccess = false;
 				}
 			}
 		}
 
-		ExternalStreamingObject->OnPopulateGeneratorPackageForCook(nullptr);
+		ExternalStreamingObject->OnPopulateGeneratorPackageForCook(CookContext, nullptr);
 	}
 	
 
@@ -492,7 +492,7 @@ bool FContentBundleEditor::PopulateGeneratedPackageForCook(class IWorldPartition
 				if (!Cell->IsAlwaysLoaded())
 				{
 					TArray<UPackage*> ModifiedPackages;
-					if (Cell->OnPopulateGeneratedPackageForCook(PackageToCook.GetPackage(), OutModifiedPackages))
+					if (Cell->OnPopulateGeneratedPackageForCook(CookContext, PackageToCook.GetPackage(), OutModifiedPackages))
 					{
 						UWorld* CellWorld = FindObject<UWorld>(PackageToCook.GetPackage(), *GetInjectedWorld()->GetName());
 						if (CellWorld != nullptr)
@@ -536,7 +536,7 @@ bool FContentBundleEditor::PopulateGeneratedPackageForCook(class IWorldPartition
 	}
 	else
 	{
-		if (!ExternalStreamingObject->OnPopulateGeneratedPackageForCook(PackageToCook.GetPackage(), OutModifiedPackages))
+		if (!ExternalStreamingObject->OnPopulateGeneratedPackageForCook(CookContext, PackageToCook.GetPackage(), OutModifiedPackages))
 		{
 			UE_LOG(LogContentBundle, Error, TEXT("%s[Cook] Failed to rename streaming object package."), *ContentBundle::Log::MakeDebugInfoString(*this));
 			bIsSuccess = false;
diff --git a/Engine/Source/Runtime/Engine/Private/WorldPartition/Cook/WorldPartitionCook.cpp b/Engine/Source/Runtime/Engine/Private/WorldPartition/Cook/WorldPartitionCook.cpp
--- a/Engine/Source/Runtime/Engine/Private/WorldPartition/Cook/WorldPartitionCook.cpp
+++ b/Engine/Source/Runtime/Engine/Private/WorldPartition/Cook/WorldPartitionCook.cpp
@@ -60,7 +60,7 @@ bool UWorldPartition::GatherPackagesToCook(IWorldPartitionCookPackageContext& Co
 			const UExternalDataLayerAsset* ExternalDataLayerAsset = InCookPackageObject->GetExternalDataLayerAsset();
 			const FString ExternalContentRoot = ExternalDataLayerAsset ? ExternalDataLayerManager->GetExternalDataLayerLevelRootPath(ExternalDataLayerAsset) : FString();
 			const FString& Root = ExternalDataLayerAsset ? ExternalContentRoot : WorldPackageName;
-			if (InCookPackageObject->IsLevelPackage() ? CookContext.AddLevelStreamingPackageToGenerate(this, Root, InPackageName) : CookContext.AddGenericPackageToGenerate(this, Root, InPackageName))
+			if (CookContext.AddPackageToGenerate(this, InCookPackageObject, Root, InPackageName))
 			{
 				return;
 			}
@@ -108,10 +108,10 @@ bool UWorldPartition::PrepareGeneratorPackageForCook(IWorldPartitionCookPackageC
 
 bool UWorldPartition::PopulateGeneratorPackageForCook(IWorldPartitionCookPackageContext& CookContext, const TArray<FWorldPartitionCookPackage*>& InPackagesToCook, TArray<UPackage*>& OutModifiedPackages)
 {
-	auto OnPopulateGeneratorPackageForCook = [this](const FWorldPartitionCookPackage* InPackageToCook)
+	auto OnPopulateGeneratorPackageForCook = [this, &CookContext](const FWorldPartitionCookPackage* InPackageToCook)
 	{
 		IWorldPartitionCookPackageObject* CookPackageObject = GetCookPackageObject(*InPackageToCook);
-		return CookPackageObject && CookPackageObject->OnPopulateGeneratorPackageForCook(InPackageToCook->GetPackage());
+		return CookPackageObject && CookPackageObject->OnPopulateGeneratorPackageForCook(CookContext, InPackageToCook->GetPackage());
 	};
 
 	TArray<const FWorldPartitionCookPackage*> GenericPackagesToCook;
@@ -161,7 +161,7 @@ bool UWorldPartition::PopulateGeneratedPackageForCook(IWorldPartitionCookPackage
 	OutModifiedPackages.Reset();
 	if (IWorldPartitionCookPackageObject* CookPackageObject = GetCookPackageObject(InPackageToCook))
 	{
-		if (CookPackageObject->OnPopulateGeneratedPackageForCook(InPackageToCook.GetPackage(), OutModifiedPackages))
+		if (CookPackageObject->OnPopulateGeneratedPackageForCook(CookContext, InPackageToCook.GetPackage(), OutModifiedPackages))
 		{
 			// Since PopulateGeneratedPackageForCook on the external streaming object package changes its outer (which affects its cells object path) and can be
 			// called after a call to PopulateGeneratedPackageForCook on a package of any cell external streaming object, we need to generate the cell path manually.
diff --git a/Engine/Source/Runtime/Engine/Private/WorldPartition/Cook/WorldPartitionCookPackageContext.cpp b/Engine/Source/Runtime/Engine/Private/WorldPartition/Cook/WorldPartitionCookPackageContext.cpp
--- a/Engine/Source/Runtime/Engine/Private/WorldPartition/Cook/WorldPartitionCookPackageContext.cpp
+++ b/Engine/Source/Runtime/Engine/Private/WorldPartition/Cook/WorldPartitionCookPackageContext.cpp
@@ -23,6 +23,39 @@ void FWorldPartitionCookPackageContext::UnregisterPackageCookPackageGenerator(IW
 	verify(CookPackageGenerators.Remove(CookPackageGenerator));
 }
 
+const FWorldPartitionCookPackage* FWorldPartitionCookPackageContext::AddPackageToGenerate(IWorldPartitionCookPackageGenerator* Generator, IWorldPartitionCookPackageObject* InCookPackageObject, const FString& Root, const FString& RelativePath)
+{
+	const FWorldPartitionCookPackage* WorldPartitionCookPackage = nullptr;
+
+	if (InCookPackageObject->IsLevelPackage())
+	{
+		WorldPartitionCookPackage = AddLevelStreamingPackageToGenerate(Generator, Root, RelativePath);
+	}
+	else
+	{
+		WorldPartitionCookPackage = AddGenericPackageToGenerate(Generator, Root, RelativePath);
+	}
+
+	if (WorldPartitionCookPackage)
+	{
+		PackageObjectToPackageId.Add(Cast<UObject>(InCookPackageObject), WorldPartitionCookPackage->PackageId);
+	}
+
+	return WorldPartitionCookPackage;
+}
+
+FString FWorldPartitionCookPackageContext::GetGeneratedPackagePath(IWorldPartitionCookPackageObject* InCookPackageObject) const
+{
+	if (const FWorldPartitionCookPackage::IDType* PackageId = PackageObjectToPackageId.Find(Cast<UObject>(InCookPackageObject)))
+	{
+		if (const TUniquePtr<FWorldPartitionCookPackage>* ExistingPackage = PackagesToCookById.Find(*PackageId))
+		{
+			return (*ExistingPackage)->GetFullGeneratedPath();
+		}
+	}
+	return FString();
+}
+
 const FWorldPartitionCookPackage* FWorldPartitionCookPackageContext::AddLevelStreamingPackageToGenerate(IWorldPartitionCookPackageGenerator* CookPackageGenerator, const FString& Root, const FString& RelativePath)
 {
 	return AddPackageToGenerateInternal(CookPackageGenerator, Root, RelativePath, FWorldPartitionCookPackage::EType::Level);
diff --git a/Engine/Source/Runtime/Engine/Private/WorldPartition/Cook/WorldPartitionCookPackageContext.h b/Engine/Source/Runtime/Engine/Private/WorldPartition/Cook/WorldPartitionCookPackageContext.h
--- a/Engine/Source/Runtime/Engine/Private/WorldPartition/Cook/WorldPartitionCookPackageContext.h
+++ b/Engine/Source/Runtime/Engine/Private/WorldPartition/Cook/WorldPartitionCookPackageContext.h
@@ -5,6 +5,7 @@
 #if WITH_EDITOR
 
 #include "CoreMinimal.h"
+#include "UObject/ObjectKey.h"
 #include "WorldPartition/Cook/WorldPartitionCookPackageContextInterface.h"
 #include "WorldPartition/Cook/WorldPartitionCookPackage.h"
 
@@ -18,9 +19,8 @@ public:
 	//~ Begin IWorldPartitionCookPackageContext Interface 
 	virtual void RegisterPackageCookPackageGenerator(IWorldPartitionCookPackageGenerator* CookPackageGenerator) override;
 	virtual void UnregisterPackageCookPackageGenerator(IWorldPartitionCookPackageGenerator* CookPackageGenerator) override;
-
-	virtual const FWorldPartitionCookPackage* AddLevelStreamingPackageToGenerate(IWorldPartitionCookPackageGenerator* CookPackageGenerator, const FString& Root, const FString& RelativePath) override;
-	virtual const FWorldPartitionCookPackage* AddGenericPackageToGenerate(IWorldPartitionCookPackageGenerator* CookPackageGenerator, const FString& Root, const FString& RelativePath) override;
+	virtual const FWorldPartitionCookPackage* AddPackageToGenerate(IWorldPartitionCookPackageGenerator* Generator, IWorldPartitionCookPackageObject* InCookPackageObject, const FString& Root, const FString& RelativePath) override;
+	virtual FString GetGeneratedPackagePath(IWorldPartitionCookPackageObject* InCookPackageObject) const override;
 	virtual bool GatherPackagesToCook() override;
 	//~ End IWorldPartitionCookPackageContext Interface 
 
@@ -35,12 +35,15 @@ public:
 	const TArray<IWorldPartitionCookPackageGenerator*>& GetCookPackageGenerators() const { return CookPackageGenerators; }
 
 private:
+	const FWorldPartitionCookPackage* AddLevelStreamingPackageToGenerate(IWorldPartitionCookPackageGenerator* CookPackageGenerator, const FString& Root, const FString& RelativePath);
+	const FWorldPartitionCookPackage* AddGenericPackageToGenerate(IWorldPartitionCookPackageGenerator* CookPackageGenerator, const FString& Root, const FString& RelativePath);
 	const FWorldPartitionCookPackage* AddPackageToGenerateInternal(IWorldPartitionCookPackageGenerator* CookPackageGenerator, const FString& Root, const FString& RelativePath, FWorldPartitionCookPackage::EType Type);
 
 	TArray<IWorldPartitionCookPackageGenerator*> CookPackageGenerators;
 	TMap<FWorldPartitionCookPackage::IDType, TUniquePtr<FWorldPartitionCookPackage>> PackagesToCookById;
 	TMap<FWorldPartitionCookPackage::IDType, IWorldPartitionCookPackageGenerator*> CookGeneratorByPackageId;
 	TMap<IWorldPartitionCookPackageGenerator*, TArray<FWorldPartitionCookPackage*>> PackagesToCookByGenerator;
+	TMap<FObjectKey, FWorldPartitionCookPackage::IDType> PackageObjectToPackageId;
 };
 
 #endif
\ No newline at end of file
diff --git a/Engine/Source/Runtime/Engine/Private/WorldPartition/WorldPartitionRuntimeHash.cpp b/Engine/Source/Runtime/Engine/Private/WorldPartition/WorldPartitionRuntimeHash.cpp
--- a/Engine/Source/Runtime/Engine/Private/WorldPartition/WorldPartitionRuntimeHash.cpp
+++ b/Engine/Source/Runtime/Engine/Private/WorldPartition/WorldPartitionRuntimeHash.cpp
@@ -13,6 +13,8 @@
 #include "Misc/ArchiveMD5.h"
 #if WITH_EDITOR
 #include "WorldPartition/Cook/WorldPartitionCookPackage.h"
+#include "WorldPartition/Cook/WorldPartitionCookPackageContextInterface.h"
+#include "UObject/ObjectSaveContext.h"
 #endif
 
 #include UE_INLINE_GENERATED_CPP_BY_NAME(WorldPartitionRuntimeHash)
@@ -114,28 +116,46 @@ FString URuntimeHashExternalStreamingObjectBase::GetPackageNameToCreate() const
 	return FString();
 }
 
-bool URuntimeHashExternalStreamingObjectBase::OnPopulateGeneratorPackageForCook(UPackage* InPackage)
+bool URuntimeHashExternalStreamingObjectBase::PrepareForCook(const IWorldPartitionCookPackageContext& InCookContext)
 {
-	ForEachStreamingCells([this](UWorldPartitionRuntimeCell& Cell)
+	bool bResult = true;
+	ForEachStreamingCells([this, &bResult, &InCookContext](UWorldPartitionRuntimeCell& Cell)
 	{
-		UWorldPartitionRuntimeLevelStreamingCell* RuntimeCell = CastChecked<UWorldPartitionRuntimeLevelStreamingCell>(&Cell);
-		UWorldPartitionLevelStreamingDynamic* LevelStreamingDynamic = RuntimeCell->GetLevelStreaming();
-		FWorldPartitionRuntimeCellStreamingData& CellStreamingData = CellToStreamingData.Add(RuntimeCell->GetFName());
-		CellStreamingData.PackageName = LevelStreamingDynamic->GetWorldAsset().GetLongPackageName();
-		// SoftObjectPath will be automatically remapped when ExternalStreamingObject will be instanced/loaded at runtime
-		CellStreamingData.WorldAsset = LevelStreamingDynamic->GetWorldAsset().ToSoftObjectPath();
-
-		// Level streaming are outered to the world and would not be saved within the ExternalStreamingObject.
-		// Do not save them, instead they will be created once the external streaming object is loaded at runtime. 
-		LevelStreamingDynamic->SetFlags(RF_Transient);
+		// Make cell is ready for cook
+		if (Cell.PrepareCellForCook(InCookContext))
+		{
+			UWorldPartitionRuntimeLevelStreamingCell* RuntimeCell = CastChecked<UWorldPartitionRuntimeLevelStreamingCell>(&Cell);
+			UWorldPartitionLevelStreamingDynamic* LevelStreamingDynamic = RuntimeCell->GetLevelStreaming();
+			FWorldPartitionRuntimeCellStreamingData& CellStreamingData = CellToStreamingData.Add(RuntimeCell->GetFName());
+			CellStreamingData.PackageName = LevelStreamingDynamic->GetWorldAsset().GetLongPackageName();
+			// SoftObjectPath will be automatically remapped when ExternalStreamingObject will be instanced/loaded at runtime
+			CellStreamingData.WorldAsset = LevelStreamingDynamic->GetWorldAsset().ToSoftObjectPath();
+
+			// Level streaming are outered to the world and would not be saved within the ExternalStreamingObject.
+			// Do not save them, instead they will be created once the external streaming object is loaded at runtime. 
+			LevelStreamingDynamic->SetFlags(RF_Transient);
+		}
+		else
+		{
+			bResult = false;
+		}
 	});
-	return true;
+	return bResult;
+}
+
+bool URuntimeHashExternalStreamingObjectBase::OnPopulateGeneratorPackageForCook(const IWorldPartitionCookPackageContext& InCookContext, UPackage* InGeneratedPackage)
+{
+	return PrepareForCook(InCookContext);
 }
 
-bool URuntimeHashExternalStreamingObjectBase::OnPopulateGeneratedPackageForCook(UPackage* InPackage, TArray<UPackage*>& OutModifiedPackages)
+bool URuntimeHashExternalStreamingObjectBase::OnPopulateGeneratedPackageForCook(const IWorldPartitionCookPackageContext& InCookContext, UPackage* InGeneratedPackage, TArray<UPackage*>& OutModifiedPackages)
 {
-	// We provide a new name for the URuntimeHashExternalStreamingObjectBase in the package so that we have a stable name (for cook determinism)
-	return Rename(TEXT("RuntimeHashExternalStreamingObjectBase"), InPackage, REN_DontCreateRedirectors);
+	if (PrepareForCook(InCookContext))
+	{
+		// We provide a new name for the URuntimeHashExternalStreamingObjectBase in the package so that we have a stable name (for cook determinism)
+		return Rename(TEXT("RuntimeHashExternalStreamingObjectBase"), InGeneratedPackage, REN_DontCreateRedirectors);
+	}
+	return false;
 }
 
 void URuntimeHashExternalStreamingObjectBase::DumpStateLog(FHierarchicalLogArchive& Ar)
diff --git a/Engine/Source/Runtime/Engine/Private/WorldPartition/WorldPartitionRuntimeLevelStreamingCell.cpp b/Engine/Source/Runtime/Engine/Private/WorldPartition/WorldPartitionRuntimeLevelStreamingCell.cpp
--- a/Engine/Source/Runtime/Engine/Private/WorldPartition/WorldPartitionRuntimeLevelStreamingCell.cpp
+++ b/Engine/Source/Runtime/Engine/Private/WorldPartition/WorldPartitionRuntimeLevelStreamingCell.cpp
@@ -15,6 +15,10 @@
 #include "UObject/Package.h"
 #include "UObject/AssetRegistryTagsContext.h"
 
+#if WITH_EDITOR
+#include "WorldPartition/Cook/WorldPartitionCookPackageContextInterface.h"
+#endif
+
 #include UE_INLINE_GENERATED_CPP_BY_NAME(WorldPartitionRuntimeLevelStreamingCell)
 
 UWorldPartitionRuntimeLevelStreamingCell::UWorldPartitionRuntimeLevelStreamingCell(const FObjectInitializer& ObjectInitializer)
@@ -417,24 +421,27 @@ bool UWorldPartitionRuntimeLevelStreamingCell::OnPrepareGeneratorPackageForCook(
 }
 
 // Do all necessary work to prepare cell object for cook.
-bool UWorldPartitionRuntimeLevelStreamingCell::PrepareCellForCook(UPackage* InPackage)
+bool UWorldPartitionRuntimeLevelStreamingCell::PrepareCellForCook(const IWorldPartitionCookPackageContext& InCookContext, UPackage* InGeneratedPackage)
 {
 	// LevelStreaming could already be created
 	if (!LevelStreaming && GetActorCount() > 0)
 	{
-		if (!InPackage)
+		FString PackageName = InCookContext.GetGeneratedPackagePath(this);
+		check(!InGeneratedPackage || (PackageName == InGeneratedPackage->GetName()));
+		if (PackageName.IsEmpty())
 		{
 			return false;
 		}
-
-		LevelStreaming = CreateLevelStreaming(InPackage->GetName());
+		// Validation
+		check(PackageName.Contains(GetPackageNameToCreate()));
+		LevelStreaming = CreateLevelStreaming(PackageName);
 	}
 	return true;
 }
 
-bool UWorldPartitionRuntimeLevelStreamingCell::OnPopulateGeneratorPackageForCook(UPackage* InPackage)
+bool UWorldPartitionRuntimeLevelStreamingCell::OnPopulateGeneratorPackageForCook(const IWorldPartitionCookPackageContext& InCookContext, UPackage* InGeneratedPackage)
 {
-	return PrepareCellForCook(InPackage);
+	return PrepareCellForCook(InCookContext, InGeneratedPackage);
 }
 
 // Helper used by UWorldPartitionRuntimeLevelStreamingCell::OnPopulateGeneratedPackageForCook
@@ -473,7 +480,7 @@ private:
 	friend class UWorldPartitionRuntimeLevelStreamingCell;
 };
 
-bool UWorldPartitionRuntimeLevelStreamingCell::OnPopulateGeneratedPackageForCook(UPackage* InPackage, TArray<UPackage*>& OutModifiedPackages)
+bool UWorldPartitionRuntimeLevelStreamingCell::OnPopulateGeneratedPackageForCook(const IWorldPartitionCookPackageContext& InCookContext, UPackage* InPackage, TArray<UPackage*>& OutModifiedPackages)
 {
 	check(!IsAlwaysLoaded());
 	if (!InPackage)
@@ -484,7 +491,7 @@ bool UWorldPartitionRuntimeLevelStreamingCell::OnPopulateGeneratedPackageForCook
 	if (GetActorCount() > 0)
 	{
 		// When cook splitter doesn't use deferred populate, cell needs to be prepared here.
-		if (!PrepareCellForCook(InPackage))
+		if (!PrepareCellForCook(InCookContext, InPackage))
 		{
 			return false;
 		}
diff --git a/Engine/Source/Runtime/Engine/Public/WorldPartition/Cook/WorldPartitionCookPackageContextInterface.h b/Engine/Source/Runtime/Engine/Public/WorldPartition/Cook/WorldPartitionCookPackageContextInterface.h
--- a/Engine/Source/Runtime/Engine/Public/WorldPartition/Cook/WorldPartitionCookPackageContextInterface.h
+++ b/Engine/Source/Runtime/Engine/Public/WorldPartition/Cook/WorldPartitionCookPackageContextInterface.h
@@ -6,19 +6,24 @@
 #include "CoreMinimal.h"
 
 class IWorldPartitionCookPackageGenerator;
+class IWorldPartitionCookPackageObject;
 struct FWorldPartitionCookPackage;
 
 class IWorldPartitionCookPackageContext
 {
-
 public:
 	virtual ~IWorldPartitionCookPackageContext() {};
 
 	virtual void RegisterPackageCookPackageGenerator(IWorldPartitionCookPackageGenerator* Generator) = 0;
 	virtual void UnregisterPackageCookPackageGenerator(IWorldPartitionCookPackageGenerator* Generator) = 0;
 
-	virtual const FWorldPartitionCookPackage* AddLevelStreamingPackageToGenerate(IWorldPartitionCookPackageGenerator* Generator, const FString& Root, const FString& RelativePath) = 0;
-	virtual const FWorldPartitionCookPackage* AddGenericPackageToGenerate(IWorldPartitionCookPackageGenerator* Generator, const FString& Root, const FString& RelativePath) = 0;
+	UE_DEPRECATED(5.5, "Use AddPackageToGenerate instead.")
+	virtual const FWorldPartitionCookPackage* AddLevelStreamingPackageToGenerate(IWorldPartitionCookPackageGenerator* Generator, const FString& Root, const FString& RelativePath) { return nullptr; }
+	UE_DEPRECATED(5.5, "Use AddPackageToGenerate instead.")
+	virtual const FWorldPartitionCookPackage* AddGenericPackageToGenerate(IWorldPartitionCookPackageGenerator* Generator, const FString& Root, const FString& RelativePath) { return nullptr; }
+
+	virtual const FWorldPartitionCookPackage* AddPackageToGenerate(IWorldPartitionCookPackageGenerator* Generator, IWorldPartitionCookPackageObject* InCookPackageObject, const FString& Root, const FString& RelativePath) = 0;
+	virtual FString GetGeneratedPackagePath(IWorldPartitionCookPackageObject* InCookPackageObject) const = 0;
 
 	virtual bool GatherPackagesToCook() = 0;
 };
diff --git a/Engine/Source/Runtime/Engine/Public/WorldPartition/Cook/WorldPartitionCookPackageInterface.h b/Engine/Source/Runtime/Engine/Public/WorldPartition/Cook/WorldPartitionCookPackageInterface.h
--- a/Engine/Source/Runtime/Engine/Public/WorldPartition/Cook/WorldPartitionCookPackageInterface.h
+++ b/Engine/Source/Runtime/Engine/Public/WorldPartition/Cook/WorldPartitionCookPackageInterface.h
@@ -6,6 +6,7 @@
 #include "WorldPartitionCookPackageInterface.generated.h"
 
 class UExternalDataLayerAsset;
+class IWorldPartitionCookPackageContext;
 
 UINTERFACE(MinimalAPI)
 class UWorldPartitionCookPackageObject : public UInterface
@@ -23,7 +24,7 @@ public:
 	virtual const UExternalDataLayerAsset* GetExternalDataLayerAsset() const = 0;
 	virtual FString GetPackageNameToCreate() const = 0;
 	virtual bool OnPrepareGeneratorPackageForCook(TArray<UPackage*>& OutModifiedPackages) = 0;
-	virtual bool OnPopulateGeneratorPackageForCook(UPackage* InPackage) = 0;
-	virtual bool OnPopulateGeneratedPackageForCook(UPackage* InPackage, TArray<UPackage*>& OutModifiedPackages) = 0;
+	virtual bool OnPopulateGeneratorPackageForCook(const IWorldPartitionCookPackageContext& InCookContext, UPackage* InPackage) = 0;
+	virtual bool OnPopulateGeneratedPackageForCook(const IWorldPartitionCookPackageContext& InCookContext, UPackage* InPackage, TArray<UPackage*>& OutModifiedPackages) = 0;
 #endif
 };
\ No newline at end of file
diff --git a/Engine/Source/Runtime/Engine/Public/WorldPartition/WorldPartitionRuntimeCell.h b/Engine/Source/Runtime/Engine/Public/WorldPartition/WorldPartitionRuntimeCell.h
--- a/Engine/Source/Runtime/Engine/Public/WorldPartition/WorldPartitionRuntimeCell.h
+++ b/Engine/Source/Runtime/Engine/Public/WorldPartition/WorldPartitionRuntimeCell.h
@@ -321,19 +321,21 @@ class UWorldPartitionRuntimeCell : public UObject, public IWorldPartitionCell, p
 	ENGINE_API virtual int32 GetActorCount() const PURE_VIRTUAL(UWorldPartitionRuntimeCell::GetActorCount, return 0;);
 
 	// Cook methods
-	virtual bool PrepareCellForCook(UPackage* InPackage) { return OnPopulateGeneratorPackageForCook(InPackage); }
+	UE_DEPRECATED(5.5, "Use version with IWorldPartitionCookPackageContext")
+	virtual bool PrepareCellForCook(UPackage* InPackage) { return false; }
+	ENGINE_API virtual bool PrepareCellForCook(const IWorldPartitionCookPackageContext& InCookContext, UPackage* InGeneratedPackage = nullptr) { return true; }
 	UE_DEPRECATED(5.4, "PopulateGeneratorPackageForCook is deprecated, it was replaced by OnPrepareGeneratorPackageForCook")
 	ENGINE_API virtual bool PopulateGeneratorPackageForCook(TArray<UPackage*>& OutModifiedPackages) { return OnPrepareGeneratorPackageForCook(OutModifiedPackages); }
 	UE_DEPRECATED(5.4, "PopulateGeneratedPackageForCook is deprecated, it was replaced by OnPopulateGeneratedPackageForCook")
-	ENGINE_API virtual bool PopulateGeneratedPackageForCook(UPackage* InPackage, TArray<UPackage*>& OutModifiedPackages) { return OnPopulateGeneratedPackageForCook(InPackage, OutModifiedPackages); }
+	ENGINE_API virtual bool PopulateGeneratedPackageForCook(UPackage* InPackage, TArray<UPackage*>& OutModifiedPackages) { return false; }
 
 	//~Begin IWorldPartitionCookPackageObject
 	ENGINE_API virtual bool IsLevelPackage() const override { return true; }
 	ENGINE_API virtual const UExternalDataLayerAsset* GetExternalDataLayerAsset() const override { return ExternalDataLayerAsset; }
 	ENGINE_API virtual FString GetPackageNameToCreate() const { return FString(); }
 	ENGINE_API virtual bool OnPrepareGeneratorPackageForCook(TArray<UPackage*>& OutModifiedPackages) override { return true; }
-	ENGINE_API virtual bool OnPopulateGeneratorPackageForCook(UPackage* InPackage) override { return true; }
-	ENGINE_API virtual bool OnPopulateGeneratedPackageForCook(UPackage* InPackage, TArray<UPackage*>& OutModifiedPackages) override { return true; }
+	ENGINE_API virtual bool OnPopulateGeneratorPackageForCook(const IWorldPartitionCookPackageContext& InCookContext, UPackage* InPackage) override { return true; }
+	ENGINE_API virtual bool OnPopulateGeneratedPackageForCook(const IWorldPartitionCookPackageContext& InCookContext, UPackage* InPackage, TArray<UPackage*>& OutModifiedPackages) override { return true; }
 	//~End IWorldPartitionCookPackageObject
 
 	ENGINE_API virtual void DumpStateLog(FHierarchicalLogArchive& Ar) const;
diff --git a/Engine/Source/Runtime/Engine/Public/WorldPartition/WorldPartitionRuntimeHash.h b/Engine/Source/Runtime/Engine/Public/WorldPartition/WorldPartitionRuntimeHash.h
--- a/Engine/Source/Runtime/Engine/Public/WorldPartition/WorldPartitionRuntimeHash.h
+++ b/Engine/Source/Runtime/Engine/Public/WorldPartition/WorldPartitionRuntimeHash.h
@@ -88,13 +88,14 @@ public:
 	ENGINE_API virtual const UExternalDataLayerAsset* GetExternalDataLayerAsset() const override { return ExternalDataLayerAsset; }
 	ENGINE_API virtual FString GetPackageNameToCreate() const override;
 	ENGINE_API virtual bool OnPrepareGeneratorPackageForCook(TArray<UPackage*>& OutModifiedPackages) override { return true; }
-	ENGINE_API virtual bool OnPopulateGeneratorPackageForCook(UPackage* InPackage) override;
-	ENGINE_API virtual bool OnPopulateGeneratedPackageForCook(UPackage* InPackage, TArray<UPackage*>& OutModifiedPackages) override;
+	ENGINE_API virtual bool OnPopulateGeneratorPackageForCook(const IWorldPartitionCookPackageContext& InCookContext, UPackage* InPackage) override;
+	ENGINE_API virtual bool OnPopulateGeneratedPackageForCook(const IWorldPartitionCookPackageContext& InCookContext, UPackage* InPackage, TArray<UPackage*>& OutModifiedPackages) override;
 	//~End IWorldPartitionCookPackageObject interface
 
 protected:
 	virtual void DumpStateLog(FHierarchicalLogArchive& Ar);
 	UWorldPartitionRuntimeCell* GetCellForCookPackage(const FString& InCookPackageName) const;
+	bool PrepareForCook(const IWorldPartitionCookPackageContext& InCookContext);
 #endif
 
 public:
diff --git a/Engine/Source/Runtime/Engine/Public/WorldPartition/WorldPartitionRuntimeLevelStreamingCell.h b/Engine/Source/Runtime/Engine/Public/WorldPartition/WorldPartitionRuntimeLevelStreamingCell.h
--- a/Engine/Source/Runtime/Engine/Public/WorldPartition/WorldPartitionRuntimeLevelStreamingCell.h
+++ b/Engine/Source/Runtime/Engine/Public/WorldPartition/WorldPartitionRuntimeLevelStreamingCell.h
@@ -51,14 +51,14 @@ class UWorldPartitionRuntimeLevelStreamingCell : public UWorldPartitionRuntimeCe
 	ENGINE_API virtual void Fixup() override;
 	ENGINE_API virtual int32 GetActorCount() const override;
 	ENGINE_API virtual void DumpStateLog(FHierarchicalLogArchive& Ar) const override;
-	ENGINE_API virtual bool PrepareCellForCook(UPackage* InPackage) override;
+	ENGINE_API virtual bool PrepareCellForCook(const IWorldPartitionCookPackageContext& InCookContext, UPackage* InGeneratedPackage) override;
 	//~End UWorldPartitionRuntimeCell interface
 
 	//~Begin IWorldPartitionCookPackageObject interface
 	ENGINE_API virtual FString GetPackageNameToCreate() const override;
 	ENGINE_API virtual bool OnPrepareGeneratorPackageForCook(TArray<UPackage*>& OutModifiedPackages) override;
-	ENGINE_API virtual bool OnPopulateGeneratorPackageForCook(UPackage* InPackage) override;
-	ENGINE_API virtual bool OnPopulateGeneratedPackageForCook(UPackage* InPackage, TArray<UPackage*>& OutModifiedPackages) override;
+	ENGINE_API virtual bool OnPopulateGeneratorPackageForCook(const IWorldPartitionCookPackageContext& InCookContext, UPackage* InPackage) override;
+	ENGINE_API virtual bool OnPopulateGeneratedPackageForCook(const IWorldPartitionCookPackageContext& InCookContext, UPackage* InPackage, TArray<UPackage*>& OutModifiedPackages) override;
 	//~End IWorldPartitionCookPackageObject interface
 
 	//~Begin IWorldPartitionCell Interface
