Fix copy paste of entities which was relying on FParse::Token to keep the white spaces after the token. This change in behavior was introduced in cl 31140890 and should probably be fixed.

diff --git a/Engine/Source/Runtime/Engine/Private/Elements/Framework/TypedElementCommonActions.cpp b/Engine/Source/Runtime/Engine/Private/Elements/Framework/TypedElementCommonActions.cpp
--- a/Engine/Source/Runtime/Engine/Private/Elements/Framework/TypedElementCommonActions.cpp
+++ b/Engine/Source/Runtime/Engine/Private/Elements/Framework/TypedElementCommonActions.cpp
@@ -250,8 +250,14 @@ TArray<FTypedElementHandle> UTypedElementCommonActions::PasteNormalizedElements(
 				{
 					if (FCString::Strcmp(TypedElementChar, Buffer) == 0)
 					{
+						// Trim any spaces at the start of the line
+						while(PositionInLine[0] == TEXT(' '))
+						{
+							++PositionInLine;
+						}
+
 						FName TypeName;
-						if (!FParse::Value(++PositionInLine, TEXT("Type="), TypeName))
+						if (!FParse::Value(PositionInLine, TEXT("Type="), TypeName))
 						{
 							continue;
 						}
