Level Instance:
- Change LevelInstance flags in AActor to an enum : None, LevelInstance, LevelInstanceEdit, LevelInstancePropertyOverride
- This allows blocking actor deletes/replaces on PropertyOverride Edits
- Will allow selection of Level Instance level actors without being able to change their properties (CanEditChange now checks that actor is in editable Level Instance)

#rnx

diff --git a/Engine/Source/Runtime/Engine/Classes/GameFramework/Actor.h b/Engine/Source/Runtime/Engine/Classes/GameFramework/Actor.h
--- a/Engine/Source/Runtime/Engine/Classes/GameFramework/Actor.h
+++ b/Engine/Source/Runtime/Engine/Classes/GameFramework/Actor.h
@@ -92,6 +92,16 @@ enum class ESpawnActorScaleMethod : uint8
 	SelectDefaultAtRuntime					UMETA(Hidden),
 };
 
+/** Determines in what type of LevelInstance the actor is in, if any. */
+UENUM()
+enum class ELevelInstanceType : uint8
+{
+	None,
+	LevelInstance,
+	LevelInstanceEdit,
+	LevelInstancePropertyOverride
+};
+
 #if WITH_EDITORONLY_DATA
 /** Enum defining how actor will be placed in the partition */
 UENUM()
@@ -281,18 +291,14 @@ private:
 	uint8 bForceNetAddressable:1;
 
 #if WITH_EDITORONLY_DATA
-	/** Whether this actor belongs to a level instance which is currently being edited. */
+	/** Whether this actor belongs to a level instance or not and what type of level instance. */
 	UPROPERTY(Transient)
-	uint8 bIsInEditLevelInstance:1;
+	ELevelInstanceType LevelInstanceType;
 
 	/** Whether this actor belongs to a level instance in a level instance hierarchy currently being edited. Itself or its parent level instances. */
 	UPROPERTY(Transient)
 	uint8 bIsInEditLevelInstanceHierarchy:1;
 
-	/** Whether this actor belongs to a level instance  */
-	UPROPERTY(Transient)
-	uint8 bIsInLevelInstance:1;
-
 	friend struct FSetActorIsInLevelInstance;
 
 public:
@@ -325,13 +331,25 @@ public:
 	UE_DEPRECATED(5.4, "Call IsInEditLevelInstanceHierarchy/IsInEditLevelInstance instead.")
 	virtual bool IsInEditingLevelInstance() const
 	{
-		return bIsInEditLevelInstance;
+		return IsInEditLevelInstance();
+	}
+
+	/** If true, the actor belongs to a level instance which is currently in an edit mode: edit or property override edit */
+	bool IsInAnyEditLevelInstance() const
+	{
+		return IsInEditLevelInstance() || IsInPropertyOverrideLevelInstance();
 	}
 
 	/** If true, the actor belongs to a level instance which is currently being edited */
 	bool IsInEditLevelInstance() const
 	{
-		return bIsInEditLevelInstance;
+		return LevelInstanceType == ELevelInstanceType::LevelInstanceEdit;
+	}
+
+	/** If true, the actor belongs to a level instance which is currently in property override edit */
+	bool IsInPropertyOverrideLevelInstance() const
+	{
+		return LevelInstanceType == ELevelInstanceType::LevelInstancePropertyOverride;
 	}
 
 	/** If true, the actor belongs to a level instance which is currently being edited or a parent level instance being edited. */
@@ -343,7 +361,7 @@ public:
 	/** If true, the actor belongs to a level instance. */
 	bool IsInLevelInstance() const
 	{
-		return bIsInLevelInstance;
+		return LevelInstanceType != ELevelInstanceType::None;
 	}
 #endif
 
@@ -2303,7 +2321,7 @@ public:
 	//~ End UObject Interface
 
 #if WITH_EDITOR
-	virtual bool CanEditChangeComponent(const UActorComponent* Component, const FProperty* InProperty) const { return true; }
+	ENGINE_API virtual bool CanEditChangeComponent(const UActorComponent* Component, const FProperty* InProperty) const;
 #endif
 
 	//~=============================================================================
@@ -4560,10 +4578,9 @@ private:
 struct FSetActorIsInLevelInstance
 {
 private:
-	FSetActorIsInLevelInstance(AActor* InActor, bool bIsEditing = false)
+	FSetActorIsInLevelInstance(AActor* InActor, ELevelInstanceType InLevelInstanceType)
 	{
-		InActor->bIsInLevelInstance = true;
-		InActor->bIsInEditLevelInstance = bIsEditing;
+		InActor->LevelInstanceType = InLevelInstanceType;
 	}
 
 	friend class ULevelStreamingLevelInstance;
diff --git a/Engine/Source/Runtime/Engine/Private/Actor.cpp b/Engine/Source/Runtime/Engine/Private/Actor.cpp
--- a/Engine/Source/Runtime/Engine/Private/Actor.cpp
+++ b/Engine/Source/Runtime/Engine/Private/Actor.cpp
@@ -130,6 +130,13 @@ namespace ActorUtils
 			return false;
 		}
 
+		// Actors in LevelInstances can't be deleted unless they are in a regular edit level instance (not property override)
+		if (InActor->IsInLevelInstance() && !InActor->IsInEditLevelInstance())
+		{
+			OutReason = LOCTEXT("CanDeleteOrReplace_Error_NonEditLevelInstance", "Actor is in a non edit level instance");
+			return false;
+		}
+
 		return true;
 	}
 }
@@ -209,8 +216,7 @@ void AActor::InitializeDefaults()
 
 #if WITH_EDITORONLY_DATA
 	bIsInEditLevelInstanceHierarchy = false;
-	bIsInEditLevelInstance = false;
-	bIsInLevelInstance = false;
+	LevelInstanceType = ELevelInstanceType::None;
 	PivotOffset = FVector::ZeroVector;
 #endif
 	SpawnCollisionHandlingMethod = ESpawnActorCollisionHandlingMethod::AlwaysSpawn;
diff --git a/Engine/Source/Runtime/Engine/Private/ActorEditor.cpp b/Engine/Source/Runtime/Engine/Private/ActorEditor.cpp
--- a/Engine/Source/Runtime/Engine/Private/ActorEditor.cpp
+++ b/Engine/Source/Runtime/Engine/Private/ActorEditor.cpp
@@ -78,6 +78,17 @@ void AActor::PreEditChange(FProperty* PropertyThatWillChange)
 	}
 }
 
+bool AActor::CanEditChangeComponent(const UActorComponent* Component, const FProperty* InProperty) const
+{
+	// Actors in LevelInstances can't be edited unless they are in any kind of editing level instance (edit or property override)
+	if (IsInLevelInstance() && !IsInAnyEditLevelInstance())
+	{
+		return false;
+	}
+
+	return true;
+}
+
 bool AActor::CanEditChange(const FProperty* PropertyThatWillChange) const
 {
 	if ((PropertyThatWillChange->GetFName() == GET_MEMBER_NAME_CHECKED(AActor, Layers)) ||
@@ -126,6 +137,12 @@ bool AActor::CanEditChange(const FProperty* PropertyThatWillChange) const
 		return false;
 	}
 
+	// Actors in LevelInstances can't be edited unless they are in any kind of editing level instance (edit or property override)
+	if (IsInLevelInstance() && !IsInAnyEditLevelInstance())
+	{
+		return false;
+	}
+
 	return Super::CanEditChange(PropertyThatWillChange);
 }
 
diff --git a/Engine/Source/Runtime/Engine/Private/LevelInstance/LevelInstanceActorImpl.cpp b/Engine/Source/Runtime/Engine/Private/LevelInstance/LevelInstanceActorImpl.cpp
--- a/Engine/Source/Runtime/Engine/Private/LevelInstance/LevelInstanceActorImpl.cpp
+++ b/Engine/Source/Runtime/Engine/Private/LevelInstance/LevelInstanceActorImpl.cpp
@@ -316,7 +316,7 @@ void FLevelInstanceActorImpl::EditorGetUnderlyingActors(TSet<AActor*>& OutUnderl
 bool FLevelInstanceActorImpl::IsLockedActor() const
 {
 	AActor* LevelInstanceActor = CastChecked<AActor>(LevelInstance);
-	if (LevelInstanceActor->IsInLevelInstance() && !LevelInstanceActor->IsInEditLevelInstance())
+	if (LevelInstanceActor->IsInLevelInstance() && !LevelInstanceActor->IsInAnyEditLevelInstance())
 	{
 		return true;
 	}
diff --git a/Engine/Source/Runtime/Engine/Private/LevelInstance/LevelInstanceEditorLevelStreaming.cpp b/Engine/Source/Runtime/Engine/Private/LevelInstance/LevelInstanceEditorLevelStreaming.cpp
--- a/Engine/Source/Runtime/Engine/Private/LevelInstance/LevelInstanceEditorLevelStreaming.cpp
+++ b/Engine/Source/Runtime/Engine/Private/LevelInstance/LevelInstanceEditorLevelStreaming.cpp
@@ -162,9 +162,8 @@ void ULevelStreamingLevelInstanceEditor::OnLevelActorAdded(AActor* InActor)
 {
 	if (InActor && InActor->GetLevel() == LoadedLevel)
 	{
-		const bool bIsEditing = true;
-		FSetActorIsInLevelInstance SetIsInLevelInstance(InActor, bIsEditing);
-		InActor->PushLevelInstanceEditingStateToProxies(bIsEditing);
+		FSetActorIsInLevelInstance SetIsInLevelInstance(InActor, ELevelInstanceType::LevelInstanceEdit);
+		InActor->PushLevelInstanceEditingStateToProxies(true);
 	}
 }
 
diff --git a/Engine/Source/Runtime/Engine/Private/LevelInstance/LevelInstanceEditorPropertyOverrideLevelStreaming.cpp b/Engine/Source/Runtime/Engine/Private/LevelInstance/LevelInstanceEditorPropertyOverrideLevelStreaming.cpp
--- a/Engine/Source/Runtime/Engine/Private/LevelInstance/LevelInstanceEditorPropertyOverrideLevelStreaming.cpp
+++ b/Engine/Source/Runtime/Engine/Private/LevelInstance/LevelInstanceEditorPropertyOverrideLevelStreaming.cpp
@@ -336,9 +336,8 @@ void ULevelStreamingLevelInstanceEditorPropertyOverride::OnCurrentStateChanged(E
 		{
 			if (IsValid(Actor))
 			{
-				const bool bIsEditing = true;
-				FSetActorIsInLevelInstance SetIsInLevelInstance(Actor, bIsEditing);
-				Actor->PushLevelInstanceEditingStateToProxies(bIsEditing);
+				FSetActorIsInLevelInstance SetIsInLevelInstance(Actor, ELevelInstanceType::LevelInstancePropertyOverride);
+				Actor->PushLevelInstanceEditingStateToProxies(true);
 			}
 		}
 	}
diff --git a/Engine/Source/Runtime/Engine/Private/LevelInstance/LevelInstanceLevelStreaming.cpp b/Engine/Source/Runtime/Engine/Private/LevelInstance/LevelInstanceLevelStreaming.cpp
--- a/Engine/Source/Runtime/Engine/Private/LevelInstance/LevelInstanceLevelStreaming.cpp
+++ b/Engine/Source/Runtime/Engine/Private/LevelInstance/LevelInstanceLevelStreaming.cpp
@@ -249,7 +249,7 @@ void ULevelStreamingLevelInstance::InitializeActors(const TArray<AActor*>& InAct
 					// Must happen before the actors are registered with the world, which is the case for this delegate.
 					FSetActorInstanceGuid SetActorInstanceGuid(Actor, ContainerID.GetActorGuid(Actor->GetActorGuid()));
 
-					FSetActorIsInLevelInstance SetIsInLevelInstance(Actor);
+					FSetActorIsInLevelInstance SetIsInLevelInstance(Actor, ELevelInstanceType::LevelInstance);
 				}
 			}
 		}
