- UCanvas, FCanvas and others updated to accept FStringView instead of const TCHAR* and const FString&. The following functions have been updated:
    - ClippedStrLen (Replaced const TCHAR* param with FStringView)
    - WrappedPrint (Replaced const TCHAR* param with FStringView)
    - DrawText (Replaced const FString& param with FStringView, added deprecated overload accepting ANSI strings except in UTF-8 mode)
    - MeasureStringInternal (Added overload to accept FStringView in addition to existing function accepting Data/Length params)
    - CanvasStringSize (Replaced const TCHAR* param with FStringView)
    - WrapString (Replaced const TCHAR* param with FStringView)
    - StrLen (Replaced const FString& param with FStringView, added deprecated overload accepting ANSI strings except in UTF-8 mode)
    - TextSize (Replaced const FString& param with FStringView, added deprecated overload accepting ANSI strings except in UTF-8 mode)
- Added FCanvasSimpleTextItem
    - This is a new class which sits between FCanvasTextItemBase and FCanvasTextItem, and contains most of the data and logic that was in FCanvasTextItem.
    - Most of the logic is now in non-virtual member functions accepting FStringView, which derived classes will pass.
- Added FCanvasTextStringViewItem
    - This is a new class which derives from FCanvasSimpleTextItem and stores a FStringView.
    - Allows stack-based strings to be drawn without allocating FString/FText.
    - Overrides call non-virtual FCanvasSimpleTextItem functions, passing the FStringView directly.
- Updated FCanvasTextItem
    - Moved most of the logic/data to FCanvasSimpleTextItem, but it keeps the FText member.
    - Overrides call non-virtual FCanvasSimpleTextItem functions, passing the FText's FString as a FStringView.
- FCanvas::DrawShadowedString and UCanvas::DrawText updated to draw via FCanvasTextStringViewItem

#tests Engine preflights passed. Launched EngineTest and ensured that various users of debug drawing functionality still works as expected (stat commands, ShowDebug, GameplayDebugger) including string wrapping functionality.

diff --git a/Engine/Source/Runtime/Engine/Classes/Engine/Canvas.h b/Engine/Source/Runtime/Engine/Classes/Engine/Canvas.h
--- a/Engine/Source/Runtime/Engine/Classes/Engine/Canvas.h
+++ b/Engine/Source/Runtime/Engine/Classes/Engine/Canvas.h
@@ -256,7 +256,6 @@ public:
 	 * @param V Vertical position of the upper left corner of the portion of the texture to be shown(texels).
 	 * @param UL The width of the portion of the texture to be drawn(texels).
 	 * @param VL The height of the portion of the texture to be drawn(texels).
-	 * @param ClipTile true to clip tile.
 	 * @param BlendMode Blending mode of texture.
 	 */
 	ENGINE_API void DrawTile(UTexture* Tex, float X, float Y, float XL, float YL, float U, float V, float UL, float VL, EBlendMode BlendMode=BLEND_Translucent);
@@ -271,7 +270,7 @@ public:
 	 * @param YL out Vertical length of string.
 	 * @param Text String to calculate for.
 	 */
-	static ENGINE_API void ClippedStrLen(const UFont* Font, float ScaleX, float ScaleY, int32& XL, int32& YL, const TCHAR* Text);
+	static ENGINE_API void ClippedStrLen(const UFont* Font, float ScaleX, float ScaleY, int32& XL, int32& YL, FStringView Text);
 
 	/**	
 	 * Calculate the size of a string built from a font, word wrapped to a specified region.
@@ -281,7 +280,7 @@ public:
 	/**
 	 * Compute size and optionally print text with word wrap.
 	 */
-	ENGINE_API int32 WrappedPrint(bool Draw, float X, float Y, int32& out_XL, int32& out_YL, const UFont* Font, float ScaleX, float ScaleY, bool bCenterTextX, bool bCenterTextY, const TCHAR* Text, const FFontRenderInfo& RenderInfo) ;
+	ENGINE_API int32 WrappedPrint(bool Draw, float X, float Y, int32& out_XL, int32& out_YL, const UFont* Font, float ScaleX, float ScaleY, bool bCenterTextX, bool bCenterTextY, FStringView Text, const FFontRenderInfo& RenderInfo) ;
 	
 	/**
 	 * Draws a string of text to the screen.
@@ -295,10 +294,16 @@ public:
 	 * @param RenderInfo Optional. The FontRenderInfo to use when drawing the text.
 	 * @return The Y extent of the rendered text.
 	 */
-	ENGINE_API float DrawText(const UFont* InFont, const FString& InText, float X, float Y, float XScale = 1.f, float YScale = 1.f, const FFontRenderInfo& RenderInfo = FFontRenderInfo());
-
+	ENGINE_API float DrawText(const UFont* InFont, FStringView InText, float X, float Y, float XScale = 1.f, float YScale = 1.f, const FFontRenderInfo& RenderInfo = FFontRenderInfo());
 	ENGINE_API float DrawText(const UFont* InFont, const FText& InText, float X, float Y, float XScale = 1.f, float YScale = 1.f, const FFontRenderInfo& RenderInfo = FFontRenderInfo());
 
+#if !PLATFORM_TCHAR_IS_UTF8CHAR
+
+	UE_DEPRECATED(5.5, "Passing an ANSI string to DrawText has been deprecated outside of UTF-8 mode. Please use the overload that takes a TCHAR string.")
+	ENGINE_API float DrawText(const UFont* InFont, FAnsiStringView InText, float X, float Y, float XScale = 1.f, float YScale = 1.f, const FFontRenderInfo& RenderInfo = FFontRenderInfo());
+
+#endif // !PLATFORM_TCHAR_IS_UTF8CHAR
+
 	enum ELastCharacterIndexFormat
 	{
 		// The last whole character before the horizontal offset
@@ -317,14 +322,16 @@ public:
 	 *							DrawYL:		[out] will be set to the height of the string
 	 *							DrawFont:	[in] specifies the font to use for retrieving the size of the characters in the string
 	 *							Scale:		[in] specifies the amount of scaling to apply to the string
-	 * @param	pText		the string to calculate the size for
-	 * @param	TextLength	the number of code units in pText
+	 * @param	Text		The string to calculate the size for
 	 * @param	StopAfterHorizontalOffset  Offset horizontally into the string to stop measuring characters after, in pixels (or INDEX_NONE)
 	 * @param	CharIndexFormat  Behavior to use for StopAfterHorizontalOffset
-	 * @param	OutCharacterIndex  The index of the last character processed (used with StopAfterHorizontalOffset)
+	 * @param	OutLastCharacterIndex  The index of the last character processed (used with StopAfterHorizontalOffset)
 	 *
 	 */
-	static ENGINE_API void MeasureStringInternal( FTextSizingParameters& Parameters, const TCHAR* const pText, const int32 TextLength, const int32 StopAfterHorizontalOffset, const ELastCharacterIndexFormat CharIndexFormat, int32& OutLastCharacterIndex );
+	static ENGINE_API void MeasureStringInternal(FTextSizingParameters& Parameters, FStringView Text, const int32 StopAfterHorizontalOffset, const ELastCharacterIndexFormat CharIndexFormat, int32& OutLastCharacterIndex);
+
+	UE_DEPRECATED(5.4, "Use MeasureStringInternal with FStringView instead.")
+	static ENGINE_API void MeasureStringInternal(FTextSizingParameters& Parameters, const TCHAR* const Text, const int32 TextLength, const int32 StopAfterHorizontalOffset, const ELastCharacterIndexFormat CharIndexFormat, int32& OutLastCharacterIndex);
 
 	/**
 	 * Calculates the size of the specified string.
@@ -334,9 +341,9 @@ public:
 	 *							DrawYL:		[out] will be set to the height of the string
 	 *							DrawFont:	[in] specifies the font to use for retrieving the size of the characters in the string
 	 *							Scale:		[in] specifies the amount of scaling to apply to the string
-	 * @param	pText		the string to calculate the size for
+	 * @param	Text		the string to calculate the size for
 	 */
-	static ENGINE_API void CanvasStringSize( FTextSizingParameters& Parameters, const TCHAR* pText );
+	static ENGINE_API void CanvasStringSize(FTextSizingParameters& Parameters, FStringView Text);
 
 	/**
 	 * Parses a single string into an array of strings that will fit inside the specified bounding region.
@@ -354,14 +361,14 @@ public:
 	 *							Scale:		[in] specifies the amount of scaling to apply to the string
 	 * @param	CurX			specifies the pixel location to begin the wrapping; usually equal to the X pos of the bounding region, unless wrapping is initiated
 	 *								in the middle of the bounding region (i.e. indentation)
-	 * @param	pText			the text that should be wrapped
+	 * @param	Text			the text that should be wrapped
 	 * @param	out_Lines		[out] will contain an array of strings which fit inside the bounding region specified.  Does
 	 *							not clear the array first.
 	 * @param	OutWrappedLineData An optional array to fill with the indices from the source string marking the begin and end points of the wrapped lines
 	 */
-	static ENGINE_API void WrapString( FCanvasWordWrapper& Wrapper, FTextSizingParameters& Parameters, const float InCurX, const TCHAR* const pText, TArray<FWrappedStringElement>& out_Lines, FCanvasWordWrapper::FWrappedLineData* const OutWrappedLineData = nullptr);
+	static ENGINE_API void WrapString(FCanvasWordWrapper& Wrapper, FTextSizingParameters& Parameters, const float InCurX, FStringView Text, TArray<FWrappedStringElement>& out_Lines, FCanvasWordWrapper::FWrappedLineData* const OutWrappedLineData = nullptr);
 
-	ENGINE_API void WrapString( FTextSizingParameters& Parameters, const float InCurX, const TCHAR* const pText, TArray<FWrappedStringElement>& out_Lines, FCanvasWordWrapper::FWrappedLineData* const OutWrappedLineData = nullptr);
+	ENGINE_API void WrapString(FTextSizingParameters& Parameters, const float InCurX, FStringView Text, TArray<FWrappedStringElement>& out_Lines, FCanvasWordWrapper::FWrappedLineData* const OutWrappedLineData = nullptr);
 
 	/**
 	 * Transforms a 3D world-space vector into 2D screen coordinates.
@@ -385,39 +392,66 @@ public:
 	 * Calculate the length of a string, taking text wrapping into account.
 	 *
 	 * @param InFont The Font use.
-	 * @param Text The string to calculate for.
+	 * @param InText The string to calculate for.
 	 * @param XL out Horizontal length of string.
 	 * @param YL out Vertical length of string.
 	 * @param bDPIAware If true measures text considering the current DPI scale factor of the canvas.  Defaults to false for backwards compatibility
-	 * @param Canvas Canvas state object
+	 * @param InCanvas Canvas state object
 	 */
-	static ENGINE_API void StrLen(const UFont* InFont, const FString& InText, float& XL, float& YL, bool bDPIAware, FCanvas* InCanvas);
-	static ENGINE_API void StrLen(const UFont* InFont, const FString& InText, double& XL, double& YL, bool bDPIAware, FCanvas* InCanvas);
+	static ENGINE_API void StrLen(const UFont* InFont, FStringView InText, float& XL, float& YL, bool bDPIAware, FCanvas* InCanvas);
+	static ENGINE_API void StrLen(const UFont* InFont, FStringView InText, double& XL, double& YL, bool bDPIAware, FCanvas* InCanvas);
+	
+#if !PLATFORM_TCHAR_IS_UTF8CHAR
+
+	UE_DEPRECATED(5.5, "Passing an ANSI string to StrLen has been deprecated outside of UTF-8 mode. Please use the overload that takes a TCHAR string.")
+	static ENGINE_API void StrLen(const UFont* InFont, FAnsiStringView InText, float& XL, float& YL, bool bDPIAware, FCanvas* InCanvas);
+	UE_DEPRECATED(5.5, "Passing an ANSI string to StrLen has been deprecated outside of UTF-8 mode. Please use the overload that takes a TCHAR string.")
+	static ENGINE_API void StrLen(const UFont* InFont, FAnsiStringView InText, double& XL, double& YL, bool bDPIAware, FCanvas* InCanvas);
+
+#endif // !PLATFORM_TCHAR_IS_UTF8CHAR
 
 	/**
 	 * Calculate the length of a string, taking text wrapping into account.
 	 *
 	 * @param InFont The Font use.
-	 * @param Text The string to calculate for.
+	 * @param InText The string to calculate for.
 	 * @param XL out Horizontal length of string.
 	 * @param YL out Vertical length of string.
 	 * @param bDPIAware If true measures text considering the current DPI scale factor of the canvas.  Defaults to false for backwards compatibility
 	 */
-	ENGINE_API void StrLen(const UFont* InFont, const FString& InText, float& XL, float& YL, bool bDPIAware = false);
-	ENGINE_API void StrLen(const UFont* InFont, const FString& InText, double& XL, double& YL, bool bDPIAware = false);
+	ENGINE_API void StrLen(const UFont* InFont, FStringView InText, float& XL, float& YL, bool bDPIAware = false);
+	ENGINE_API void StrLen(const UFont* InFont, FStringView InText, double& XL, double& YL, bool bDPIAware = false);
+	
+#if !PLATFORM_TCHAR_IS_UTF8CHAR
+
+	UE_DEPRECATED(5.5, "Passing an ANSI string to StrLen has been deprecated outside of UTF-8 mode. Please use the overload that takes a TCHAR string.")
+	ENGINE_API void StrLen(const UFont* InFont, FAnsiStringView InText, float& XL, float& YL, bool bDPIAware = false);
+	UE_DEPRECATED(5.5, "Passing an ANSI string to StrLen has been deprecated outside of UTF-8 mode. Please use the overload that takes a TCHAR string.")
+	ENGINE_API void StrLen(const UFont* InFont, FAnsiStringView InText, double& XL, double& YL, bool bDPIAware = false);
+
+#endif // !PLATFORM_TCHAR_IS_UTF8CHAR
 
 	/** 
 	 * Calculates the horizontal and vertical size of a given string. This is used for clipped text as it does not take wrapping into account.
 	 *
-	 * @param Font The font to use.
-	 * @param Text String to calculate for.
+	 * @param InFont The font to use.
+	 * @param InText String to calculate for.
 	 * @param XL out Horizontal length of string.
 	 * @param YL out Vertical length of string.
 	 * @param ScaleX Scale that the string is expected to draw at horizontally.
 	 * @param ScaleY Scale that the string is expected to draw at vertically.
 	 */
-	ENGINE_API void TextSize( const UFont* InFont, const FString& InText, float& XL, float& YL, float ScaleX=1.f, float ScaleY=1.f);
-	ENGINE_API void TextSize( const UFont* InFont, const FString& InText, double& XL, double& YL, double ScaleX=1.f, double ScaleY=1.f);
+	ENGINE_API void TextSize( const UFont* InFont, FStringView InText, float& XL, float& YL, float ScaleX=1.f, float ScaleY=1.f);
+	ENGINE_API void TextSize( const UFont* InFont, FStringView InText, double& XL, double& YL, double ScaleX=1.f, double ScaleY=1.f);
+	
+#if !PLATFORM_TCHAR_IS_UTF8CHAR
+
+	UE_DEPRECATED(5.5, "Passing an ANSI string to TextSize has been deprecated outside of UTF-8 mode. Please use the overload that takes a TCHAR string.")
+	ENGINE_API void TextSize( const UFont* InFont, FAnsiStringView InText, float& XL, float& YL, float ScaleX=1.f, float ScaleY=1.f);
+	UE_DEPRECATED(5.5, "Passing an ANSI string to TextSize has been deprecated outside of UTF-8 mode. Please use the overload that takes a TCHAR string.")
+	ENGINE_API void TextSize( const UFont* InFont, FAnsiStringView InText, double& XL, double& YL, double ScaleX=1.f, double ScaleY=1.f);
+
+#endif // !PLATFORM_TCHAR_IS_UTF8CHAR
 
 	/** Set DrawColor with a FLinearColor and optional opacity override */
 	ENGINE_API void SetLinearDrawColor(FLinearColor InColor, float OpacityOverride=-1.f);
@@ -583,6 +617,7 @@ public:
 	 * @param ScreenPosition			Screen space position to render the text.
 	 * @param ScreenSize				Screen space size to render the texture.
 	 * @param Thickness					How many pixels thick the box lines should be.
+	 * @param RenderColor				Color to tint the box.
 	 */
 	UFUNCTION(BlueprintCallable, Category=Canvas, meta=(DisplayName="Draw Box", ScriptName="DrawBox"))
 	ENGINE_API void K2_DrawBox(FVector2D ScreenPosition, FVector2D ScreenSize, float Thickness=1.0f, FLinearColor RenderColor=FLinearColor::White);
diff --git a/Engine/Source/Runtime/Engine/Private/UserInterface/Canvas.cpp b/Engine/Source/Runtime/Engine/Private/UserInterface/Canvas.cpp
--- a/Engine/Source/Runtime/Engine/Private/UserInterface/Canvas.cpp
+++ b/Engine/Source/Runtime/Engine/Private/UserInterface/Canvas.cpp
@@ -93,22 +93,22 @@ FTexture* GetTextureForCanvasItem(UTexture* RenderTexture)
 	}
 }
 
-void FCanvasWordWrapper::Execute(const TCHAR* const InString, const FTextSizingParameters& InParameters, TArray<FWrappedStringElement>& OutStrings, FWrappedLineData* const OutWrappedLineData)
+void FCanvasWordWrapper::Execute(FStringView InString, const FTextSizingParameters& InParameters, TArray<FWrappedStringElement>& OutStrings, FWrappedLineData* const OutWrappedLineData)
 {
 	SCOPE_CYCLE_COUNTER(STAT_Canvas_WordWrappingTime);
 
-	FWrappingState WrappingState(InString, FCString::Strlen(InString), InParameters, OutStrings, OutWrappedLineData);
+	FWrappingState WrappingState(InString, InParameters, OutStrings, OutWrappedLineData);
 	if (WrappingState.WrappedLineData)
 	{
 		WrappingState.WrappedLineData->Empty();
 	}
 
-	if (WrappingState.StringLength > 0)
+	if (!WrappingState.String.IsEmpty())
 	{
-		GraphemeBreakIterator->SetString(WrappingState.String, WrappingState.StringLength);
-		LineBreakIterator->SetString(WrappingState.String, WrappingState.StringLength);
+		GraphemeBreakIterator->SetStringRef(WrappingState.String);
+		LineBreakIterator->SetStringRef(WrappingState.String);
 
-		for(int32 i = 0; i < WrappingState.StringLength; ++i) // Sanity check: Doesn't seem valid to have more lines than code units.
+		for(int32 i = 0; i < WrappingState.String.Len(); ++i) // Sanity check: Doesn't seem valid to have more lines than code units.
 		{	
 			if( !ProcessLine(WrappingState) )
 			{
@@ -121,17 +121,18 @@ void FCanvasWordWrapper::Execute(const TCHAR* const InString, const FTextSizingP
 bool FCanvasWordWrapper::ProcessLine(FWrappingState& WrappingState)
 {
 	bool bHasAddedLine = false;
-	if(WrappingState.StartIndex < WrappingState.StringLength)
+	const int32 StringLength = WrappingState.String.Len();
+	if (WrappingState.StartIndex < StringLength)
 	{
-		int32 BreakIndex = FindFirstMandatoryBreakBetween(WrappingState, WrappingState.StringLength);
+		int32 BreakIndex = FindFirstMandatoryBreakBetween(WrappingState, StringLength);
 
 		int32 NextStartIndex;
-		if( BreakIndex == INDEX_NONE || !DoesSubstringFit(WrappingState, BreakIndex) )
+		if (BreakIndex == INDEX_NONE || !DoesSubstringFit(WrappingState, BreakIndex))
 		{
 			BreakIndex = INDEX_NONE;
 			int32 WrapIndex = FindIndexAtOrAfterWrapWidth(WrappingState);
 
-			if (WrapIndex == WrappingState.StringLength)
+			if (WrapIndex == StringLength)
 			{
 				BreakIndex = WrapIndex;
 			}
@@ -175,12 +176,12 @@ bool FCanvasWordWrapper::ProcessLine(FWrappingState& WrappingState)
 			bHasAddedLine = true;
 		}
 
-		while (NextStartIndex < WrappingState.StringLength && FText::IsWhitespace(WrappingState.String[NextStartIndex]))
+		while (NextStartIndex < StringLength && FText::IsWhitespace(WrappingState.String[NextStartIndex]))
 		{
 			++NextStartIndex;
 		}
 
-		if(WrappingState.WrappedLineData)
+		if (WrappingState.WrappedLineData)
 		{
 			WrappingState.WrappedLineData->Emplace(WrappingState.StartIndex, BreakIndex);
 		}
@@ -194,7 +195,7 @@ bool FCanvasWordWrapper::DoesSubstringFit(FWrappingState& WrappingState, const i
 {
 	FTextSizingParameters MeasureParameters(WrappingState.Parameters);
 	int32 Unused;
-	UCanvas::MeasureStringInternal( MeasureParameters, WrappingState.String + WrappingState.StartIndex, EndIndex - WrappingState.StartIndex, 0, UCanvas::ELastCharacterIndexFormat::Unused, Unused );
+	UCanvas::MeasureStringInternal( MeasureParameters, WrappingState.String.Mid(WrappingState.StartIndex, EndIndex - WrappingState.StartIndex), 0, UCanvas::ELastCharacterIndexFormat::Unused, Unused );
 	return MeasureParameters.DrawXL <= WrappingState.Parameters.DrawXL;
 }
 
@@ -202,16 +203,16 @@ int32 FCanvasWordWrapper::FindIndexAtOrAfterWrapWidth(FWrappingState& WrappingSt
 {
 	FTextSizingParameters MeasureParameters(WrappingState.Parameters);
 	int32 Return = INDEX_NONE;
-	UCanvas::MeasureStringInternal(MeasureParameters, WrappingState.String + WrappingState.StartIndex, WrappingState.StringLength - WrappingState.StartIndex, WrappingState.Parameters.DrawXL, UCanvas::ELastCharacterIndexFormat::CharacterAtOffset, Return);
+	UCanvas::MeasureStringInternal(MeasureParameters, WrappingState.String.RightChop(WrappingState.StartIndex), WrappingState.Parameters.DrawXL, UCanvas::ELastCharacterIndexFormat::CharacterAtOffset, Return);
 	return WrappingState.StartIndex + Return;
 }
 
 void FCanvasWordWrapper::AddLine(FWrappingState& WrappingState, const int32 EndIndex)
 {
 	FTextSizingParameters MeasureParameters(WrappingState.Parameters);
-	FString Substring(EndIndex - WrappingState.StartIndex, WrappingState.String + WrappingState.StartIndex);
-	FWrappedStringElement Element(*Substring, 0.0f, 0.0f);
-	UCanvas::CanvasStringSize(MeasureParameters, *Element.Value);
+	const FStringView Substring(WrappingState.String.Mid(WrappingState.StartIndex, EndIndex - WrappingState.StartIndex));
+	UCanvas::CanvasStringSize(MeasureParameters, Substring);
+	FWrappedStringElement Element(FString(Substring), 0.0f, 0.0f);
 	Element.LineExtent.X = MeasureParameters.DrawXL;
 	Element.LineExtent.Y = MeasureParameters.DrawYL;
 	WrappingState.Results.Add(Element);
@@ -225,7 +226,7 @@ int32 FCanvasWordWrapper::FindFirstMandatoryBreakBetween(FWrappingState& Wrappin
 		const TCHAR Previous = WrappingState.String[i - 1];
 		if( FChar::IsLinebreak(Previous) ) // Line break occurs *after* linebreak character.
 		{
-			const TCHAR* const Current = i < WrapIndex ? WrappingState.String + i : NULL;
+			const TCHAR* const Current = i < WrapIndex ? &WrappingState.String[i] : nullptr;
 			if(	Previous != FChar::CarriageReturn || !(Current && *Current == FChar::LineFeed) ) // Line break cannot occur within CR LF pair.
 			{
 				BreakIndex = i;
@@ -1068,10 +1069,10 @@ void FCanvas::DrawTile(double X, double Y, double SizeX, double SizeY, float U,
 	DrawItem(TileItem);
 }
 
-int32 FCanvas::DrawShadowedString(double StartX, double StartY, const TCHAR* Text, const UFont* Font, const FLinearColor& Color, const FLinearColor& ShadowColor)
+int32 FCanvas::DrawShadowedString(double StartX, double StartY, FStringView Text, const UFont* Font, const FLinearColor& Color, const FLinearColor& ShadowColor)
 {
 	const float Z = 1.0f;
-	FCanvasTextItem TextItem( FVector2D( StartX, StartY ), FText::FromString( Text ), Font, Color);
+	FCanvasTextStringViewItem TextItem( FVector2D( StartX, StartY ), Text, Font, Color);
 	// just render text in single pass for distance field drop shadow
 	if (Font && Font->ImportOptions.bUseDistanceFieldAlpha)
 	{	
@@ -1110,16 +1111,16 @@ int32 FCanvas::DrawShadowedText(double StartX, double StartY, const FText& Text,
 	return TextItem.DrawnSize.Y;	
 }
 
-void FCanvas::WrapString( FTextSizingParameters& Parameters, const float InCurX, const TCHAR* const pText, TArray<FWrappedStringElement>& out_Lines, FCanvasWordWrapper::FWrappedLineData* const OutWrappedLineData)
+void FCanvas::WrapString(FTextSizingParameters& Parameters, const float InCurX, FStringView Text, TArray<FWrappedStringElement>& out_Lines, FCanvasWordWrapper::FWrappedLineData* const OutWrappedLineData)
 {
 	if (!WordWrapper.IsValid())
 	{
 		WordWrapper = MakeShareable(new FCanvasWordWrapper());
 	}
-	UCanvas::WrapString( *WordWrapper, Parameters, InCurX, pText, out_Lines, OutWrappedLineData);
+	UCanvas::WrapString( *WordWrapper, Parameters, InCurX, Text, out_Lines, OutWrappedLineData);
 }
 
-ENGINE_API void StringSize(const UFont* Font,int32& XL,int32& YL,const TCHAR* Text)
+ENGINE_API void StringSize(const UFont* Font,int32& XL,int32& YL, FStringView Text)
 {
 	// this functionality has been moved to a static function in UIString
 	FTextSizingParameters Parameters(Font,1.f,1.f);
@@ -1147,7 +1148,7 @@ static void GetDefaultCharSize( const UFont* DrawFont, float& DefaultCharWidth,
 	}
 }
 
-void UCanvas::MeasureStringInternal( FTextSizingParameters& Parameters, const TCHAR* const pText, const int32 TextLength, const int32 StopAfterHorizontalOffset, const ELastCharacterIndexFormat CharIndexFormat, int32& OutLastCharacterIndex )
+void UCanvas::MeasureStringInternal(FTextSizingParameters& Parameters, FStringView Text, const int32 StopAfterHorizontalOffset, const ELastCharacterIndexFormat CharIndexFormat, int32& OutLastCharacterIndex)
 {
 	// initialize output so it always makes some sense
 	OutLastCharacterIndex = INDEX_NONE;
@@ -1155,7 +1156,7 @@ void UCanvas::MeasureStringInternal( FTextSizingParameters& Parameters, const TC
 	Parameters.DrawXL = 0.f;
 	Parameters.DrawYL = 0.f;
 
-	if( Parameters.DrawFont )
+	if (Parameters.DrawFont)
 	{
 		// get a default character width and height to be used for non-renderable characters
 		float DefaultCharWidth, DefaultCharHeight;
@@ -1167,38 +1168,35 @@ void UCanvas::MeasureStringInternal( FTextSizingParameters& Parameters, const TC
 		const float ScaleY = Parameters.Scaling.Y;
 
 		const float DefaultCharIncrement = Parameters.SpacingAdjust.X * ScaleX;
-		const float DefaultScaledHeight = DefaultCharHeight * ScaleY + Parameters.SpacingAdjust.Y * ScaleY;
-		const TCHAR* pCurrentPos;
-		const TCHAR* pPrevPos = nullptr;
-		for ( pCurrentPos = pText; *pCurrentPos && pCurrentPos < pText + TextLength; ++pCurrentPos )
+		const TCHAR* const BeginPos = Text.GetData();
+		const TCHAR* const EndPos = Text.GetData() + Text.Len();
+		const TCHAR* CurrentPos;
+		const TCHAR* PrevPos = nullptr;
+		for (CurrentPos = BeginPos; CurrentPos < EndPos && *CurrentPos; ++CurrentPos)
 		{
 			float CharWidth, CharHeight;
-			const TCHAR* const pNextPos = pCurrentPos + 1;
 
-			TCHAR Ch = *pCurrentPos;
+			const TCHAR Ch = *CurrentPos;
 			Parameters.DrawFont->GetCharSize(Ch, CharWidth, CharHeight);
-			if ( CharHeight == 0 && Ch == TEXT('\n') )
+			if (CharHeight == 0 && Ch == TEXT('\n'))
 			{
 				CharHeight = DefaultCharHeight;
 			}
 
 			float CharSpacing = DefaultCharIncrement;
-			if ( pPrevPos )
+			if (PrevPos)
 			{
-				CharSpacing += Parameters.DrawFont->GetCharKerning( *pPrevPos, Ch ) * ScaleX;
+				CharSpacing += Parameters.DrawFont->GetCharKerning(*PrevPos, Ch) * ScaleX;
 			}
 
 			CharWidth *= ScaleX;
 			CharHeight *= ScaleY;
 
-			// never add character spacing if the next character is whitespace
-			if ( !FChar::IsWhitespace(*pNextPos) )
+			// If we have another character that isn't whitespace, append the character spacing
+			const TCHAR* const NextPos = CurrentPos + 1;
+			if (NextPos < EndPos && *NextPos && !FChar::IsWhitespace(*NextPos))
 			{
-				// if we have another character, append the character spacing
-				if ( *pNextPos )
-				{
-					CharWidth += CharSpacing;
-				}
+				CharWidth += CharSpacing;
 			}
 
 			const float ScaledVertSpacing = Parameters.SpacingAdjust.Y * ScaleY;
@@ -1207,47 +1205,52 @@ void UCanvas::MeasureStringInternal( FTextSizingParameters& Parameters, const TC
 			Parameters.DrawYL = FMath::Max<float>(Parameters.DrawYL, CharHeight + ScaledVertSpacing );
 
 			// Were we asked to stop measuring after the specified horizontal offset in pixels?
-			if( StopAfterHorizontalOffset != INDEX_NONE )
+			if (StopAfterHorizontalOffset != INDEX_NONE)
 			{
-				if( CharIndexFormat == ELastCharacterIndexFormat::CharacterAtOffset )
+				if (CharIndexFormat == ELastCharacterIndexFormat::CharacterAtOffset)
 				{
 					// Round our test toward the character's center position
 					const float TotalCharWidth = CharWidth + Parameters.DrawFont->GetCharHorizontalOffset(Ch);
-					if( StopAfterHorizontalOffset < Parameters.DrawXL - TotalCharWidth / 2 )
+					if (StopAfterHorizontalOffset < Parameters.DrawXL - TotalCharWidth / 2)
 					{
 						// We've reached the stopping point, so bail
 						break;
 					}
 				}
-				else if( CharIndexFormat == ELastCharacterIndexFormat::LastWholeCharacterBeforeOffset )
+				else if (CharIndexFormat == ELastCharacterIndexFormat::LastWholeCharacterBeforeOffset)
 				{
-					if( StopAfterHorizontalOffset < Parameters.DrawXL - CharWidth )
+					if (StopAfterHorizontalOffset < Parameters.DrawXL - CharWidth)
 					{
-						--pCurrentPos;
+						--CurrentPos;
 						// We've reached the stopping point, so bail
 						break;
 					}
 				}
 			}
 
-			pPrevPos = pCurrentPos;
+			PrevPos = CurrentPos;
 		}
 
-		OutLastCharacterIndex = pCurrentPos - pText;
+		OutLastCharacterIndex = CurrentPos - Text.GetData();
 	}
 }
 
-void UCanvas::CanvasStringSize( FTextSizingParameters& Parameters, const TCHAR* const pText )
+void UCanvas::MeasureStringInternal(FTextSizingParameters& Parameters, const TCHAR* const Text, const int32 TextLength, const int32 StopAfterHorizontalOffset, const ELastCharacterIndexFormat CharIndexFormat, int32& OutLastCharacterIndex)
+{
+	MeasureStringInternal(Parameters, MakeStringView(Text, TextLength), StopAfterHorizontalOffset, CharIndexFormat, OutLastCharacterIndex);
+}
+
+void UCanvas::CanvasStringSize(FTextSizingParameters& Parameters, FStringView Text)
 {
 	int32 Unused;
-	MeasureStringInternal(Parameters, pText, FCString::Strlen(pText), 0, ELastCharacterIndexFormat::Unused, Unused);
+	MeasureStringInternal(Parameters, Text, 0, ELastCharacterIndexFormat::Unused, Unused);
 }
 
-void UCanvas::WrapString( FCanvasWordWrapper& Wrapper, FTextSizingParameters& Parameters, const float InCurX, const TCHAR* const pText, TArray<FWrappedStringElement>& out_Lines, FCanvasWordWrapper::FWrappedLineData* const OutWrappedLineData)
+void UCanvas::WrapString( FCanvasWordWrapper& Wrapper, FTextSizingParameters& Parameters, const float InCurX, FStringView Text, TArray<FWrappedStringElement>& out_Lines, FCanvasWordWrapper::FWrappedLineData* const OutWrappedLineData)
 {
 	if (Parameters.DrawFont)
 	{
-		Wrapper.Execute(pText, Parameters, out_Lines, OutWrappedLineData);
+		Wrapper.Execute(Text, Parameters, out_Lines, OutWrappedLineData);
 	}
 	else
 	{
@@ -1255,9 +1258,9 @@ void UCanvas::WrapString( FCanvasWordWrapper& Wrapper, FTextSizingParameters& Pa
 	}
 }
 
-void UCanvas::WrapString( FTextSizingParameters& Parameters, const float InCurX, const TCHAR* const pText, TArray<FWrappedStringElement>& out_Lines, FCanvasWordWrapper::FWrappedLineData* const OutWrappedLineData)
+void UCanvas::WrapString( FTextSizingParameters& Parameters, const float InCurX, FStringView Text, TArray<FWrappedStringElement>& out_Lines, FCanvasWordWrapper::FWrappedLineData* const OutWrappedLineData)
 {
-	Canvas->WrapString(Parameters, InCurX, pText, out_Lines, OutWrappedLineData);
+	Canvas->WrapString(Parameters, InCurX, Text, out_Lines, OutWrappedLineData);
 }
 
 /*-----------------------------------------------------------------------------
@@ -1591,7 +1594,7 @@ void UCanvas::DrawTile( UTexture* Tex, float X, float Y, float XL, float YL, flo
 }
 
 
-void UCanvas::ClippedStrLen( const UFont* Font, float ScaleX, float ScaleY, int32& XL, int32& YL, const TCHAR* Text )
+void UCanvas::ClippedStrLen(const UFont* Font, float ScaleX, float ScaleY, int32& XL, int32& YL, FStringView Text)
 {
 	XL = 0;
 	YL = 0;
@@ -1615,16 +1618,21 @@ void VARARGS UCanvas::WrappedStrLenf( const UFont* Font, float ScaleX, float Sca
 }
 
 float UCanvas::DrawText(const UFont* InFont, const FText& InText, float X, float Y, float XScale, float YScale, const FFontRenderInfo& RenderInfo)
+{
+	return DrawText(InFont, InText.ToString(), X, Y, XScale, YScale, RenderInfo);
+}
+
+float UCanvas::DrawText(const UFont* InFont, FStringView InText, float X, float Y, float XScale, float YScale, const FFontRenderInfo& RenderInfo)
 {
 	ensure(InFont);
 	int32		XL		= 0;
 	int32		YL		= 0; 
 	// need this call in any case to update YL and XL - one of them will be needed anyway
-	WrappedPrint(RenderInfo.bClipText == false, X, Y, XL, YL, InFont, XScale, YScale, bCenterX, bCenterY, *InText.ToString(), RenderInfo);
+	WrappedPrint(RenderInfo.bClipText == false, X, Y, XL, YL, InFont, XScale, YScale, bCenterX, bCenterY, InText, RenderInfo);
 
 	if (RenderInfo.bClipText)
 	{
-		FCanvasTextItem TextItem(FVector2D(FMath::TruncToFloat(OrgX + X), FMath::TruncToFloat(OrgY + Y)), InText, InFont, DrawColor);
+		FCanvasTextStringViewItem TextItem(FVector2D(FMath::TruncToFloat(OrgX + X), FMath::TruncToFloat(OrgY + Y)), InText, InFont, DrawColor);
 		TextItem.Scale = FVector2D( XScale, YScale ), 
 		TextItem.BlendMode = SE_BLEND_Translucent;
 		TextItem.FontRenderInfo = RenderInfo;
@@ -1634,12 +1642,16 @@ float UCanvas::DrawText(const UFont* InFont, const FText& InText, float X, float
 	return (float)YL;
 }
 
-float UCanvas::DrawText(const UFont* InFont, const FString& InText, float X, float Y, float XScale, float YScale, const FFontRenderInfo& RenderInfo)
+#if !PLATFORM_TCHAR_IS_UTF8CHAR
+
+float UCanvas::DrawText(const UFont* InFont, FAnsiStringView InText, float X, float Y, float XScale, float YScale, const FFontRenderInfo& RenderInfo)
 {
-	return DrawText(InFont, FText::FromString(InText), X, Y, XScale, YScale, RenderInfo);
+	return DrawText(InFont, StringCast<TCHAR>(InText.GetData(), InText.Len()), X, Y, XScale, YScale, RenderInfo);
 }
 
-int32 UCanvas::WrappedPrint(bool Draw, float X, float Y, int32& out_XL, int32& out_YL, const UFont* Font, float ScaleX, float ScaleY, bool bCenterTextX, bool bCenterTextY, const TCHAR* Text, const FFontRenderInfo& RenderInfo) 
+#endif // !PLATFORM_TCHAR_IS_UTF8CHAR
+
+int32 UCanvas::WrappedPrint(bool Draw, float X, float Y, int32& out_XL, int32& out_YL, const UFont* Font, float ScaleX, float ScaleY, bool bCenterTextX, bool bCenterTextY, FStringView Text, const FFontRenderInfo& RenderInfo)
 {
 	if (ClipX < 0 || ClipY < 0)
 	{
@@ -1698,7 +1710,7 @@ int32 UCanvas::WrappedPrint(bool Draw, float X, float Y, int32& out_XL, int32& o
 		{
 			int32 TempX;
 			int32 TempY;
-			ClippedStrLen(Font, ScaleX, ScaleY, TempX, TempY, *WrappedString.Value);
+			ClippedStrLen(Font, ScaleX, ScaleY, TempX, TempY, WrappedString.Value);
 			LineXL = TempX;
 		}
 		XL = FMath::Max<float>(XL, LineXL);
@@ -1711,17 +1723,30 @@ int32 UCanvas::WrappedPrint(bool Draw, float X, float Y, int32& out_XL, int32& o
 	return WrappedStrings.Num();
 }
 
-void UCanvas::StrLen(const UFont* InFont, const FString& InText, float& XL, float& YL, bool bDPIAware)
+void UCanvas::StrLen(const UFont* InFont, FStringView InText, float& XL, float& YL, bool bDPIAware)
 {
 	UCanvas::StrLen(InFont, InText, XL, YL, bDPIAware, Canvas);
 }
-void UCanvas::StrLen(const UFont* InFont, const FString& InText, double& XL, double& YL, bool bDPIAware)
+void UCanvas::StrLen(const UFont* InFont, FStringView InText, double& XL, double& YL, bool bDPIAware)
 {
 	UCanvas::StrLen(InFont, InText, XL, YL, bDPIAware, Canvas);
 }
 
+#if !PLATFORM_TCHAR_IS_UTF8CHAR
+
+void UCanvas::StrLen(const UFont* InFont, FAnsiStringView InText, float& XL, float& YL, bool bDPIAware)
+{
+	return StrLen(InFont, StringCast<TCHAR>(InText.GetData(), InText.Len()), XL, YL, bDPIAware);
+}
+
+void UCanvas::StrLen(const UFont* InFont, FAnsiStringView InText, double& XL, double& YL, bool bDPIAware)
+{
+	return StrLen(InFont, StringCast<TCHAR>(InText.GetData(), InText.Len()), XL, YL, bDPIAware);
+}
+
+#endif // !PLATFORM_TCHAR_IS_UTF8CHAR
 
-void UCanvas::StrLen(const UFont* InFont, const FString& InText, float& XL, float& YL, bool bDPIAware, FCanvas* InCanvas)
+void UCanvas::StrLen(const UFont* InFont, FStringView InText, float& XL, float& YL, bool bDPIAware, FCanvas* InCanvas)
 {
 	if (InFont == NULL)
 	{
@@ -1732,7 +1757,7 @@ void UCanvas::StrLen(const UFont* InFont, const FString& InText, float& XL, floa
 		if (InFont->FontCacheType == EFontCacheType::Offline || bDPIAware == false)
 		{
 			FTextSizingParameters Parameters(InFont, 1.0f, 1.0f);
-			UCanvas::CanvasStringSize(Parameters, *InText);
+			UCanvas::CanvasStringSize(Parameters, InText);
 			XL = Parameters.DrawXL;
 			YL = Parameters.DrawYL;
 		}
@@ -1747,7 +1772,7 @@ void UCanvas::StrLen(const UFont* InFont, const FString& InText, float& XL, floa
 	}
 }
 
-void UCanvas::StrLen(const UFont* InFont, const FString& InText, double& XL, double& YL, bool bDPIAware, FCanvas* InCanvas)
+void UCanvas::StrLen(const UFont* InFont, FStringView InText, double& XL, double& YL, bool bDPIAware, FCanvas* InCanvas)
 {
 	float XLF = static_cast<float>(XL);
 	float YLF = static_cast<float>(YL);
@@ -1756,7 +1781,21 @@ void UCanvas::StrLen(const UFont* InFont, const FString& InText, double& XL, dou
 	YL = YLF;
 }
 
-void UCanvas::TextSize(const UFont* InFont, const FString& InText, float& XL, float& YL, float ScaleX, float ScaleY)
+#if !PLATFORM_TCHAR_IS_UTF8CHAR
+
+void UCanvas::StrLen(const UFont* InFont, FAnsiStringView InText, float& XL, float& YL, bool bDPIAware, FCanvas* InCanvas)
+{
+	return StrLen(InFont, StringCast<TCHAR>(InText.GetData(), InText.Len()), XL, YL, bDPIAware, InCanvas);
+}
+
+void UCanvas::StrLen(const UFont* InFont, FAnsiStringView InText, double& XL, double& YL, bool bDPIAware, FCanvas* InCanvas)
+{
+	return StrLen(InFont, StringCast<TCHAR>(InText.GetData(), InText.Len()), XL, YL, bDPIAware, InCanvas);
+}
+
+#endif // !PLATFORM_TCHAR_IS_UTF8CHAR
+
+void UCanvas::TextSize(const UFont* InFont, FStringView InText, float& XL, float& YL, float ScaleX, float ScaleY)
 {
 	int32 XLi, YLi;
 
@@ -1766,12 +1805,12 @@ void UCanvas::TextSize(const UFont* InFont, const FString& InText, float& XL, fl
 		return;
 	}
 
-	ClippedStrLen(InFont, ScaleX, ScaleY, XLi, YLi, *InText);
+	ClippedStrLen(InFont, ScaleX, ScaleY, XLi, YLi, InText);
 
 	XL = XLi;
 	YL = YLi;
 }
-void UCanvas::TextSize(const UFont* InFont, const FString& InText, double& XL, double& YL, double ScaleX, double ScaleY)
+void UCanvas::TextSize(const UFont* InFont, FStringView InText, double& XL, double& YL, double ScaleX, double ScaleY)
 {
 	float XLF = static_cast<float>(XL);
 	float YLF = static_cast<float>(YL);
@@ -1780,6 +1819,20 @@ void UCanvas::TextSize(const UFont* InFont, const FString& InText, double& XL, d
 	YL = YLF;
 }
 
+#if !PLATFORM_TCHAR_IS_UTF8CHAR
+
+void UCanvas::TextSize(const UFont* InFont, FAnsiStringView InText, float& XL, float& YL, float ScaleX, float ScaleY)
+{
+	return TextSize(InFont, StringCast<TCHAR>(InText.GetData(), InText.Len()), XL, YL, ScaleX, ScaleY);
+}
+
+void UCanvas::TextSize(const UFont* InFont, FAnsiStringView InText, double& XL, double& YL, double ScaleX, double ScaleY)
+{
+	return TextSize(InFont, StringCast<TCHAR>(InText.GetData(), InText.Len()), XL, YL, ScaleX, ScaleY);
+}
+
+#endif // !PLATFORM_TCHAR_IS_UTF8CHAR
+
 FVector UCanvas::Project(FVector Location, bool bClampToZeroPlane) const 
 {
 	FPlane V(0,0,0,0);
diff --git a/Engine/Source/Runtime/Engine/Private/UserInterface/CanvasItem.cpp b/Engine/Source/Runtime/Engine/Private/UserInterface/CanvasItem.cpp
--- a/Engine/Source/Runtime/Engine/Private/UserInterface/CanvasItem.cpp
+++ b/Engine/Source/Runtime/Engine/Private/UserInterface/CanvasItem.cpp
@@ -919,21 +919,18 @@ void FCanvasTextItemBase::Draw( class FCanvas* InCanvas )
 	DrawStringInternal( InCanvas, DrawPos, DrawColor, TextEffects);
 }
 
-FCanvasTextItem::~FCanvasTextItem()
-{
-}
-
-EFontCacheType FCanvasTextItem::GetFontCacheType() const
+EFontCacheType FCanvasSimpleTextItem::GetFontCacheType() const
 {
 	return Font->FontCacheType;
 }
 
-bool FCanvasTextItem::HasValidText() const
+// Declaration is pure-virtual, but we provide an implementation to allow checking the validity of the font in derived classes.
+bool FCanvasSimpleTextItem::HasValidText() const
 {
-	return Font && !Text.IsEmpty();
+	return Font != nullptr;
 }
 
-ESimpleElementBlendMode FCanvasTextItem::GetTextBlendMode( const bool bHasShadow ) const
+ESimpleElementBlendMode FCanvasSimpleTextItem::GetTextBlendMode( const bool bHasShadow ) const
 {
 	ESimpleElementBlendMode BlendModeToUse = BlendMode;
 	if (Font->ImportOptions.bUseDistanceFieldAlpha)
@@ -958,7 +955,7 @@ ESimpleElementBlendMode FCanvasTextItem::GetTextBlendMode( const bool bHasShadow
 	return BlendModeToUse;
 }
 
-FVector2D FCanvasTextItem::GetTextSize(float DPIScale) const
+FVector2D FCanvasSimpleTextItem::GetTextSizeInternal(FStringView Text, float DPIScale) const
 {
 	FVector2D MeasuredTextSize = FVector2D::ZeroVector;
 	switch( GetFontCacheType() )
@@ -966,7 +963,7 @@ FVector2D FCanvasTextItem::GetTextSize(float DPIScale) const
 	case EFontCacheType::Offline:
 		{
 			FTextSizingParameters Parameters( Font, Scale.X ,Scale.Y );
-			UCanvas::CanvasStringSize( Parameters, *Text.ToString() );
+			UCanvas::CanvasStringSize( Parameters, Text);
 			MeasuredTextSize.X = Parameters.DrawXL;
 			MeasuredTextSize.Y = Parameters.DrawYL;
 		}
@@ -976,7 +973,7 @@ FVector2D FCanvasTextItem::GetTextSize(float DPIScale) const
 		{
 			const FSlateFontInfo LegacyFontInfo = (SlateFontInfo.IsSet()) ? SlateFontInfo.GetValue() : Font->GetLegacySlateFontInfo();
 			const TSharedRef<FSlateFontMeasure> FontMeasure = FSlateApplication::Get().GetRenderer()->GetFontMeasureService();
-			MeasuredTextSize = FontMeasure->Measure( Text, LegacyFontInfo, DPIScale) * Scale;
+			MeasuredTextSize = FontMeasure->Measure(Text, LegacyFontInfo, DPIScale) * Scale;
 		}
 		break;
 
@@ -986,16 +983,16 @@ FVector2D FCanvasTextItem::GetTextSize(float DPIScale) const
 	return MeasuredTextSize;
 }
 
-void FCanvasTextItem::DrawStringInternal(FCanvas* InCanvas, const FVector2D& DrawPos, const FLinearColor& InColor, TArrayView<FTextEffect> TextEffects)
+void FCanvasSimpleTextItem::DrawStringInternal(FCanvas* InCanvas, FStringView Text, const FVector2D& DrawPos, const FLinearColor& InColor, TArrayView<FTextEffect> TextEffects)
 {
 	switch(GetFontCacheType())
 	{
 	case EFontCacheType::Offline:
-		DrawStringInternal_OfflineCache(InCanvas, DrawPos, InColor, TextEffects);
+		DrawStringInternal_OfflineCache(InCanvas, Text, DrawPos, InColor, TextEffects);
 		break;
 
 	case EFontCacheType::Runtime:
-		DrawStringInternal_RuntimeCache(InCanvas, DrawPos, InColor, TextEffects);
+		DrawStringInternal_RuntimeCache(InCanvas, Text, DrawPos, InColor, TextEffects);
 		break;
 
 	default:
@@ -1003,15 +1000,14 @@ void FCanvasTextItem::DrawStringInternal(FCanvas* InCanvas, const FVector2D& Dra
 	}
 }
 
-void FCanvasTextItem::DrawStringInternal_OfflineCache(FCanvas* InCanvas, const FVector2D& DrawPos, const FLinearColor& InColor, TArrayView<FTextEffect> TextEffects)
+void FCanvasSimpleTextItem::DrawStringInternal_OfflineCache(FCanvas* InCanvas, FStringView Text, const FVector2D& DrawPos, const FLinearColor& InColor, TArrayView<FTextEffect> TextEffects)
 {
 	QUICK_SCOPE_CYCLE_COUNTER(DrawStringInternal_OfflineCache);
 
 	DrawnSize = FVector2D::ZeroVector;
 
 	// Nothing to do if no text
-	const FString& TextString = Text.ToString();
-	if( TextString.Len() == 0 )
+	if (Text.IsEmpty())
 	{
 		return;
 	}
@@ -1024,12 +1020,11 @@ void FCanvasTextItem::DrawStringInternal_OfflineCache(FCanvas* InCanvas, const F
 
 	const float CharIncrement = ( (float)Font->Kerning + HorizSpacingAdjust ) * Scale.X;
 
-	const TArray< TCHAR, FString::AllocatorType >& Chars = TextString.GetCharArray();
 	// Draw all characters in string.
-	const int32 TextLen = TextString.Len();
+	const int32 TextLen = Text.Len();
 	for( int32 i=0; i < TextLen; i++ )
 	{
-		int32 Ch = (int32)Font->RemapChar(Chars[i]);
+		int32 Ch = (int32)Font->RemapChar(Text[i]);
 
 		// Skip invalid characters.
 		if (!Font->Characters.IsValidIndex(Ch))
@@ -1045,7 +1040,7 @@ void FCanvasTextItem::DrawStringInternal_OfflineCache(FCanvas* InCanvas, const F
 			DrawnSize.Y = Font->GetMaxCharHeight() * Scale.Y;
 		}
 
-		if (FChar::IsLinebreak(Chars[i]))
+		if (FChar::IsLinebreak(Text[i]))
 		{
 			// Set current character offset to the beginning of next line.
 			CurrentPos.X = 0.0f;
@@ -1135,7 +1130,7 @@ void FCanvasTextItem::DrawStringInternal_OfflineCache(FCanvas* InCanvas, const F
 			AddTriangles(FVector2f::ZeroVector, InColor);
 
 			// if we have another non-whitespace character to render, add the font's kerning.
-			if ( Chars[i+1] && !FChar::IsWhitespace(Chars[i+1]) )
+			if (Text.IsValidIndex(i+1) && !FChar::IsWhitespace(Text[i+1]) )
 			{
 				SizeX += CharIncrement;
 			}
@@ -1152,7 +1147,7 @@ void FCanvasTextItem::DrawStringInternal_OfflineCache(FCanvas* InCanvas, const F
 	}
 }
 
-void FCanvasTextItem::DrawStringInternal_RuntimeCache(FCanvas* InCanvas, const FVector2D& DrawPos, const FLinearColor& InColor, TArrayView<FTextEffect> TextEffects)
+void FCanvasSimpleTextItem::DrawStringInternal_RuntimeCache(FCanvas* InCanvas, FStringView Text, const FVector2D& DrawPos, const FLinearColor& InColor, TArrayView<FTextEffect> TextEffects)
 {
 	QUICK_SCOPE_CYCLE_COUNTER(DrawStringInternal_RuntimeCache);
 
@@ -1161,8 +1156,7 @@ void FCanvasTextItem::DrawStringInternal_RuntimeCache(FCanvas* InCanvas, const F
 	DrawnSize = FVector2D::ZeroVector;
 
 	// Nothing to do if no text
-	const FString& TextString = Text.ToString();
-	if( TextString.Len() == 0 )
+	if (Text.IsEmpty())
 	{
 		return;
 	}
@@ -1198,10 +1192,10 @@ void FCanvasTextItem::DrawStringInternal_RuntimeCache(FCanvas* InCanvas, const F
 	float LineX = PosX;
 
 	TCHAR PreviousChar = 0;
-	const int32 TextLen = TextString.Len();
+	const int32 TextLen = Text.Len();
 	for (int32 CharIndex = 0; CharIndex < TextLen; ++CharIndex)
 	{
-		const TCHAR CurrentChar = TextString[CharIndex];
+		const TCHAR CurrentChar = Text[CharIndex];
 
 		if (DrawnSize.Y == 0)
 		{
@@ -1344,6 +1338,46 @@ void FCanvasTextItem::DrawStringInternal_RuntimeCache(FCanvas* InCanvas, const F
 	}
 }
 
+bool FCanvasTextItem::HasValidText() const
+{
+	return FCanvasSimpleTextItem::HasValidText() && !Text.IsEmpty();
+}
+
+FVector2D FCanvasTextItem::GetTextSize(float DPIScale) const
+{
+	return GetTextSizeInternal(Text.ToString(), DPIScale);
+}
+
+void FCanvasTextItem::DrawStringInternal(FCanvas* InCanvas, const FVector2D& DrawPos, const FLinearColor& DrawColor, TArrayView<FTextEffect> TextEffects)
+{
+	FCanvasSimpleTextItem::DrawStringInternal(InCanvas, Text.ToString(), DrawPos, DrawColor, TextEffects);
+}
+
+void FCanvasTextItem::DrawStringInternal_OfflineCache(FCanvas* InCanvas, const FVector2D& DrawPos, const FLinearColor& DrawColor, TArrayView<FTextEffect> TextEffects)
+{
+	FCanvasSimpleTextItem::DrawStringInternal_OfflineCache(InCanvas, Text.ToString(), DrawPos, DrawColor, TextEffects);
+}
+
+void FCanvasTextItem::DrawStringInternal_RuntimeCache(FCanvas* InCanvas, const FVector2D& DrawPos, const FLinearColor& DrawColor, TArrayView<FTextEffect> TextEffects)
+{
+	FCanvasSimpleTextItem::DrawStringInternal_RuntimeCache(InCanvas, Text.ToString(), DrawPos, DrawColor, TextEffects);
+}
+
+bool FCanvasTextStringViewItem::HasValidText() const
+{
+	return FCanvasSimpleTextItem::HasValidText() && !Text.IsEmpty();
+}
+
+FVector2D FCanvasTextStringViewItem::GetTextSize(float DPIScale) const
+{
+	return GetTextSizeInternal(Text, DPIScale);
+}
+
+void FCanvasTextStringViewItem::DrawStringInternal(FCanvas* InCanvas, const FVector2D& DrawPos, const FLinearColor& DrawColor, TArrayView<FTextEffect> TextEffects)
+{
+	return FCanvasSimpleTextItem::DrawStringInternal(InCanvas, Text, DrawPos, DrawColor, TextEffects);
+}
+
 bool FCanvasShapedTextItem::HasValidText() const
 {
 	return ShapedGlyphSequence.IsValid() && ShapedGlyphSequence->GetGlyphsToRender().Num();
diff --git a/Engine/Source/Runtime/Engine/Public/CanvasItem.h b/Engine/Source/Runtime/Engine/Public/CanvasItem.h
--- a/Engine/Source/Runtime/Engine/Public/CanvasItem.h
+++ b/Engine/Source/Runtime/Engine/Public/CanvasItem.h
@@ -433,47 +433,46 @@ protected:
 	FBatchedElements* BatchedElements;
 };
 
-/* Text item with misc optional items such as shadow, centering etc. */
-class FCanvasTextItem : public FCanvasTextItemBase
+/*
+ * Base item used for drawing simple text strings with misc optional items such as shadow, centering etc.
+ * Provides functions which operate on FStringView, which derived classes can use.
+ */
+class FCanvasSimpleTextItem : public FCanvasTextItemBase
 {
 public:
-	/** 	 
-	 * Text item
+
+	/**
+	 * Constructor.
 	 *
-	 * @param	InPosition		Draw position
-	 * @param	InText			String to draw
-	 * @param	InFont			Font to draw with
+	 * @param	InPosition		Draw position.
+	 * @param	InFont			Font to draw with.
+	 * @param	InColor			Color to draw the text with.
 	 */
-	FCanvasTextItem( const FVector2D& InPosition, const FText& InText, const UFont* InFont, const FLinearColor& InColor )
-		: FCanvasTextItemBase( InPosition, InColor )
-		, Text( InText )
-		, Font( InFont )
+	FCanvasSimpleTextItem(const FVector2D& InPosition, const UFont* InFont, const FLinearColor& InColor)
+		: FCanvasTextItemBase(InPosition, InColor)
+		, Font(InFont)
 		, SlateFontInfo()
 	{
 		BlendMode = SE_BLEND_Translucent;
 	}
 
-	/** 	 
-	 * Text item
+	/**
+	 * Constructor.
 	 *
-	 * @param	InPosition		Draw position
-	 * @param	InText			String to draw
-	 * @param	InFontInfo		Font info to draw with
+	 * @param	InPosition		Draw position.
+	 * @param	InFontInfo		Font info to draw with.
+	 * @param	InColor			Color to draw the text with.
 	 */
-	FCanvasTextItem( const FVector2D& InPosition, const FText& InText, const FSlateFontInfo& InFontInfo, const FLinearColor& InColor )
-		: FCanvasTextItemBase( InPosition, InColor )
-		, Text( InText )
-		, Font( Cast<const UFont>(InFontInfo.FontObject) )
-		, SlateFontInfo( InFontInfo )
+	FCanvasSimpleTextItem(const FVector2D& InPosition, const FSlateFontInfo& InFontInfo, const FLinearColor& InColor)
+		: FCanvasTextItemBase(InPosition, InColor)
+		, Font(Cast<const UFont>(InFontInfo.FontObject))
+		, SlateFontInfo(InFontInfo)
 	{
 		BlendMode = SE_BLEND_TranslucentAlphaOnly;
 	}
-	
-	ENGINE_API ~FCanvasTextItem();
 
-	/* The text to draw. */
-	FText Text;
-	
+	ENGINE_API virtual ~FCanvasSimpleTextItem() override = default;
+
 	/* Font to draw text with. */
 	const UFont* Font;
 
@@ -484,32 +483,147 @@ protected:
 	/** Get the type of font cache the UFont is using */
 	ENGINE_API EFontCacheType GetFontCacheType() const;
 
-	//~ FCanvasTextItemBase overrides
+	//~ Begin FCanvasTextItemBase interface
+	ENGINE_API virtual bool HasValidText() const override = 0;
+	ENGINE_API virtual ESimpleElementBlendMode GetTextBlendMode(const bool bHasShadow) const override;
+	virtual FVector2D GetTextSize(float DPIScale) const override = 0;
+	virtual void DrawStringInternal(FCanvas* InCanvas, const FVector2D& DrawPos, const FLinearColor& DrawColor, TArrayView<FTextEffect> Offsets) override = 0;
+	//~ End FCanvasTextItemBase interface
+
+	ENGINE_API FVector2D GetTextSizeInternal(FStringView Text, float DPIScale) const;
+	ENGINE_API void DrawStringInternal(FCanvas* InCanvas, FStringView Text, const FVector2D& DrawPos, const FLinearColor& DrawColor, TArrayView<FTextEffect> TextEffects);
+
+	/**
+	 * Internal string draw
+	 *
+	 * In a method to make it simpler to do effects like shadow, outline
+	 */
+	ENGINE_API void DrawStringInternal_OfflineCache(FCanvas* InCanvas, FStringView Text, const FVector2D& DrawPos, const FLinearColor& DrawColor, TArrayView<FTextEffect> TextEffects);
+	ENGINE_API void DrawStringInternal_RuntimeCache(FCanvas* InCanvas, FStringView Text, const FVector2D& DrawPos, const FLinearColor& DrawColor, TArrayView<FTextEffect> TextEffects);
+
+};
+
+/*
+ * Text item with misc optional items such as shadow, centering etc.
+ * This class operates on FText.
+ */
+class FCanvasTextItem : public FCanvasSimpleTextItem
+{
+public:
+
+	/** 	 
+	 * Constructor.
+	 *
+	 * @param	InPosition		Draw position.
+	 * @param	InText			String to draw.
+	 * @param	InFont			Font to draw with.
+	 * @param	InColor			Color to draw the text with.
+	 */
+	FCanvasTextItem(const FVector2D& InPosition, const FText& InText, const UFont* InFont, const FLinearColor& InColor)
+		: FCanvasSimpleTextItem(InPosition, InFont, InColor)
+		, Text(InText)
+	{
+	}
+
+	/** 	 
+	 * Constructor
+	 *
+	 * @param	InPosition		Draw position.
+	 * @param	InText			String to draw.
+	 * @param	InFontInfo		Font info to draw with.
+	 * @param	InColor			Color to draw the text with.
+	 */
+	FCanvasTextItem(const FVector2D& InPosition, const FText& InText, const FSlateFontInfo& InFontInfo, const FLinearColor& InColor)
+		: FCanvasSimpleTextItem(InPosition, InFontInfo, InColor)
+		, Text(InText)
+	{
+	}
+	
+	ENGINE_API virtual ~FCanvasTextItem() override = default;
+
+	/* The text to draw. */
+	FText Text;
+
+protected:
+
+	//~ Begin FCanvasTextItemBase interface
 	ENGINE_API virtual bool HasValidText() const override;
-	ENGINE_API virtual ESimpleElementBlendMode GetTextBlendMode( const bool bHasShadow ) const override;
 	ENGINE_API virtual FVector2D GetTextSize(float DPIScale) const override;
 	ENGINE_API virtual void DrawStringInternal(FCanvas* InCanvas, const FVector2D& DrawPos, const FLinearColor& DrawColor, TArrayView<FTextEffect> TextEffects) override;
-
-
+	//~ End FCanvasTextItemBase interface
+	
 	/** 
-	 * Internal string draw
+	 * Internal string draw.
 	 *
 	 * In a method to make it simpler to do effects like shadow, outline
 	 */
 	ENGINE_API void DrawStringInternal_OfflineCache(FCanvas* InCanvas, const FVector2D& DrawPos, const FLinearColor& DrawColor, TArrayView<FTextEffect> TextEffects);
 	ENGINE_API void DrawStringInternal_RuntimeCache(FCanvas* InCanvas, const FVector2D& DrawPos, const FLinearColor& DrawColor, TArrayView<FTextEffect> TextEffects);
+
+};
+
+/*
+ * Text item with misc optional items such as shadow, centering etc.
+ * This class operates on FStringView - ensure that the underlying string is valid for
+ * the lifetime of the string view stored in this class.
+ */
+class FCanvasTextStringViewItem : public FCanvasSimpleTextItem
+{
+public:
+
+	/**
+	 * Constructor.
+	 *
+	 * @param	InPosition		Draw position.
+	 * @param	InText			String to draw.
+	 * @param	InFont			Font to draw with.
+	 * @param	InColor			Color to draw the text with.
+	 */
+	FCanvasTextStringViewItem(const FVector2D& InPosition, FStringView InText, const UFont* InFont, const FLinearColor& InColor)
+		: FCanvasSimpleTextItem(InPosition, InFont, InColor)
+		, Text(InText)
+	{
+	}
+
+	/**
+	 * Constructor.
+	 *
+	 * @param	InPosition		Draw position.
+	 * @param	InText			String to draw.
+	 * @param	InFontInfo		Font info to draw with.
+	 * @param	InColor			Color to draw the text with.
+	 */
+	FCanvasTextStringViewItem(const FVector2D& InPosition, FStringView InText, const FSlateFontInfo& InFontInfo, const FLinearColor& InColor)
+		: FCanvasSimpleTextItem(InPosition, InFontInfo, InColor)
+		, Text(InText)
+	{
+	}
+
+	ENGINE_API virtual ~FCanvasTextStringViewItem() override = default;
+
+	/* The text to draw. */
+	FStringView Text;
+
+protected:
+
+	//~ Begin FCanvasTextItemBase interface
+	ENGINE_API virtual bool HasValidText() const override;
+	ENGINE_API virtual FVector2D GetTextSize(float DPIScale) const override;
+	ENGINE_API virtual void DrawStringInternal(FCanvas* InCanvas, const FVector2D& DrawPos, const FLinearColor& DrawColor, TArrayView<FTextEffect> TextEffects) override;
+	//~ End FCanvasTextItemBase interface
+
 };
 
 /* Text item which can handle complex shaped text. */
 class FCanvasShapedTextItem : public FCanvasTextItemBase
 {
 public:
-	/** 	 
-	 * Text item
+	/**
+	 * Shaped text item for complex script rendering.
 	 *
-	 * @param	InPosition		Draw position
-	 * @param	InText			String to draw
-	 * @param	InFontInfo		Font info to draw with
+	 * @param InPosition             Draw position
+	 * @param InShapedGlyphSequence  Shaped text information to draw
+	 * @param InColor                Color to draw the text with
 	 */
 	FCanvasShapedTextItem( const FVector2D& InPosition, FShapedGlyphSequenceRef InShapedGlyphSequence, const FLinearColor& InColor )
 		: FCanvasTextItemBase( InPosition, InColor )
@@ -522,11 +636,14 @@ public:
 	FShapedGlyphSequencePtr ShapedGlyphSequence;
 
 protected:
-	//~ FCanvasTextItemBase overrides
+
+	//~ Begin FCanvasTextItemBase interface
 	ENGINE_API virtual bool HasValidText() const override;
 	ENGINE_API virtual ESimpleElementBlendMode GetTextBlendMode( const bool bHasShadow ) const override;
 	ENGINE_API virtual FVector2D GetTextSize(float DPIScale) const override;
 	ENGINE_API virtual void DrawStringInternal(FCanvas* InCanvas, const FVector2D& DrawPos, const FLinearColor& DrawColor, TArrayView<FTextEffect> TextEffects) override;
+	//~ End FCanvasTextItemBase interface
+
 };
 
 /* Line item. Note blend mode will be disregarded for these - only SE_BLEND_Opaque is currently supported. */
diff --git a/Engine/Source/Runtime/Engine/Public/CanvasTypes.h b/Engine/Source/Runtime/Engine/Public/CanvasTypes.h
--- a/Engine/Source/Runtime/Engine/Public/CanvasTypes.h
+++ b/Engine/Source/Runtime/Engine/Public/CanvasTypes.h
@@ -119,17 +119,18 @@ struct FWrappedStringElement
 	UPROPERTY()
 	FVector2D LineExtent;
 
-
 	FWrappedStringElement()
 		: LineExtent(ForceInit)
 	{
 	}
 
+	FWrappedStringElement(FString&& InValue, float Width, float Height )
+		: Value(MoveTemp(InValue)), LineExtent(Width,Height)
+	{}
 
-		/** Constructor */
-		FWrappedStringElement( const TCHAR* InValue, float Width, float Height )
-		: Value(InValue), LineExtent(Width,Height)
-		{}
+	FWrappedStringElement(const FString& InValue, float Width, float Height)
+		: Value(CopyTemp(InValue)), LineExtent(Width, Height)
+	{}
 	
 };
 
@@ -142,21 +143,18 @@ public:
 private:
 	struct FWrappingState
 	{
-		FWrappingState(	const TCHAR* const InString,
-						const int32 InStringLength,
+		FWrappingState(	FStringView InString,
 						const FTextSizingParameters& InParameters,
 						TArray<FWrappedStringElement>& InResults,
 						FWrappedLineData* const InWrappedLineData)
 			: String(InString)
-			, StringLength(InStringLength)
 			, Parameters(InParameters)
 			, StartIndex(0)
 			, Results(InResults)
 			, WrappedLineData(InWrappedLineData)
 		{}
 
-		const TCHAR* const String;
-		const int32 StringLength;
+		FStringView String;
 		const FTextSizingParameters& Parameters;
 		int32 StartIndex;
 		TArray<FWrappedStringElement>& Results;
@@ -174,7 +172,7 @@ public:
 	* @param InWrapWidth The width available.
 	* @param OutWrappedLineData An optional array to fill with the indices from the source string marking the begin and end points of the wrapped lines
 	*/
-	ENGINE_API void Execute(const TCHAR* const InString, const FTextSizingParameters& InParameters, TArray<FWrappedStringElement>& OutStrings, FWrappedLineData* const OutWrappedLineData);
+	ENGINE_API void Execute(FStringView InString, const FTextSizingParameters& InParameters, TArray<FWrappedStringElement>& OutStrings, FWrappedLineData* const OutWrappedLineData);
 
 private:
 	/**
@@ -747,11 +745,11 @@ public:
 	* @param ShadowColor - Shadow color to draw underneath the text (ignored for distance field fonts)
 	* @return total size in pixels of text drawn
 	*/
-	ENGINE_API int32 DrawShadowedString(double StartX, double StartY, const TCHAR* Text, const UFont* Font, const FLinearColor& Color, const FLinearColor& ShadowColor = FLinearColor::Black );
+	ENGINE_API int32 DrawShadowedString(double StartX, double StartY, FStringView Text, const UFont* Font, const FLinearColor& Color, const FLinearColor& ShadowColor = FLinearColor::Black );
 	
 	ENGINE_API int32 DrawShadowedText(double StartX, double StartY, const FText& Text, const UFont* Font, const FLinearColor& Color, const FLinearColor& ShadowColor = FLinearColor::Black );
 
-	ENGINE_API void WrapString( FTextSizingParameters& Parameters, const float InCurX, const TCHAR* const pText, TArray<FWrappedStringElement>& out_Lines, FCanvasWordWrapper::FWrappedLineData* const OutWrappedLineData = nullptr);
+	ENGINE_API void WrapString( FTextSizingParameters& Parameters, const float InCurX, FStringView Text, TArray<FWrappedStringElement>& out_Lines, FCanvasWordWrapper::FWrappedLineData* const OutWrappedLineData = nullptr);
 
 	ENGINE_API void DrawNGon(const FVector2D& Center, const FColor& Color, int32 NumSides, float Radius);
 
@@ -863,7 +861,7 @@ public:
 * @param YL - out height
 * @param Text - string of text to be measured
 */
-extern ENGINE_API void StringSize( const UFont* Font, int32& XL, int32& YL, const TCHAR* Text);
+extern ENGINE_API void StringSize( const UFont* Font, int32& XL, int32& YL, FStringView Text);
 
 /**
  * Helper class to write a line of texts on screen
