#ue5 - Improve CharacterMovementComponent ledge movement functionality. Fixed deflection attempt reapplying acceleration, often making it fail by moving off the ledge again. Now GetLedgeMove() attempts to compute the ledge normal based on the hit of the capsule with the previous valid floor, and deflects movement as if hitting an invisible barrier.

GetLedgeMove() and CheckLedgeDirection() had their "GravDir" parameter removed (in favor of using the CharacterMovement gravity direction), while adding the previous valid floor as a param.

#tests PIE while walking and crouching near the edge of various geometry, with "can walk of ledges" toggled as well as "use flat base for floor" toggled.

diff --git a/Engine/Source/Runtime/Engine/Classes/GameFramework/CharacterMovementComponent.h b/Engine/Source/Runtime/Engine/Classes/GameFramework/CharacterMovementComponent.h
--- a/Engine/Source/Runtime/Engine/Classes/GameFramework/CharacterMovementComponent.h
+++ b/Engine/Source/Runtime/Engine/Classes/GameFramework/CharacterMovementComponent.h
@@ -1654,7 +1654,7 @@ public:
 	/** Rotate a vector from world to gravity space. */
 	FVector RotateGravityToWorld(const FVector& World) const { return WorldToGravityTransform.RotateVector(World); }
 
-	/** Rotate a vector gravity to world space. */
+	/** Rotate a vector from gravity to world space. */
 	FVector RotateWorldToGravity(const FVector& Gravity) const { return GravityToWorldTransform.RotateVector(Gravity); }
 
 protected:
@@ -1755,13 +1755,13 @@ public:
 	ENGINE_API float GetCrouchedHalfHeight() const;
 
 	/** Returns true if there is a suitable floor SideStep from current position. */
-	ENGINE_API virtual bool CheckLedgeDirection(const FVector& OldLocation, const FVector& SideStep, const FVector& GravDir) const;
+	ENGINE_API virtual bool CheckLedgeDirection(const FVector& OldLocation, const FVector& SideStep, const FFindFloorResult& OldFloor) const;
 
 	/** 
 	 * @param Delta is the current move delta (which ended up going over a ledge).
 	 * @return new delta which moves along the ledge
 	 */
-	ENGINE_API virtual FVector GetLedgeMove(const FVector& OldLocation, const FVector& Delta, const FVector& GravDir) const;
+	ENGINE_API virtual FVector GetLedgeMove(const FVector& OldLocation, const FVector& Delta, const FFindFloorResult& OldFloor) const;
 
 	/** Check if pawn is falling */
 	ENGINE_API virtual bool CheckFall(const FFindFloorResult& OldFloor, const FHitResult& Hit, const FVector& Delta, const FVector& OldLocation, float remainingTime, float timeTick, int32 Iterations, bool bMustJump);
diff --git a/Engine/Source/Runtime/Engine/Private/Components/CharacterMovementComponent.cpp b/Engine/Source/Runtime/Engine/Private/Components/CharacterMovementComponent.cpp
--- a/Engine/Source/Runtime/Engine/Private/Components/CharacterMovementComponent.cpp
+++ b/Engine/Source/Runtime/Engine/Private/Components/CharacterMovementComponent.cpp
@@ -332,6 +332,20 @@ namespace CharacterMovementCVars
 		TEXT("Optimization - When enabled, defers CharacterMesh move propagation for all corrections until the end of larger scoped moves. Requires `bDeferCharacterMeshMovement=true'."),
 		ECVF_Default);
 
+	static bool bLedgeMovementDetectEdgeNormal = true;
+	FAutoConsoleVariableRef CVarLedgeMovementDetectEdgeNormal(
+		TEXT("p.LedgeMovement.DetectEdgeNormal"),
+		bLedgeMovementDetectEdgeNormal,
+		TEXT("Detect the normal of the ledge when avoiding walking off ledges, to try to find a better movement direction."),
+		ECVF_Default);
+
+	static bool bLedgeMovementApplyDirectMove = true;
+	FAutoConsoleVariableRef CVarLedgeMovementApplyDirectMove(
+		TEXT("p.LedgeMovement.ApplyDirectMove"),
+		bLedgeMovementApplyDirectMove,
+		TEXT("Apply the ledge movement vector directly, rather than the old method that reapplied acceleration."),
+		ECVF_Default);
+
 #if !UE_BUILD_SHIPPING
 
 	int32 NetShowCorrections = 0;
@@ -5035,7 +5049,8 @@ FVector UCharacterMovementComponent::LimitAirControl(float DeltaTime, const FVec
 	return Result;
 }
 
-bool UCharacterMovementComponent::CheckLedgeDirection(const FVector& OldLocation, const FVector& SideStep, const FVector& GravDir) const
+
+bool UCharacterMovementComponent::CheckLedgeDirection(const FVector& OldLocation, const FVector& SideStep, const FFindFloorResult& OldFloor) const
 {
 	const FVector SideDest = OldLocation + SideStep;
 	FCollisionQueryParams CapsuleParams(SCENE_QUERY_STAT(CheckLedgeDirection), false, CharacterOwner);
@@ -5050,7 +5065,7 @@ bool UCharacterMovementComponent::CheckLedgeDirection(const FVector& OldLocation
 	{
 		if ( !Result.bBlockingHit )
 		{
-			GetWorld()->SweepSingleByChannel(Result, SideDest, SideDest + GravDir * (MaxStepHeight + LedgeCheckThreshold), FQuat::Identity, CollisionChannel, CapsuleShape, CapsuleParams, ResponseParam);
+			GetWorld()->SweepSingleByChannel(Result, SideDest, SideDest + GetGravityDirection() * (MaxStepHeight + LedgeCheckThreshold), FQuat::Identity, CollisionChannel, CapsuleShape, CapsuleParams, ResponseParam);
 		}
 		if ( (Result.Time < 1.f) && IsWalkable(Result) )
 		{
@@ -5061,28 +5076,83 @@ bool UCharacterMovementComponent::CheckLedgeDirection(const FVector& OldLocation
 }
 
 
-FVector UCharacterMovementComponent::GetLedgeMove(const FVector& OldLocation, const FVector& Delta, const FVector& GravDir) const
+FVector UCharacterMovementComponent::GetLedgeMove(const FVector& OldLocation, const FVector& Delta, const FFindFloorResult& OldFloor) const
 {
 	if (!HasValidData() || Delta.IsZero())
 	{
 		return FVector::ZeroVector;
 	}
 
-	FVector SideDir(Delta.Y, -1.f * Delta.X, 0.f);
-		
-	// try left
-	if ( CheckLedgeDirection(OldLocation, SideDir, GravDir) )
+	// Try to base the parallel movement test on the collision with the bottom of the capsule,
+	// which should inform where the perpendicular surface of the ledge is.
+	bool bUseLedgeNormal = false;
+	FVector LedgeNormal = FVector::ZeroVector;
+	FVector LedgeMove = FVector::ZeroVector;
+
+	if (CharacterMovementCVars::bLedgeMovementDetectEdgeNormal)
 	{
-		return SideDir;
+		FFindFloorResult LedgeFloor = OldFloor;
+
+		// Flat base doesn't get capsule impact data for the floor in a meaningful way, so trace using a capsule at the last valid location.
+		if (bUseFlatBaseForFloorChecks)
+		{
+			UCharacterMovementComponent* MutableThis = const_cast<UCharacterMovementComponent*>(this);
+			MutableThis->bUseFlatBaseForFloorChecks = false;
+			FindFloor(OldLocation, LedgeFloor, false, NULL);
+			MutableThis->bUseFlatBaseForFloorChecks = true;
+		}
+
+		if (!LedgeFloor.IsWalkableFloor())
+		{
+			// Fall back to current floor, though this is likely unwalkable, but may contain a usable hit.
+			LedgeFloor = CurrentFloor;
+		}
+
+		// Check previous (valid) walkable floor location to get data about the ledge
+		if (LedgeFloor.bBlockingHit && RotateWorldToGravity(LedgeFloor.HitResult.Normal).Z > UE_KINDA_SMALL_NUMBER)
+		{
+			// It would be nice to use HitResult.Normal, but the quality of this normal doesn't seem as good as ImpactNormal (contains off-axis values for simple shapes).
+			LedgeNormal = (LedgeFloor.HitResult.ImpactNormal ^ LedgeFloor.HitResult.Normal);
+			LedgeNormal = (LedgeFloor.HitResult.ImpactNormal ^ LedgeNormal);
+			LedgeNormal = FVector::VectorPlaneProject(LedgeNormal, -GetGravityDirection()).GetSafeNormal();
+			LedgeMove = ComputeSlideVector(Delta, 1.0f, -LedgeNormal, LedgeFloor.HitResult);
+			bUseLedgeNormal = true;
+		}
 	}
 
-	// try right
-	SideDir *= -1.f;
-	if ( CheckLedgeDirection(OldLocation, SideDir, GravDir) )
+	// In case of no good estimate of ledge, fall back to movement direction.
+	if (LedgeNormal.IsZero())
 	{
-		return SideDir;
+		LedgeMove = Delta;
+		bUseLedgeNormal = false;
 	}
-	
+
+	if (bUseLedgeNormal)
+	{
+		bool bLedgeHit = CheckLedgeDirection(OldLocation, LedgeMove, OldFloor);
+		if (bLedgeHit)
+		{
+			return LedgeMove;
+		}
+	}
+	else
+	{
+		FVector SideDir(Delta.Y, -1.f * Delta.X, 0.f);
+
+		// try left
+		if ( CheckLedgeDirection(OldLocation, SideDir, OldFloor) )
+		{
+			return SideDir;
+		}
+
+		// try right
+		SideDir *= -1.f;
+		if ( CheckLedgeDirection(OldLocation, SideDir, OldFloor) )
+		{
+			return SideDir;
+		}
+	}
+
 	return FVector::ZeroVector;
 }
 
@@ -5422,7 +5492,8 @@ void UCharacterMovementComponent::PhysWalking(float deltaTime, int32 Iterations)
 		Acceleration = FVector::VectorPlaneProject(Acceleration, -GravityDirection);
 
 		// Apply acceleration
-		if( !HasAnimRootMotion() && !CurrentRootMotion.HasOverrideVelocity() )
+		const bool bSkipForLedgeMove = bTriedLedgeMove && CharacterMovementCVars::bLedgeMovementApplyDirectMove;
+		if( !HasAnimRootMotion() && !CurrentRootMotion.HasOverrideVelocity() && !bSkipForLedgeMove )
 		{
 			CalcVelocity(timeTick, GroundFriction, false, GetMaxBrakingDeceleration());
 			devCode(ensureMsgf(!Velocity.ContainsNaN(), TEXT("PhysWalking: Velocity contains NaN after CalcVelocity (%s)\n%s"), *GetPathNameSafe(this), *Velocity.ToString()));
@@ -5490,8 +5561,7 @@ void UCharacterMovementComponent::PhysWalking(float deltaTime, int32 Iterations)
 		if ( bCheckLedges && !CurrentFloor.IsWalkableFloor() )
 		{
 			// calculate possible alternate movement
-			const FVector GravDir = GravityDirection;
-			const FVector NewDelta = bTriedLedgeMove ? FVector::ZeroVector : GetLedgeMove(OldLocation, Delta, GravDir);
+			const FVector NewDelta = bTriedLedgeMove ? FVector::ZeroVector : GetLedgeMove(OldLocation, Delta, OldFloor);
 			if ( !NewDelta.IsZero() )
 			{
 				// first revert this move
@@ -5503,6 +5573,7 @@ void UCharacterMovementComponent::PhysWalking(float deltaTime, int32 Iterations)
 				// Try new movement direction
 				Velocity = NewDelta/timeTick;
 				remainingTime += timeTick;
+				Iterations--;
 				continue;
 			}
 			else
