Fixes issue where a character sliding down an unwalkable surface could detect landing early and get hung up, with its capsule sitting up on the unwalkable surface, rather than continuing to slide all the way down to the true walkable floor.

diff --git a/Engine/Source/Runtime/Engine/Private/Components/CharacterMovementComponent.cpp b/Engine/Source/Runtime/Engine/Private/Components/CharacterMovementComponent.cpp
--- a/Engine/Source/Runtime/Engine/Private/Components/CharacterMovementComponent.cpp
+++ b/Engine/Source/Runtime/Engine/Private/Components/CharacterMovementComponent.cpp
@@ -4847,7 +4847,9 @@ void UCharacterMovementComponent::PhysFalling(float deltaTime, int32 Iterations)
 					const FVector PawnLocation = UpdatedComponent->GetComponentLocation();
 					FFindFloorResult FloorResult;
 					FindFloor(PawnLocation, FloorResult, false);
-					if (FloorResult.IsWalkableFloor() && IsValidLandingSpot(PawnLocation, FloorResult.HitResult))
+
+					// Note that we only care about capsule sweep floor results, since the line trace may detect a lower walkable surface that our falling capsule wouldn't actually reach yet.
+					if (!FloorResult.bLineTrace && FloorResult.IsWalkableFloor() && IsValidLandingSpot(PawnLocation, FloorResult.HitResult))
 					{
 						remainingTime += subTimeTickRemaining;
 						ProcessLanded(FloorResult.HitResult, remainingTime, Iterations);
