TextureDerivedDataTask : FTextureSourceData : use UpdateChannelMinMaxFromIncomingTextureData so result is cached and code is not duplicated
also makes the LayerColorInfo thread safe

diff --git a/Engine/Plugins/Interchange/Runtime/Source/Pipelines/Private/InterchangeGenericTexturePipeline.cpp b/Engine/Plugins/Interchange/Runtime/Source/Pipelines/Private/InterchangeGenericTexturePipeline.cpp
--- a/Engine/Plugins/Interchange/Runtime/Source/Pipelines/Private/InterchangeGenericTexturePipeline.cpp
+++ b/Engine/Plugins/Interchange/Runtime/Source/Pipelines/Private/InterchangeGenericTexturePipeline.cpp
@@ -402,7 +402,7 @@ void UInterchangeGenericTexturePipeline::PostImportTextureAssetImport(UObject* C
 	bool bRunNormapMapDetection = !bIsAReimport && bDetectNormalMapTexture && !Texture->IsNormalMap();
 
 	 // we probably got the info via Init() - if we didn't it's because it's compressed. Here we can decompress, so do it if needed.
-	bool bRunChannelScan = Source.GetLayerColorInfo().Num() == 0;
+	bool bRunChannelScan = !Source.HasLayerColorInfo();
 
 	bool bNeedLockedMip = bRunChannelScan || bRunNormapMapDetection;
 	if (!bNeedLockedMip)
diff --git a/Engine/Source/Editor/TextureEditor/Private/TextureEditorToolkit.cpp b/Engine/Source/Editor/TextureEditor/Private/TextureEditorToolkit.cpp
--- a/Engine/Source/Editor/TextureEditor/Private/TextureEditorToolkit.cpp
+++ b/Engine/Source/Editor/TextureEditor/Private/TextureEditorToolkit.cpp
@@ -751,10 +751,14 @@ void FTextureEditorToolkit::PopulateQuickInfo( )
 		} // end if results metadata valid
 
 		if (Texture->Source.IsValid() &&
-			Texture->Source.GetLayerColorInfo().Num())
+			Texture->Source.HasLayerColorInfo())
 		{
-			// Make a 1x1 image with our max colors to use for alpha detection.
-			FImageView View(&Texture->Source.GetLayerColorInfo()[0].ColorMin, 1, 1);
+			// Make a 1x1 image with our min color to use for alpha detection.
+			
+			TArray<FTextureSourceLayerColorInfo> LayerColorInfo;
+			Texture->Source.GetLayerColorInfo(LayerColorInfo);
+
+			FImageView View(&LayerColorInfo[0].ColorMin, 1, 1);
 
 			bool bSourceAlphaDetected = FImageCore::DetectAlphaChannel(View);
 			SourceMipsAlphaDetectedText->SetText(FText::Format(NSLOCTEXT("TextureEditor", "QuickInfo_SourceAlphaDetected", "Source Alpha Detected: {0}"),
diff --git a/Engine/Source/Runtime/Engine/Classes/Engine/Texture.h b/Engine/Source/Runtime/Engine/Classes/Engine/Texture.h
--- a/Engine/Source/Runtime/Engine/Classes/Engine/Texture.h
+++ b/Engine/Source/Runtime/Engine/Classes/Engine/Texture.h
@@ -497,10 +497,21 @@ struct FTextureSource
 		ENGINE_API FSharedBuffer GetMipData(int32 BlockIndex, int32 LayerIndex, int32 MipIndex) const;
 		ENGINE_API FSharedBuffer GetMipDataWithInfo(int32 InBlockIndex, int32 InLayerIndex, int32 InMipIndex, FImageInfo& OutMipImageInfo) const;
 
+		// note: FImageView does not hold a ref on the FSharedBuffer memory it points at
+		//	you must keep the FMipData around or hold a ref on the FSharedBuffer elsewhere
+		inline FImageView GetMipDataImageView(int32 BlockIndex, int32 LayerIndex, int32 MipIndex) const
+		{
+			FImageInfo Info;
+			FSharedBuffer Buffer = GetMipDataWithInfo(BlockIndex,LayerIndex,MipIndex,Info);
+			return FImageView(Info,const_cast<void *>(Buffer.GetData()));
+		}
+
 		inline bool IsValid() const { return !MipData.IsNull(); }
 
 		inline void ResetData() { MipData.Reset(); }
 
+		inline FSharedBuffer GetData() const { return MipData; }
+
 	private:
 		// We only want to allow FTextureSource to create FMipData objects
 		friend struct FTextureSource;
@@ -612,15 +623,7 @@ private:
 
 	// Internal implementation for locking the mip data, called by LockMipReadOnly or LockMip.
 	FMutableMemoryView LockMipInternal(int32 BlockIndex, int32 LayerIndex, int32 MipIndex, ELockState RequestedLockState);
-	
-	// As per UpdateChannelLinearMinMax(), except acts on incoming new data rather than locking existing mips.
-	// InNewTextureData must be uncompressed
-	// UpdateChannelMinMaxFromIncomingTextureData does not use the BulkData or CompressionFormat on the TextureSource
-	//	but it does use the dimensions/blocks/etc. they must be set before calling this.
-	// returns true/false for success/failure.
-	// failure may occur if the size of InNewTextureData does not match the dimensions set in the Texturesource
-	bool UpdateChannelMinMaxFromIncomingTextureData(FMemoryView InNewTextureData);
-	
+		
 	/** Returns the source data fully decompressed */
 	// ImageWrapperModule is not used
 	FSharedBuffer Decompress(class IImageWrapperModule* ImageWrapperModule = nullptr) const;
@@ -662,23 +665,35 @@ private:
 	bool EnsureBlocksAreSorted();
 
 public:
-	// Runs FImageCore::ComputeChannelLinearMinMax on all blocks and layers (but only mip0), returns false
-	// if the source was unable to be locked and leaves the channel minmax as unknown. Compute just gets
-	// the values and leaves the source untouched.
+	// UpdateChannelLinearMinMax runs FImageCore::ComputeChannelLinearMinMax on all blocks and layers (but only mip0), returns false
+	//   if the source was unable to be locked and leaves the channel minmax as unknown. 
+	// Compute just gets the values and leaves the cached minax in the Source untouched.
 	ENGINE_API bool UpdateChannelLinearMinMax();
 	ENGINE_API bool ComputeChannelLinearMinMax(int32 InLayerIndex, FLinearColor& OutMinColor, FLinearColor& OutMaxColor) const;
-	ENGINE_API const TArray<FTextureSourceLayerColorInfo>& GetLayerColorInfo() const { return LayerColorInfo; }
+
+	ENGINE_API void GetLayerColorInfo(TArray<FTextureSourceLayerColorInfo> & OutLayerColorInfo) const;
+	ENGINE_API void SetLayerColorInfo(const TArray<FTextureSourceLayerColorInfo> & InLayerColorInfo);
+	ENGINE_API void ResetLayerColorInfo();
+	ENGINE_API bool HasLayerColorInfo() const;
+
+	// As per UpdateChannelLinearMinMax(), except acts on incoming new data rather than locking existing mips.
+	// InNewTextureData must be uncompressed
+	// UpdateChannelMinMaxFromIncomingTextureData does not use the BulkData or CompressionFormat on the TextureSource
+	//	but it does use the dimensions/blocks/etc. they must be set before calling this.
+	// returns true/false for success/failure.
+	// failure may occur if the size of InNewTextureData does not match the dimensions set in the Texturesource
+	ENGINE_API bool UpdateChannelMinMaxFromIncomingTextureData(FMemoryView InNewTextureData);
 
 	/** Uses a hash as the GUID, useful to prevent creating new GUIDs on load for legacy assets.
 	This is automatically done by Init() and Mip Lock/Unlock.  New textures should always have the data hash as Id. */
 	ENGINE_API void UseHashAsGuid();
 
-	void ReleaseSourceMemory(); // release the memory from the mips (does almost the same as remove source data except doesn't rebuild the guid)
+	ENGINE_API void ReleaseSourceMemory(); // release the memory from the mips (does almost the same as remove source data except doesn't rebuild the guid)
 	FORCEINLINE bool HasHadBulkDataCleared() const { return bHasHadBulkDataCleared; }
 private:
 	/** Used while cooking to clear out unneeded memory after compression */
 	bool bHasHadBulkDataCleared;
-#endif
+#endif // WITH_EDITOR
 
 #if WITH_EDITORONLY_DATA
 	/** GUID used to track changes to the source data.
@@ -736,9 +751,10 @@ private:
 	UPROPERTY(VisibleAnywhere, Category=TextureSource)
 	bool bGuidIsHash;
 
-	/** Per layer color info. If this is empty we don't have the data, otherwise count is == NumLayers. */
+	/** Per layer color info. If this is empty we don't have the data, otherwise count is == NumLayers.
+	Protected by BulkDataLock for thread safety.  Use Get/Set accessors which do the locking for you.*/
 	UPROPERTY(VisibleAnywhere, Category=TextureSource)
-	TArray<FTextureSourceLayerColorInfo> LayerColorInfo;
+	TArray<FTextureSourceLayerColorInfo> LayerColorInfo_LockProtected;
 
 	/** Format in which the source data is stored. */
 	UPROPERTY(VisibleAnywhere, Category=TextureSource)
diff --git a/Engine/Source/Runtime/Engine/Private/Texture.cpp b/Engine/Source/Runtime/Engine/Private/Texture.cpp
--- a/Engine/Source/Runtime/Engine/Private/Texture.cpp
+++ b/Engine/Source/Runtime/Engine/Private/Texture.cpp
@@ -2336,7 +2336,7 @@ void FTextureSource::RemoveCompression()
 
 		FSharedBuffer Buffer = Decompress();
 		
-		if (LayerColorInfo.Num() == 0 )
+		if ( ! HasLayerColorInfo() )
 		{
 			// since we're changing compression, go ahead and also update channel minmax now if not done
 			UpdateChannelMinMaxFromIncomingTextureData(Buffer.GetView());
@@ -2380,7 +2380,7 @@ void FTextureSource::Compress()
 			
 			FSharedBuffer Buffer = Decompress();
 		
-			if (LayerColorInfo.Num() == 0 )
+			if ( ! HasLayerColorInfo() )
 			{
 				// since we're changing compression, go ahead and also update channel minmax now if not done
 				UpdateChannelMinMaxFromIncomingTextureData(Buffer.GetView());
@@ -3269,7 +3269,7 @@ void FTextureSource::ImportCustomProperties(const TCHAR* SourceText, FFeedbackCo
 				// Data changed - we don't know the bounds anymore.
 				// This seems very suspicious - expected data size doesn't seem to be checked at all? When is this used? Is then input data
 				// compressed? If its uncompressed then we can run the color analysis on it..?
-				LayerColorInfo.Empty();
+				ResetLayerColorInfo();
 
 				BulkData.UpdatePayload(Buffer.MoveToShared(), Owner);
 			}
@@ -3353,7 +3353,7 @@ void FTextureSource::RemoveSourceData()
 	NumLockedMips = 0u;
 	LockState = ELockState::None;
 	
-	LayerColorInfo.Empty();
+	ResetLayerColorInfo();
 
 	BulkData.UnloadData();
 
@@ -4367,11 +4367,14 @@ bool UTexture::ComputeTextureSourceChannelMinMax(FLinearColor & OutColorMin, FLi
 	OutColorMax = FLinearColor(ForceInit);
 
 #if WITH_EDITORONLY_DATA
-	if (Source.LayerColorInfo.Num())
+	if (Source.HasLayerColorInfo())
 	{
+		TArray<FTextureSourceLayerColorInfo> LayerColorInfo;
+		Source.GetLayerColorInfo(LayerColorInfo);
+
 		// This function only operates on layer 1. 
-		OutColorMin = Source.LayerColorInfo[0].ColorMin;
-		OutColorMax = Source.LayerColorInfo[0].ColorMax;
+		OutColorMin = LayerColorInfo[0].ColorMin;
+		OutColorMax = LayerColorInfo[0].ColorMax;
 		return true;
 	}
 	else if (Source.ComputeChannelLinearMinMax(0 /* layer index */, OutColorMin, OutColorMax))
@@ -4567,6 +4570,32 @@ bool FTextureSource::ComputeChannelLinearMinMax(int32 InLayerIndex, FLinearColor
 	return true;
 }
 
+void FTextureSource::GetLayerColorInfo(TArray<FTextureSourceLayerColorInfo> & OutLayerColorInfo) const
+{
+	FScopeLock BulkDataExclusiveScope(&const_cast<FCriticalSection &>(BulkDataLock.Get()));
+	OutLayerColorInfo = LayerColorInfo_LockProtected;
+}
+void FTextureSource::SetLayerColorInfo(const TArray<FTextureSourceLayerColorInfo> & InLayerColorInfo)
+{
+	FScopeLock BulkDataExclusiveScope(&BulkDataLock.Get());
+	int32 Num = InLayerColorInfo.Num();
+	check( Num == 0 || Num == GetNumLayers() );
+	LayerColorInfo_LockProtected = InLayerColorInfo;
+}
+void FTextureSource::ResetLayerColorInfo()
+{
+	FScopeLock BulkDataExclusiveScope(&BulkDataLock.Get());
+	LayerColorInfo_LockProtected.Empty();
+}
+bool FTextureSource::HasLayerColorInfo() const
+{
+	FScopeLock BulkDataExclusiveScope(&const_cast<FCriticalSection &>(BulkDataLock.Get()));
+	int32 Num = LayerColorInfo_LockProtected.Num();
+	check( Num == 0 || Num == GetNumLayers() );
+	return Num != 0;
+}
+
+
 // UpdateChannelMinMaxFromIncomingTextureData does not use the BulkData or CompressionFormat on the TextureSource
 //	but it does use the dimensions/blocks/etc. they must be set before calling this
 bool FTextureSource::UpdateChannelMinMaxFromIncomingTextureData(FMemoryView InNewTextureData)
@@ -4574,6 +4603,7 @@ bool FTextureSource::UpdateChannelMinMaxFromIncomingTextureData(FMemoryView InNe
 	// InNewTextureData must be uncompressed
 	//	if it's not, will likely hit the check on mip size below
 	
+	TArray<FTextureSourceLayerColorInfo> LayerColorInfo;
 	LayerColorInfo.SetNum(NumLayers);
 
 	for (int32 LayerIndex = 0; LayerIndex < NumLayers; LayerIndex++)
@@ -4596,7 +4626,6 @@ bool FTextureSource::UpdateChannelMinMaxFromIncomingTextureData(FMemoryView InNe
 			int64 MipSize = CalcMipSize(BlockIndex, LayerIndex, 0);
 
 			FMemoryView MipView = InNewTextureData.Mid(MipOffset, MipSize);
-			check(MipView.GetSize() == MipSize);
 
 			if (MipView.GetSize() == MipSize)
 			{
@@ -4624,7 +4653,7 @@ bool FTextureSource::UpdateChannelMinMaxFromIncomingTextureData(FMemoryView InNe
 			else
 			{
 				UE_LOG(LogTexture, Error, TEXT("Invalid mip size in texture source init: passed in size doesn't accomodate all mips!"));
-				LayerColorInfo.Empty();
+				ResetLayerColorInfo();
 				return false;
 			}
 		} // end each block
@@ -4633,6 +4662,8 @@ bool FTextureSource::UpdateChannelMinMaxFromIncomingTextureData(FMemoryView InNe
 		LayerInfo.ColorMin = TotalMin;
 	} // end each layer
 
+	SetLayerColorInfo(LayerColorInfo);
+
 	return true;
 }
 
@@ -4643,21 +4674,25 @@ bool FTextureSource::UpdateChannelLinearMinMax()
 	FTextureSource::FMipLock LockedMip0(ELockState::ReadOnly, const_cast<FTextureSource*>(this), 0);
 	if (LockedMip0.IsValid() == false)
 	{
-		LayerColorInfo.Empty();
+		ResetLayerColorInfo();
 		return false;
 	}
-
+	
+	TArray<FTextureSourceLayerColorInfo> LayerColorInfo;
 	LayerColorInfo.SetNum(NumLayers);
+
 	for (int32 LayerIndex = 0; LayerIndex < NumLayers; LayerIndex++)
 	{
 		FTextureSourceLayerColorInfo& LayerInfo = LayerColorInfo[LayerIndex];
 
 		if ( ! ComputeChannelLinearMinMax(LayerIndex, LayerInfo.ColorMin, LayerInfo.ColorMax) )
 		{
-			LayerColorInfo.Empty();
+			ResetLayerColorInfo();
 			return false;
 		}
 	}
+	
+	SetLayerColorInfo(LayerColorInfo);
 
 	return true;
 }
diff --git a/Engine/Source/Runtime/Engine/Private/TextureDerivedData.cpp b/Engine/Source/Runtime/Engine/Private/TextureDerivedData.cpp
--- a/Engine/Source/Runtime/Engine/Private/TextureDerivedData.cpp
+++ b/Engine/Source/Runtime/Engine/Private/TextureDerivedData.cpp
@@ -811,9 +811,12 @@ static void FinalizeBuildSettingsForLayer(
 
 	// If we have channel boundary information, use that to determine whether we expect to have
 	// a non opaque alpha.
-	if (LayerIndex < Texture.Source.GetLayerColorInfo().Num())
+	TArray<FTextureSourceLayerColorInfo> LayerColorInfo;
+	Texture.Source.GetLayerColorInfo(LayerColorInfo);
+	if (LayerIndex < LayerColorInfo.Num())
 	{
-		const FTextureSourceLayerColorInfo& LayerChannelBounds = Texture.Source.GetLayerColorInfo()[LayerIndex];
+		const FTextureSourceLayerColorInfo& LayerChannelBounds = LayerColorInfo[LayerIndex];
+
 		OutSettings.bKnowAlphaTransparency = ITextureCompressorModule::DetermineAlphaChannelTransparency(OutSettings, 
 			LayerChannelBounds.ColorMin, LayerChannelBounds.ColorMax, OutSettings.bHasTransparentAlpha);
 	}
diff --git a/Engine/Source/Runtime/Engine/Private/TextureDerivedDataTask.cpp b/Engine/Source/Runtime/Engine/Private/TextureDerivedDataTask.cpp
--- a/Engine/Source/Runtime/Engine/Private/TextureDerivedDataTask.cpp
+++ b/Engine/Source/Runtime/Engine/Private/TextureDerivedDataTask.cpp
@@ -194,18 +194,6 @@ void FTextureSourceData::Init(UTexture& InTexture, TextureMipGenSettings InMipGe
 {
 	check( bValid == false ); // we set to true at the end, acts as our return value
 
-	// Copy the channel min/max if we have it to avoid redoing it.
-	if (InTexture.Source.GetLayerColorInfo().Num())
-	{
-		LayerChannelMinMax.Reset();
-		for (const FTextureSourceLayerColorInfo& LayerColorInfo : InTexture.Source.GetLayerColorInfo())
-		{
-			TPair<FLinearColor, FLinearColor>& MinMax = LayerChannelMinMax.AddDefaulted_GetRef();
-			MinMax.Key = LayerColorInfo.ColorMin;
-			MinMax.Value = LayerColorInfo.ColorMax;
-		}
-	}
-
 	const int32 NumBlocks = InTexture.Source.GetNumBlocks();
 	const int32 NumLayers = InTexture.Source.GetNumLayers();
 	if (NumBlocks < 1 || NumLayers < 1)
@@ -213,6 +201,20 @@ void FTextureSourceData::Init(UTexture& InTexture, TextureMipGenSettings InMipGe
 		UE_LOG(LogTexture, Warning, TEXT("Texture has no source data: %s"), *InTexture.GetPathName());
 		return;
 	}
+	
+	// Copy the channel min/max if we have it already
+	// if Texture Source did not already have SourceLayerColorInfo, we will update it in GetSourceMips (when we have decompressed data)
+	
+	TArray<FTextureSourceLayerColorInfo> SourceLayerColorInfo;
+	InTexture.Source.GetLayerColorInfo(SourceLayerColorInfo);
+
+	check( SourceLayerColorInfo.Num() == 0 || SourceLayerColorInfo.Num() == NumLayers );
+	LayerChannelMinMax.SetNum( SourceLayerColorInfo.Num() );
+	for(int32 i=0;i<SourceLayerColorInfo.Num();i++)
+	{
+		LayerChannelMinMax[i].Key   = SourceLayerColorInfo[i].ColorMin;
+		LayerChannelMinMax[i].Value = SourceLayerColorInfo[i].ColorMax;
+	}
 
 	Layers.Reserve(NumLayers);
 	for (int LayerIndex = 0; LayerIndex < NumLayers; ++LayerIndex)
@@ -292,6 +294,22 @@ void FTextureSourceData::GetSourceMips(FTextureSource& Source, IImageWrapperModu
 {
 	if (bValid)
 	{
+		const int32 NumBlocks = Source.GetNumBlocks();
+		const int32 NumLayers = Source.GetNumLayers();
+		
+		// these arrays were sized in Init but not fully filled out :
+		check( Blocks.Num() == NumBlocks );
+		check( Layers.Num() == NumLayers );
+
+		check( NumBlocks > 0 && NumLayers > 0 )
+		
+		if ( Blocks[0].MipsPerLayer[0].Num() > 0 )
+		{
+			// If we already got valid data, nothing to do. (GetSourceMips was called before now it's being called again)
+			// @@ is this ever hit ? how?
+ 			return;
+		}
+
 		if (Source.HasHadBulkDataCleared())
 		{	// don't do any work we can't reload this
 			UE_LOG(LogTexture, Error, TEXT("Unable to get texture source mips because its bulk data was released. %s"), *TextureFullName);
@@ -316,88 +334,77 @@ void FTextureSourceData::GetSourceMips(FTextureSource& Source, IImageWrapperModu
 			bValid = false;
 			return;
 		}
+		
 
 		// If we didn't get ChannelMinMax from the texture source, then compute it now. As time goes on this will get hit less and less.
-		if (LayerChannelMinMax.Num() != Layers.Num())
+		if (LayerChannelMinMax.Num() != NumLayers)
 		{
 			TRACE_CPUPROFILER_EVENT_SCOPE(FTextureSourceData::GetSourceMips_ChannelMinMax);
 
-			// @@ can we do this on the Texture instead so it's cached for next time?
-			//	I think you can just call Source.UpdateChannelMinMaxFromIncomingTextureData(ScopedMipData) here
-
-			LayerChannelMinMax.Reset();
-			for (int32 LayerIndex = 0; LayerIndex < Layers.Num(); ++LayerIndex)
+			// Update MipMax if it wasn't found before
+			// do this after GetMipData() so we do it on decompressed data
+			
+			if ( Source.UpdateChannelMinMaxFromIncomingTextureData(ScopedMipData.GetData().GetView()) )
 			{
-				TPair<FLinearColor, FLinearColor>& LayerInfo = LayerChannelMinMax.AddDefaulted_GetRef();
+				TArray<FTextureSourceLayerColorInfo> SourceLayerColorInfo;
+				Source.GetLayerColorInfo(SourceLayerColorInfo);
 
-				FLinearColor TotalMin(FLT_MAX, FLT_MAX, FLT_MAX, FLT_MAX);
-				FLinearColor TotalMax(-FLT_MAX, -FLT_MAX, -FLT_MAX, -FLT_MAX);
-
-				for (int32 BlockIndex = 0; BlockIndex < Blocks.Num(); BlockIndex++)
+				check( SourceLayerColorInfo.Num() == 0 || SourceLayerColorInfo.Num() == NumLayers );
+				LayerChannelMinMax.SetNum( SourceLayerColorInfo.Num() );
+				for(int32 i=0;i<SourceLayerColorInfo.Num();i++)
 				{
-					FImageView MipImageView;
-					FSharedBuffer MipData = ScopedMipData.GetMipDataWithInfo(BlockIndex, LayerIndex, 0, MipImageView);
-
-					MipImageView.RawData = (void*)MipData.GetData();
-
-					FLinearColor MinColor, MaxColor;
-					FImageCore::ComputeChannelLinearMinMax(MipImageView, MinColor, MaxColor);
-
-					TotalMin.R = FMath::Min(MinColor.R, TotalMin.R);
-					TotalMin.G = FMath::Min(MinColor.G, TotalMin.G);
-					TotalMin.B = FMath::Min(MinColor.B, TotalMin.B);
-					TotalMin.A = FMath::Min(MinColor.A, TotalMin.A);
-
-					TotalMax.R = FMath::Max(MaxColor.R, TotalMax.R);
-					TotalMax.G = FMath::Max(MaxColor.G, TotalMax.G);
-					TotalMax.B = FMath::Max(MaxColor.B, TotalMax.B);
-					TotalMax.A = FMath::Max(MaxColor.A, TotalMax.A);
+					LayerChannelMinMax[i].Key   = SourceLayerColorInfo[i].ColorMin;
+					LayerChannelMinMax[i].Value = SourceLayerColorInfo[i].ColorMax;
 				}
-
-				LayerInfo.Key = TotalMin;
-				LayerInfo.Value = TotalMax;
+			}
+			else
+			{
+				UE_LOG(LogTexture, Warning, TEXT("Unexpected failure in UpdateChannelMinMaxFromIncomingTextureData on %s"), *TextureFullName);
 			}
 		}
 
 		{
 			TRACE_CPUPROFILER_EVENT_SCOPE(FTextureSourceData::GetSourceMips_CopyMips);
 
-			check( Blocks.Num() == Source.GetNumBlocks() );
-
-			for (int32 BlockIndex = 0; BlockIndex < Blocks.Num(); ++BlockIndex)
+			for (int32 BlockIndex = 0; BlockIndex < NumBlocks; ++BlockIndex)
 			{
 				FTextureSourceBlock SourceBlock;
 				Source.GetBlock(BlockIndex, SourceBlock);
 
 				FTextureSourceBlockData& BlockData = Blocks[BlockIndex];
-				check( BlockData.MipsPerLayer.Num() == Source.GetNumLayers() );
-				check( Layers.Num() == Source.GetNumLayers() );
+				check( BlockData.MipsPerLayer.Num() == NumLayers );
 
-				for (int32 LayerIndex = 0; LayerIndex < Layers.Num(); ++LayerIndex)
+				for (int32 LayerIndex = 0; LayerIndex < NumLayers; ++LayerIndex)
 				{
 					const FTextureSourceLayerData& LayerData = Layers[LayerIndex];
-					if ( BlockData.MipsPerLayer[LayerIndex].Num() == 0 ) // If we already got valid data, nothing to do. (@@ this is weird, should detect this and early out sooner, before GetMipData)
-					{
-						BlockData.MipsPerLayer[LayerIndex].SetNum( BlockData.NumMips );
 
-						for (int32 MipIndex = 0; MipIndex < BlockData.NumMips; ++MipIndex)
-						{
-							FImageInfo MipImageInfo;
-							FSharedBuffer MipData = ScopedMipData.GetMipDataWithInfo(BlockIndex, LayerIndex, MipIndex, MipImageInfo);
-							
-							check(MipImageInfo.GammaSpace == LayerData.SourceGammaSpace);
-							check(MipImageInfo.Format == LayerData.ImageFormat);
+					check( BlockData.MipsPerLayer[LayerIndex].Num() == 0 );
+					//if Source had existing mips but TMGS was not LeaveExisting, then BlockData.NumMips is set to 1
+					check( BlockData.NumMips > 0 && ( BlockData.NumMips == SourceBlock.NumMips || ( BlockData.NumMips == 1 && SourceBlock.NumMips > 1 ) ) );
 
-							FImageView MipView( MipImageInfo, const_cast<void *>(MipData.GetData()) );
+					BlockData.MipsPerLayer[LayerIndex].SetNum( BlockData.NumMips );
+
+					for (int32 MipIndex = 0; MipIndex < BlockData.NumMips; ++MipIndex)
+					{
+						FImageView MipView = ScopedMipData.GetMipDataImageView(BlockIndex, LayerIndex, MipIndex);
 							
-							// allocates the destination FImage and copies into it :
-							MipView.CopyTo( BlockData.MipsPerLayer[LayerIndex][MipIndex] );
+						check(MipView.GammaSpace == LayerData.SourceGammaSpace);
+						check(MipView.Format == LayerData.ImageFormat);
 
-							// this copy takes a while, and potentially we could just point at the FSharedBuffer from the ScopedMipData
-							//   (like MipView here does)
-							// at the moment that's not easy because all the code around TextureCompressorModule expects FImage, not FImageView
-							//	perhaps ideally we'd have an FImage variant that's COW
-						}
+						// allocates the destination FImage and copies into it :
+						MipView.CopyTo( BlockData.MipsPerLayer[LayerIndex][MipIndex] );
+
+						// CB notes 04/02/2024 :
+
+						// this copy takes a while, and potentially we could just instead point at the FSharedBuffer from the ScopedMipData
+						//   (like MipView here does)
+						// at the moment that's not easy because all the code around TextureCompressorModule/Formats expects FImage, not FImageView
+						//	perhaps ideally we'd have an FImage variant that's COW
+
+						// fundamentally, this whole alloc and copy is totally unecessary, so it would be great to get rid of it
+						//	but practically we need a better way to have FImage point at FSharedBuffer
+						// one issue is there's no way to MoveTemp into a TArray from Shared/Unique buffer.
+						// also at the moment the TextureCompressorModule/Formats assume the FImage is mutable, that would have to be cleaned up
 					}
 				}
 			}
diff --git a/Engine/Source/Runtime/Engine/Private/TextureDerivedDataTask.h b/Engine/Source/Runtime/Engine/Private/TextureDerivedDataTask.h
--- a/Engine/Source/Runtime/Engine/Private/TextureDerivedDataTask.h
+++ b/Engine/Source/Runtime/Engine/Private/TextureDerivedDataTask.h
@@ -115,7 +115,7 @@ struct FTextureSourceData
 		bValid = false;
 	}
 
-	TArray<TPair<FLinearColor, FLinearColor>> LayerChannelMinMax;
+	TArray<TPair<FLinearColor, FLinearColor>> LayerChannelMinMax; // Key == Min, Value == Max
 
 	FString TextureFullName;
 	FTextureSource AsyncSource;
