Iris - only call BeginReplication while flushing dormancy if the actor has begun play or is beginning play. Before beginning play, the actor's replicated components may not be fully set up yet.

diff --git a/Engine/Source/Runtime/Engine/Private/Net/Iris/ReplicationSystem/ActorReplicationBridge.cpp b/Engine/Source/Runtime/Engine/Private/Net/Iris/ReplicationSystem/ActorReplicationBridge.cpp
--- a/Engine/Source/Runtime/Engine/Private/Net/Iris/ReplicationSystem/ActorReplicationBridge.cpp
+++ b/Engine/Source/Runtime/Engine/Private/Net/Iris/ReplicationSystem/ActorReplicationBridge.cpp
@@ -223,6 +223,12 @@ UE::Net::FNetRefHandle UActorReplicationBridge::BeginReplication(AActor* Actor,
 		return FNetRefHandle::GetInvalid();
 	}
 
+	if (!ensureMsgf(Actor->HasActorBegunPlay() || Actor->IsActorBeginningPlay(), TEXT("Actor %s hasn't begun or isn't beginning play in BeginReplication."), ToCStr(GetFullNameSafe(Actor))))
+	{
+		UE_LOG_ACTORREPLICATIONBRIDGE(Error, TEXT("Actor %s hasn't begun or isn't beginning play in BeginReplication."), ToCStr(GetFullNameSafe(Actor)));
+		return FNetRefHandle::GetInvalid();
+	}
+
 	if (!NetDriver)
 	{
 		UE_LOG_ACTORREPLICATIONBRIDGE(VeryVerbose, TEXT("There's no NetDriver so nothing can be replicated."));
diff --git a/Engine/Source/Runtime/Engine/Private/Net/Iris/ReplicationSystem/ReplicationSystemUtil.cpp b/Engine/Source/Runtime/Engine/Private/Net/Iris/ReplicationSystem/ReplicationSystemUtil.cpp
--- a/Engine/Source/Runtime/Engine/Private/Net/Iris/ReplicationSystem/ReplicationSystemUtil.cpp
+++ b/Engine/Source/Runtime/Engine/Private/Net/Iris/ReplicationSystem/ReplicationSystemUtil.cpp
@@ -480,8 +480,15 @@ void FReplicationSystemUtil::FlushNetDormancy(UReplicationSystem* ReplicationSys
 			}
 		}
 	}
-	else
+	else if (Actor->HasActorBegunPlay() || Actor->IsActorBeginningPlay())
 	{
+		// Call BeginReplication for DORM_Initial actors the first time their dormancy is flushed
+		// (since it's not called when they BeginPlay).
+		// 
+		// We still don't want to call BeginReplication before BeginPlay though:
+		// -The actor and its components/subobjects may not be completely set up for replication yet
+		// -If the actor is DormInitial, and is flushed before BeginPlay, its dormancy state will change to DormantAll and it will BeginReplication normally
+
 		UE_CLOG(!Actor->GetIsReplicated(), LogIris, Warning, TEXT("FReplicationSystemUtil::FlushNetDormancy Actor %s that is not replicated"), ToCStr(Actor->GetName()));
 		UE_CLOG(!bWasDormInitial, LogIris, Display, TEXT("FReplicationSystemUtil::FlushNetDormancy For not replicated Actor %s is not initially dormant"), ToCStr(Actor->GetName()));
 
