Optimized Matrix4x4 -> Matrix3x4 transpose SIMD in GPUSkinVertexFactory, and also optimized the EPreviousBoneTransformUpdateMode::DuplicateCurrentToPrevious to copy the current transforms directly instead of redundantly recomputing them all over again with UpdateRefToLocalMatrices

diff --git a/Engine/Source/Runtime/Engine/Private/GPUSkinVertexFactory.cpp b/Engine/Source/Runtime/Engine/Private/GPUSkinVertexFactory.cpp
--- a/Engine/Source/Runtime/Engine/Private/GPUSkinVertexFactory.cpp
+++ b/Engine/Source/Runtime/Engine/Private/GPUSkinVertexFactory.cpp
@@ -291,7 +291,7 @@ void FGPUBaseSkinVertexFactory::FShaderDataType::UpdateBoneData(FRHICommandList&
 				FMatrix3x4& BoneMat = ChunkMatrices[BoneIdx];
 				const FMatrix44f& RefToLocal = ReferenceToLocalMatrices[RefToLocalIdx];
 				// Explicit SIMD implementation seems to be faster than standard implementation
-#if PLATFORM_ENABLE_VECTORINTRINSICS
+			#if PLATFORM_ENABLE_VECTORINTRINSICS
 				VectorRegister4Float InRow0 = VectorLoadAligned(&(RefToLocal.M[0][0]));
 				VectorRegister4Float InRow1 = VectorLoadAligned(&(RefToLocal.M[1][0]));
 				VectorRegister4Float InRow2 = VectorLoadAligned(&(RefToLocal.M[2][0]));
@@ -302,17 +302,12 @@ void FGPUBaseSkinVertexFactory::FShaderDataType::UpdateBoneData(FRHICommandList&
 				VectorRegister4Float Temp2 = VectorShuffle(InRow0, InRow1, 2, 3, 2, 3);
 				VectorRegister4Float Temp3 = VectorShuffle(InRow2, InRow3, 2, 3, 2, 3);
 
-				Temp0 = VectorSwizzle(Temp0, 0, 2, 1, 3);
-				Temp1 = VectorSwizzle(Temp1, 0, 2, 1, 3);
-				Temp2 = VectorSwizzle(Temp2, 0, 2, 1, 3);
-				Temp3 = VectorSwizzle(Temp3, 0, 2, 1, 3);
-
-				VectorStoreAligned(VectorShuffle(Temp0, Temp1, 0, 1, 0, 1), &(BoneMat.M[0][0]));
-				VectorStoreAligned(VectorShuffle(Temp0, Temp1, 2, 3, 2, 3), &(BoneMat.M[1][0]));
-				VectorStoreAligned(VectorShuffle(Temp2, Temp3, 0, 1, 0, 1), &(BoneMat.M[2][0]));
-#else
+				VectorStoreAligned(VectorShuffle(Temp0, Temp1, 0, 2, 0, 2), &(BoneMat.M[0][0]));
+				VectorStoreAligned(VectorShuffle(Temp0, Temp1, 1, 3, 1, 3), &(BoneMat.M[1][0]));
+				VectorStoreAligned(VectorShuffle(Temp2, Temp3, 0, 2, 0, 2), &(BoneMat.M[2][0]));
+			#else
 				RefToLocal.To3x4MatrixTranspose((float*)BoneMat.M);
-#endif
+			#endif
 			}
 		}
 	}
diff --git a/Engine/Source/Runtime/Engine/Private/SkeletalRenderGPUSkin.cpp b/Engine/Source/Runtime/Engine/Private/SkeletalRenderGPUSkin.cpp
--- a/Engine/Source/Runtime/Engine/Private/SkeletalRenderGPUSkin.cpp
+++ b/Engine/Source/Runtime/Engine/Private/SkeletalRenderGPUSkin.cpp
@@ -2404,11 +2404,17 @@ void FDynamicSkelMeshObjectDataGPUSkin::InitDynamicSkelMeshObjectDataGPUSkin(
 	#endif
 		break;
 	case EPreviousBoneTransformUpdateMode::DuplicateCurrentToPrevious:
-		UpdateRefToLocalMatrices(PreviousReferenceToLocal, InMeshComponent, InSkeletalMeshRenderData, LODIndex, ExtraRequiredBoneIndices);
+		for (int32 Index = 0; Index < ReferenceToLocal.Num(); ++Index)
+		{
+			PreviousReferenceToLocal[Index] = ReferenceToLocal[Index];
+		}
 	#if RHI_RAYTRACING
 		if (RayTracingLODIndex != LODIndex)
 		{
-			UpdateRefToLocalMatrices(PreviousReferenceToLocalForRayTracing, InMeshComponent, InSkeletalMeshRenderData, RayTracingLODIndex, ExtraRequiredBoneIndices);
+			for (int32 Index = 0; Index < ReferenceToLocalForRayTracing.Num(); ++Index)
+			{
+				PreviousReferenceToLocalForRayTracing[Index] = ReferenceToLocalForRayTracing[Index];
+			}
 		}
 	#endif
 		break;
