Texture.cpp : fix use of incorrect EditorBulkData call UnloadData should be Reset

diff --git a/Engine/Source/Runtime/Engine/Private/Texture.cpp b/Engine/Source/Runtime/Engine/Private/Texture.cpp
--- a/Engine/Source/Runtime/Engine/Private/Texture.cpp
+++ b/Engine/Source/Runtime/Engine/Private/Texture.cpp
@@ -3355,7 +3355,7 @@ void FTextureSource::RemoveSourceData()
 	
 	ResetLayerColorInfo();
 
-	BulkData.UnloadData();
+	BulkData.Reset();
 
 	ForceGenerateGuid();
 }
