Set shader compilation errors to fatal again: using the .cook odsc command will now spawn an attended process, and will display a message box if r.ShaderDevelopmentMode=1 in ConsoleVariables.ini, just like in editor

diff --git a/Engine/Source/Runtime/Engine/Private/ShaderCompiler/ShaderCompiler.cpp b/Engine/Source/Runtime/Engine/Private/ShaderCompiler/ShaderCompiler.cpp
--- a/Engine/Source/Runtime/Engine/Private/ShaderCompiler/ShaderCompiler.cpp
+++ b/Engine/Source/Runtime/Engine/Private/ShaderCompiler/ShaderCompiler.cpp
@@ -10453,9 +10453,6 @@ void RecompileShadersForRemote(
 {
 	TRACE_CPUPROFILER_EVENT_SCOPE(RecompileShadersForRemote);
 
-	CVarAreShaderErrorsFatal->Set(0, ECVF_SetByCode);
-	ON_SCOPE_EXIT{ CVarAreShaderErrorsFatal->Unset(ECVF_SetByCode); };
-
 	// figure out what shader platforms to recompile
 	ITargetPlatformManagerModule* TPM = GetTargetPlatformManager();
 	ITargetPlatform* TargetPlatform = TPM->FindTargetPlatform(Args.PlatformName);
