Add relevant constructors to asset locator fragment

Allows for easier construction of fragments from FAssetData

diff --git a/Engine/Source/Runtime/Engine/Private/UniversalObjectLocators/AssetLocatorFragment.cpp b/Engine/Source/Runtime/Engine/Private/UniversalObjectLocators/AssetLocatorFragment.cpp
--- a/Engine/Source/Runtime/Engine/Private/UniversalObjectLocators/AssetLocatorFragment.cpp
+++ b/Engine/Source/Runtime/Engine/Private/UniversalObjectLocators/AssetLocatorFragment.cpp
@@ -9,9 +9,20 @@
 #include "UObject/ObjectRedirector.h"
 #include "UObject/Package.h"
 #include "IUniversalObjectLocatorModule.h"
+#include "AssetRegistry/AssetData.h"
 
 UE::UniversalObjectLocator::TFragmentTypeHandle<FAssetLocatorFragment> FAssetLocatorFragment::FragmentType;
 
+FAssetLocatorFragment::FAssetLocatorFragment(const FAssetData& InAssetData)
+	: Path(InAssetData.PackageName, InAssetData.AssetName)
+{
+}
+
+FAssetLocatorFragment::FAssetLocatorFragment(const FTopLevelAssetPath& InPath)
+	: Path(InPath)
+{
+}
+
 UE::UniversalObjectLocator::FResolveResult FAssetLocatorFragment::Resolve(const UE::UniversalObjectLocator::FResolveParams& Params) const
 {
 	using namespace UE::UniversalObjectLocator;
diff --git a/Engine/Source/Runtime/Engine/Public/UniversalObjectLocators/AssetLocatorFragment.h b/Engine/Source/Runtime/Engine/Public/UniversalObjectLocators/AssetLocatorFragment.h
--- a/Engine/Source/Runtime/Engine/Public/UniversalObjectLocators/AssetLocatorFragment.h
+++ b/Engine/Source/Runtime/Engine/Public/UniversalObjectLocators/AssetLocatorFragment.h
@@ -6,6 +6,8 @@
 #include "UniversalObjectLocatorFwd.h"
 #include "AssetLocatorFragment.generated.h"
 
+struct FAssetData;
+
 /**
  * 32 Bytes (40 in-editor).
  */
@@ -14,6 +16,12 @@ struct FAssetLocatorFragment
 {
 	GENERATED_BODY()
 
+	ENGINE_API FAssetLocatorFragment() = default;
+
+	ENGINE_API explicit FAssetLocatorFragment(const FAssetData& InAssetData);
+
+	ENGINE_API explicit FAssetLocatorFragment(const FTopLevelAssetPath& InPath);
+
 	UPROPERTY()
 	FTopLevelAssetPath Path;
 
