FPreviewScene: Listen to FCoreDelegates::OnEnginePreExit to allow proper preview UWorld cleanup

Fixes random ensure in UTickableWorldSubsystem::BeginDestroy on preview worlds getting GCed without getting properly cleaned up (depending on order of GC)

#rnx

diff --git a/Engine/Source/Runtime/Engine/Private/PreviewScene.cpp b/Engine/Source/Runtime/Engine/Private/PreviewScene.cpp
--- a/Engine/Source/Runtime/Engine/Private/PreviewScene.cpp
+++ b/Engine/Source/Runtime/Engine/Private/PreviewScene.cpp
@@ -90,9 +90,17 @@ FPreviewScene::FPreviewScene(FPreviewScene::ConstructionValues CVS)
 		LineBatcher->bCalculateAccurateBounds = false;
 		AddComponent(LineBatcher, FTransform::Identity);
 	}
+
+	FCoreDelegates::OnEnginePreExit.AddRaw(this, &FPreviewScene::Uninitialize);
 }
 
 FPreviewScene::~FPreviewScene()
+{
+	FCoreDelegates::OnEnginePreExit.RemoveAll(this);
+	Uninitialize();
+}
+
+void FPreviewScene::Uninitialize()
 {
 	// Stop any audio components playing in this scene
 	if (GEngine)
@@ -108,9 +116,9 @@ FPreviewScene::~FPreviewScene()
 	}
 
 	// Remove all the attached components
-	for( int32 ComponentIndex = 0; ComponentIndex < Components.Num(); ComponentIndex++ )
+	for (int32 ComponentIndex = 0; ComponentIndex < Components.Num(); ComponentIndex++)
 	{
-		UActorComponent* Component = Components[ ComponentIndex ];
+		UActorComponent* Component = Components[ComponentIndex];
 
 		if (bForceAllUsedMipsResident)
 		{
@@ -124,20 +132,26 @@ FPreviewScene::~FPreviewScene()
 
 		Component->UnregisterComponent();
 	}
-	
+
+	// Uninitialize can get called from destructor or FCoreDelegates::OnPreExit (or both)
+	// so make sure we empty Components and set PreviewWorld to nullptr
+	Components.Empty();
+		
+	UWorld* LocalPreviewWorld = PreviewWorld;
+	PreviewWorld = nullptr;
+
 	// The world may be released by now.
-	if (PreviewWorld && GEngine)
+	if (LocalPreviewWorld && GEngine)
 	{
-		PreviewWorld->CleanupWorld();
-		GEngine->DestroyWorldContext(GetWorld());
+		LocalPreviewWorld->CleanupWorld();
+		GEngine->DestroyWorldContext(LocalPreviewWorld);
 		// Release PhysicsScene for fixing big fbx importing bug
-		PreviewWorld->ReleasePhysicsScene();
+		LocalPreviewWorld->ReleasePhysicsScene();
 
 		// The preview world is a heavy-weight object and may hold a significant amount of resources,
 		// including various GPU render targets and buffers required for rendering the scene.
 		// Since UWorld is garbage-collected, this memory may not be cleaned for an indeterminate amount of time.
 		// By forcing garbage collection explicitly, we allow memory to be reused immediately.
-		PreviewWorld = nullptr;
 		GEngine->ForceGarbageCollection(true /*bFullPurge*/);
 	}
 }
diff --git a/Engine/Source/Runtime/Engine/Public/PreviewScene.h b/Engine/Source/Runtime/Engine/Public/PreviewScene.h
--- a/Engine/Source/Runtime/Engine/Public/PreviewScene.h
+++ b/Engine/Source/Runtime/Engine/Public/PreviewScene.h
@@ -116,6 +116,7 @@ public:
 	class USkyLightComponent* SkyLight = nullptr;;
 
 private:
+	void Uninitialize();
 
 	TArray<TObjectPtr<class UActorComponent>> Components;
 
