MaterialInterface: Refactor GetMaterialShaderMapKey into RecordOrEmitMaterialShaderMapKey which takes an FMaterialKeyGeneratorContext instead of an FShaderKeyGenerator. FMaterialKeyGeneratorContext can handle any one of three modes: emit hash strings or save/load data needed for hash strings to compact binary. Use the new RecordOrEmitMaterialShaderMapKey to Save/Load and Emit the material dependency hash for incremental cooking.
#rnx

diff --git a/Engine/Source/Runtime/Core/Private/Serialization/MemoryImage.cpp b/Engine/Source/Runtime/Core/Private/Serialization/MemoryImage.cpp
--- a/Engine/Source/Runtime/Core/Private/Serialization/MemoryImage.cpp
+++ b/Engine/Source/Runtime/Core/Private/Serialization/MemoryImage.cpp
@@ -16,6 +16,10 @@
 #include "Serialization/Archive.h"
 #include "Serialization/ShaderKeyGenerator.h"
 #include "Async/ParallelFor.h"
+#if WITH_EDITOR
+#include "Serialization/CompactBinary.h"
+#include "Serialization/CompactBinaryWriter.h"
+#endif
 
 DEFINE_LOG_CATEGORY_STATIC(LogMemoryImage, Log, All);
 
@@ -132,6 +136,34 @@ void FPlatformTypeLayoutParameters::Append(FShaderKeyGenerator& KeyGen) const
 	}
 }
 
+#if WITH_EDITOR
+void FPlatformTypeLayoutParameters::Save(FCbWriter& Writer) const
+{
+	Writer.BeginArray();
+	Writer << MaxFieldAlignment;
+	Writer << Flags;
+	Writer.EndArray();
+}
+
+bool FPlatformTypeLayoutParameters::TryLoad(FCbFieldView Field)
+{
+	*this = FPlatformTypeLayoutParameters();
+
+	FCbFieldViewIterator ElementField(Field.CreateViewIterator());
+	MaxFieldAlignment = ElementField.AsUInt32();
+	if ((ElementField++).HasError()) return false;
+	Flags = ElementField.AsUInt32();
+	if ((ElementField++).HasError()) return false;
+	return true;
+}
+
+bool LoadFromCompactBinary(FCbFieldView Field, FPlatformTypeLayoutParameters& OutValue)
+{
+	return OutValue.TryLoad(Field);
+}
+
+#endif
+
 // evaluated during static-initialization, so logging from regular check() macros won't work correctly
 static void InitializeSizeFromFields(FTypeLayoutDesc& TypeLayout, const FPlatformTypeLayoutParameters& PlatformLayoutParams)
 {
@@ -1258,6 +1290,50 @@ FHashedName& FHashedName::operator=(const FHashedName& InName)
 	return *this;
 }
 
+#if WITH_EDITOR
+void FHashedName::Save(FCbWriter& Writer) const
+{
+	Writer.BeginArray();
+	Writer << GetHash();
+	FUtf8StringView DebugName;
+	DebugName = GetDebugString().String.Get();
+	Writer << DebugName;
+	Writer.EndArray();
+}
+
+bool FHashedName::TryLoad(FCbFieldView Field)
+{
+	*this = FHashedName();
+	FCbFieldViewIterator ElementField(Field.CreateViewIterator());
+	uint64 LocalHash = ElementField.AsUInt64();
+	if ((ElementField++).HasError())
+	{
+		return false;
+	}
+	FUtf8StringView DebugName = ElementField.AsString();
+	if ((ElementField++).HasError())
+	{
+		return false;
+	}
+
+	if (!DebugName.IsEmpty())
+	{
+		*this = FHashedName(*WriteToString<256>(DebugName));
+		if (GetHash() == LocalHash)
+		{
+			return true;
+		}
+	}
+	*this = FHashedName(LocalHash);
+	return true;
+}
+
+bool LoadFromCompactBinary(FCbFieldView Field, FHashedName& OutValue)
+{
+	return OutValue.TryLoad(Field);
+}
+#endif // WITH_EDITOR
+
 void FPointerTableBase::SaveToArchive(FArchive& Ar, const FPlatformTypeLayoutParameters& LayoutParams, const void* FrozenObject) const
 {
 	int32 NumDependencies = 0;
diff --git a/Engine/Source/Runtime/Core/Public/Serialization/MemoryImage.h b/Engine/Source/Runtime/Core/Public/Serialization/MemoryImage.h
--- a/Engine/Source/Runtime/Core/Public/Serialization/MemoryImage.h
+++ b/Engine/Source/Runtime/Core/Public/Serialization/MemoryImage.h
@@ -29,6 +29,8 @@
 #include "Traits/IsCharType.h"
 #include "UObject/NameTypes.h"
 
+class FCbFieldView;
+class FCbWriter;
 template <typename T> struct TIsContiguousContainer;
 
 #if defined(WITH_RTTI) || defined(_CPPRTTI) || defined(__GXX_RTTI) || WITH_EDITOR
@@ -895,6 +897,18 @@ public:
 private:
 	LAYOUT_FIELD(uint64, Hash);
 	LAYOUT_FIELD_EDITORONLY(FHashedNameDebugString, DebugString);
+
+#if WITH_EDITOR
+	// Compact binary API with hidden friend operator<<
+	CORE_API void Save(FCbWriter& Writer) const;
+	bool TryLoad(FCbFieldView Field);
+	friend inline FCbWriter& operator<<(FCbWriter& Writer, const FHashedName& Value)
+	{
+		Value.Save(Writer);
+		return Writer;
+	}
+	friend CORE_API bool LoadFromCompactBinary(FCbFieldView Field, FHashedName& OutValue);
+#endif
 };
 
 namespace Freeze
diff --git a/Engine/Source/Runtime/Core/Public/Serialization/MemoryLayout.h b/Engine/Source/Runtime/Core/Public/Serialization/MemoryLayout.h
--- a/Engine/Source/Runtime/Core/Public/Serialization/MemoryLayout.h
+++ b/Engine/Source/Runtime/Core/Public/Serialization/MemoryLayout.h
@@ -14,6 +14,8 @@
 #include "Templates/Models.h"
 #include "Templates/UnrealTemplate.h"
 
+class FCbFieldView;
+class FCbWriter;
 class FHashedName;
 class FShaderKeyGenerator;
 class FSHA1;
@@ -863,6 +865,17 @@ private:
 	{
 		Value.Append(KeyGen);
 	}
+#if WITH_EDITOR
+	// Compact binary API with hidden friend operator<<
+	CORE_API void Save(FCbWriter& Writer) const;
+	bool TryLoad(FCbFieldView Field);
+	friend inline FCbWriter& operator<<(FCbWriter& Writer, const FPlatformTypeLayoutParameters& Value)
+	{
+		Value.Save(Writer);
+		return Writer;
+	}
+	friend CORE_API bool LoadFromCompactBinary(FCbFieldView Field, FPlatformTypeLayoutParameters& OutValue);
+#endif
 };
 
 #if UE_ENABLE_INCLUDE_ORDER_DEPRECATED_IN_5_5
diff --git a/Engine/Source/Runtime/Core/Public/Serialization/ShaderKeyGenerator.h b/Engine/Source/Runtime/Core/Public/Serialization/ShaderKeyGenerator.h
--- a/Engine/Source/Runtime/Core/Public/Serialization/ShaderKeyGenerator.h
+++ b/Engine/Source/Runtime/Core/Public/Serialization/ShaderKeyGenerator.h
@@ -32,6 +32,7 @@ public:
 	inline FShaderKeyGenerator(TUniqueFunction<void(const void* Data, uint64 Size)>&& InResultFunc);
 	/** Constructor that writes the appended data to a long human-readable debug string. */
 	inline FShaderKeyGenerator(FString& InResultString);
+	inline ~FShaderKeyGenerator();
 
 	/**
 	 * True iff the KeyGen is writing to a hash function. Writing to a hash function also implies that debugtext and
@@ -132,6 +133,9 @@ inline FShaderKeyGenerator::FShaderKeyGenerator(FString& InResultString)
 {
 }
 
+// This destructor needs to be defined so that it can be called manually from union destructors.
+inline FShaderKeyGenerator::~FShaderKeyGenerator() = default;
+
 inline bool FShaderKeyGenerator::IsBinary() const
 {
 	return OutputType == EOutputType::Binary;
diff --git a/Engine/Source/Runtime/Engine/Classes/Materials/MaterialInterface.h b/Engine/Source/Runtime/Engine/Classes/Materials/MaterialInterface.h
--- a/Engine/Source/Runtime/Engine/Classes/Materials/MaterialInterface.h
+++ b/Engine/Source/Runtime/Engine/Classes/Materials/MaterialInterface.h
@@ -402,6 +402,9 @@ public:
 	UE_DEPRECATED(5.4, "Implement the version that takes FAssetRegistryTagsContext instead.")
 	ENGINE_API virtual void GetAssetRegistryTags(TArray<FAssetRegistryTag>& OutTags) const override;
 #endif // WITH_EDITOR
+#if WITH_EDITOR
+	static ENGINE_API void AppendToClassSchema(FAppendToClassSchemaContext& Context);
+#endif
 	//~ End UObject Interface.
 
 	//~ Begin Begin Interface IBlendableInterface
diff --git a/Engine/Source/Runtime/Engine/Private/Materials/Material.cpp b/Engine/Source/Runtime/Engine/Private/Materials/Material.cpp
--- a/Engine/Source/Runtime/Engine/Private/Materials/Material.cpp
+++ b/Engine/Source/Runtime/Engine/Private/Materials/Material.cpp
@@ -46,6 +46,7 @@
 #include "Materials/MaterialExpressionMaterialFunctionCall.h"
 #include "Materials/MaterialExpressionSingleLayerWaterMaterialOutput.h"
 #include "Materials/MaterialExpressionMultiply.h"
+#include "Materials/MaterialSharedPrivate.h"
 
 #include "SceneManagement.h"
 #include "SceneView.h"
@@ -1058,6 +1059,15 @@ void UMaterial::PreSave(FObjectPreSaveContext ObjectSaveContext)
 	Super::PreSave(ObjectSaveContext);
 #if WITH_EDITOR
 	GMaterialsWithDirtyUsageFlags.RemoveAnnotation(this);
+	if (ObjectSaveContext.IsCooking())
+	{
+		const ITargetPlatform* TargetPlatform = ObjectSaveContext.GetTargetPlatform();
+		check(TargetPlatform);
+
+		TArray<FMaterialResourceForCooking>* Resources = CachedMaterialResourcesForCooking.Find(TargetPlatform);
+		UE::MaterialInterface::Private::RecordMaterialDependenciesForCook(ObjectSaveContext,
+			Resources ? *Resources : TArray<FMaterialResourceForCooking>());
+	}
 #endif
 }
 
diff --git a/Engine/Source/Runtime/Engine/Private/Materials/MaterialInstance.cpp b/Engine/Source/Runtime/Engine/Private/Materials/MaterialInstance.cpp
--- a/Engine/Source/Runtime/Engine/Private/Materials/MaterialInstance.cpp
+++ b/Engine/Source/Runtime/Engine/Private/Materials/MaterialInstance.cpp
@@ -31,6 +31,7 @@
 #include "Materials/MaterialInstanceConstant.h"
 #include "Materials/MaterialInstanceUpdateParameterSet.h"
 #include "Materials/MaterialInstanceSupport.h"
+#include "Materials/MaterialSharedPrivate.h"
 #include "Engine/SubsurfaceProfile.h"
 #include "Engine/SpecularProfile.h"
 #include "ProfilingDebugging/CookStats.h"
@@ -4671,6 +4672,16 @@ void UMaterialInstance::PreSave(FObjectPreSaveContext ObjectSaveContext)
 {
 	// @TODO : Remove any duplicate data from parent? Aims at improving change propagation (if controlled by parent)
 	Super::PreSave(ObjectSaveContext);
+#if WITH_EDITOR
+	if (ObjectSaveContext.IsCooking())
+	{
+		const ITargetPlatform * TargetPlatform = ObjectSaveContext.GetTargetPlatform();
+		check(TargetPlatform);
+		TArray<FMaterialResourceForCooking>* Resources = CachedMaterialResourcesForCooking.Find(TargetPlatform);
+		UE::MaterialInterface::Private::RecordMaterialDependenciesForCook(ObjectSaveContext,
+			Resources ? *Resources : TArray<FMaterialResourceForCooking>());
+	}
+#endif
 }
 
 float UMaterialInstance::GetTextureDensity(FName TextureName, const struct FMeshUVChannelInfo& UVChannelData) const
diff --git a/Engine/Source/Runtime/Engine/Private/Materials/MaterialInterface.cpp b/Engine/Source/Runtime/Engine/Private/Materials/MaterialInterface.cpp
--- a/Engine/Source/Runtime/Engine/Private/Materials/MaterialInterface.cpp
+++ b/Engine/Source/Runtime/Engine/Private/Materials/MaterialInterface.cpp
@@ -12,6 +12,7 @@
 #include "PrimitiveViewRelevance.h"
 #include "MaterialShared.h"
 #include "Materials/Material.h"
+#include "Materials/MaterialSharedPrivate.h"
 #include "UObject/AssetRegistryTagsContext.h"
 #include "UObject/ObjectSaveContext.h"
 #include "UObject/PropertyPortFlags.h"
@@ -792,6 +793,15 @@ void UMaterialInterface::GetAssetRegistryTags(FAssetRegistryTagsContext Context)
 }
 #endif // WITH_EDITOR
 
+#if WITH_EDITOR
+void UMaterialInterface::AppendToClassSchema(FAppendToClassSchemaContext& Context)
+{
+	Super::AppendToClassSchema(Context);
+
+	UE::MaterialInterface::Private::HashMaterialStaticClassDependenciesForCook(Context);
+}
+#endif
+
 void UMaterialInterface::GetLightingGuidChain(bool bIncludeTextures, TArray<FGuid>& OutGuids) const
 {
 	const FMaterialCachedExpressionData& CachedData = GetCachedExpressionData();
diff --git a/Engine/Source/Runtime/Engine/Private/Materials/MaterialKeyGeneratorContext.cpp b/Engine/Source/Runtime/Engine/Private/Materials/MaterialKeyGeneratorContext.cpp
new file mode 100644
--- /dev/null
+++ b/Engine/Source/Runtime/Engine/Private/Materials/MaterialKeyGeneratorContext.cpp
@@ -0,0 +1,120 @@
+// Copyright Epic Games, Inc. All Rights Reserved.
+
+#include "MaterialKeyGeneratorContext.h"
+
+#if WITH_EDITOR
+#include "RHIStrings.h"
+#endif
+
+#if WITH_EDITOR
+
+FMaterialKeyGeneratorContext::FMaterialKeyGeneratorContext(
+	TUniqueFunction<void(const void* Data, uint64 Size)>&& InResultFunc, EShaderPlatform InShaderPlatform)
+	: KeyGen(MoveTemp(InResultFunc))
+	, ShaderFormat(InShaderPlatform < EShaderPlatform::SP_NumPlatforms ? 
+		LegacyShaderPlatformToShaderFormat(InShaderPlatform) : NAME_None)
+	, ShaderPlatform(InShaderPlatform)
+	, Mode(EMode::Emitting)
+{
+}
+
+FMaterialKeyGeneratorContext::FMaterialKeyGeneratorContext(FString& InResultString, EShaderPlatform InShaderPlatform)
+	: KeyGen(InResultString)
+	, ShaderFormat(InShaderPlatform < EShaderPlatform::SP_NumPlatforms ? 
+		LegacyShaderPlatformToShaderFormat(InShaderPlatform) : NAME_None)
+	, ShaderPlatform(InShaderPlatform)
+	, Mode(EMode::Emitting)
+{
+	InResultString.Reserve(16384);
+}
+
+FMaterialKeyGeneratorContext::FMaterialKeyGeneratorContext(FCbWriter& InWriter, EShaderPlatform InShaderPlatform)
+	: Writer(&InWriter)
+	, ShaderFormat(InShaderPlatform < EShaderPlatform::SP_NumPlatforms ? 
+		LegacyShaderPlatformToShaderFormat(InShaderPlatform) : NAME_None)
+	, ShaderPlatform(InShaderPlatform)
+	, Mode(EMode::Saving)
+{
+}
+
+FMaterialKeyGeneratorContext::FMaterialKeyGeneratorContext(FCbObjectView LoadRoot, EShaderPlatform InShaderPlatform)
+	: ObjectStack()
+	, ShaderFormat(InShaderPlatform < EShaderPlatform::SP_NumPlatforms ? 
+		LegacyShaderPlatformToShaderFormat(InShaderPlatform) : NAME_None)
+	, ShaderPlatform(InShaderPlatform)
+	, Mode(EMode::Loading)
+{
+	ObjectStack.Add(LoadRoot);
+}
+
+FMaterialKeyGeneratorContext::~FMaterialKeyGeneratorContext()
+{
+	switch (Mode)
+	{
+	case EMode::Emitting:
+		KeyGen.~FShaderKeyGenerator();
+		break;
+	case EMode::Saving:
+		break;
+	case EMode::Loading:
+		ObjectStack.~TArray();
+		break;
+	default:
+		checkNoEntry();
+		break;
+	}
+}
+
+FCbObjectView FMaterialKeyGeneratorContext::GetCurrentObject()
+{
+	// Caller only calls when type is EMode::Loading
+	return ObjectStack.Last();
+}
+
+void FMaterialKeyGeneratorContext::RecordObjectStart(FUtf8StringView Name)
+{
+	switch (Mode)
+	{
+	case EMode::Emitting:
+		break;
+	case EMode::Saving:
+		*Writer << Name;
+		Writer->BeginObject();
+		break;
+	case EMode::Loading:
+	{
+		FCbFieldView SubObjectField = GetCurrentObject()[Name];
+		FCbObjectView SubObject = SubObjectField.AsObjectView();
+		if (SubObjectField.HasError())
+		{
+			bHasLoadError = true;
+		}
+		ObjectStack.Add(SubObject);
+		break;
+	}
+	default:
+		checkNoEntry();
+		break;
+	}
+}
+
+void FMaterialKeyGeneratorContext::RecordObjectEnd()
+{
+	switch (Mode)
+	{
+	case EMode::Emitting:
+		break;
+	case EMode::Saving:
+		Writer->EndObject();
+		break;
+	case EMode::Loading:
+		check(ObjectStack.Num() >= 2);
+		ObjectStack.Pop(EAllowShrinking::No);
+		break;
+	default:
+		checkNoEntry();
+		break;
+	}
+}
+
+#endif
\ No newline at end of file
diff --git a/Engine/Source/Runtime/Engine/Private/Materials/MaterialShader.cpp b/Engine/Source/Runtime/Engine/Private/Materials/MaterialShader.cpp
--- a/Engine/Source/Runtime/Engine/Private/Materials/MaterialShader.cpp
+++ b/Engine/Source/Runtime/Engine/Private/Materials/MaterialShader.cpp
@@ -36,6 +36,16 @@
 #include "PSOPrecacheValidation.h"
 
 #if WITH_EDITOR
+#include "Algo/Sort.h"
+#include "Containers/AnsiString.h"
+#include "Cooker/CookDependency.h"
+#include "Interfaces/ITargetPlatform.h"
+#include "MaterialCachedData.h"
+#include "MaterialKeyGeneratorContext.h"
+#include "Materials/MaterialSharedPrivate.h"
+#include "Misc/ConfigCacheIni.h"
+#include "Serialization/CompactBinarySerialization.h"
+#include "Serialization/CompactBinaryWriter.h"
 #include "Serialization/MemoryReader.h"
 #endif
 
@@ -195,74 +205,104 @@ FString GetMaterialShaderMapKeyString(const FMaterialShaderMapId& ShaderMapId,
 {
 	FString ShaderMapKeyString;
 	ShaderMapKeyString.Reserve(16384);
-	FShaderKeyGenerator KeyGen(ShaderMapKeyString);
-	GetMaterialShaderMapKey(KeyGen, ShaderMapId, ShaderParameters, Platform,
-		bIncludeKeyStringShaderDependencies);
+	FMaterialKeyGeneratorContext Context(ShaderMapKeyString, Platform);
+	Context.SetFlags(EMaterialKeyInclude::ShaderDependencies, bIncludeKeyStringShaderDependencies);
+	RecordOrEmitMaterialShaderMapKey(Context, ShaderMapId, ShaderParameters);
+
 	return ShaderMapKeyString;
 }
 
-void GetMaterialShaderMapKey(FShaderKeyGenerator& KeyGen, const FMaterialShaderMapId& ShaderMapId,
-	const FMaterialShaderParameters& ShaderParameters, EShaderPlatform Platform,
-	bool bIncludeKeyStringShaderDependencies)
+void RecordOrEmitMaterialShaderMapKey(FMaterialKeyGeneratorContext& Context,
+	const FMaterialShaderMapId& ShaderMapId, const FMaterialShaderParameters& ShaderParameters)
 {
-	TRACE_CPUPROFILER_EVENT_SCOPE(GetMaterialShaderMapKeyString);
+	check(!Context.IsLoading());
+	RecordOrEmitMaterialShaderMapKey(Context, const_cast<FMaterialShaderMapId&>(ShaderMapId),
+		const_cast<FMaterialShaderParameters&>(ShaderParameters));
+}
 
-	FName Format = LegacyShaderPlatformToShaderFormat(Platform);
+/**
+ * Creates a ShaderKey for a given a shader map id, or saves and loads the fields of the ShaderMapId
+ * for creating the ShaderKey later.
+ */
+void RecordOrEmitMaterialShaderMapKey(FMaterialKeyGeneratorContext& Context,
+	FMaterialShaderMapId& ShaderMapId, FMaterialShaderParameters& ShaderParameters)
+{
+	TRACE_CPUPROFILER_EVENT_SCOPE(RecordOrEmitMaterialShaderMapKey);
+
+	FName Format = Context.GetShaderFormat();
+	EShaderPlatform Platform = Context.GetShaderPlatform();
+
+	// Emitting and recording are mutually exclusive and we rely on that
+	check(!(Context.IsRecording() && Context.IsEmitting()));
 	
-	KeyGen.AppendDebugText(TEXTVIEW("MATSM_"));
-	KeyGen.Append(GetMaterialShaderMapDDCGuid());
-	KeyGen.AppendSeparator();
+	if (Context.IsRecording())
+	{
+		Context.RecordObjectStart("ShaderMapId");
+		ShaderMapId.RecordAndEmit(Context);
+		Context.RecordObjectEnd();
+		Context.Record("ShaderParameters", ShaderParameters);
+	}
+	else
+	{
+		FShaderKeyGenerator& KeyGen = Context.GetKeyGenIfEmittingChecked();
 
-	KeyGen.Append(Format);
-	KeyGen.AppendSeparator();
-	KeyGen.Append(GetTargetPlatformManagerRef().ShaderFormatVersion(Format));
-	KeyGen.AppendSeparator();
+		KeyGen.AppendDebugText(TEXTVIEW("MATSM_"));
+		KeyGen.Append(GetMaterialShaderMapDDCGuid());
+		KeyGen.AppendSeparator();
 
-	ShaderMapAppendKey(Platform, KeyGen);
-	ShaderMapId.Append(KeyGen, true, bIncludeKeyStringShaderDependencies);
+		KeyGen.Append(Format);
+		KeyGen.AppendSeparator();
+		KeyGen.Append(GetTargetPlatformManagerRef().ShaderFormatVersion(Format));
+		KeyGen.AppendSeparator();
 
+		ShaderMapAppendKey(Platform, KeyGen);
+		ShaderMapId.RecordAndEmit(Context);
 
-	FMemoryHasherBlake3 EnvHasher;
+		FMemoryHasherBlake3 EnvHasher;
 
-	// Hash any modifications applied to the compilation environment for each shadertype that will be compiled.
-	// When the shadermap ID is constructed, individual shadertype dependencies are added for any pipelines, hence
-	// we don't need to have equivalent hashing code for the pipeline dependencies.
-	for (const FShaderTypeDependency& TypeDep : ShaderMapId.ShaderTypeDependencies)
-	{
-		const FShaderType* ShaderType = FindShaderTypeByName(TypeDep.ShaderTypeName);
-		if (const FMaterialShaderType* MatShaderType = ShaderType->GetMaterialShaderType())
-		{
-			if (MatShaderType->ShouldCompilePermutation(Platform, ShaderParameters, TypeDep.PermutationId, ShaderMapId.GetPermutationFlags()))
-			{
-				FShaderCompilerEnvironment EnvModifications(EnvHasher);
-				MatShaderType->SetupCompileEnvironment(Platform, ShaderParameters, TypeDep.PermutationId, ShaderMapId.GetPermutationFlags(), EnvModifications);
-				EnvModifications.SerializeEverythingButFiles(EnvHasher);
-			}
-		}
-		else if (const FMeshMaterialShaderType* MeshMatShaderType = ShaderType->GetMeshMaterialShaderType())
+		// Hash any modifications applied to the compilation environment for each shadertype that will be compiled.
+		// When the shadermap ID is constructed, individual shadertype dependencies are added for any pipelines, hence
+		// we don't need to have equivalent hashing code for the pipeline dependencies.
+		for (const FShaderTypeDependency& TypeDep : ShaderMapId.ShaderTypeDependencies)
 		{
-			for (const FVertexFactoryTypeDependency& VFDep : ShaderMapId.VertexFactoryTypeDependencies)
+			const FShaderType* ShaderType = FindShaderTypeByName(TypeDep.ShaderTypeName);
+			if (const FMaterialShaderType* MatShaderType = ShaderType->GetMaterialShaderType())
 			{
-				const FVertexFactoryType* VFType = FindVertexFactoryType(VFDep.VertexFactoryTypeName);
-				const bool bVFShouldCache = FMeshMaterialShaderType::ShouldCompileVertexFactoryPermutation(Platform, ShaderParameters, VFType, ShaderType, ShaderMapId.GetPermutationFlags());
-				const bool bShaderShouldCache = MeshMatShaderType->ShouldCompilePermutation(Platform, ShaderParameters, VFType, TypeDep.PermutationId, ShaderMapId.GetPermutationFlags());
-				if (bVFShouldCache && bShaderShouldCache)
+				if (MatShaderType->ShouldCompilePermutation(Platform, ShaderParameters, TypeDep.PermutationId, ShaderMapId.GetPermutationFlags()))
 				{
 					FShaderCompilerEnvironment EnvModifications(EnvHasher);
-					VFType->ModifyCompilationEnvironment(FVertexFactoryShaderPermutationParameters(Platform, ShaderParameters, VFType, ShaderType, ShaderMapId.GetPermutationFlags()), EnvModifications);
-					MeshMatShaderType->SetupCompileEnvironment(Platform, ShaderParameters, VFType, TypeDep.PermutationId, ShaderMapId.GetPermutationFlags(), EnvModifications);
+					MatShaderType->SetupCompileEnvironment(Platform, ShaderParameters, TypeDep.PermutationId, ShaderMapId.GetPermutationFlags(), EnvModifications);
 					EnvModifications.SerializeEverythingButFiles(EnvHasher);
 				}
 			}
+			else if (const FMeshMaterialShaderType* MeshMatShaderType = ShaderType->GetMeshMaterialShaderType())
+			{
+				for (const FVertexFactoryTypeDependency& VFDep : ShaderMapId.VertexFactoryTypeDependencies)
+				{
+					const FVertexFactoryType* VFType = FindVertexFactoryType(VFDep.VertexFactoryTypeName);
+					const bool bVFShouldCache = FMeshMaterialShaderType::ShouldCompileVertexFactoryPermutation(Platform, ShaderParameters, VFType, ShaderType, ShaderMapId.GetPermutationFlags());
+					const bool bShaderShouldCache = MeshMatShaderType->ShouldCompilePermutation(Platform, ShaderParameters, VFType, TypeDep.PermutationId, ShaderMapId.GetPermutationFlags());
+					if (bVFShouldCache && bShaderShouldCache)
+					{
+						FShaderCompilerEnvironment EnvModifications(EnvHasher);
+						VFType->ModifyCompilationEnvironment(FVertexFactoryShaderPermutationParameters(Platform, ShaderParameters, VFType, ShaderType, ShaderMapId.GetPermutationFlags()), EnvModifications);
+						MeshMatShaderType->SetupCompileEnvironment(Platform, ShaderParameters, VFType, TypeDep.PermutationId, ShaderMapId.GetPermutationFlags(), EnvModifications);
+						EnvModifications.SerializeEverythingButFiles(EnvHasher);
+					}
+				}
+			}
 		}
-	}
-	KeyGen.AppendDebugText(TEXTVIEW("_EMH_"));
-	KeyGen.Append(EnvHasher.Finalize());
-	KeyGen.AppendSeparator();
+		KeyGen.AppendDebugText(TEXTVIEW("_EMH_"));
+		KeyGen.Append(EnvHasher.Finalize());
+		KeyGen.AppendSeparator();
 
-	FMaterialAttributeDefinitionMap::AppendDDCKey(KeyGen);
-	FShaderCompileUtilities::AppendGBufferDDCKey(Platform, KeyGen);
-	FHLSLMaterialTranslator::AppendVersion(KeyGen, Platform);
+		if (Context.HasAllFlags(EMaterialKeyInclude::Globals))
+		{
+			FMaterialAttributeDefinitionMap::AppendDDCKey(KeyGen);
+		}
+		FShaderCompileUtilities::AppendGBufferDDCKey(Platform, KeyGen);
+		FHLSLMaterialTranslator::AppendVersion(KeyGen, Platform);
+	}
 }
 
 static UE::DerivedData::FCacheKey GetMaterialShaderMapKey(const FStringView MaterialShaderMapKey)
@@ -282,6 +322,209 @@ static UE::DerivedData::FSharedString GetMaterialShaderMapName(const FStringView
 		TEXTVIEW("]")));
 }
 
+
+FCbWriter& operator<<(FCbWriter& Writer, const EShaderPermutationFlags& PermutationFlags)
+{
+	Writer << static_cast<uint32>(PermutationFlags);
+	return Writer;
+}
+
+bool LoadFromCompactBinary(FCbFieldView Field, EShaderPermutationFlags& OutPermutationFlags)
+{
+	OutPermutationFlags = EShaderPermutationFlags::None;
+	uint32 IntValue = Field.AsUInt32();
+	if (Field.HasError())
+	{
+		return false;
+	}
+	OutPermutationFlags = static_cast<EShaderPermutationFlags>(IntValue);
+	return true;
+}
+
+#endif // WITH_EDITOR
+
+#if WITH_EDITOR
+
+namespace UE::MaterialInterface::Private
+{
+
+void FRecordedMaterialResourceForCooking::Save(FCbWriter& Writer) const
+{
+	Writer.BeginArray();
+	Writer << ShaderPlatform;
+	Writer << ExpressionIncludes;
+	Writer.BeginObject();
+	FMaterialKeyGeneratorContext KeyGenContext(Writer, ShaderPlatform);
+	KeyGenContext.RemoveFlags(EMaterialKeyInclude::Globals | EMaterialKeyInclude::UObjectData);
+	RecordOrEmitMaterialShaderMapKey(KeyGenContext, ShaderMapId, ShaderParameters);
+	Writer.EndObject();
+	Writer.EndArray();
+}
+
+bool FRecordedMaterialResourceForCooking::TryLoad(FCbFieldView Field)
+{
+	*this = FRecordedMaterialResourceForCooking();
+
+	FCbFieldViewIterator ElementIter = Field.CreateViewIterator();
+	if (!LoadFromCompactBinary(ElementIter++, ShaderPlatform))
+	{
+		return false;
+	}
+	if (!LoadFromCompactBinary(ElementIter++, ExpressionIncludes))
+	{
+		return false;
+	}
+	FCbObjectView ObjectView = ElementIter.AsObjectView();
+	if ((ElementIter++).HasError())
+	{
+		return false;
+	}
+	FMaterialKeyGeneratorContext KeyGenContext(ObjectView, ShaderPlatform);
+	KeyGenContext.RemoveFlags(EMaterialKeyInclude::Globals | EMaterialKeyInclude::UObjectData);
+	RecordOrEmitMaterialShaderMapKey(KeyGenContext, ShaderMapId, ShaderParameters);
+	if (KeyGenContext.HasLoadError())
+	{
+		return false;
+	}
+	return true;
+}
+
+} // namespace UE::MaterialInterface::Private
+
+bool LoadFromCompactBinary(FCbFieldView Field,
+	UE::MaterialInterface::Private::FRecordedMaterialResourceForCooking& OutValue)
+{
+	return OutValue.TryLoad(Field);
+}
+
+namespace UE::MaterialInterface::Private
+{
+
+void HashMaterialStaticClassDependenciesForCook(FAppendToClassSchemaContext& Context)
+{
+	Context.Update(&GetMaterialShaderMapDDCGuid(), sizeof(GetMaterialShaderMapDDCGuid()));
+
+	FShaderKeyGenerator KeyGen([&Context](const void* Data, uint64 Size) { Context.Update(Data, Size); });
+	FMaterialAttributeDefinitionMap::AppendDDCKey(KeyGen);
+
+	// Record the data used by all FPlatformTypeLayoutParameters by appending a default-constructed one.
+	FPlatformTypeLayoutParameters LayoutParams;
+	LayoutParams.Append(KeyGen);
+
+	bool bSupportCookedEditorConfigValue = false;
+	GConfig->GetBool(TEXT("CookedEditorSettings"), TEXT("bSupportCookedEditor"), bSupportCookedEditorConfigValue, GGameIni);
+	uint8 Byte = bSupportCookedEditorConfigValue ? 1 : 0;
+	Context.Update(&Byte, sizeof(Byte));
+
+	for (EMaterialQualityLevel::Type QualityLevel = static_cast<EMaterialQualityLevel::Type>(0);
+		QualityLevel < EMaterialQualityLevel::Type::Num;
+		QualityLevel = static_cast<EMaterialQualityLevel::Type>(static_cast<uint32>(QualityLevel) + 1))
+	{
+		TStringBuilder<32> QualityLevelName(InPlace, GetMaterialQualityLevelFName(QualityLevel));
+		Context.Update(*QualityLevelName, QualityLevelName.Len() * sizeof((*QualityLevelName)[0]));
+	}
+	for (ERHIFeatureLevel::Type FeatureLevel = static_cast<ERHIFeatureLevel::Type>(0);
+		FeatureLevel < ERHIFeatureLevel::Type::Num;
+		FeatureLevel = static_cast<ERHIFeatureLevel::Type>(static_cast<uint32>(FeatureLevel) + 1))
+	{
+		FName FeatureLevelName;
+		GetFeatureLevelName(FeatureLevel, FeatureLevelName);
+		TStringBuilder<32> FeatureLevelNameStr(InPlace, FeatureLevelName);
+		Context.Update(*FeatureLevelNameStr, FeatureLevelNameStr.Len() * sizeof((*FeatureLevelNameStr)[0]));
+	}
+}
+
+constexpr int32 HashMaterialDependenciesForCookArgsVersion = 1;
+void HashMaterialDependenciesForCook(FCbFieldViewIterator Args, UE::Cook::FCookDependencyContext& Context)
+{
+	int32 ArgsVersion = -1;
+	bool bValid = false;
+
+	TArray<FRecordedMaterialResourceForCooking> CookResources;
+	FCbFieldViewIterator ArgField(Args);
+	ArgsVersion = (ArgField++).AsInt32();
+	if (ArgsVersion == HashMaterialDependenciesForCookArgsVersion)
+	{
+		bValid = LoadFromCompactBinary(ArgField++, CookResources);
+	}
+	if (!bValid)
+	{
+		Context.LogError(FString::Printf(TEXT("Unsupported arguments version %d."), ArgsVersion));
+		return;
+	}
+
+	for (FRecordedMaterialResourceForCooking& Resource : CookResources)
+	{
+		FMaterialKeyGeneratorContext KeyGenContext([&Context](const void* Data, uint64 Size)
+			{
+				Context.Update(Data, Size);
+			}, Resource.ShaderPlatform);
+		KeyGenContext.RemoveFlags(EMaterialKeyInclude::Globals | EMaterialKeyInclude::UObjectData);
+		RecordOrEmitMaterialShaderMapKey(KeyGenContext, Resource.ShaderMapId, Resource.ShaderParameters);
+		for (const FString& FileName : Resource.ExpressionIncludes)
+		{
+			const FSHAHash* Hash = TryGetShaderFileHash(*FileName, Resource.ShaderPlatform);
+			if (!Hash)
+			{
+				Context.LogError(FString::Printf(TEXT("Could not find file %s for hashing."), *FileName));
+				continue;
+			}
+			Context.Update(&Hash->Hash, sizeof(Hash->Hash));
+		}
+	}
+}
+
+UE_COOK_DEPENDENCY_FUNCTION(HashMaterialDependenciesForCook, UE::MaterialInterface::Private::HashMaterialDependenciesForCook);
+
+void RecordMaterialDependenciesForCook(FObjectPreSaveContext ObjectSaveContext,
+	TConstArrayView<FMaterialResourceForCooking> Resources)
+{
+	const ITargetPlatform* TargetPlatform = ObjectSaveContext.GetTargetPlatform();
+	check(TargetPlatform); // Caller should call only when cooking
+
+	TArray<FRecordedMaterialResourceForCooking> RecordedResources;
+	RecordedResources.Reserve(Resources.Num());
+
+	// Sort the array of resources for determinism
+	TArray<FMaterialResourceForCooking> SortedResources(Resources);
+	Algo::Sort(SortedResources, [](const FMaterialResourceForCooking& A, const FMaterialResourceForCooking& B)
+		{
+			return static_cast<uint32>(A.Platform) < static_cast<uint32>(B.Platform);
+		});
+
+	// Calculate the ShaderMapId for each resource and save it into an array of FRecordedMaterialResourceForCooking
+	// that is suitable for persistent saving as parameters of HashMaterialDependenciesForCook.
+	for (const FMaterialResourceForCooking& ResourceForCooking : SortedResources)
+	{
+		const FMaterialResource& Resource = *ResourceForCooking.Resource;
+
+		FRecordedMaterialResourceForCooking& RecordedResource = RecordedResources.Emplace_GetRef();
+		RecordedResource.ShaderPlatform = ResourceForCooking.Platform;
+		RecordedResource.ShaderParameters = FMaterialShaderParameters(&Resource);
+		Resource.GetShaderMapId(RecordedResource.ShaderPlatform, TargetPlatform, RecordedResource.ShaderMapId);
+
+		TSet<FString>& ExpressionIncludes = Resource.GetCachedExpressionData().EditorOnlyData->ExpressionIncludeFilePaths;
+		// ExpressionIncludeFilePaths is already deterministically sorted, so no need to sort it here.
+		RecordedResource.ExpressionIncludes.Reserve(ExpressionIncludes.Num());
+		for (const FString& FilePath : ExpressionIncludes)
+		{
+			RecordedResource.ExpressionIncludes.Add(FilePath);
+		}
+	}
+
+	// Store the FRecordedMaterialResourceForCooking in compact binary and pass them to the cooker to save in 
+	// ObjectSaveContext's list of dependencies for our package.
+	FCbWriter Writer;
+	Writer << HashMaterialDependenciesForCookArgsVersion;
+	Writer << RecordedResources;
+
+	ObjectSaveContext.AddCookDependency(
+		UE::Cook::FCookDependency::Function(
+			UE_COOK_DEPENDENCY_FUNCTION_CALL(HashMaterialDependenciesForCook), Writer.Save()));
+}
+
+} // namespace UE::MaterialInterface::Private
+
 #endif // WITH_EDITOR
 
 /** Called for every material shader to update the appropriate stats. */
@@ -1219,90 +1462,136 @@ void FMaterialShaderMapId::AppendStaticParameters(FShaderKeyGenerator& KeyGen) c
 
 void FMaterialShaderMapId::AppendKeyString(FString& KeyString, bool bIncludeSourceAndMaterialState, bool bIncludeKeyStringShaderDependencies) const
 {
-	FShaderKeyGenerator KeyGen(KeyString);
-	Append(KeyGen, bIncludeSourceAndMaterialState, bIncludeKeyStringShaderDependencies);
+	FMaterialKeyGeneratorContext Context(KeyString, EShaderPlatform::SP_NumPlatforms);
+	Context.SetFlags(EMaterialKeyInclude::SourceAndMaterialState, bIncludeSourceAndMaterialState);
+	Context.SetFlags(EMaterialKeyInclude::ShaderDependencies, bIncludeKeyStringShaderDependencies);
+	const_cast<FMaterialShaderMapId&>(*this).RecordAndEmit(Context);
 }
 
-void FMaterialShaderMapId::Append(FShaderKeyGenerator& KeyGen, bool bIncludeSourceAndMaterialState, bool bIncludeKeyStringShaderDependencies) const
+void FMaterialShaderMapId::RecordAndEmit(FMaterialKeyGeneratorContext& Context)
 {
-	check(IsContentValid());
-	if (bIncludeSourceAndMaterialState)
+	check(Context.IsLoading() || IsContentValid());
+	if (Context.HasAllFlags(EMaterialKeyInclude::SourceAndMaterialState | EMaterialKeyInclude::UObjectData))
 	{
-		KeyGen.Append(BaseMaterialId);
-		KeyGen.AppendSeparator();
+		Context.RecordAndEmit("BaseMaterialId", BaseMaterialId);
+		Context.EmitSeparator();
 	}
 
-	KeyGen.Append(GetMaterialQualityLevelFName(QualityLevel));
-	KeyGen.AppendSeparator();
-
-	FName FeatureLevelName;
-	GetFeatureLevelName(FeatureLevel, FeatureLevelName);
-	KeyGen.Append(FeatureLevelName);
-	KeyGen.AppendSeparator();
+	if (Context.HasAllFlags(EMaterialKeyInclude::UObjectData))
+	{
+		Context.Record("QualityLevel", QualityLevel);
+		Context.Emit(GetMaterialQualityLevelFName(QualityLevel));
+		Context.EmitSeparator();
 
-	LayoutParams.Append(KeyGen);
+		Context.Record("FeatureLevel", FeatureLevel);
+		FName FeatureLevelName;
+		GetFeatureLevelName(FeatureLevel, FeatureLevelName);
+		Context.Emit(FeatureLevelName);
+		Context.EmitSeparator();
+	}
 
-	AppendStaticParameters(KeyGen);
+	Context.RecordAndEmit("LayoutParams", LayoutParams);
 
-	if (MaterialLayersId)
+	// Recording of some fields necessary for IncludeUObjectData is not yet implemented; we only support emitting them
+	check(!Context.HasAllFlags(EMaterialKeyInclude::UObjectData) || !Context.IsRecording());
+	if (Context.HasAllFlags(EMaterialKeyInclude::UObjectData) && Context.IsEmitting())
 	{
-		MaterialLayersId->Append(KeyGen);
-	}
+		FShaderKeyGenerator& KeyGen = Context.GetKeyGenIfEmittingChecked();
+		AppendStaticParameters(KeyGen);
 
-	KeyGen.AppendSeparator();
-	KeyGen.Append(static_cast<int32>(Usage));
-	KeyGen.AppendSeparator();
+		if (MaterialLayersId)
+		{
+			MaterialLayersId->Append(KeyGen);
+		}
 
-	if (Usage == EMaterialShaderMapUsage::MaterialExportCustomOutput)
-	{
-		KeyGen.Append(UsageCustomOutput);
 		KeyGen.AppendSeparator();
-	}
+		KeyGen.Append(static_cast<int32>(Usage));
+		KeyGen.AppendSeparator();
 
-	if (bIncludeSourceAndMaterialState)
-	{
-		::Append(KeyGen, ReferencedFunctions);
+		if (Usage == EMaterialShaderMapUsage::MaterialExportCustomOutput)
+		{
+			KeyGen.Append(UsageCustomOutput);
+			KeyGen.AppendSeparator();
+		}
+
+		if (Context.HasAllFlags(EMaterialKeyInclude::SourceAndMaterialState))
+		{
+			// Add any referenced functions to the key so that we will recompile when they are changed
+			::Append(KeyGen, ReferencedFunctions);
+		}
 	}
 
+	if (Context.HasAllFlags(EMaterialKeyInclude::Globals))
 	{
 		const FSHAHash LayoutHash = GetShaderTypeLayoutHash(StaticGetTypeLayoutDesc<FMaterialShaderMapContent>(), LayoutParams);
-		KeyGen.AppendSeparator();
-		KeyGen.Append(LayoutHash);
-		KeyGen.AppendSeparator();
+		Context.EmitSeparator();
+		Context.Emit(LayoutHash);
+		Context.EmitSeparator();
 	}
 
-	KeyGen.AppendSeparator();
+	Context.EmitSeparator();
 
-	if (bIncludeSourceAndMaterialState)
+	if (Context.HasAllFlags(EMaterialKeyInclude::SourceAndMaterialState | EMaterialKeyInclude::UObjectData))
 	{
-		::Append(KeyGen, ReferencedParameterCollections);
+		Context.RecordAndEmit("ReferencedParameterCollections", ReferencedParameterCollections);
 	}
 
 	// Add the inputs for any shaders that are stored inline in the shader map
-	if (bIncludeKeyStringShaderDependencies)
+	// Callers can use the Dependencies arrays, so record them all even if not HasAllFlags(ShaderDependencies)
+	Context.Record("ShaderTypeDependencies", ShaderTypeDependencies);
+	Context.Record("ShaderPipelineTypeDependencies", ShaderPipelineTypeDependencies);
+	Context.Record("VertexFactoryTypeDependencies", VertexFactoryTypeDependencies);
+	Context.PostLoad([this, ShaderPlatform=Context.GetShaderPlatform()]()
+		{
+			for (FShaderTypeDependency& Value : ShaderTypeDependencies)
+			{
+				Value.RefreshCachedSourceHash(ShaderPlatform);
+			}
+			for (FShaderPipelineTypeDependency& Value : ShaderPipelineTypeDependencies)
+			{
+				Value.RefreshCachedSourceHash(ShaderPlatform);
+			}
+			for (FVertexFactoryTypeDependency& Value : VertexFactoryTypeDependencies)
+			{
+				Value.RefreshCachedSourceHash(ShaderPlatform);
+			}
+		});
+	if (Context.HasAllFlags(EMaterialKeyInclude::ShaderDependencies) && Context.IsEmitting())
 	{
+		FShaderKeyGenerator& KeyGen = Context.GetKeyGenIfEmittingChecked();
 		AppendShaderDependencies(
 			KeyGen,
 			MakeArrayView(ShaderTypeDependencies),
 			MakeArrayView(ShaderPipelineTypeDependencies),
 			MakeArrayView(VertexFactoryTypeDependencies),
 			LayoutParams,
-			bIncludeSourceAndMaterialState);
+			Context.HasAllFlags(EMaterialKeyInclude::SourceAndMaterialState));
 	}
 
-	KeyGen.Append(TextureReferencesHash);
-
-	KeyGen.Append(ExpressionIncludesHash);
-
-	KeyGen.Append(BasePropertyOverridesHash);
+	// Recording of these hashes is not useful, because the hashes do not include the list of files that contributed
+	// to the hash.
+	check(!Context.HasAllFlags(EMaterialKeyInclude::UObjectData) || !Context.IsRecording());
+	if (Context.HasAllFlags(EMaterialKeyInclude::UObjectData))
+	{
+		Context.RecordAndEmit("TextureReferencesHash", TextureReferencesHash);
+		Context.RecordAndEmit("ExpressionIncludesHash", ExpressionIncludesHash);
+		Context.RecordAndEmit("BasePropertyOverridesHash", BasePropertyOverridesHash);
+	}
 
+	Context.Record("bUsingNewHLSLGenerator", bUsingNewHLSLGenerator);
 	if (bUsingNewHLSLGenerator)
 	{
-		KeyGen.AppendDebugText(TEXT("_NewHLSL"));
-		KeyGen.Append(FMaterialHLSLGenerator::Version);
+		Context.EmitDebugText(TEXTVIEW("_NewHLSL"));
+		Context.Emit(FMaterialHLSLGenerator::Version);
 	}
 
-	SubstrateCompilationConfig.Append(KeyGen);
+	// Recording of some fields necessary for IncludeUObjectData is not yet implemented; we only support emitting them
+	check(!Context.HasAllFlags(EMaterialKeyInclude::UObjectData) || !Context.IsRecording());
+	if (Context.HasAllFlags(EMaterialKeyInclude::UObjectData) && Context.IsEmitting())
+	{
+		FShaderKeyGenerator& KeyGen = Context.GetKeyGenIfEmittingChecked();
+		SubstrateCompilationConfig.Append(KeyGen);
+	}
 }
 
 void FMaterialShaderMapId::SetShaderDependencies(const TArray<FShaderType*>& ShaderTypes, const TArray<const FShaderPipelineType*>& ShaderPipelineTypes, const TArray<FVertexFactoryType*>& VFTypes, EShaderPlatform ShaderPlatform)
diff --git a/Engine/Source/Runtime/Engine/Private/Materials/MaterialShared.cpp b/Engine/Source/Runtime/Engine/Private/Materials/MaterialShared.cpp
--- a/Engine/Source/Runtime/Engine/Private/Materials/MaterialShared.cpp
+++ b/Engine/Source/Runtime/Engine/Private/Materials/MaterialShared.cpp
@@ -48,6 +48,8 @@
 #include "MaterialHLSLGenerator.h"
 #include "MaterialHLSLEmitter.h"
 #include "HLSLTree/HLSLTreeCommon.h"
+#include "Serialization/CompactBinary.h"
+#include "Serialization/CompactBinaryWriter.h"
 #include "Shader/PreshaderEvaluate.h"
 #endif
 #if WITH_ODSC
@@ -5502,4 +5504,28 @@ FMaterialShaderParameters::FMaterialShaderParameters(const FMaterial* InMaterial
 	bAlwaysEvaluateWorldPositionOffset = InMaterial->ShouldAlwaysEvaluateWorldPositionOffset();
 }
 
+#if WITH_EDITOR
+void FMaterialShaderParameters::Save(FCbWriter& Writer) const
+{
+	Writer.AddBinary(FMemoryView(this, sizeof(*this)));
+}
+
+bool FMaterialShaderParameters::TryLoad(FCbFieldView Field)
+{
+	*this = FMaterialShaderParameters();
+	FMemoryView MemoryView = Field.AsBinaryView();
+	if (Field.HasError() || MemoryView.GetSize() != sizeof(*this))
+	{
+		return false;
+	}
+	FMemory::Memcpy(this, MemoryView.GetData(), sizeof(*this));
+	return true;
+}
+
+bool LoadFromCompactBinary(FCbFieldView Field, FMaterialShaderParameters& OutValue)
+{
+	return OutValue.TryLoad(Field);
+}
+#endif // WITH_EDITOR
+
 #undef LOCTEXT_NAMESPACE
\ No newline at end of file
diff --git a/Engine/Source/Runtime/Engine/Private/Materials/MaterialSharedPrivate.h b/Engine/Source/Runtime/Engine/Private/Materials/MaterialSharedPrivate.h
new file mode 100644
--- /dev/null
+++ b/Engine/Source/Runtime/Engine/Private/Materials/MaterialSharedPrivate.h
@@ -0,0 +1,118 @@
+// Copyright Epic Games, Inc. All Rights Reserved.
+
+#pragma once
+
+#if WITH_EDITOR
+#include "Containers/Array.h"
+#include "Containers/ArrayView.h"
+#include "Containers/UnrealString.h"
+#include "MaterialShared.h"
+#include "RHIFeatureLevel.h"
+#include "RHIShaderPlatform.h"
+#include "SceneTypes.h"
+#include "Serialization/CompactBinary.h"
+#include "Serialization/CompactBinaryWriter.h"
+
+class FObjectPreSaveContext;
+struct FAppendToClassSchemaContext;
+struct FMaterialResourceForCooking;
+
+namespace UE::MaterialInterface::Private
+{
+
+/** Record the instance-specific cook dependencies (config, hlsl) from the resources for a UMaterialInterface. */
+void RecordMaterialDependenciesForCook(FObjectPreSaveContext ObjectSaveContext,
+	TConstArrayView<FMaterialResourceForCooking> Resources);
+
+/** Record the cook dependencies that apply to every UMaterialInterface. */
+void HashMaterialStaticClassDependenciesForCook(FAppendToClassSchemaContext& Context);
+
+/**
+ * A copy of the fields saved from an FMaterialResourceForCooking that includes only the data necessary for
+ * calculating non-UObject-based material dependencies. This struct is marshalled to CompactBinary and stored in cook
+ * metadata and read at the beginning of incremental cooks to check whether the material's dependencies have changed.
+ */
+struct FRecordedMaterialResourceForCooking
+{
+	FMaterialShaderMapId ShaderMapId;
+	FMaterialShaderParameters ShaderParameters;
+	EShaderPlatform ShaderPlatform;
+	/**
+	 * We have to handle ExpressionIncludes separately rather than using ShaderMapId.ExpressionIncludesHash,
+	 * because we need to record the filename of each expression include to recalculate it.
+	 */
+	TArray<FString> ExpressionIncludes;
+
+	void Save(FCbWriter& Writer) const;
+	bool TryLoad(FCbFieldView Field);
+};
+
+} // namespace UE::MaterialInterface::Private
+
+///////////////////////////////////////////////////////
+// CompactBinary interface for FRecordedMaterialResourceForCooking
+///////////////////////////////////////////////////////
+
+inline FCbWriter& operator<<(FCbWriter& Writer,
+	const UE::MaterialInterface::Private::FRecordedMaterialResourceForCooking& Value)
+{
+	Value.Save(Writer);
+	return Writer;
+}
+bool LoadFromCompactBinary(FCbFieldView Field,
+	UE::MaterialInterface::Private::FRecordedMaterialResourceForCooking& OutValue);
+
+///////////////////////////////////////////////////////
+// Implementation details for UE_DEFINEINLINE_COMPACTBINARY_ENUM_INT
+///////////////////////////////////////////////////////
+
+template <typename IntType>
+inline IntType GetViewFieldAsInteger(FCbFieldView FieldView, IntType DefaultValue)
+{
+	static_assert(sizeof(IntType) != 1, "Not implemented");
+}
+template <>
+inline uint8 GetViewFieldAsInteger(FCbFieldView FieldView, uint8 DefaultValue)
+{
+	return FieldView.AsUInt8(DefaultValue);
+}
+template <>
+inline uint16 GetViewFieldAsInteger(FCbFieldView FieldView, uint16 DefaultValue)
+{
+	return FieldView.AsUInt16(DefaultValue);
+}
+
+/**
+ * Declare and Define inline functions for LoadFromCompactBinary and operator<< for an enum, by converting
+ * the enum to an integer.
+ */
+#define UE_DEFINEINLINE_COMPACTBINARY_ENUM_INT(EnumType, IntType, NumValues, InvalidValue)				\
+	inline bool LoadFromCompactBinary(FCbFieldView Field, EnumType& OutValue)							\
+	{																									\
+		IntType IntValue = GetViewFieldAsInteger<IntType>(Field, static_cast<IntType>(InvalidValue));	\
+		if (Field.HasError() || IntValue >= static_cast<IntType>(NumValues))							\
+		{																								\
+			OutValue = static_cast<EnumType>(InvalidValue);												\
+			return false;																				\
+		}																								\
+		OutValue = static_cast<EnumType>(IntValue);														\
+		return true;																					\
+	}																									\
+	inline FCbWriter& operator<<(FCbWriter& Writer, EnumType Value)										\
+	{																									\
+		Writer << static_cast<IntType>(Value);															\
+		return Writer;																					\
+	}
+
+///////////////////////////////////////////////////////
+// CompactBinary functions for enums used by FRecordedMaterialResourceForCooking
+///////////////////////////////////////////////////////
+
+UE_DEFINEINLINE_COMPACTBINARY_ENUM_INT(EShaderPlatform, uint16, SP_NumPlatforms, SP_NumPlatforms);
+UE_DEFINEINLINE_COMPACTBINARY_ENUM_INT(EMaterialQualityLevel::Type, uint8,
+	EMaterialQualityLevel::Num, EMaterialQualityLevel::Low);
+UE_DEFINEINLINE_COMPACTBINARY_ENUM_INT(ERHIFeatureLevel::Type, uint8,
+	ERHIFeatureLevel::Num, ERHIFeatureLevel::ES3_1);
+
+
+#endif // WITH_EDITOR
diff --git a/Engine/Source/Runtime/Engine/Public/MaterialKeyGeneratorContext.h b/Engine/Source/Runtime/Engine/Public/MaterialKeyGeneratorContext.h
new file mode 100644
--- /dev/null
+++ b/Engine/Source/Runtime/Engine/Public/MaterialKeyGeneratorContext.h
@@ -0,0 +1,341 @@
+// Copyright Epic Games, Inc. All Rights Reserved.
+
+#pragma once
+
+#if WITH_EDITOR
+#include "Containers/StringView.h"
+#include "Containers/UnrealString.h"
+#include "Misc/EnumClassFlags.h"
+#include "Misc/StringBuilder.h"
+#include "RHIShaderPlatform.h"
+#include "Serialization/CompactBinary.h"
+#include "Serialization/CompactBinaryWriter.h"
+#include "Serialization/ShaderKeyGenerator.h"
+#include "Templates/Function.h"
+#include "UObject/NameTypes.h"
+
+/**
+ * Categories of fields which should be included in the output of an FMaterialKeyGeneratorContext. Code recording or
+ * emitting fields to a GeneratorContext calls HasAllFlags for the flags relevant to a given field before
+ * writing it.
+ */
+enum class EMaterialKeyInclude : uint32
+{
+	/** shadertype, shaderpipelinetype, vertexfactortype dependencies are included in the output */
+	ShaderDependencies = 0x1,
+	/** shader hlsl file hashes are included in the output */
+	SourceAndMaterialState = 0x2,
+	/** global data that applies to all shaders is included in the output */
+	Globals = 0x4,
+	/** data stored in UObject exports in .uasset and .umap files is included in the output */
+	UObjectData = 0x8,
+
+	All = ShaderDependencies | SourceAndMaterialState | Globals | UObjectData,
+};
+
+/**
+ * Output class passed to RecordAndEmit functions for material shader data. It receives function calls that either
+ * save/load variables to compact binary or emits those variables to an FShaderKeyGenerator to construct the key.
+ * It is written to by the function RecordOrEmitMaterialShaderMapKey and its helper functions, to create the
+ * shadermap's DDC key, to hash the Material's cook dependencies for incremental cooks, and to save/load the inputs
+ * for those cook dependencies to cook metadata, with a single function definition that lists all the relevant
+ * variables. This reduces the amount of boilerplate that has to be written and maintained for each variable that
+ * can affect the DDC key and cook dependencies.
+ * 
+ * The Context additionally holds flags for which portion of data should be included in its output; these flags
+ * are set differently by the calling function depending on whether we are saving a DDC key or material dependencies,
+ * and on what kind of DDC key we are saving.
+ * 
+ * The Context additionally holds some common arguments - ShaderPlatform and ShaderFormat - that are provided to the
+ * functions from the calling functions.
+ */
+class FMaterialKeyGeneratorContext
+{
+public:
+	/** Which mode the GeneratorContext is operating in; these modes are mutually exclusive, not flags. */
+	enum class EMode : uint8
+	{
+		Emitting,
+		Saving,
+		Loading,
+	};
+
+public:
+	/**
+	 * Construct a Context that is emitting variables into a binary-formatted shader DDC key or into a material
+	 * dependencies hash function.
+	 */
+	ENGINE_API FMaterialKeyGeneratorContext(TUniqueFunction<void(const void* Data, uint64 Size)>&& HashFunction,
+		EShaderPlatform InShaderPlatform);
+	/** Construct a Context that is emitting variables into a text-formatted shader DDC key. */
+	ENGINE_API FMaterialKeyGeneratorContext(FString& InResultString, EShaderPlatform InShaderPlatform);
+	/** Construct a Context that is saving data to compact binary. */
+	ENGINE_API FMaterialKeyGeneratorContext(FCbWriter& InWriter, EShaderPlatform InShaderPlatform);
+	/** Construct a Context that is loading data from compact binary. */
+	ENGINE_API FMaterialKeyGeneratorContext(FCbObjectView LoadRoot, EShaderPlatform InShaderPlatform);
+	ENGINE_API ~FMaterialKeyGeneratorContext();
+
+	/** Add include flags. @see EMaterialKeyInclude. */
+	inline void AddFlags(EMaterialKeyInclude Flags);
+	/** Remove include flags. @see EMaterialKeyInclude. */
+	inline void RemoveFlags(EMaterialKeyInclude Flags);
+	/** Add or remove include flags, depending on the value of bIncluded. @see EMaterialKeyInclude. */
+	inline void SetFlags(EMaterialKeyInclude Flags, bool bIncluded);
+	/** Report whether all requested flags are included. @see EMaterialKeyInclude. */
+	inline bool HasAllFlags(EMaterialKeyInclude Flags) const;
+	/** Return the list of include flags. @see EMaterialKeyInclude. */
+	inline EMaterialKeyInclude GetFlags() const;
+
+	/** Return the mode the Context is operating in. See also the Is functions that returns true when in each mode. */
+	inline EMode GetMode() const;
+	/** True iff the Context is in emit mode - creating DDC key or hashing material dependencies. */
+	inline bool IsEmitting() const;
+	/** True iff the Context is in either IsSaving or IsLoading mode. */
+	inline bool IsRecording() const;
+	/** True iff the Context is saving to compact binary. */
+	inline bool IsSaving() const;
+	/** True iff the Context is loading from compact binary. */
+	inline bool IsLoading() const;
+	/** The ShaderPlatform provided by the caller. */
+	inline EShaderPlatform GetShaderPlatform() const;
+	/** The ShaderFormat that corresponds to the ShaderPlatform provided by the caller. */
+	inline FName GetShaderFormat() const;
+
+	/**
+	 * Output function used by all three modes. The given data is either saved, loaded, or emitted. If saving or
+	 * loading, the Name is used as the id for the data. If emitting the Name is not used.
+	 * The name is scoped by the object level defined by RecordObjectStart/RecordObjectEnd.
+	 */
+	template <typename T>
+	inline void RecordAndEmit(FUtf8StringView Name, T&& Data);
+
+	/**
+	 * When saving or loading, save the given Data with the given name. When emitting, this function is a noop.
+	 * The name is scoped by the object level defined by RecordObjectStart/RecordObjectEnd.
+	 */
+	template <typename T>
+	inline void Record(FUtf8StringView Name, T&& Data);
+	/**
+	 * When saving or loading, start a new object scope with the given name. When emitting, this function is a noop.
+	 * The name is scoped in the parent object level defined by previous calls to RecordObjectStart/RecordObjectEnd.
+	 */
+	ENGINE_API void RecordObjectStart(FUtf8StringView Name);
+	/**
+	 * When saving or loading, end the object scope started by the last call to RecordObjectStart, and return to the
+	 * parent object scope. When emitting, this function is a noop.
+	 */
+	ENGINE_API void RecordObjectEnd();
+	/** When loading, execute the given function. In all other modes, this function is a noop. */
+	inline void PostLoad(TFunctionRef<void()> Action);
+
+	/**
+	 * When emitting, append the given Data to the ShaderKeyGenerator, using operator<<. In all other modes, this
+	 * function is a noop.
+	 */
+	template <typename T>
+	inline void Emit(T&& Data);
+	/** When emitting, call Callback with the ShaderKeyGenerator. In all other modes, this function is a noop. */
+	inline void EmitFunc(TFunctionRef<void(FShaderKeyGenerator&)> Callback);
+	/**  When emitting, call the ShaderKeyGenerator's AppendDebugText. In all other modes, this function is a noop. */
+	inline void EmitDebugText(FStringView Data);
+	/**  When emitting, call the ShaderKeyGenerator's AppendSeparator. In all other modes, this function is a noop. */
+	inline void EmitSeparator();
+	/**  When emitting, call the ShaderKeyGenerator's AppendBoolInt. In all other modes, this function is a noop. */
+	inline void EmitBoolInt(bool Data);
+	/** When emitting, return pointer to the ShaderKeyGenerator. In all other modes, returns nullptr. */
+	inline FShaderKeyGenerator* GetKeyGenIfEmitting();
+	/** Must only be called when emitting; fails assertion if not. Returns reference to the ShaderKeyGenerator. */
+	inline FShaderKeyGenerator& GetKeyGenIfEmittingChecked();
+
+	/** When loading, report whether any Record function reported an error. In all other modes, always returns true. */
+	inline bool HasLoadError() const;
+
+private:
+	FCbObjectView GetCurrentObject();
+
+private:
+	union
+	{
+		FShaderKeyGenerator KeyGen;
+		TArray<FCbObjectView, TInlineAllocator<2>> ObjectStack;
+		FCbWriter* Writer;
+	};
+	FName ShaderFormat;
+	EShaderPlatform ShaderPlatform = EShaderPlatform::SP_NumPlatforms;
+	EMode Mode = EMode::Emitting;
+	EMaterialKeyInclude IncludeFlags = EMaterialKeyInclude::All;
+	bool bHasLoadError = false;
+};
+
+ENUM_CLASS_FLAGS(EMaterialKeyInclude);
+
+///////////////////////////////////////////////////////
+// Inline implementations
+///////////////////////////////////////////////////////
+
+inline void FMaterialKeyGeneratorContext::AddFlags(EMaterialKeyInclude Flags)
+{
+	EnumAddFlags(IncludeFlags, Flags);
+}
+
+inline void FMaterialKeyGeneratorContext::RemoveFlags(EMaterialKeyInclude Flags)
+{
+	EnumRemoveFlags(IncludeFlags, Flags);
+}
+
+inline void FMaterialKeyGeneratorContext::SetFlags(EMaterialKeyInclude Flags, bool bIncluded)
+{
+	if (bIncluded)
+	{
+		AddFlags(Flags);
+	}
+	else
+	{
+		RemoveFlags(Flags);
+	}
+}
+
+inline bool FMaterialKeyGeneratorContext::HasAllFlags(EMaterialKeyInclude Flags) const
+{
+	return EnumHasAllFlags(IncludeFlags, Flags);
+}
+
+inline EMaterialKeyInclude FMaterialKeyGeneratorContext::GetFlags() const
+{
+	return IncludeFlags;
+}
+
+FMaterialKeyGeneratorContext::EMode FMaterialKeyGeneratorContext::GetMode() const
+{
+	return Mode;
+}
+
+inline bool FMaterialKeyGeneratorContext::IsEmitting() const
+{
+	return Mode == EMode::Emitting;
+}
+
+inline bool FMaterialKeyGeneratorContext::IsRecording() const
+{
+	return (Mode == EMode::Saving) | (Mode == EMode::Loading);
+}
+
+inline bool FMaterialKeyGeneratorContext::IsSaving() const
+{
+	return Mode == EMode::Saving;
+}
+
+inline bool FMaterialKeyGeneratorContext::IsLoading() const
+{
+	return Mode == EMode::Loading;
+}
+
+inline EShaderPlatform FMaterialKeyGeneratorContext::GetShaderPlatform() const
+{
+	return ShaderPlatform;
+}
+
+inline FName FMaterialKeyGeneratorContext::GetShaderFormat() const
+{
+	return ShaderFormat;
+}
+
+template <typename T>
+inline void FMaterialKeyGeneratorContext::RecordAndEmit(FUtf8StringView Name, T&& Data)
+{
+	Record(Name, Data);
+	Emit(Data);
+}
+
+template <typename T>
+inline void FMaterialKeyGeneratorContext::Record(FUtf8StringView Name, T&& Data)
+{
+	switch (Mode)
+	{
+	case EMode::Emitting:
+		break;
+	case EMode::Saving:
+		*Writer << Name << Data;
+		break;
+	case EMode::Loading:
+		if (!LoadFromCompactBinary(GetCurrentObject()[Name], Data))
+		{
+			bHasLoadError = true;
+		}
+		break;
+	default:
+		checkNoEntry();
+		break;
+	}
+}
+
+inline void FMaterialKeyGeneratorContext::PostLoad(TFunctionRef<void()> Action)
+{
+	if (Mode == EMode::Loading)
+	{
+		Action();
+	}
+}
+
+template <typename T>
+inline void FMaterialKeyGeneratorContext::Emit(T&& Data)
+{
+	if (Mode == EMode::Emitting)
+	{
+		KeyGen << Data;
+	}
+}
+
+inline void FMaterialKeyGeneratorContext::EmitFunc(TFunctionRef<void(FShaderKeyGenerator&)> Callback)
+{
+	if (Mode == EMode::Emitting)
+	{
+		Callback(KeyGen);
+	}
+}
+
+inline void FMaterialKeyGeneratorContext::EmitDebugText(FStringView Data)
+{
+	if (Mode == EMode::Emitting)
+	{
+		KeyGen.AppendDebugText(Data);
+	}
+}
+
+inline void FMaterialKeyGeneratorContext::EmitSeparator()
+{
+	if (Mode == EMode::Emitting)
+	{
+		KeyGen.AppendSeparator();
+	}
+}
+
+inline void FMaterialKeyGeneratorContext::EmitBoolInt(bool Data)
+{
+	if (Mode == EMode::Emitting)
+	{
+		KeyGen.AppendBoolInt(Data);
+	}
+}
+
+inline FShaderKeyGenerator* FMaterialKeyGeneratorContext::GetKeyGenIfEmitting()
+{
+	if (Mode == EMode::Emitting)
+	{
+		return &KeyGen;
+	}
+	return nullptr;
+}
+
+inline FShaderKeyGenerator& FMaterialKeyGeneratorContext::GetKeyGenIfEmittingChecked()
+{
+	check(Mode == EMode::Emitting);
+	return KeyGen;
+}
+
+inline bool FMaterialKeyGeneratorContext::HasLoadError() const
+{
+	return Mode != EMode::Loading || bHasLoadError;
+}
+
+#endif
\ No newline at end of file
diff --git a/Engine/Source/Runtime/Engine/Public/MaterialShaderType.h b/Engine/Source/Runtime/Engine/Public/MaterialShaderType.h
--- a/Engine/Source/Runtime/Engine/Public/MaterialShaderType.h
+++ b/Engine/Source/Runtime/Engine/Public/MaterialShaderType.h
@@ -37,6 +37,10 @@ struct FMaterialShadingModelField;
 enum EBlendMode : int;
 enum EMaterialShadingModel : int;
 enum class EShaderCompileJobPriority : uint8;
+#if WITH_EDITOR
+class FMaterialKeyGeneratorContext;
+#endif
+
 
 DECLARE_DELEGATE_RetVal_OneParam(FString, FShadingModelToStringDelegate, EMaterialShadingModel)
 
@@ -65,12 +69,25 @@ extern ENGINE_API FString GetMaterialShaderMapKeyString(
 	const FMaterialShaderParameters& ShaderParameters,
 	EShaderPlatform Platform,
 	bool bIncludeKeyStringShaderDependencies = true);
-extern ENGINE_API void GetMaterialShaderMapKey(
-	FShaderKeyGenerator& KeyGen,
+
+extern ENGINE_API void RecordOrEmitMaterialShaderMapKey(
+	FMaterialKeyGeneratorContext& Context,
+	FMaterialShaderMapId& ShaderMapId,
+	FMaterialShaderParameters& ShaderParameters);
+
+// Alternate arguments for RecordOrEmitMaterialShaderMapKey: Support being called with const&
+// when emitting or saving.
+extern ENGINE_API void RecordOrEmitMaterialShaderMapKey(
+	FMaterialKeyGeneratorContext& Context,
 	const FMaterialShaderMapId& ShaderMapId,
-	const FMaterialShaderParameters& ShaderParameters,
-	EShaderPlatform Platform,
-	bool bIncludeKeyStringShaderDependencies = true);
+	const FMaterialShaderParameters& ShaderParameters);
+inline void RecordOrEmitMaterialShaderMapKey(
+	FMaterialKeyGeneratorContext& Context,
+	FMaterialShaderMapId& ShaderMapId,
+	FMaterialShaderParameters&& ShaderParameters)
+{
+	RecordOrEmitMaterialShaderMapKey(Context, ShaderMapId, ShaderParameters);
+}
 #endif
 
 /** Called for every material shader to update the appropriate stats. */
diff --git a/Engine/Source/Runtime/Engine/Public/MaterialShared.h b/Engine/Source/Runtime/Engine/Public/MaterialShared.h
--- a/Engine/Source/Runtime/Engine/Public/MaterialShared.h
+++ b/Engine/Source/Runtime/Engine/Public/MaterialShared.h
@@ -55,6 +55,8 @@
 
 struct FExpressionInput;
 struct FExtraShaderCompilerSettings;
+class FCbFieldView;
+class FCbWriter;
 class FMaterial;
 class FMaterialCompiler;
 class FMaterialRenderProxy;
@@ -91,6 +93,7 @@ class FMaterialHLSLGenerator;
 class FShaderMapLayout;
 #if WITH_EDITOR
 class FMaterialCachedHLSLTree;
+class FMaterialKeyGeneratorContext;
 #endif
 
 enum EMaterialDomain : int;
@@ -979,6 +982,19 @@ struct FMaterialShaderParameters
 	};
 
 	FMaterialShaderParameters(const FMaterial* InMaterial=nullptr);
+
+private:
+#if WITH_EDITOR
+	// Compact binary API with hidden friend operator<<
+	ENGINE_API void Save(FCbWriter& Writer) const;
+	bool TryLoad(FCbFieldView Field);
+	friend inline FCbWriter& operator<<(FCbWriter& Writer, const FMaterialShaderParameters& Value)
+	{
+		Value.Save(Writer);
+		return Writer;
+	}
+	friend ENGINE_API bool LoadFromCompactBinary(FCbFieldView Field, FMaterialShaderParameters& OutValue);
+#endif
 };
 
 /** Contains all the information needed to uniquely identify a FMaterialShaderMap. */
@@ -1170,8 +1186,8 @@ public:
 	/** Appends string representations of this Id to a key string. */
 	void AppendKeyString(FString& KeyString, bool bIncludeSourceAndMaterialState = true,
 		bool bIncludeKeyStringShaderDependencies = true) const;
-	void Append(FShaderKeyGenerator& KeyGen, bool bIncludeSourceAndMaterialState = true,
-		bool bIncludeKeyStringShaderDependencies = true) const;
+	void RecordAndEmit(FMaterialKeyGeneratorContext& Context);
+
 	UE_DEPRECATED(5.5, "Internal function only, called through AppendKeyString.")
 	void AppendStaticParametersString(FString& ParamsString) const;
 
diff --git a/Engine/Source/Runtime/RenderCore/Private/Shader.cpp b/Engine/Source/Runtime/RenderCore/Private/Shader.cpp
--- a/Engine/Source/Runtime/RenderCore/Private/Shader.cpp
+++ b/Engine/Source/Runtime/RenderCore/Private/Shader.cpp
@@ -31,6 +31,11 @@
 #include "Interfaces/IShaderFormat.h"
 #endif
 
+#if WITH_EDITOR
+#include "Serialization/CompactBinary.h"
+#include "Serialization/CompactBinaryWriter.h"
+#endif
+
 #if RHI_RAYTRACING
 #include "RayTracingPayloadType.h"
 #endif
@@ -1187,6 +1192,95 @@ EShaderPermutationPrecacheRequest FShaderPipelineType::ShouldPrecachePermutation
 	return Result;
 }
 
+void FShaderTypeDependency::RefreshCachedSourceHash(EShaderPlatform ShaderPlatform)
+{
+	const FShaderType*const* ShaderTypePtr = FShaderType::GetNameToTypeMap().Find(ShaderTypeName);
+	const FShaderType* ShaderType = ShaderTypePtr ? *ShaderTypePtr : nullptr;
+	if (!ShaderType)
+	{
+		SourceHash = FSHAHash();
+		return;
+	}
+	SourceHash = ShaderType->GetSourceHash(ShaderPlatform);
+}
+
+void FShaderPipelineTypeDependency::RefreshCachedSourceHash(EShaderPlatform ShaderPlatform)
+{
+	const FShaderPipelineType* ShaderPipelineType =
+		FShaderPipelineType::GetShaderPipelineTypeByName(ShaderPipelineTypeName);
+	if (!ShaderPipelineType)
+	{
+		StagesSourceHash = FSHAHash();
+		return;
+	}
+	StagesSourceHash = ShaderPipelineType->GetSourceHash(ShaderPlatform);
+}
+
+#if WITH_EDITOR
+
+void FShaderTypeDependency::Save(FCbWriter& Writer) const
+{
+	Writer.BeginArray();
+	Writer << ShaderTypeName;
+	Writer << SourceHash;
+	Writer << PermutationId;
+	Writer.EndArray();
+}
+
+bool FShaderTypeDependency::TryLoad(FCbFieldView Field)
+{
+	*this = FShaderTypeDependency();
+	FCbFieldViewIterator ElementField(Field.CreateViewIterator());
+	if (!LoadFromCompactBinary(ElementField++, ShaderTypeName))
+	{
+		return false;
+	}
+	if (!LoadFromCompactBinary(ElementField++, SourceHash))
+	{
+		return false;
+	}
+	PermutationId = ElementField.AsInt32();
+	if ((ElementField++).HasError())
+	{
+		return false;
+	}
+	return true;
+}
+
+bool LoadFromCompactBinary(FCbFieldView Field, FShaderTypeDependency& OutValue)
+{
+	return OutValue.TryLoad(Field);
+}
+
+void FShaderPipelineTypeDependency::Save(FCbWriter& Writer) const
+{
+	Writer.BeginArray();
+	Writer << ShaderPipelineTypeName;
+	Writer << StagesSourceHash;
+	Writer.EndArray();
+}
+
+bool FShaderPipelineTypeDependency::TryLoad(FCbFieldView Field)
+{
+	*this = FShaderPipelineTypeDependency();
+	FCbFieldViewIterator ElementField(Field.CreateViewIterator());
+	if (!LoadFromCompactBinary(ElementField++, ShaderPipelineTypeName))
+	{
+		return false;
+	}
+	if (!LoadFromCompactBinary(ElementField++, StagesSourceHash))
+	{
+		return false;
+	}
+	return true;
+}
+
+bool LoadFromCompactBinary(FCbFieldView Field, FShaderPipelineTypeDependency& OutValue)
+{
+	return OutValue.TryLoad(Field);
+}
+#endif // WITH_EDITOR
+
 void FShaderPipeline::AddShader(FShader* Shader, int32 PermutationId)
 {
 	const EShaderFrequency Frequency = Shader->GetFrequency();
diff --git a/Engine/Source/Runtime/RenderCore/Private/VertexFactory.cpp b/Engine/Source/Runtime/RenderCore/Private/VertexFactory.cpp
--- a/Engine/Source/Runtime/RenderCore/Private/VertexFactory.cpp
+++ b/Engine/Source/Runtime/RenderCore/Private/VertexFactory.cpp
@@ -10,6 +10,10 @@
 #include "ShaderCompilerCore.h"
 #include "RenderUtils.h"
 #include "DataDrivenShaderPlatformInfo.h"
+#if WITH_EDITOR
+#include "Serialization/CompactBinary.h"
+#include "Serialization/CompactBinaryWriter.h"
+#endif
 
 IMPLEMENT_TYPE_LAYOUT(FVertexFactoryShaderParameters);
 
@@ -195,6 +199,47 @@ FVertexFactoryType* FindVertexFactoryType(const FHashedName& TypeName)
 	return FVertexFactoryType::GetVFByName(TypeName);
 }
 
+void FVertexFactoryTypeDependency::RefreshCachedSourceHash(EShaderPlatform ShaderPlatform)
+{
+	const FVertexFactoryType* VertexFactory = FindVertexFactoryType(VertexFactoryTypeName);
+	if (!VertexFactory)
+	{
+		VFSourceHash = FSHAHash();
+		return;
+	}
+	VFSourceHash = VertexFactory->GetSourceHash(ShaderPlatform);
+}
+
+#if WITH_EDITOR
+void FVertexFactoryTypeDependency::Save(FCbWriter& Writer) const
+{
+	Writer.BeginArray();
+	Writer << VertexFactoryTypeName;
+	Writer << VFSourceHash;
+	Writer.EndArray();
+}
+
+bool FVertexFactoryTypeDependency::TryLoad(FCbFieldView Field)
+{
+	*this = FVertexFactoryTypeDependency();
+	FCbFieldViewIterator ElementField(Field.CreateViewIterator());
+	if (!LoadFromCompactBinary(ElementField++, VertexFactoryTypeName))
+	{
+		return false;
+	}
+	if (!LoadFromCompactBinary(ElementField++, VFSourceHash))
+	{
+		return false;
+	}
+	return true;
+}
+
+bool LoadFromCompactBinary(FCbFieldView Field, FVertexFactoryTypeDependency& OutValue)
+{
+	return OutValue.TryLoad(Field);
+}
+#endif // WITH_EDITOR
+
 void FVertexFactory::GetStreams(ERHIFeatureLevel::Type InFeatureLevel, EVertexInputStreamType VertexStreamType, FVertexInputStreamArray& OutVertexStreams) const
 {
 	check(IsInitialized());
diff --git a/Engine/Source/Runtime/RenderCore/Public/Shader.h b/Engine/Source/Runtime/RenderCore/Public/Shader.h
--- a/Engine/Source/Runtime/RenderCore/Public/Shader.h
+++ b/Engine/Source/Runtime/RenderCore/Public/Shader.h
@@ -62,6 +62,8 @@
 #endif
 
 class ITargetPlatform;
+class FCbFieldView;
+class FCbWriter;
 class FComputeKernelShaderType;
 class FGlobalShaderType;
 class FMaterialShaderType;
@@ -1906,6 +1908,9 @@ public:
 		return !(*this == Reference);
 	}
 
+	/** Call GetShaderFileHash to get the cached value for the filename's hash in the current process. */
+	RENDERCORE_API void RefreshCachedSourceHash(EShaderPlatform ShaderPlatform);
+
 	/** Shader type */
 	LAYOUT_FIELD(FHashedName, ShaderTypeName);
 
@@ -1914,6 +1919,19 @@ public:
 
 	/** Used to detect changes to the shader source files. This is always present, as this type is sometimes frozen. */
 	LAYOUT_FIELD(FSHAHash, SourceHash);
+
+private:
+#if WITH_EDITOR
+	// Compact binary API with hidden friend operator<<
+	RENDERCORE_API void Save(FCbWriter& Writer) const;
+	bool TryLoad(FCbFieldView Field);
+	friend inline FCbWriter& operator<<(FCbWriter& Writer, const FShaderTypeDependency& Value)
+	{
+		Value.Save(Writer);
+		return Writer;
+	}
+	friend RENDERCORE_API bool LoadFromCompactBinary(FCbFieldView Field, FShaderTypeDependency& OutValue);
+#endif
 };
 
 
@@ -1952,6 +1970,22 @@ public:
 	{
 		return !(*this == Reference);
 	}
+
+	/** Call GetShaderFileHash to get the cached value for the filename's hash in the current process. */
+	RENDERCORE_API void RefreshCachedSourceHash(EShaderPlatform ShaderPlatform);
+
+private:
+#if WITH_EDITOR
+	// Compact binary API with hidden friend operator<<
+	RENDERCORE_API void Save(FCbWriter& Writer) const;
+	bool TryLoad(FCbFieldView Field);
+	friend inline FCbWriter& operator<<(FCbWriter& Writer, const FShaderPipelineTypeDependency& Value)
+	{
+		Value.Save(Writer);
+		return Writer;
+	}
+	friend RENDERCORE_API bool LoadFromCompactBinary(FCbFieldView Field, FShaderPipelineTypeDependency& OutValue);
+#endif
 };
 
 /** Used to compare two shader types by name. */
diff --git a/Engine/Source/Runtime/RenderCore/Public/VertexFactory.h b/Engine/Source/Runtime/RenderCore/Public/VertexFactory.h
--- a/Engine/Source/Runtime/RenderCore/Public/VertexFactory.h
+++ b/Engine/Source/Runtime/RenderCore/Public/VertexFactory.h
@@ -34,6 +34,8 @@
 
 #include <atomic>
 
+class FCbFieldView;
+class FCbWriter;
 class FMaterial;
 class FMeshDrawSingleShaderBindings;
 class FPrimitiveSceneProxy;
@@ -601,6 +603,22 @@ public:
 	{
 		return !(*this == Reference);
 	}
+
+	/** Call GetShaderFileHash to get the cached value for the filename's hash in the current process. */
+	RENDERCORE_API void RefreshCachedSourceHash(EShaderPlatform ShaderPlatform);
+
+private:
+#if WITH_EDITOR
+	// Compact binary API with hidden friend operator<<
+	RENDERCORE_API void Save(FCbWriter& Writer) const;
+	bool TryLoad(FCbFieldView Field);
+	friend inline FCbWriter& operator<<(FCbWriter& Writer, const FVertexFactoryTypeDependency& Value)
+	{
+		Value.Save(Writer);
+		return Writer;
+	}
+	friend RENDERCORE_API bool LoadFromCompactBinary(FCbFieldView Field, FVertexFactoryTypeDependency& OutValue);
+#endif
 };
 
 /** Used to compare two Vertex Factory types by name. */
