r.ShaderCompiler.MaxShaderStatsToPrint to allow changing how many jobs FShaderJobCache stats logs out for.

diff --git a/Engine/Source/Runtime/Engine/Private/ShaderCompiler/ShaderCompiler.cpp b/Engine/Source/Runtime/Engine/Private/ShaderCompiler/ShaderCompiler.cpp
--- a/Engine/Source/Runtime/Engine/Private/ShaderCompiler/ShaderCompiler.cpp
+++ b/Engine/Source/Runtime/Engine/Private/ShaderCompiler/ShaderCompiler.cpp
@@ -228,6 +228,14 @@ static TAutoConsoleVariable<bool> CVarAreShaderErrorsFatal(
 	TEXT("Default: true"),
 	ECVF_RenderThreadSafe);
 
+int32 GMaxShaderStatsToLog = 5;
+static FAutoConsoleVariableRef CVarShaderCompilerMaxShaderStatsToPrint(
+	TEXT("r.ShaderCompiler.MaxShaderStatsToPrint"),
+	GMaxShaderStatsToLog,
+	TEXT("Max number of shaders FShaderJobCache stats logs out"),
+	ECVF_Default
+);
+
 bool AreShaderErrorsFatal()
 {
 	return CVarAreShaderErrorsFatal.GetValueOnAnyThread();
@@ -4836,7 +4844,7 @@ void FShaderCompilerStats::WriteStatSummary()
 		// sort by avg time
 		ShaderTimings.ValueSort([](const FShaderTimings& A, const FShaderTimings& B) { return A.AverageCompileTime > B.AverageCompileTime; });
 
-		const int32 MaxShadersToPrint = FMath::Min(ShaderTimings.Num(), 5);
+		const int32 MaxShadersToPrint = FMath::Min(ShaderTimings.Num(), GMaxShaderStatsToLog);
 		UE_LOG(LogShaderCompilers, Display, TEXT("Top %d most expensive shader types by average time:"), MaxShadersToPrint);
 
 		int32 Idx = 0;
