Follow-up on CL 32050495: return correct skin weight profile name for static/default overridden setups

diff --git a/Engine/Source/Runtime/Engine/Classes/Components/SkinnedMeshComponent.h b/Engine/Source/Runtime/Engine/Classes/Components/SkinnedMeshComponent.h
--- a/Engine/Source/Runtime/Engine/Classes/Components/SkinnedMeshComponent.h
+++ b/Engine/Source/Runtime/Engine/Classes/Components/SkinnedMeshComponent.h
@@ -1347,7 +1347,7 @@ public:
 
 	/** Return the name of the Skin Weight Profile that is currently set otherwise returns 'None' */
 	UFUNCTION(BlueprintCallable, Category = "Components|SkinnedMesh")
-	FName GetCurrentSkinWeightProfileName() const { return CurrentSkinWeightProfileName; }
+	ENGINE_API FName GetCurrentSkinWeightProfileName() const;
 
 	/** Check whether or not a Skin Weight Profile is currently set */
 	UFUNCTION(BlueprintCallable, Category = "Components|SkinnedMesh")
diff --git a/Engine/Source/Runtime/Engine/Private/Components/SkinnedMeshComponent.cpp b/Engine/Source/Runtime/Engine/Private/Components/SkinnedMeshComponent.cpp
--- a/Engine/Source/Runtime/Engine/Private/Components/SkinnedMeshComponent.cpp
+++ b/Engine/Source/Runtime/Engine/Private/Components/SkinnedMeshComponent.cpp
@@ -4971,18 +4971,37 @@ bool USkinnedMeshComponent::IsUsingSkinWeightProfile() const
 	{		
 		return true;
 	}
-	else
+
+	if (USkinnedAsset* Asset = GetSkinnedAsset())
+	{
+		if (const FSkinWeightProfilesData* ProfileData = Asset->GetSkinWeightProfilesData(GetPredictedLODLevel()))
+		{
+			return ProfileData->IsDefaultOverridden() || ProfileData->IsStaticOverridden();
+		}
+	}	
+
+	return false;
+}
+
+FName USkinnedMeshComponent::GetCurrentSkinWeightProfileName() const
+{
+	if (bSkinWeightProfileSet)
+	{
+		return CurrentSkinWeightProfileName;
+	}
+
+	if (USkinnedAsset* Asset = GetSkinnedAsset())
 	{
-		if (USkinnedAsset* Asset = GetSkinnedAsset())
+		if (const FSkinWeightProfilesData* ProfileData = Asset->GetSkinWeightProfilesData(GetPredictedLODLevel()))
 		{
-			if (const FSkinWeightProfilesData* ProfileData = Asset->GetSkinWeightProfilesData(GetPredictedLODLevel()))
+			if( ProfileData->IsDefaultOverridden() || ProfileData->IsStaticOverridden())
 			{
-				return ProfileData->IsDefaultOverridden() || ProfileData->IsStaticOverridden();
+				return ProfileData->GetDefaultProfileName();
 			}
-		}		
+		}
 	}
 
-	return false;
+	return NAME_None;
 }
 
 void USkinnedMeshComponent::UpdateSkinWeightOverrideBuffer()
diff --git a/Engine/Source/Runtime/Engine/Public/Animation/SkinWeightProfile.h b/Engine/Source/Runtime/Engine/Public/Animation/SkinWeightProfile.h
--- a/Engine/Source/Runtime/Engine/Public/Animation/SkinWeightProfile.h
+++ b/Engine/Source/Runtime/Engine/Public/Animation/SkinWeightProfile.h
@@ -179,6 +179,7 @@ struct FSkinWeightProfilesData
 
 	ENGINE_API bool IsDefaultOverridden() const { return bDefaultOverridden; }
 	ENGINE_API bool IsStaticOverridden() const { return bStaticOverridden; }
+	ENGINE_API FName GetDefaultProfileName() const { return DefaultProfileName; }
 protected:
 	ENGINE_API void ApplyOverrideProfile(FSkinWeightVertexBuffer* OverrideBuffer, const FName& ProfileName);
 
