MaterialHLSLGenerator: Fix missing export on symbol that causes link error on some platforms for UnrealEditor build.
#rnx

diff --git a/Engine/Source/Runtime/Engine/Public/MaterialHLSLGenerator.h b/Engine/Source/Runtime/Engine/Public/MaterialHLSLGenerator.h
--- a/Engine/Source/Runtime/Engine/Public/MaterialHLSLGenerator.h
+++ b/Engine/Source/Runtime/Engine/Public/MaterialHLSLGenerator.h
@@ -65,7 +65,7 @@ struct TMaterialHLSLGeneratorType;
 class FMaterialHLSLGenerator
 {
 public:
-	static const uint32 Version = 1;
+	static constexpr uint32 Version = 1;
 
 	FMaterialHLSLGenerator(UMaterial* Material,
 		const FMaterialLayersFunctions* InLayerOverrides,
