Only release task execution resource in case anim compression was successful (otherwise another task will be spawned and the resource is still required)

diff --git a/Engine/Source/Runtime/Engine/Private/Animation/AnimationCompressionDerivedData.cpp b/Engine/Source/Runtime/Engine/Private/Animation/AnimationCompressionDerivedData.cpp
--- a/Engine/Source/Runtime/Engine/Private/Animation/AnimationCompressionDerivedData.cpp
+++ b/Engine/Source/Runtime/Engine/Private/Animation/AnimationCompressionDerivedData.cpp
@@ -117,7 +117,7 @@ void FAnimationSequenceAsyncCacheTask::EndCache(UE::DerivedData::FCacheGetValueR
 
 			{
 				// Release execution resource as soon as the task is done
-				ON_SCOPE_EXIT{ ExecutionResource = nullptr; };
+				ON_SCOPE_EXIT{ if (bIsDataValid) { ExecutionResource = nullptr; } };
 
 				if (UAnimSequence* AnimSequence = WeakAnimSequence.Get())
 				{
