Promoting functionality to be able to set and reset console variables across various automation tests as to not be limited for Functional Blueprint tests.

#rnx

diff --git a/Engine/Source/Developer/FunctionalTesting/Classes/FunctionalTest.h b/Engine/Source/Developer/FunctionalTesting/Classes/FunctionalTest.h
--- a/Engine/Source/Developer/FunctionalTesting/Classes/FunctionalTest.h
+++ b/Engine/Source/Developer/FunctionalTesting/Classes/FunctionalTest.h
@@ -11,6 +11,7 @@
 #include "GameFramework/Actor.h"
 #include "ProfilingDebugging/ExternalProfiler.h"
 #include "Math/StatisticalFloat.h"
+#include "Tests/AutomationCommon.h"
 #include "FunctionalTest.generated.h"
 
 class Error;
@@ -220,43 +221,6 @@ enum class EFunctionalTestLogHandling : uint8
 	OutputIgnored
 };
 
-
-class FConsoleVariableBPSetter
-{
-	friend class FAutomationFunctionalTestEnvSetup;
-
-public:
-	FConsoleVariableBPSetter(FString InConsoleVariableName);
-
-	void Set(const FString& Value);
-	FString Get();
-	void Restore();
-
-private:
-	bool bModified;
-	FString ConsoleVariableName;
-
-	FString OriginalValue;
-};
-
-class FAutomationFunctionalTestEnvSetup
-{
-public:
-	FAutomationFunctionalTestEnvSetup() = default;
-	~FAutomationFunctionalTestEnvSetup();
-
-	void SetVariable(const FString& VariableName, const FString& Value);
-
-	FString GetVariable(const FString& VariableName);
-
-	/** Restore the old settings. */
-	void Restore();
-
-private:
-	TArray<FConsoleVariableBPSetter> Variables;
-};
-
-
 DECLARE_DYNAMIC_MULTICAST_DELEGATE(FFunctionalTestEventSignature);
 DECLARE_DELEGATE_OneParam(FFunctionalTestDoneSignature, class AFunctionalTest*);
 
@@ -796,7 +760,7 @@ protected:
 public:
 	FFunctionalTestDoneSignature TestFinishedObserver;
 
-	// AG TEMP - solving a compile issue in a temp way to unblock the bui.d
+	// AG TEMP - solving a compile issue in a temp way to unblock the build
 	UPROPERTY(Transient)
 	bool bIsRunning;
 
@@ -813,7 +777,7 @@ public:
 
 private:
 	bool bIsReady;
-	FAutomationFunctionalTestEnvSetup EnvSetup;
+	TSharedPtr<FScopedTestEnvironment> EnvSetup;
 
 public:
 	/** Returns SpriteComponent subobject **/
diff --git a/Engine/Source/Developer/FunctionalTesting/Private/FunctionalTest.cpp b/Engine/Source/Developer/FunctionalTesting/Private/FunctionalTest.cpp
--- a/Engine/Source/Developer/FunctionalTesting/Private/FunctionalTest.cpp
+++ b/Engine/Source/Developer/FunctionalTesting/Private/FunctionalTest.cpp
@@ -235,7 +235,7 @@ bool AFunctionalTest::RunTest(const TArray<FString>& Params)
 	SCOPE_CYCLE_COUNTER(STAT_FunctionalTest_RunTest);
 	UWorld* World = GetWorld();
 	ensure(World->HasBegunPlay());
-
+	
 	FFunctionalTestBase* FunctionalTest = static_cast<FFunctionalTestBase*>(FAutomationTestFramework::Get().GetCurrentTest());
 
 	// Set handling of warnings/errors based on this test. Tests can either specify an explicit option or choose to go with the
@@ -443,7 +443,7 @@ void AFunctionalTest::FinishTest(EFunctionalTestResult TestResult, const FString
 
 	TestFinishedObserver.ExecuteIfBound(this);
 
-	EnvSetup.Restore();
+	EnvSetup = nullptr;
 }
 
 void AFunctionalTest::EndPlay(const EEndPlayReason::Type EndPlayReason)
@@ -537,22 +537,27 @@ FName AFunctionalTest::GetCurrentRerunReason()const
 
 void AFunctionalTest::SetConsoleVariable(const FString& Name, const FString& InValue)
 {
-	EnvSetup.SetVariable(Name, InValue);
+	// Initialize our Scoped Test Environment to be used during the duration of the test
+	if (!EnvSetup.IsValid())
+	{
+		EnvSetup = FScopedTestEnvironment::Get();
+	}
+	EnvSetup->SetConsoleVariableValue(Name, InValue);
 }
 
 void AFunctionalTest::SetConsoleVariableFromInteger(const FString& Name, const int32 InValue)
 {
-	EnvSetup.SetVariable(Name, FString::FromInt(InValue));
+	SetConsoleVariable(Name, FString::FromInt(InValue));
 }
 
 void AFunctionalTest::SetConsoleVariableFromFloat(const FString& Name, const float InValue)
 {
-	EnvSetup.SetVariable(Name, FString::SanitizeFloat(InValue));
+	SetConsoleVariable(Name, FString::SanitizeFloat(InValue));
 }
 
 void AFunctionalTest::SetConsoleVariableFromBoolean(const FString& Name, const bool InValue)
 {
-	EnvSetup.SetVariable(Name, FString::FromInt(InValue));
+	SetConsoleVariable(Name, FString::FromInt(InValue));
 }
 
 void AFunctionalTest::RegisterAutoDestroyActor(AActor* ActorToAutoDestroy)
@@ -1713,94 +1718,3 @@ void UAutomationPerformaceHelper::EndStatsFile()
 		GEngine->Exec(World, TEXT("Stat StopFile"));
 	}
 }
-
-FConsoleVariableBPSetter::FConsoleVariableBPSetter(FString InConsoleVariableName)
-	: bModified(false)
-	, ConsoleVariableName(InConsoleVariableName)
-{
-}
-
-void FConsoleVariableBPSetter::Set(const FString& Value)
-{
-	IConsoleVariable* ConsoleVariable = IConsoleManager::Get().FindConsoleVariable(*ConsoleVariableName);
-	if (ensure(ConsoleVariable))
-	{
-		if (bModified == false)
-		{
-			bModified = true;
-			OriginalValue = ConsoleVariable->GetString();
-		}
-
-		ConsoleVariable->AsVariable()->SetWithCurrentPriority(*Value);
-	}
-}
-
-FString FConsoleVariableBPSetter::Get()
-{
-	IConsoleVariable* ConsoleVariable = IConsoleManager::Get().FindConsoleVariable(*ConsoleVariableName);
-
-	if (ensure(ConsoleVariable))
-	{
-		return ConsoleVariable->GetString();
-	}
-
-	return FString{};
-}
-
-void FConsoleVariableBPSetter::Restore()
-{
-	if (bModified)
-	{
-		IConsoleVariable* ConsoleVariable = IConsoleManager::Get().FindConsoleVariable(*ConsoleVariableName);
-		if (ensure(ConsoleVariable))
-		{
-			ConsoleVariable->AsVariable()->SetWithCurrentPriority(*OriginalValue);
-		}
-
-		bModified = false;
-	}
-}
-
-FAutomationFunctionalTestEnvSetup::~FAutomationFunctionalTestEnvSetup()
-{
-	Restore();
-}
-
-void FAutomationFunctionalTestEnvSetup::SetVariable(const FString& VariableName, const FString& Value)
-{
-	check(IsInGameThread());
-
-	FConsoleVariableBPSetter Variable(VariableName);
-	Variable.Set(Value);
-	Variables.Add(MoveTemp(Variable));
-}
-
-FString FAutomationFunctionalTestEnvSetup::GetVariable(const FString& VariableName)
-{
-	check(IsInGameThread());
-
-	for (auto& Variable : Variables)
-	{
-		if (Variable.ConsoleVariableName == VariableName)
-		{
-			return Variable.Get();
-		}
-	}
-
-	return FString{};
-}
-
-void FAutomationFunctionalTestEnvSetup::Restore()
-{
-	if (!Variables.IsEmpty())
-	{
-		check(IsInGameThread());
-
-		for (auto& Variable : Variables)
-		{
-			Variable.Restore();
-		}
-
-		Variables.Empty();
-	}
-}
diff --git a/Engine/Source/Runtime/Engine/Private/Tests/AutomationCommon.cpp b/Engine/Source/Runtime/Engine/Private/Tests/AutomationCommon.cpp
--- a/Engine/Source/Runtime/Engine/Private/Tests/AutomationCommon.cpp
+++ b/Engine/Source/Runtime/Engine/Private/Tests/AutomationCommon.cpp
@@ -267,6 +267,136 @@ void FTestWorldWrapper::ForwardErrorMessages(FAutomationTestBase* AutomationTest
 	}
 }
 
+///////////////////////////////////////////////////////////////////////
+// FTestConsoleVariable
+
+FTestConsoleVariable::FTestConsoleVariable(const FString& InConsoleVariableName)
+	: bModified(false)
+	, ConsoleVariableName(InConsoleVariableName)
+{
+}
+
+FTestConsoleVariable::FTestConsoleVariable(FTestConsoleVariable&& Other)
+	: bModified(MoveTemp(Other.bModified))
+	, ConsoleVariableName(MoveTemp(Other.ConsoleVariableName))
+	, OriginalValue(MoveTemp(Other.OriginalValue))
+{
+	// Reset our other's modified flag to avoid triggering a reset of the CVar
+	Other.bModified = false;
+}
+
+FTestConsoleVariable::~FTestConsoleVariable()
+{
+	Restore();
+}
+
+void FTestConsoleVariable::Set(const FString& Value)
+{
+	IConsoleVariable* ConsoleVariable = IConsoleManager::Get().FindConsoleVariable(*ConsoleVariableName);
+	if (ensure(ConsoleVariable))
+	{
+		if (!bModified)
+		{
+			bModified = true;
+			OriginalValue = ConsoleVariable->GetString();
+		}
+
+		ConsoleVariable->AsVariable()->SetWithCurrentPriority(*Value);
+	}
+}
+
+FString FTestConsoleVariable::Get()
+{
+	IConsoleVariable* ConsoleVariable = IConsoleManager::Get().FindConsoleVariable(*ConsoleVariableName);
+	if (ensure(ConsoleVariable))
+	{
+		return ConsoleVariable->GetString();
+	}
+
+	return FString{};
+}
+
+void FTestConsoleVariable::Restore()
+{
+	if (bModified)
+	{
+		IConsoleVariable* ConsoleVariable = IConsoleManager::Get().FindConsoleVariable(*ConsoleVariableName);
+		if (ensure(ConsoleVariable))
+		{
+			ConsoleVariable->AsVariable()->SetWithCurrentPriority(*OriginalValue);
+		}
+
+		bModified = false;
+	}
+}
+
+///////////////////////////////////////////////////////////////////////
+// FScopedTestEnvironment
+
+TWeakPtr<FScopedTestEnvironment> FScopedTestEnvironment::EnvironmentInstance = nullptr;
+
+FScopedTestEnvironment::~FScopedTestEnvironment()
+{
+	Restore();
+}
+
+TSharedPtr<FScopedTestEnvironment> FScopedTestEnvironment::Get()
+{
+	check(IsInGameThread());
+
+	TSharedPtr<FScopedTestEnvironment> Instance = EnvironmentInstance.Pin();
+	if (!Instance.IsValid())
+	{
+		Instance = MakeShareable(new FScopedTestEnvironment());
+		EnvironmentInstance = Instance;
+	}
+
+	check(Instance.IsValid());
+	return Instance;
+}
+
+void FScopedTestEnvironment::SetConsoleVariableValue(const FString& ConsoleVariableName, const FString& Value)
+{
+	check(IsInGameThread());
+
+	if (FTestConsoleVariable* Variable = Variables.Find(ConsoleVariableName))
+	{
+		Variable->Set(Value);
+	}
+	else
+	{
+		FTestConsoleVariable NewVariable(ConsoleVariableName);
+		NewVariable.Set(Value);
+		Variables.Emplace(ConsoleVariableName, MoveTemp(NewVariable));
+	}
+}
+
+bool FScopedTestEnvironment::TryGetConsoleVariableValue(const FString& ConsoleVariableName, FString* OutValue)
+{
+	check(IsInGameThread());
+
+	if (FTestConsoleVariable* Variable = Variables.Find(ConsoleVariableName))
+	{
+		if (OutValue)
+		{
+			*OutValue = Variable->Get();
+		}
+		return true;
+	}
+
+	return false;
+}
+
+void FScopedTestEnvironment::Restore()
+{
+	if (!Variables.IsEmpty())
+	{
+		check(IsInGameThread());
+
+		Variables.Empty();
+	}
+}
+
 ///////////////////////////////////////////////////////////////////////
 // Common Latent commands
 
diff --git a/Engine/Source/Runtime/Engine/Private/Tests/EngineAutomationTests.cpp b/Engine/Source/Runtime/Engine/Private/Tests/EngineAutomationTests.cpp
--- a/Engine/Source/Runtime/Engine/Private/Tests/EngineAutomationTests.cpp
+++ b/Engine/Source/Runtime/Engine/Private/Tests/EngineAutomationTests.cpp
@@ -227,6 +227,133 @@ bool FSaveGameTest::RunTest(const FString& Parameters)
 	return LoadedData == SavedData;
 }
 
+/**
+ * FCVarEnvironmentTest
+ * Test makes sure that CVars are set and restore properly
+ */
+IMPLEMENT_SIMPLE_AUTOMATION_TEST(FCVarEnvironmentTest, "System.Engine.Automation.Environment.CVar", EAutomationTestFlags::ApplicationContextMask | EAutomationTestFlags::EngineFilter)
+
+/**
+ * Set and restore a CVar
+ *
+ * @param Parameters - Unused for this test
+ * @return	TRUE if the test was successful, FALSE otherwise
+ */
+bool FCVarEnvironmentTest::RunTest(const FString& Parameters)
+{
+	const FString DummyTestName = TEXT("Automation.DummyTestVariable");
+	const int32 TestValue = 12345;
+
+	TAutoConsoleVariable<int32> CVarDummyTestVariable(
+		*DummyTestName,
+		111,
+		TEXT("Used for the purposes of testing if the CVar is getting set and reset."),
+		ECVF_Default);
+
+	IConsoleVariable* ConsoleVariable = IConsoleManager::Get().FindConsoleVariable(*DummyTestName);
+	if (!ConsoleVariable)
+	{
+		AddError(FString::Format(TEXT("Could not find CVar for '{0}'"), { DummyTestName }));
+		return false;
+	}
+
+	int32 InitialValue = ConsoleVariable->GetInt();
+	if (InitialValue == TestValue)
+	{
+		AddError(FString::Format(TEXT("Initial and values to test are the same '{0}'"), { InitialValue }));
+		return false;
+	}
+
+	// Because we're testing that the CVar is properly restored, we want to create our FTestEnvironment object inside a scope to be destructed
+	{
+		TSharedPtr<FScopedTestEnvironment> TestCVarEnvironment = FScopedTestEnvironment::Get();
+		TestCVarEnvironment->SetConsoleVariableValue(DummyTestName, FString::FromInt(TestValue));
+
+		// Verify that setting the CVar through our test environment actually sets the Console Variable
+		int32 CurrentValue = ConsoleVariable->GetInt();
+		if (CurrentValue == InitialValue)
+		{
+			AddError(FString::Format(TEXT("CVar was not set as the current value matches the initial value of '{0}'"), { InitialValue }));
+			return false;
+		}
+
+		// Verify that retrieving the CVar from our test environment matches the value fetched directly from the CVar
+		FString ConsoleVariableValue;
+		bool bWasCVarSet = TestCVarEnvironment->TryGetConsoleVariableValue(DummyTestName, &ConsoleVariableValue);
+		if (!bWasCVarSet)
+		{
+			AddError(TEXT("CVar was not found as being set in the ScopedTestEnvironment."));
+			return false;
+		}
+
+		if (ConsoleVariableValue != FString::FromInt(CurrentValue))
+		{
+			AddError(FString::Format(TEXT("CVar value of '{0}' does not match the CVar value fetched from ScopedTestEnvironment '{1}'"), { CurrentValue, ConsoleVariableValue }));
+			return false;
+		}
+	}
+
+	int32 CurrentValue = ConsoleVariable->GetInt();
+	AddErrorIfFalse(CurrentValue == InitialValue, FString::Format(TEXT("CVar was not reset as the current value of '{0}' does not match the initial value of '{1}'"), { CurrentValue, InitialValue }));
+
+	return !HasAnyErrors();
+}
+
+/**
+ * FCVarEnvironmentReuseTest
+ * Test makes sure that a CVar can be set multiple times before restoring back to the original value prior to being set
+ */
+IMPLEMENT_SIMPLE_AUTOMATION_TEST(FCVarEnvironmentReuseTest, "System.Engine.Automation.Environment.CVar Reuse", EAutomationTestFlags::ApplicationContextMask | EAutomationTestFlags::EngineFilter)
+
+/**
+ * Set a CVar multiple times before restoring to the original value
+ *
+ * @param Parameters - Unused for this test
+ * @return	TRUE if the test was successful, FALSE otherwise
+ */
+bool FCVarEnvironmentReuseTest::RunTest(const FString& Parameters)
+{
+	const FString DummyTestName = TEXT("Automation.DummyTestVariable");
+
+	TAutoConsoleVariable<int32> CVarDummyTestVariable(
+		*DummyTestName,
+		111,
+		TEXT("Used for the purposes of testing if the CVar is getting set and reset."),
+		ECVF_Default);
+
+	IConsoleVariable* ConsoleVariable = IConsoleManager::Get().FindConsoleVariable(*DummyTestName);
+	if (!ConsoleVariable)
+	{
+		AddError(FString::Format(TEXT("Could not find CVar for '{0}'"), { DummyTestName }));
+		return false;
+	}
+
+	int32 InitialValue = ConsoleVariable->GetInt();
+
+	// Because we're testing that the CVar is properly restored, we want to create our FTestEnvironment object inside a scope to be destructed
+	{
+		TSharedPtr<FScopedTestEnvironment> TestCVarEnvironment = FScopedTestEnvironment::Get();
+
+		// Loop through a range of values to set our CVar
+		for (int32 TestValue = 0; TestValue < 5; ++TestValue)
+		{
+			TestCVarEnvironment->SetConsoleVariableValue(DummyTestName, FString::FromInt(TestValue));
+			
+			int32 CurrentValue = ConsoleVariable->GetInt();
+			if (TestValue != CurrentValue)
+			{
+				AddError(FString::Format(TEXT("CVar was not set as the current value '{0}' does not match the expected value of '{1}'"), { CurrentValue, TestValue }));
+				return false;
+			}
+		}
+	}
+
+	int32 CurrentValue = ConsoleVariable->GetInt();
+	AddErrorIfFalse(CurrentValue == InitialValue, FString::Format(TEXT("CVar was not reset as the current value of '{0}' does not match the initial value of '{1}'"), { CurrentValue, InitialValue }));
+
+	return !HasAnyErrors();
+}
+
 IMPLEMENT_SIMPLE_AUTOMATION_TEST(FAutomationLogAddMessage, "TestFramework.Log.Add Log Message", EAutomationTestFlags::ApplicationContextMask | EAutomationTestFlags::EngineFilter)
 
 bool FAutomationLogAddMessage::RunTest(const FString& Parameters)
diff --git a/Engine/Source/Runtime/Engine/Public/Tests/AutomationCommon.h b/Engine/Source/Runtime/Engine/Public/Tests/AutomationCommon.h
--- a/Engine/Source/Runtime/Engine/Public/Tests/AutomationCommon.h
+++ b/Engine/Source/Runtime/Engine/Public/Tests/AutomationCommon.h
@@ -71,6 +71,83 @@ protected:
 	TArray<FString> FailureErrors;
 };
 
+/**
+ * Utility for setting and restoring of a Console Variable (CVar).
+ * 
+ * Note that race conditions are possible when multiple `FTestConsoleVariable` objects refer to the same Console Variable and get restored/destroyed in an order different than they were set.
+ * Please use `FScopedTestEnvironment` to manage multiple `FTestConsoleVariable` objects.
+ */
+struct FTestConsoleVariable
+{
+	ENGINE_API FTestConsoleVariable(const FString& InConsoleVariableName);
+	ENGINE_API FTestConsoleVariable(FTestConsoleVariable&& Other);
+
+	ENGINE_API ~FTestConsoleVariable();
+
+	/**
+	 * Sets a Console Variable to the specified value. Will keep a reference to the original value regardless of how many times the value has been set.
+	 * @param Value		Value to set on the Console Variable
+	 */
+	ENGINE_API void Set(const FString& Value);
+
+	/** Returns the current value of the Console Variable */
+	ENGINE_API FString Get();
+
+	/** Returns the Console Variable to the original value */
+	ENGINE_API void Restore();
+
+private:
+	bool bModified;
+	FString ConsoleVariableName;
+	FString OriginalValue;
+};
+
+/**
+ * Utility for setting and management of the test environment
+ * Will handle restoring the Console Variables (CVars) back to the original state on destruction.
+ */
+struct FScopedTestEnvironment
+{
+	/** Restores all set Console Variables back to the original value */
+	ENGINE_API ~FScopedTestEnvironment();
+
+	/**
+	 * Returns a shared pointer to current instance of the scoped test environment. Will create an instance if current instance is invalid.
+	 * @return the shared pointer to the instance
+	 */
+	static ENGINE_API TSharedPtr<FScopedTestEnvironment> Get();
+
+	/**
+	 * Sets a Console Variable to the specified value. Will keep a reference to the original value regardless of how many times the value has been set.
+	 * @param VariableName	Name of the Console Variable to set
+	 * @param Value			Value to be applied
+	 */
+	ENGINE_API void SetConsoleVariableValue(const FString& ConsoleVariableName, const FString& Value);
+
+	/**
+	 * Gets the current overridden value for the specified Console Variable
+	 * @param VariableName	Name of the Console Variable to fetch the current value from
+	 * @param OutValue		Current value of the Console Variable if it was overridden.
+	 * @return true if the Console Variable was overridden.
+	 */
+	ENGINE_API bool TryGetConsoleVariableValue(const FString& ConsoleVariableName, FString* OutValue);
+
+	/** Restores all set Console Variables back to the original value */
+	ENGINE_API void Restore();
+
+private:
+	FScopedTestEnvironment() = default;
+
+	FScopedTestEnvironment(FScopedTestEnvironment&&) = delete;
+	FScopedTestEnvironment(const FScopedTestEnvironment&) = delete;
+	FScopedTestEnvironment& operator=(FScopedTestEnvironment&&) = delete;
+	FScopedTestEnvironment& operator=(const FScopedTestEnvironment&) = delete;
+
+	TMap<FString, FTestConsoleVariable> Variables;
+
+	static TWeakPtr<struct FScopedTestEnvironment> EnvironmentInstance;
+};
+
 /** Common automation functions */
 namespace AutomationCommon
 {
