Silence false V1078 warning caused by population of the array inside a lambda

diff --git a/Engine/Source/Runtime/Engine/Private/WorldPartition/DataLayer/ExternalDataLayerHelper.cpp b/Engine/Source/Runtime/Engine/Private/WorldPartition/DataLayer/ExternalDataLayerHelper.cpp
--- a/Engine/Source/Runtime/Engine/Private/WorldPartition/DataLayer/ExternalDataLayerHelper.cpp
+++ b/Engine/Source/Runtime/Engine/Private/WorldPartition/DataLayer/ExternalDataLayerHelper.cpp
@@ -287,7 +287,7 @@ bool FExternalDataLayerHelper::MoveActorsToExternalDataLayer(const TArray<AActor
 
 		// Move dependant objects into the new actor package
 		UPackage* NewActorPackage = InActor->GetExternalPackage();
-		for (UObject* DependantObject : DependantObjects)
+		for (UObject* DependantObject : DependantObjects) //-V1078
 		{
 			DependantObject->Rename(nullptr, NewActorPackage, REN_NonTransactional | REN_DontCreateRedirectors | REN_ForceNoResetLoaders | REN_DoNotDirty);
 		}
