URO now ignores non-visible mesh components to prevent invalid LOD params being set

diff --git a/Engine/Source/Runtime/Engine/Private/Components/SkinnedMeshComponent.cpp b/Engine/Source/Runtime/Engine/Private/Components/SkinnedMeshComponent.cpp
--- a/Engine/Source/Runtime/Engine/Private/Components/SkinnedMeshComponent.cpp
+++ b/Engine/Source/Runtime/Engine/Private/Components/SkinnedMeshComponent.cpp
@@ -340,11 +340,14 @@ namespace FAnimUpdateRateManager
 		const TArray<USkinnedMeshComponent*>& SkinnedComponents = Tracker->RegisteredComponents;
 		for (USkinnedMeshComponent* Component : SkinnedComponents)
 		{
-			bRecentlyRendered |= Component->bRecentlyRendered;
-			MaxDistanceFactor = FMath::Max(MaxDistanceFactor, Component->MaxDistanceFactor);
-			bPlayingNetworkedRootMotionMontage |= Component->IsPlayingNetworkedRootMotionMontage();
-			bUsingRootMotionFromEverything &= Component->IsPlayingRootMotionFromEverything();
-			MinLod = FMath::Min(MinLod, Tracker->UpdateRateParameters.bShouldUseMinLod ? Component->MinLodModel : Component->GetPredictedLODLevel());
+			if (Component->IsVisible())
+			{
+				bRecentlyRendered |= Component->bRecentlyRendered;
+				MaxDistanceFactor = FMath::Max(MaxDistanceFactor, Component->MaxDistanceFactor);
+				bPlayingNetworkedRootMotionMontage |= Component->IsPlayingNetworkedRootMotionMontage();
+				bUsingRootMotionFromEverything &= Component->IsPlayingRootMotionFromEverything();
+				MinLod = FMath::Min(MinLod, Tracker->UpdateRateParameters.bShouldUseMinLod ? Component->MinLodModel : Component->GetPredictedLODLevel());
+			}
 		}
 
 		bNeedsValidRootMotion &= bPlayingNetworkedRootMotionMontage;
