LightFunctionAtlas - fix compilation output missing initialisation.

diff --git a/Engine/Source/Runtime/Engine/Public/MaterialShared.h b/Engine/Source/Runtime/Engine/Public/MaterialShared.h
--- a/Engine/Source/Runtime/Engine/Public/MaterialShared.h
+++ b/Engine/Source/Runtime/Engine/Public/MaterialShared.h
@@ -720,6 +720,7 @@ public:
 		bUsesVertexInterpolator(false),
 		bHasRuntimeVirtualTextureOutputNode(false),
 		bUsesAnisotropy(false),
+		bIsLightFunctionAtlasCompatible(false),
 		bUsesDisplacement(false),
 		bUsedWithNeuralNetworks(false)
 	{
