Cooking fix: When OFPA actors are moved into their own generated level, make sure to trash their original package to avoid any future load of this same actor to find the empty package.
One use case is if APCGWorldActor::BeginCacheForCookedPlatformData is called after one of its dependant landscape actor has already been moved and saved into its own generated package.

diff --git a/Engine/Source/Runtime/Engine/Private/WorldPartition/WorldPartitionLevelHelper.cpp b/Engine/Source/Runtime/Engine/Private/WorldPartition/WorldPartitionLevelHelper.cpp
--- a/Engine/Source/Runtime/Engine/Private/WorldPartition/WorldPartitionLevelHelper.cpp
+++ b/Engine/Source/Runtime/Engine/Private/WorldPartition/WorldPartitionLevelHelper.cpp
@@ -297,6 +297,11 @@ void FWorldPartitionLevelHelper::MoveExternalActorsToLevel(const TArray<FWorldPa
 				}
 			}
 
+			// Trash this package to guarantee that any potential future load of this actor won't find the old empty package
+			// @todo_ow: Decide if we want to support actor reloads during cook. If not, remove this code, detect the reload and report an error.
+			FName NewPackageName = MakeUniqueObjectName(nullptr, UPackage::StaticClass(), FName(*FString::Printf(TEXT("%s_Trashed"), *ActorExternalPackage->GetName())));
+			ActorExternalPackage->Rename(*NewPackageName.ToString(), nullptr, REN_ForceNoResetLoaders | REN_DontCreateRedirectors | REN_NonTransactional | REN_DoNotDirty);
+
 			OutModifiedPackages.Add(ActorExternalPackage);
 			LevelActors.Add(Actor->GetFName());
 		}
