Fail SoundWaveProxyReader immediately if it fails to create a decoder.

Prevents downstream users from attempting to continue to read from the failed decoder

#tests PIE with local cvar to trigger failure

diff --git a/Engine/Source/Runtime/Engine/Private/AudioDecompress.cpp b/Engine/Source/Runtime/Engine/Private/AudioDecompress.cpp
--- a/Engine/Source/Runtime/Engine/Private/AudioDecompress.cpp
+++ b/Engine/Source/Runtime/Engine/Private/AudioDecompress.cpp
@@ -87,6 +87,8 @@ bool IStreamedCompressedInfo::ReadCompressedInfo(const uint8* InSrcBufferData, u
 	// Parse the format header, this is done different for each format
 	if (!ParseHeader(InSrcBufferData, InSrcBufferDataSize, QualityInfo))
 	{
+		UE_LOG(LogAudio, Error, TEXT("IStreamedCompressedInfo::ReadCompressedInfo: Failed to Parse Header"));
+		bHasError = true;
 		return false;
 	}
 
diff --git a/Engine/Source/Runtime/Engine/Private/SoundWaveProxyReader.cpp b/Engine/Source/Runtime/Engine/Private/SoundWaveProxyReader.cpp
--- a/Engine/Source/Runtime/Engine/Private/SoundWaveProxyReader.cpp
+++ b/Engine/Source/Runtime/Engine/Private/SoundWaveProxyReader.cpp
@@ -72,6 +72,9 @@ FSoundWaveProxyReader::FSoundWaveProxyReader(FSoundWaveProxyRef InWaveProxy, con
 
 	// Prepare to read audio
 	bIsDecoderValid = InitializeDecoder(Settings.StartTimeInSeconds);
+
+	// set the decoder to "Fail" if we're unable to create a decoder
+	DecodeResult = bIsDecoderValid ? DecodeResult : EDecodeResult::Fail;
 }
 
 /** Create a wave proxy reader.
@@ -139,12 +142,21 @@ void FSoundWaveProxyReader::SetLoopDuration(float InLoopDurationInSeconds)
 bool FSoundWaveProxyReader::SeekToTime(float InSeconds)
 {
 	int32 InFrameIndex = FMath::Clamp(static_cast<int32>(InSeconds * GetSampleRate()), 0, GetNumFramesInWave());
+
+	if (!bIsDecoderValid)
+	{
+		// set the current frame index, but indicate that we've still failed
+		CurrentFrameIndex = InFrameIndex;
+		UE_LOG(LogAudio, Verbose, TEXT("FSoundWaveProxyReader::SeekToTime, failed to seek due to the decoder being invalid!: %s"), *WaveProxy->GetFName().ToString());
+		return false;
+	}
+	
 	// ignore seek request if we're already at the specified time
 	if (InFrameIndex == CurrentFrameIndex)
 	{
-		return bIsDecoderValid;
+		return true;
 	}
-
+	
 	if (WaveProxy->IsSeekable() && CompressedAudioInfo)
 	{
 		CompressedAudioInfo->SeekToFrame(InFrameIndex);
@@ -152,7 +164,7 @@ bool FSoundWaveProxyReader::SeekToTime(float InSeconds)
 		DecoderOutput.SetNum(0);
 		NumDecodeSamplesToDiscard = 0;
 		DecodeResult = EDecodeResult::MoreDataRemaining;
-		return bIsDecoderValid;
+		return true;
 	}
 	// Direct seeking is not supported. A new decoder must be created. 
 	bIsDecoderValid = InitializeDecoder(InSeconds);
@@ -168,10 +180,18 @@ bool FSoundWaveProxyReader::CanProduceMoreAudio() const
 
 bool FSoundWaveProxyReader::SeekToFrame(uint32 InFrameNum)
 {
+	if (!bIsDecoderValid)
+	{
+		// set the current frame index, but indicate that we've still failed
+		CurrentFrameIndex = InFrameNum;
+		UE_LOG(LogAudio, Verbose, TEXT("FSoundWaveProxyReader::SeekToTime, failed to seek due to the decoder being invalid!"));
+		return false;
+	}
+	
 	// ignore seek request if we're already at the specified time
 	if (InFrameNum == CurrentFrameIndex)
 	{
-		return bIsDecoderValid;
+		return true;
 	}
 
 	if (WaveProxy->IsSeekable() && CompressedAudioInfo)
@@ -181,7 +201,7 @@ bool FSoundWaveProxyReader::SeekToFrame(uint32 InFrameNum)
 		DecoderOutput.SetNum(0);
 		NumDecodeSamplesToDiscard = 0;
 		DecodeResult = EDecodeResult::MoreDataRemaining;
-		return bIsDecoderValid;
+		return true;
 	}
 
 	// Direct seeking is not supported. A new decoder must be created.
@@ -340,7 +360,7 @@ bool FSoundWaveProxyReader::InitializeDecoder(float InStartTimeInSeconds)
 	IAudioInfoFactory* Factory = IAudioInfoFactoryRegistry::Get().Find(Format);
 	if (!ensure(Factory))
 	{
-		UE_LOG(LogAudio, Error, TEXT("Failed to CompressedAudioInfo for wave (package: %s). Unable to find AudioInfoFactory for format: %s"), *WaveProxy->GetPackageName().ToString(), *Format.ToString());
+		UE_LOG(LogAudio, Error, TEXT("FSoundWaveProxyReader::InitializeDecoder: Failed to create CompressedAudioInfo for wave (package: %s). Unable to find AudioInfoFactory for format: %s"), *WaveProxy->GetPackageName().ToString(), *Format.ToString());
 		return false;
 	}
 
@@ -349,7 +369,7 @@ bool FSoundWaveProxyReader::InitializeDecoder(float InStartTimeInSeconds)
 
 	if (!ensure(InfoInstance.IsValid()))
 	{
-		UE_LOG(LogAudio, Error, TEXT("Failed to created CompressedAudioInfo for wave (package: %s). Unable to create info from factory for for format: %s"), *WaveProxy->GetPackageName().ToString(), *Format.ToString());
+		UE_LOG(LogAudio, Error, TEXT("FSoundWaveProxyReader::InitializeDecoder: Failed to create CompressedAudioInfo for wave (package: %s). Unable to create info from factory for for format: %s"), *WaveProxy->GetPackageName().ToString(), *Format.ToString());
 		return false;
 	}
 
@@ -358,7 +378,7 @@ bool FSoundWaveProxyReader::InitializeDecoder(float InStartTimeInSeconds)
 	{
 		if (!InfoInstance->StreamCompressedInfo(WaveProxy, &Info))
 		{
-			UE_LOG(LogAudio, Error, TEXT("Failed to created CompressedAudioInfo for wave (package: %s). Unable to stream compressed info for streaming wave"), *WaveProxy->GetPackageName().ToString());
+			UE_LOG(LogAudio, Error, TEXT("FSoundWaveProxyReader::InitializeDecoder: Failed to create CompressedAudioInfo for wave (package: %s). Unable to stream compressed info for streaming wave"), *WaveProxy->GetPackageName().ToString());
 			return false;
 		}
 	}
@@ -366,7 +386,7 @@ bool FSoundWaveProxyReader::InitializeDecoder(float InStartTimeInSeconds)
 	{
 		if (!InfoInstance->ReadCompressedInfo(WaveProxy->GetResourceData(), WaveProxy->GetResourceSize(), &Info))
 		{
-			UE_LOG(LogAudio, Error, TEXT("Failed to created decoder input for wave (package: %s). Unable to read compressed info for non-streaming wave"), *WaveProxy->GetPackageName().ToString());
+			UE_LOG(LogAudio, Error, TEXT("FSoundWaveProxyReader::InitializeDecoder: Failed to create CompressedAudioInfo for wave (package: %s). Unable to read compressed info for non-streaming wave"), *WaveProxy->GetPackageName().ToString());
 			return false;
 		}
 	}
