Fix crash when destroying a world context where it would set the GameInstance's context pointer to the deleted world context

diff --git a/Engine/Source/Runtime/Engine/Private/GameInstance.cpp b/Engine/Source/Runtime/Engine/Private/GameInstance.cpp
--- a/Engine/Source/Runtime/Engine/Private/GameInstance.cpp
+++ b/Engine/Source/Runtime/Engine/Private/GameInstance.cpp
@@ -593,18 +593,26 @@ UGameViewportClient* UGameInstance::GetGameViewportClient() const
 
 void UGameInstance::OnWorldChanged(UWorld* OldWorld, UWorld* NewWorld)
 {
-	// Fix up our world context if it is incorrect
-	FWorldContext* NewWorldContext = const_cast<FWorldContext*>(GEngine->GetWorldContextFromWorld(NewWorld));
-
-	if (WorldContext == nullptr)
+	if (NewWorld)
 	{
-		// This may be a test/preview world that did not have a world context set before
-		WorldContext = NewWorldContext;
+		// Fix up our world context if it is incorrect and we have a valid new world
+		FWorldContext* NewWorldContext = const_cast<FWorldContext*>(GEngine->GetWorldContextFromWorld(NewWorld));
+
+		if (WorldContext == nullptr)
+		{
+			// This may be a test/preview world that did not have a world context set before
+			WorldContext = NewWorldContext;
+		}
+		else if (WorldContext != NewWorldContext)
+		{
+			UE_LOG(LogLoad, Warning, TEXT("GameInstance %s changed from world %s to world %s with a different world context!"), *GetName(), *GetPathNameSafe(OldWorld), *GetPathNameSafe(NewWorld));
+			WorldContext = NewWorldContext;
+		}
 	}
-	else if (WorldContext != NewWorldContext)
+	else
 	{
-		UE_LOG(LogLoad, Warning, TEXT("GameInstance %s changed from world %s to world %s with a different world context!"), *GetName(), *GetPathNameSafe(OldWorld), *GetPathNameSafe(NewWorld));
-		WorldContext = NewWorldContext;
+		// Clear our world context if our world is cleared
+		WorldContext = nullptr;
 	}
 }
 
