Fix clang 18 errors

#rnx

diff --git a/Engine/Source/Programs/UnrealBuildTool/ToolChain/ClangWarnings.cs b/Engine/Source/Programs/UnrealBuildTool/ToolChain/ClangWarnings.cs
--- a/Engine/Source/Programs/UnrealBuildTool/ToolChain/ClangWarnings.cs
+++ b/Engine/Source/Programs/UnrealBuildTool/ToolChain/ClangWarnings.cs
@@ -50,11 +50,6 @@ namespace UnrealBuildTool
 					Arguments.Add("-Wno-invalid-unevaluated-string");   // https://clang.llvm.org/docs/DiagnosticsReference.html#winvalid-unevaluated-string			// new warning for clang 17
 				}
 			}
-			if (ClangVersion >= new VersionNumber(18))
-			{
-				Arguments.Add("-Wno-deprecated-this-capture");          // https://clang.llvm.org/docs/DiagnosticsReference.html#wdeprecated-this-capture
-				Arguments.Add("-Wno-enum-constexpr-conversion");        // https://clang.llvm.org/docs/DiagnosticsReference.html#wenum-constexpr-conversion
-			}
 
 			Arguments.Add("-Wno-gnu-string-literal-operator-template"); // https://clang.llvm.org/docs/DiagnosticsReference.html#wgnu-string-literal-operator-template	// We use this feature to allow static FNames.
 			Arguments.Add("-Wno-inconsistent-missing-override");        // https://clang.llvm.org/docs/DiagnosticsReference.html#winconsistent-missing-override			// ?? no reason given
diff --git a/Engine/Source/Runtime/Engine/Private/TraceFilter.cpp b/Engine/Source/Runtime/Engine/Private/TraceFilter.cpp
--- a/Engine/Source/Runtime/Engine/Private/TraceFilter.cpp
+++ b/Engine/Source/Runtime/Engine/Private/TraceFilter.cpp
@@ -267,14 +267,14 @@ FAutoConsoleCommand FlushFilterStateCommand(TEXT("TraceFilter.FlushState"), TEXT
 	);
 
 template<>
-bool ENGINE_API FTraceFilter::IsObjectTraceable</*bForceThreadSafe = */ true>(const UObject* InObject)
+ENGINE_API bool FTraceFilter::IsObjectTraceable</*bForceThreadSafe = */ true>(const UObject* InObject)
 {
 	// Object not found in the AnnotationMap means that it is at the default value, which is bIsTraceable == true
 	return GObjectFilterAnnotations.GetAnnotationMap().Find(InObject) == nullptr;
 }
 
 template<>
-bool ENGINE_API FTraceFilter::IsObjectTraceable</*bForceThreadSafe = */ false>(const UObject* InObject)
+ENGINE_API bool FTraceFilter::IsObjectTraceable</*bForceThreadSafe = */ false>(const UObject* InObject)
 {
 	check(GObjectFilterAnnotations.IsLocked());
 	// Object not found in the AnnotationMap means that it is at the default value, which is bIsTraceable == true
@@ -282,7 +282,7 @@ bool ENGINE_API FTraceFilter::IsObjectTraceable</*bForceThreadSafe = */ false>(c
 }
 
 template<>
-void ENGINE_API FTraceFilter::SetObjectIsTraceable</*bForceThreadSafe = */ true>(const UObject* InObject, bool bIsTraceable)
+ENGINE_API void FTraceFilter::SetObjectIsTraceable</*bForceThreadSafe = */ true>(const UObject* InObject, bool bIsTraceable)
 {
 	ensure(InObject);
 		
@@ -297,7 +297,7 @@ void ENGINE_API FTraceFilter::SetObjectIsTraceable</*bForceThreadSafe = */ true>
 }
 
 template<>
-void ENGINE_API FTraceFilter::SetObjectIsTraceable</*bForceThreadSafe = */ false>(const UObject* InObject, bool bIsTraceable)
+ENGINE_API void FTraceFilter::SetObjectIsTraceable</*bForceThreadSafe = */ false>(const UObject* InObject, bool bIsTraceable)
 {
 	ensure(InObject);
 
@@ -315,7 +315,7 @@ void ENGINE_API FTraceFilter::SetObjectIsTraceable</*bForceThreadSafe = */ false
 }
 
 template<>
-void ENGINE_API FTraceFilter::MarkObjectTraceable</*bForceThreadSafe = */ true>(const UObject* InObject)
+ENGINE_API void FTraceFilter::MarkObjectTraceable</*bForceThreadSafe = */ true>(const UObject* InObject)
 {
 	ensure(InObject);	
 	FTraceFilterObjectAnnotation Annotation;
@@ -324,7 +324,7 @@ void ENGINE_API FTraceFilter::MarkObjectTraceable</*bForceThreadSafe = */ true>(
 }
 
 template<>
-void ENGINE_API FTraceFilter::MarkObjectTraceable</*bForceThreadSafe = */ false>(const UObject* InObject)
+ENGINE_API void FTraceFilter::MarkObjectTraceable</*bForceThreadSafe = */ false>(const UObject* InObject)
 {
 	ensure(InObject);
 	check(GObjectFilterAnnotations.IsLocked());
diff --git a/Engine/Source/Runtime/GeometryCore/Public/Curve/DynamicGraph2.h b/Engine/Source/Runtime/GeometryCore/Public/Curve/DynamicGraph2.h
--- a/Engine/Source/Runtime/GeometryCore/Public/Curve/DynamicGraph2.h
+++ b/Engine/Source/Runtime/GeometryCore/Public/Curve/DynamicGraph2.h
@@ -139,7 +139,7 @@ public:
 	value_iteration<TVector2<T>> Vertices() const
 	{
 		return vertices_refcount.MappedIndices<TVector2<T>>(
-			[=](int vid) {
+			[&](int vid) {
 				return vertices.template AsVector2<T>(vid);
 			});
 	}
diff --git a/Engine/Source/Runtime/GeometryCore/Public/Curve/DynamicGraph3.h b/Engine/Source/Runtime/GeometryCore/Public/Curve/DynamicGraph3.h
--- a/Engine/Source/Runtime/GeometryCore/Public/Curve/DynamicGraph3.h
+++ b/Engine/Source/Runtime/GeometryCore/Public/Curve/DynamicGraph3.h
@@ -111,7 +111,7 @@ public:
 	value_iteration<TVector<T>> VerticesItr() const
 	{
 		return vertices_refcount.MappedIndices<TVector<T>>(
-			[=](int vid) {
+			[&](int vid) {
 				return Vertices.template AsVector3<T>(vid);
 			});
 	}
