[NetPhysics] Serialize angular velocity in degrees for physics states.
- Also fix a bug where State instead of Input was called in NetworkPhysicsComponent

diff --git a/Engine/Source/Runtime/Engine/Private/PhysicsEngine/Experimental/PhysScene_Chaos.cpp b/Engine/Source/Runtime/Engine/Private/PhysicsEngine/Experimental/PhysScene_Chaos.cpp
--- a/Engine/Source/Runtime/Engine/Private/PhysicsEngine/Experimental/PhysScene_Chaos.cpp
+++ b/Engine/Source/Runtime/Engine/Private/PhysicsEngine/Experimental/PhysScene_Chaos.cpp
@@ -1374,7 +1374,7 @@ void FPhysScene_Chaos::PopulateReplicationCache(const int32 PhysicsStep)
 			ReplicationState.Position = Handle->GetX();
 			ReplicationState.Quaternion = Handle->GetR();
 			ReplicationState.LinVel = Handle->GetV();
-			ReplicationState.AngVel = Handle->GetW();
+			ReplicationState.AngVel = FMath::RadiansToDegrees(Handle->GetW());
 			ReplicationState.Flags = Handle->ObjectState() == Chaos::EObjectStateType::Sleeping ? ERigidBodyFlags::Sleeping : 0;
 			StateWasCached = true;
 		}
diff --git a/Engine/Source/Runtime/Engine/Private/PhysicsEngine/NetworkPhysicsComponent.cpp b/Engine/Source/Runtime/Engine/Private/PhysicsEngine/NetworkPhysicsComponent.cpp
--- a/Engine/Source/Runtime/Engine/Private/PhysicsEngine/NetworkPhysicsComponent.cpp
+++ b/Engine/Source/Runtime/Engine/Private/PhysicsEngine/NetworkPhysicsComponent.cpp
@@ -722,7 +722,7 @@ void UNetworkPhysicsComponent::OnRep_SetReplicatedInputs()
 				{
 						if (bCompareInputToTriggerRewind)
 						{
-							int32 ResimFrame = StateHistory->ReceiveNewData(*ReceivedInputs, LocalOffset, /*CompareDataForRewind*/ bCompareInputToTriggerRewind);
+							int32 ResimFrame = InputHistory->ReceiveNewData(*ReceivedInputs, LocalOffset, /*CompareDataForRewind*/ bCompareInputToTriggerRewind);
 							if (ResimFrame != INDEX_NONE)
 							{
 								if (Chaos::FPhysicsSolver* Solver = PhysScene->GetSolver())
diff --git a/Engine/Source/Runtime/Engine/Private/PhysicsEngine/PhysicsReplication.cpp b/Engine/Source/Runtime/Engine/Private/PhysicsEngine/PhysicsReplication.cpp
--- a/Engine/Source/Runtime/Engine/Private/PhysicsEngine/PhysicsReplication.cpp
+++ b/Engine/Source/Runtime/Engine/Private/PhysicsEngine/PhysicsReplication.cpp
@@ -898,7 +898,7 @@ void FPhysicsReplicationAsync::UpdateRewindDataTarget(const FPhysicsRepAsyncInpu
 		const int32 LocalFrame = Input.ServerFrame - Input.FrameOffset;
 		RewindData->SetTargetStateAtFrame(*Handle, LocalFrame, Chaos::FFrameAndPhase::EParticleHistoryPhase::PostPushData,
 			Input.TargetState.Position, Input.TargetState.Quaternion,
-			Input.TargetState.LinVel, Input.TargetState.AngVel, (Input.TargetState.Flags & ERigidBodyFlags::Sleeping));
+			Input.TargetState.LinVel, FMath::DegreesToRadians(Input.TargetState.AngVel), (Input.TargetState.Flags & ERigidBodyFlags::Sleeping));
 	}
 }
 
@@ -1567,7 +1567,7 @@ bool FPhysicsReplicationAsync::PredictiveInterpolation(Chaos::FPBDRigidParticleH
 	const FVector TargetPos = FVector(Target.TargetState.Position);
 	const FQuat TargetRot = Target.TargetState.Quaternion;
 	const FVector TargetLinVel = FVector(Target.TargetState.LinVel);
-	const FVector TargetAngVel = FVector(Target.TargetState.AngVel); // Radians
+	const FVector TargetAngVel = FVector(FMath::DegreesToRadians(Target.TargetState.AngVel)); // Radians
 
 	/** --- Reconciliation ---
 	* If target velocities are low enough, check the traveled direction and distance from previous frame and compare with replicated linear velocity.
@@ -1799,6 +1799,7 @@ void FPhysicsReplicationAsync::ExtrapolateTarget(FReplicatedPhysicsTargetAsync&
 	float TargetAngVelSize;
 	FVector TargetAngVelAxis;
 	Target.TargetState.AngVel.FVector::ToDirectionAndLength(TargetAngVelAxis, TargetAngVelSize);
+	TargetAngVelSize = FMath::DegreesToRadians(TargetAngVelSize);
 	const FQuat TargetRotExtrapDelta = FQuat(TargetAngVelAxis, TargetAngVelSize * ExtrapolationTime);
 	Target.TargetState.Quaternion = TargetRotExtrapDelta * Target.TargetState.Quaternion;
 }
@@ -1903,7 +1904,7 @@ bool FPhysicsReplicationAsync::ResimulationReplication(Chaos::FPBDRigidParticleH
 					const FVector CorrectedV = Handle->GetV() + VelCorrection;
 
 					// Angular Velocity Correction
-					const FVector AngVelDiff = Target.TargetState.AngVel - PastState.GetW(); // Angular velocity vector that the server covers but the client doesn't
+					const FVector AngVelDiff = FMath::DegreesToRadians(Target.TargetState.AngVel) - PastState.GetW(); // Angular velocity vector that the server covers but the client doesn't
 					const float CorrectionAmountW = SettingsCurrent.ResimulationSettings.GetAngVelStabilityMultiplier() / NumPredictedFrames;
 					const FVector AngVelCorrection = AngVelDiff * CorrectionAmountW; // Same result as (AngVelDiff / NumPredictedFrames) * VelStabilityMultiplier
 					const FVector CorrectedW = Handle->GetW() + AngVelCorrection;
diff --git a/Engine/Source/Runtime/Engine/Private/PrimitiveComponentPhysics.cpp b/Engine/Source/Runtime/Engine/Private/PrimitiveComponentPhysics.cpp
--- a/Engine/Source/Runtime/Engine/Private/PrimitiveComponentPhysics.cpp
+++ b/Engine/Source/Runtime/Engine/Private/PrimitiveComponentPhysics.cpp
@@ -92,7 +92,7 @@ bool UPrimitiveComponent::GetRigidBodyState(FRigidBodyState& OutState, FName Bon
 			OutState.Position = Transform.GetLocation();
 			OutState.Quaternion = Transform.GetRotation();
 			OutState.LinVel = Interface->GetV(PhysicsObject);
-			OutState.AngVel = Interface->GetW(PhysicsObject);
+			OutState.AngVel = FMath::RadiansToDegrees(Interface->GetW(PhysicsObject));
 			OutState.Flags = (Interface->AreAllSleeping({ PhysicsObject }) ? ERigidBodyFlags::Sleeping : ERigidBodyFlags::None);
 			return true;
 		}
