Iris
- Fixed issue with disconnect caused by Moves being to large, old system would invalidate and ignore the error. Added code to invalidate move on sending side.

diff --git a/Engine/Source/Runtime/Engine/Private/GameFramework/CharacterNetworkSerializationPackedBitsNetSerializer.cpp b/Engine/Source/Runtime/Engine/Private/GameFramework/CharacterNetworkSerializationPackedBitsNetSerializer.cpp
--- a/Engine/Source/Runtime/Engine/Private/GameFramework/CharacterNetworkSerializationPackedBitsNetSerializer.cpp
+++ b/Engine/Source/Runtime/Engine/Private/GameFramework/CharacterNetworkSerializationPackedBitsNetSerializer.cpp
@@ -7,8 +7,9 @@
 #if UE_WITH_IRIS
 
 #include "GameFramework/CharacterNetworkSerializationPackedBitsNetSerializer.h"
+#include "EngineLogs.h"
+#include "HAL/IConsoleManager.h"
 #include "Iris/ReplicationState/PropertyNetSerializerInfoRegistry.h"
-
 #include "Iris/Core/NetObjectReference.h"
 #include "Iris/Serialization/NetBitStreamUtil.h"
 #include "Iris/Serialization/NetSerializerDelegates.h"
@@ -91,6 +92,7 @@ private:
 
 	static FCharacterNetworkSerializationPackedBitsNetSerializer::FNetSerializerRegistryDelegates NetSerializerRegistryDelegates;
 	static const FNetSerializer* ObjectNetSerializer;
+	inline static IConsoleVariable* CVarNetPackedMovementMaxBits = nullptr;
 };
 
 UE_NET_IMPLEMENT_SERIALIZER(FCharacterNetworkSerializationPackedBitsNetSerializer);
@@ -179,9 +181,12 @@ void FCharacterNetworkSerializationPackedBitsNetSerializer::Deserialize(FNetSeri
 	{
 		const uint32 NumDataBits = UE::Net::ReadPackedUint32(Reader);
 
-		if (NumDataBits > Config->MaxAllowedDataBits)
+		const uint32 MaxNumDataBits = CVarNetPackedMovementMaxBits ? static_cast<uint32>(CVarNetPackedMovementMaxBits->GetInt()) : Config->MaxAllowedDataBits;
+		if (NumDataBits > MaxNumDataBits)
 		{
 			Context.SetError(GNetError_ArraySizeTooLarge);
+			UE_LOG(LogNetPlayerMovement, Error, TEXT("FCharacterNetworkSerializationPackedBits::Deserialize: Invalidating move due to NumBits (%u) exceeding allowable limit (%u). See NetPackedMovementMaxBits."), NumDataBits, MaxNumDataBits);
+			ensureMsgf(false, TEXT("Invalidating move due to NumBits exceeding allowable limit"));
 			return;
 		}
 		
@@ -208,13 +213,13 @@ void FCharacterNetworkSerializationPackedBitsNetSerializer::Quantize(FNetSeriali
 
 	if (NumObjectReferences > 0)
 	{
-		FObjectNetSerializerConfig Config;
+		FObjectNetSerializerConfig ObjectNetSerializerConfig;
 		const TObjectPtr<UObject>* SourceReferences = SourceValue.ObjectReferences.GetData();
 		FNetObjectReference* TargetReferences = TargetValue.ObjectReferenceStorage.GetData();
 		for (uint32 ReferenceIndex = 0; ReferenceIndex < NumObjectReferences; ++ReferenceIndex)
 		{
 			FNetQuantizeArgs ObjectArgs;
-			ObjectArgs.NetSerializerConfig = &Config;
+			ObjectArgs.NetSerializerConfig = &ObjectNetSerializerConfig;
 			ObjectArgs.Source = NetSerializerValuePointer(SourceReferences + ReferenceIndex);
 			ObjectArgs.Target = NetSerializerValuePointer(TargetReferences + ReferenceIndex);
 
@@ -222,7 +227,19 @@ void FCharacterNetworkSerializationPackedBitsNetSerializer::Quantize(FNetSeriali
 		}
 	}
 
-	const uint32 NumDataBits = SourceValue.DataBits.Num();
+	uint32 NumDataBits = SourceValue.DataBits.Num();
+
+	const ConfigType* Config = static_cast<const ConfigType*>(Args.NetSerializerConfig);
+	const uint32 MaxNumDataBits = CVarNetPackedMovementMaxBits ? static_cast<uint32>(CVarNetPackedMovementMaxBits->GetInt()) : Config->MaxAllowedDataBits;
+
+	if (NumDataBits > MaxNumDataBits)
+	{
+		// This is just to avoid disconnect and instead warn and invalidate the data on the sending side.
+		UE_LOG(LogNetPlayerMovement, Error, TEXT("FCharacterNetworkSerializationPackedBits::Quantize: Invalidating move due to NumBits (%u) exceeding allowable limit (%u). See NetPackedMovementMaxBits."), NumDataBits, MaxNumDataBits);
+		NumDataBits = 0U;
+		ensureMsgf(false, TEXT("Invalidating move due to NumBits exceeding allowable limit"));
+	}
+
 	TargetValue.DataBitsStorage.AdjustSize(Context, Private::CalculateRequiredWordCount(NumDataBits));
 	if (NumDataBits > 0)
 	{
@@ -296,7 +313,9 @@ bool FCharacterNetworkSerializationPackedBitsNetSerializer::Validate(FNetSeriali
 	const ConfigType* Config = static_cast<const ConfigType*>(Args.NetSerializerConfig);
 	const QuantizedType& SourceValue = *reinterpret_cast<const QuantizedType*>(Args.Source);
 
-	if (SourceValue.ObjectReferenceStorage.Num() > Config->MaxAllowedObjectReferences || SourceValue.NumDataBits > Config->MaxAllowedDataBits)
+	const uint32 MaxNumDataBits = CVarNetPackedMovementMaxBits ? static_cast<uint32>(CVarNetPackedMovementMaxBits->GetInt()) : Config->MaxAllowedDataBits;
+
+	if (SourceValue.ObjectReferenceStorage.Num() > Config->MaxAllowedObjectReferences || SourceValue.NumDataBits > MaxNumDataBits)
 	{
 		return false;
 	}
@@ -347,6 +366,11 @@ FCharacterNetworkSerializationPackedBitsNetSerializer::FNetSerializerRegistryDel
 
 void FCharacterNetworkSerializationPackedBitsNetSerializer::FNetSerializerRegistryDelegates::OnPreFreezeNetSerializerRegistry()
 {
+	FCharacterNetworkSerializationPackedBitsNetSerializer::CVarNetPackedMovementMaxBits = IConsoleManager::Get().FindConsoleVariable(TEXT("p.NetPackedMovementMaxBits"), false);
+#if WITH_SERVER_CODE
+	ensureMsgf(CVarNetPackedMovementMaxBits != nullptr, TEXT("%s"), TEXT("Unable to find cvar p.NetPackedMovementMaxBits"));
+#endif
+
 	UE_NET_REGISTER_NETSERIALIZER_INFO(PropertyNetSerializerRegistry_NAME_CharacterMoveResponsePackedBits);
 	UE_NET_REGISTER_NETSERIALIZER_INFO(PropertyNetSerializerRegistry_NAME_CharacterServerMovePackedBitsPackedBits);
 	UE_NET_REGISTER_NETSERIALIZER_INFO(PropertyNetSerializerRegistry_NAME_CharacterNetworkSerializationPackedBits);
