In RBAN add an option not to run the narrow phase every solver iteration.
#prefligh 660f13bda77c355f53880216

diff --git a/Engine/Plugins/Experimental/PhysicsControl/Source/PhysicsControl/Private/AnimNode_RigidBodyWithControl.cpp b/Engine/Plugins/Experimental/PhysicsControl/Source/PhysicsControl/Private/AnimNode_RigidBodyWithControl.cpp
--- a/Engine/Plugins/Experimental/PhysicsControl/Source/PhysicsControl/Private/AnimNode_RigidBodyWithControl.cpp
+++ b/Engine/Plugins/Experimental/PhysicsControl/Source/PhysicsControl/Private/AnimNode_RigidBodyWithControl.cpp
@@ -832,7 +832,8 @@ void FAnimNode_RigidBodyWithControl::EvaluateSkeletalControl_AnyThread(FComponen
 				SolverSettings.bUseLinearJointSolver,
 				SolverSettings.PositionIterations,
 				SolverSettings.VelocityIterations,
-				SolverSettings.ProjectionIterations);
+				SolverSettings.ProjectionIterations,
+				SolverSettings.bUseManifolds);
 
 			if (!bUseDeferredSimulationTask)
 			{
@@ -1257,7 +1258,8 @@ void FAnimNode_RigidBodyWithControl::InitPhysics(const UAnimInstance* InAnimInst
 			SolverSettings.bUseLinearJointSolver,
 			SolverSettings.PositionIterations,
 			SolverSettings.VelocityIterations,
-			SolverSettings.ProjectionIterations);
+			SolverSettings.ProjectionIterations,
+			SolverSettings.bUseManifolds);
 
 		SolverIterations = PhysicsAssetToUse->SolverIterations;
 
diff --git a/Engine/Source/Editor/PhysicsAssetEditor/Private/PhysicsAssetEditorAnimInstanceProxy.cpp b/Engine/Source/Editor/PhysicsAssetEditor/Private/PhysicsAssetEditorAnimInstanceProxy.cpp
--- a/Engine/Source/Editor/PhysicsAssetEditor/Private/PhysicsAssetEditorAnimInstanceProxy.cpp
+++ b/Engine/Source/Editor/PhysicsAssetEditor/Private/PhysicsAssetEditorAnimInstanceProxy.cpp
@@ -88,7 +88,8 @@ bool FPhysicsAssetEditorAnimInstanceProxy::Evaluate_WithRoot(FPoseContext& Outpu
 			SolverSettings.bUseLinearJointSolver,
 			SolverSettings.PositionIterations,
 			SolverSettings.VelocityIterations,
-			SolverSettings.ProjectionIterations);
+			SolverSettings.ProjectionIterations,
+			SolverSettings.bUseManifolds);
 	}
 
 	if (CurrentAsset != nullptr)
diff --git a/Engine/Source/Runtime/AnimGraphRuntime/Private/BoneControllers/AnimNode_RigidBody.cpp b/Engine/Source/Runtime/AnimGraphRuntime/Private/BoneControllers/AnimNode_RigidBody.cpp
--- a/Engine/Source/Runtime/AnimGraphRuntime/Private/BoneControllers/AnimNode_RigidBody.cpp
+++ b/Engine/Source/Runtime/AnimGraphRuntime/Private/BoneControllers/AnimNode_RigidBody.cpp
@@ -32,7 +32,7 @@
 
 LLM_DEFINE_TAG(Animation_RigidBody);
 
-//PRAGMA_DISABLE_OPTIMIZATION
+//UE_DISABLE_OPTIMIZATION
 
 /////////////////////////////////////////////////////
 // FAnimNode_RigidBody
@@ -847,7 +847,9 @@ void FAnimNode_RigidBody::EvaluateSkeletalControl_AnyThread(FComponentSpacePoseC
 				SolverSettings.bUseLinearJointSolver,
 				SolverSettings.PositionIterations,
 				SolverSettings.VelocityIterations,
-				SolverSettings.ProjectionIterations);
+				SolverSettings.ProjectionIterations,
+				SolverSettings.bUseManifolds);
+
 
 			if (!bUseDeferredSimulationTask)
 			{
@@ -1296,7 +1298,8 @@ void FAnimNode_RigidBody::InitPhysics(const UAnimInstance* InAnimInstance)
 			SolverSettings.bUseLinearJointSolver,
 			SolverSettings.PositionIterations,
 			SolverSettings.VelocityIterations,
-			SolverSettings.ProjectionIterations);
+			SolverSettings.ProjectionIterations,
+			SolverSettings.bUseManifolds);
 
 		SolverIterations = UsePhysicsAsset->SolverIterations;
 	}
diff --git a/Engine/Source/Runtime/Engine/Classes/PhysicsEngine/PhysicsAsset.h b/Engine/Source/Runtime/Engine/Classes/PhysicsEngine/PhysicsAsset.h
--- a/Engine/Source/Runtime/Engine/Classes/PhysicsEngine/PhysicsAsset.h
+++ b/Engine/Source/Runtime/Engine/Classes/PhysicsEngine/PhysicsAsset.h
@@ -85,12 +85,19 @@ struct FPhysicsAssetSolverSettings
 	*/
 	UPROPERTY(EditAnywhere, BlueprintReadWrite, Category = SolverSettings)
 	bool bUseLinearJointSolver;
+
+	/**
+	 * RBAN: It enables the use of multi-point contact manifolds, which are created only once at the start of each tick.
+	 * When disabled, a single-point contact is generated in each solver iteration which is more expensive.
+	*/
+	UPROPERTY(EditAnywhere, BlueprintReadWrite, Category = CollisionSettings)
+	bool bUseManifolds;
 };
 
 
 /**
  * Solver settings for use by the Legacy RigidBody AnimNode (RBAN) solver.
- * Thse settings are no longer used by default and will eventually be deprecated and then removed.
+ * These settings are no longer used by default and will eventually be deprecated and then removed.
  * 
  * @note These settings have no effect when the Physics Asset is used in a world simulation (ragdoll).
  */
@@ -227,7 +234,7 @@ public:
 	FSolverIterations SolverIterations;
 
 	/** 
-	 * Solver type used in physics asset editor. This can be used to make what you see in the asset editror more closely resembles what you
+	 * Solver type used in physics asset editor. This can be used to make what you see in the asset editor more closely resembles what you
 	 * see in game (though there will be differences owing to framerate variation etc). If your asset will primarily be used as a ragdoll 
 	 * select "World", but if it will be used in the AnimGraph select "RBAN".
 	*/
diff --git a/Engine/Source/Runtime/Engine/Private/PhysicsEngine/ImmediatePhysics/ImmediatePhysicsChaos/ImmediatePhysicsSimulation_Chaos.cpp b/Engine/Source/Runtime/Engine/Private/PhysicsEngine/ImmediatePhysics/ImmediatePhysicsChaos/ImmediatePhysicsSimulation_Chaos.cpp
--- a/Engine/Source/Runtime/Engine/Private/PhysicsEngine/ImmediatePhysics/ImmediatePhysicsChaos/ImmediatePhysicsSimulation_Chaos.cpp
+++ b/Engine/Source/Runtime/Engine/Private/PhysicsEngine/ImmediatePhysics/ImmediatePhysicsChaos/ImmediatePhysicsSimulation_Chaos.cpp
@@ -367,7 +367,6 @@ namespace ImmediatePhysics_Chaos
 		DetectorSettings.bFilteringEnabled = false;
 		DetectorSettings.bAllowManifoldReuse = false;
 		DetectorSettings.bDeferNarrowPhase = (ChaosImmediate_Collision_DeferNarrowPhase != 0);
-		DetectorSettings.bAllowManifolds = (ChaosImmediate_Collision_UseManifolds != 0);
 		DetectorSettings.bAllowCCD = false;
 		DetectorSettings.bAllowMACD = false;
 		Implementation->Collisions.SetDetectorSettings(DetectorSettings);
@@ -780,7 +779,7 @@ namespace ImmediatePhysics_Chaos
 		SimSpaceSettings.EulerAlpha = ChaosImmediate_Evolution_SimSpaceEulerAlpha;
 	}
 
-	void FSimulation::SetSolverSettings(const FReal FixedDt, const FReal CullDistance, const FReal MaxDepenetrationVelocity, const int32 UseLinearJointSolver, const int32 PositionIts, const int32 VelocityIts, const int32 ProjectionIts)
+	void FSimulation::SetSolverSettings(const FReal FixedDt, const FReal CullDistance, const FReal MaxDepenetrationVelocity, const int32 UseLinearJointSolver, const int32 PositionIts, const int32 VelocityIts, const int32 ProjectionIts, const bool bUseManifoilds)
 	{
 		if (FixedDt >= FReal(0))
 		{
@@ -816,6 +815,8 @@ namespace ImmediatePhysics_Chaos
 		{
 			Implementation->Joints.SetUseLinearJointSolver(UseLinearJointSolver != 0);
 		}
+
+		Implementation->Collisions.SetAllowManifolds(bUseManifoilds || (ChaosImmediate_Collision_UseManifolds != 0));
 	}
 
 	void FSimulation::DebugDraw()
@@ -867,7 +868,8 @@ namespace ImmediatePhysics_Chaos
 				ChaosImmediate_Joint_UseLinearSolver,
 				ChaosImmediate_Evolution_PositionIterations,
 				ChaosImmediate_Evolution_VelocityIterations,
-				ChaosImmediate_Evolution_ProjectionIterations);
+				ChaosImmediate_Evolution_ProjectionIterations,
+				ChaosImmediate_Collision_UseManifolds != 0);
 
 			FPBDJointSolverSettings JointsSettings = Implementation->Joints.GetSettings();
 			JointsSettings.SwingTwistAngleTolerance = ChaosImmediate_Joint_SwingTwistAngleTolerance;
@@ -913,7 +915,6 @@ namespace ImmediatePhysics_Chaos
 			FCollisionDetectorSettings DetectorSettings = Implementation->Collisions.GetDetectorSettings();
 			DetectorSettings.bAllowManifoldReuse = false;
 			DetectorSettings.bDeferNarrowPhase = (ChaosImmediate_Collision_DeferNarrowPhase != 0);;
-			DetectorSettings.bAllowManifolds = (ChaosImmediate_Collision_UseManifolds != 0);;
 			Implementation->Collisions.SetDetectorSettings(DetectorSettings);
 
 			if (ChaosImmediate_Evolution_StepTime > 0)
diff --git a/Engine/Source/Runtime/Engine/Private/PhysicsEngine/PhysicsAsset.cpp b/Engine/Source/Runtime/Engine/Private/PhysicsEngine/PhysicsAsset.cpp
--- a/Engine/Source/Runtime/Engine/Private/PhysicsEngine/PhysicsAsset.cpp
+++ b/Engine/Source/Runtime/Engine/Private/PhysicsEngine/PhysicsAsset.cpp
@@ -38,6 +38,7 @@ FPhysicsAssetSolverSettings::FPhysicsAssetSolverSettings()
 	, MaxDepenetrationVelocity(0.0f)
 	, FixedTimeStep(0.0f)
 	, bUseLinearJointSolver(true)
+	, bUseManifolds(false)
 {
 }
 
diff --git a/Engine/Source/Runtime/Engine/Public/Physics/ImmediatePhysics/ImmediatePhysicsChaos/ImmediatePhysicsSimulation_Chaos.h b/Engine/Source/Runtime/Engine/Public/Physics/ImmediatePhysics/ImmediatePhysicsChaos/ImmediatePhysicsSimulation_Chaos.h
--- a/Engine/Source/Runtime/Engine/Public/Physics/ImmediatePhysics/ImmediatePhysicsChaos/ImmediatePhysicsSimulation_Chaos.h
+++ b/Engine/Source/Runtime/Engine/Public/Physics/ImmediatePhysics/ImmediatePhysicsChaos/ImmediatePhysicsSimulation_Chaos.h
@@ -106,7 +106,8 @@ namespace ImmediatePhysics_Chaos
 			const int32 UseLinearJointSolver,
 			const int32 PositionIts,
 			const int32 VelocityIts,
-			const int32 ProjectionIts);
+			const int32 ProjectionIts,
+			const bool bUseManifolds);
 
 		/** Explicit debug draw path if the use case needs it to happen at a point outside of the simulation **/
 		ENGINE_API void DebugDraw();
diff --git a/Engine/Source/Runtime/Experimental/Chaos/Public/Chaos/PBDCollisionConstraints.h b/Engine/Source/Runtime/Experimental/Chaos/Public/Chaos/PBDCollisionConstraints.h
--- a/Engine/Source/Runtime/Experimental/Chaos/Public/Chaos/PBDCollisionConstraints.h
+++ b/Engine/Source/Runtime/Experimental/Chaos/Public/Chaos/PBDCollisionConstraints.h
@@ -323,6 +323,11 @@ public:
 		DetectorSettings.MaxVelocityBoundsExpansionMACD = MaxVelocityBoundsExpansion;
 	}
 
+	void SetAllowManifolds(const bool bAllowManifolds)
+	{
+		DetectorSettings.bAllowManifolds = bAllowManifolds;
+	}
+
 protected:
 	CHAOS_API FPBDCollisionConstraint& GetConstraint(int32 Index);
 
