Moved FBlueprintDebugData function definitions out of the header.

diff --git a/Engine/Source/Runtime/Engine/Classes/Engine/BlueprintGeneratedClass.h b/Engine/Source/Runtime/Engine/Classes/Engine/BlueprintGeneratedClass.h
--- a/Engine/Source/Runtime/Engine/Classes/Engine/BlueprintGeneratedClass.h
+++ b/Engine/Source/Runtime/Engine/Classes/Engine/BlueprintGeneratedClass.h
@@ -199,12 +199,14 @@ USTRUCT()
 struct FBlueprintDebugData
 {
 	GENERATED_USTRUCT_BODY()
-	FBlueprintDebugData()
-	{
-	}
 
-	~FBlueprintDebugData()
-	{ }
+	ENGINE_API FBlueprintDebugData();
+	ENGINE_API FBlueprintDebugData(FBlueprintDebugData&&);
+	ENGINE_API FBlueprintDebugData(const FBlueprintDebugData&);
+	ENGINE_API FBlueprintDebugData& operator=(FBlueprintDebugData&&);
+	ENGINE_API FBlueprintDebugData& operator=(const FBlueprintDebugData&);
+	ENGINE_API ~FBlueprintDebugData();
+
 #if WITH_EDITORONLY_DATA
 
 protected:
@@ -234,267 +236,67 @@ protected:
 public:
 
 	// Returns the UEdGraphNode associated with the UUID, or nullptr if there isn't one.
-	UEdGraphNode* FindNodeFromUUID(int32 UUID) const
-	{
-		if (const TWeakObjectPtr<UEdGraphNode>* pParentNode = DebugNodesAllocatedUniqueIDsMap.Find(UUID))
-		{
-			return pParentNode->Get();
-		}
-	
-		return nullptr;
-	}
+	ENGINE_API UEdGraphNode* FindNodeFromUUID(int32 UUID) const;
 
-	bool IsValid() const
-	{
-		return DebugNodeLineNumbers.Num() > 0;
-	}
+	ENGINE_API bool IsValid() const;
 
 	// Finds the UEdGraphNode associated with the code location Function+CodeOffset, or nullptr if there isn't one
-	UEdGraphNode* FindSourceNodeFromCodeLocation(UFunction* Function, int32 CodeOffset, bool bAllowImpreciseHit) const
-	{
-		if (const FDebuggingInfoForSingleFunction* pFuncInfo = PerFunctionLineNumbers.Find(Function))
-		{
-			UEdGraphNode* Result = pFuncInfo->LineNumberToSourceNodeMap.FindRef(CodeOffset).Get();
-
-			if ((Result == nullptr) && bAllowImpreciseHit)
-			{
-				for (int32 TrialOffset = CodeOffset + 1; (Result == nullptr) && (TrialOffset < Function->Script.Num()); ++TrialOffset)
-				{
-					Result = pFuncInfo->LineNumberToSourceNodeMap.FindRef(TrialOffset).Get();
-				}
-			}
-
-			return Result;
-		}
-
-		return nullptr;
-	}
+	ENGINE_API UEdGraphNode* FindSourceNodeFromCodeLocation(UFunction* Function, int32 CodeOffset, bool bAllowImpreciseHit) const;
 
 	// Finds the source pin associated with the code location Function+CodeOffset, or nullptr if there isn't one
-	UEdGraphPin* FindSourcePinFromCodeLocation(UFunction* Function, int32 CodeOffset) const
-	{
-		if (const FDebuggingInfoForSingleFunction* pFuncInfo = PerFunctionLineNumbers.Find(Function))
-		{
-			return pFuncInfo->LineNumberToSourcePinMap.FindRef(CodeOffset).Get();
-		}
-
-		return nullptr;
-	}
+	ENGINE_API UEdGraphPin* FindSourcePinFromCodeLocation(UFunction* Function, int32 CodeOffset) const;
 
 	// Finds all code locations (Function+CodeOffset) associated with the source pin
-	void FindAllCodeLocationsFromSourcePin(UEdGraphPin const* SourcePin, UFunction* InFunction, TArray<int32>& OutPinToCodeAssociations) const
-	{
-		OutPinToCodeAssociations.Empty();
-
-		if (const FDebuggingInfoForSingleFunction* pFuncInfo = PerFunctionLineNumbers.Find(InFunction))
-		{
-			pFuncInfo->SourcePinToLineNumbersMap.MultiFind(SourcePin, OutPinToCodeAssociations, true);
-		}
-	}
+	ENGINE_API void FindAllCodeLocationsFromSourcePin(UEdGraphPin const* SourcePin, UFunction* InFunction, TArray<int32>& OutPinToCodeAssociations) const;
 
 	// Finds the first code location (Function+CodeOffset) associated with the source pin within the given range, or INDEX_NONE if there isn't one
-	int32 FindCodeLocationFromSourcePin(UEdGraphPin const* SourcePin, UFunction* InFunction, FInt32Range InRange = FInt32Range()) const
-	{
-		TArray<int32> PinToCodeAssociations;
-		FindAllCodeLocationsFromSourcePin(SourcePin, InFunction, PinToCodeAssociations);
-
-		for (int32 i = 0; i < PinToCodeAssociations.Num(); ++i)
-		{
-			if (InRange.Contains(PinToCodeAssociations[i]))
-			{
-				return PinToCodeAssociations[i];
-			}
-		}
-
-		return INDEX_NONE;
-	}
+	ENGINE_API int32 FindCodeLocationFromSourcePin(UEdGraphPin const* SourcePin, UFunction* InFunction, FInt32Range InRange = FInt32Range()) const;
 
 	// Finds all code locations (Function+CodeOffset) associated with the source node
-	void FindAllCodeLocationsFromSourceNode(UEdGraphNode* SourceNode, UFunction* InFunction, TArray<int32>& OutNodeToCodeAssociations) const
-	{
-		OutNodeToCodeAssociations.Empty();
-
-		if (const FDebuggingInfoForSingleFunction* pFuncInfo = PerFunctionLineNumbers.Find(InFunction))
-		{
-			for (auto CodeLocation : pFuncInfo->LineNumberToSourceNodeMap)
-			{
-				if (CodeLocation.Value == SourceNode)
-				{
-					OutNodeToCodeAssociations.Add(CodeLocation.Key);
-				}
-			}
-		}
-	}
+	ENGINE_API void FindAllCodeLocationsFromSourceNode(UEdGraphNode* SourceNode, UFunction* InFunction, TArray<int32>& OutNodeToCodeAssociations) const;
 
 	// Finds the pure node script code range associated with the [impure] source node, or FInt32Range(INDEX_NONE) if there is no existing association
-	FInt32Range FindPureNodeScriptCodeRangeFromSourceNode(const UEdGraphNode* SourceNode, UFunction* InFunction) const
-	{
-		FInt32Range Result = FInt32Range(INDEX_NONE);
+	ENGINE_API FInt32Range FindPureNodeScriptCodeRangeFromSourceNode(const UEdGraphNode* SourceNode, UFunction* InFunction) const;
 
-		if (const FDebuggingInfoForSingleFunction* DebugInfoPtr = PerFunctionLineNumbers.Find(InFunction))
-		{
-			if (const FInt32Range* ValuePtr = DebugInfoPtr->PureNodeScriptCodeRangeMap.Find(MakeWeakObjectPtr(const_cast<UEdGraphNode*>(SourceNode))))
-			{
-				Result = *ValuePtr;
-			}
-		}
-
-		return Result;
-	}
-
-	const TArray<TWeakObjectPtr<UEdGraphNode> >* FindExpansionSourceNodesFromCodeLocation(UFunction* Function, int32 CodeOffset) const
-	{
-		if (const FDebuggingInfoForSingleFunction* pFuncInfo = PerFunctionLineNumbers.Find(Function))
-		{
-			return pFuncInfo->LineNumberToTunnelInstanceSourceNodesMap.Find(CodeOffset);
-		}
-
-		return nullptr;
-	}
+	ENGINE_API const TArray<TWeakObjectPtr<UEdGraphNode> >* FindExpansionSourceNodesFromCodeLocation(UFunction* Function, int32 CodeOffset) const;
 
 	// Finds the breakpoint injection site(s) in bytecode if any were associated with the given node
-	void FindBreakpointInjectionSites(UEdGraphNode* Node, TArray<uint8*>& InstallSites) const
-	{
-		TArray<int32> RecordIndices;
-		DebugNodeIndexLookup.MultiFind(Node, RecordIndices, true);
-		for(int i = 0; i < RecordIndices.Num(); ++i)
-		{
-			int32 RecordIndex = RecordIndices[i];
-			if (DebugNodeLineNumbers.IsValidIndex(RecordIndex))
-			{
-				const FNodeToCodeAssociation& Record = DebugNodeLineNumbers[RecordIndex];
-				if (UFunction* Scope = Record.Scope.Get())
-				{
-					if (Scope->Script.IsValidIndex(Record.Offset))
-					{
-						InstallSites.Add(&(Scope->Script[Record.Offset]));
-					}
-				}
-			}
-		}
-	}
+	ENGINE_API void FindBreakpointInjectionSites(UEdGraphNode* Node, TArray<uint8*>& InstallSites) const;
 
 	// Looks thru the debugging data for any class variables associated with the node
-	FProperty* FindClassPropertyForPin(const UEdGraphPin* Pin) const
-	{
-		if (!Pin)
-		{
-			return nullptr;
-		}
-
-		TFieldPath<FProperty> PropertyPtr = DebugPinToPropertyMap.FindRef(Pin);
-		if ((PropertyPtr == nullptr) && (Pin->LinkedTo.Num() > 0))
-		{
-			// Try checking the other side of the connection
-			PropertyPtr = DebugPinToPropertyMap.FindRef(Pin->LinkedTo[0]);
-		}
-
-		return *PropertyPtr;
-	}
+	ENGINE_API FProperty* FindClassPropertyForPin(const UEdGraphPin* Pin) const;
 
 	// Looks thru the debugging data for any class variables associated with the node (e.g., temporary variables or timelines)
-	FProperty* FindClassPropertyForNode(const UEdGraphNode* Node) const
-	{
-		return *DebugObjectToPropertyMap.FindRef(MakeWeakObjectPtr(const_cast<UEdGraphNode*>(Node)));
-	}
+	ENGINE_API FProperty* FindClassPropertyForNode(const UEdGraphNode* Node) const;
 
 	// Adds a debug record for a source node and destination in the bytecode of a specified function
-	void RegisterNodeToCodeAssociation(UEdGraphNode* SourceNode, const TArray<TWeakObjectPtr<UEdGraphNode> >& ExpansionSourceNodes, UFunction* InFunction, int32 CodeOffset, bool bBreakpointSite)
-	{
-		//@TODO: Nasty expansion behavior during compile time
-		if (bBreakpointSite)
-		{
-			DebugNodeLineNumbers.Emplace(SourceNode, InFunction, CodeOffset);
-			DebugNodeIndexLookup.Add(SourceNode, DebugNodeLineNumbers.Num() - 1);
-		}
+	ENGINE_API void RegisterNodeToCodeAssociation(UEdGraphNode* SourceNode, const TArray<TWeakObjectPtr<UEdGraphNode> >& ExpansionSourceNodes, UFunction* InFunction, int32 CodeOffset, bool bBreakpointSite);
 
-		FDebuggingInfoForSingleFunction& PerFuncInfo = PerFunctionLineNumbers.FindOrAdd(InFunction);
-		PerFuncInfo.LineNumberToSourceNodeMap.Add(CodeOffset, SourceNode);
+	ENGINE_API void RegisterPureNodeScriptCodeRange(UEdGraphNode* SourceNode, UFunction* InFunction, FInt32Range InPureNodeScriptCodeRange);
 
-		if (ExpansionSourceNodes.Num() > 0)
-		{
-			PerFuncInfo.LineNumberToTunnelInstanceSourceNodesMap.Add(CodeOffset, ExpansionSourceNodes);
-		}
-	}
+	ENGINE_API void RegisterPinToCodeAssociation(UEdGraphPin const* SourcePin, UFunction* InFunction, int32 CodeOffset);
 
-	void RegisterPureNodeScriptCodeRange(UEdGraphNode* SourceNode, UFunction* InFunction, FInt32Range InPureNodeScriptCodeRange)
-	{
-		FDebuggingInfoForSingleFunction& PerFuncInfo = PerFunctionLineNumbers.FindOrAdd(InFunction);
-		PerFuncInfo.PureNodeScriptCodeRangeMap.Add(SourceNode, InPureNodeScriptCodeRange);
-	}
+	ENGINE_API const TMap<int32, FName>& GetEntryPoints() const;
 
-	void RegisterPinToCodeAssociation(UEdGraphPin const* SourcePin, UFunction* InFunction, int32 CodeOffset)
-	{
-		FDebuggingInfoForSingleFunction& PerFuncInfo = PerFunctionLineNumbers.FindOrAdd(InFunction);
-		PerFuncInfo.LineNumberToSourcePinMap.Add(CodeOffset, SourcePin);
-		PerFuncInfo.SourcePinToLineNumbersMap.Add(SourcePin, CodeOffset);
-	}
-
-	const TMap<int32, FName>& GetEntryPoints() const { return EntryPoints; }
-
-	bool IsValidEntryPoint(const int32 LinkId) const { return EntryPoints.Contains(LinkId); }
+	ENGINE_API bool IsValidEntryPoint(const int32 LinkId) const;
 
-	void RegisterEntryPoint(const int32 ScriptOffset, const FName FunctionName)
-	{
-		EntryPoints.Add(ScriptOffset, FunctionName);
-	}
+	ENGINE_API void RegisterEntryPoint(const int32 ScriptOffset, const FName FunctionName);
 
 	// Registers an association between an object (pin or node typically) and an associated class member property
-	void RegisterClassPropertyAssociation(class UObject* TrueSourceObject, class FProperty* AssociatedProperty)
-	{
-		DebugObjectToPropertyMap.Add(TrueSourceObject, AssociatedProperty);
-	}
+	ENGINE_API void RegisterClassPropertyAssociation(class UObject* TrueSourceObject, class FProperty* AssociatedProperty);
 
-	void RegisterClassPropertyAssociation(const UEdGraphPin* TrueSourcePin, class FProperty* AssociatedProperty)
-	{
-		if (TrueSourcePin)
-		{
-			DebugPinToPropertyMap.Add(TrueSourcePin, AssociatedProperty);
-		}
-	}
+	ENGINE_API void RegisterClassPropertyAssociation(const UEdGraphPin* TrueSourcePin, class FProperty* AssociatedProperty);
 
 	// Registers an association between a UUID and a node
-	void RegisterUUIDAssociation(UEdGraphNode* TrueSourceNode, int32 UUID)
-	{
-		DebugNodesAllocatedUniqueIDsMap.Add(UUID, TrueSourceNode);
-	}
+	ENGINE_API void RegisterUUIDAssociation(UEdGraphNode* TrueSourceNode, int32 UUID);
 
 	// Returns the object that caused the specified property to be created (can return nullptr if the association is unknown)
-	UObject* FindObjectThatCreatedProperty(class FProperty* AssociatedProperty) const
-	{
-		if (const TWeakObjectPtr<UObject>* pValue = DebugObjectToPropertyMap.FindKey(AssociatedProperty))
-		{
-			return pValue->Get();
-		}
-		else
-		{
-			return nullptr;
-		}
-	}
+	ENGINE_API UObject* FindObjectThatCreatedProperty(class FProperty* AssociatedProperty) const;
 
 	// Returns the pin that caused the specified property to be created (can return nullptr if the association is unknown or the association is from an object instead)
-	UEdGraphPin* FindPinThatCreatedProperty(class FProperty* AssociatedProperty) const
-	{
-		if (const FEdGraphPinReference* pValue = DebugPinToPropertyMap.FindKey(AssociatedProperty))
-		{
-			return pValue->Get();
-		}
-		else
-		{
-			return nullptr;
-		}
-	}
+	ENGINE_API UEdGraphPin* FindPinThatCreatedProperty(class FProperty* AssociatedProperty) const;
 
-	void GenerateReversePropertyMap(TMap<FProperty*, UObject*>& PropertySourceMap)
-	{
-		for (TMap<TWeakObjectPtr<UObject>, TFieldPath<FProperty>>::TIterator MapIt(DebugObjectToPropertyMap); MapIt; ++MapIt)
-		{
-			if (UObject* SourceObj = MapIt.Key().Get())
-			{
-				PropertySourceMap.Add(*MapIt.Value(), SourceObj);
-			}
-		}
-	}
+	ENGINE_API void GenerateReversePropertyMap(TMap<FProperty*, UObject*>& PropertySourceMap);
 #endif
 };
 
diff --git a/Engine/Source/Runtime/Engine/Private/BlueprintGeneratedClass.cpp b/Engine/Source/Runtime/Engine/Private/BlueprintGeneratedClass.cpp
--- a/Engine/Source/Runtime/Engine/Private/BlueprintGeneratedClass.cpp
+++ b/Engine/Source/Runtime/Engine/Private/BlueprintGeneratedClass.cpp
@@ -3078,3 +3078,267 @@ ESoftObjectPathCollectType FImportExportCollector::Union(ESoftObjectPathCollectT
 
 #endif //if WITH_EDITORONLY_DATA
 
+FBlueprintDebugData::FBlueprintDebugData() = default;
+FBlueprintDebugData::FBlueprintDebugData(FBlueprintDebugData&&) = default;
+FBlueprintDebugData::FBlueprintDebugData(const FBlueprintDebugData&) = default;
+FBlueprintDebugData& FBlueprintDebugData::operator=(FBlueprintDebugData&&) = default;
+FBlueprintDebugData& FBlueprintDebugData::operator=(const FBlueprintDebugData&) = default;
+FBlueprintDebugData::~FBlueprintDebugData() = default;
+
+#if WITH_EDITORONLY_DATA
+
+UEdGraphNode* FBlueprintDebugData::FindNodeFromUUID(int32 UUID) const
+{
+	if (const TWeakObjectPtr<UEdGraphNode>* pParentNode = DebugNodesAllocatedUniqueIDsMap.Find(UUID))
+	{
+		return pParentNode->Get();
+	}
+	
+	return nullptr;
+}
+
+bool FBlueprintDebugData::IsValid() const
+{
+	return DebugNodeLineNumbers.Num() > 0;
+}
+
+UEdGraphNode* FBlueprintDebugData::FindSourceNodeFromCodeLocation(UFunction* Function, int32 CodeOffset, bool bAllowImpreciseHit) const
+{
+	if (const FDebuggingInfoForSingleFunction* pFuncInfo = PerFunctionLineNumbers.Find(Function))
+	{
+		UEdGraphNode* Result = pFuncInfo->LineNumberToSourceNodeMap.FindRef(CodeOffset).Get();
+
+		if ((Result == nullptr) && bAllowImpreciseHit)
+		{
+			for (int32 TrialOffset = CodeOffset + 1; (Result == nullptr) && (TrialOffset < Function->Script.Num()); ++TrialOffset)
+			{
+				Result = pFuncInfo->LineNumberToSourceNodeMap.FindRef(TrialOffset).Get();
+			}
+		}
+
+		return Result;
+	}
+
+	return nullptr;
+}
+
+UEdGraphPin* FBlueprintDebugData::FindSourcePinFromCodeLocation(UFunction* Function, int32 CodeOffset) const
+{
+	if (const FDebuggingInfoForSingleFunction* pFuncInfo = PerFunctionLineNumbers.Find(Function))
+	{
+		return pFuncInfo->LineNumberToSourcePinMap.FindRef(CodeOffset).Get();
+	}
+
+	return nullptr;
+}
+
+void FBlueprintDebugData::FindAllCodeLocationsFromSourcePin(UEdGraphPin const* SourcePin, UFunction* InFunction, TArray<int32>& OutPinToCodeAssociations) const
+{
+	OutPinToCodeAssociations.Empty();
+
+	if (const FDebuggingInfoForSingleFunction* pFuncInfo = PerFunctionLineNumbers.Find(InFunction))
+	{
+		pFuncInfo->SourcePinToLineNumbersMap.MultiFind(SourcePin, OutPinToCodeAssociations, true);
+	}
+}
+
+int32 FBlueprintDebugData::FindCodeLocationFromSourcePin(UEdGraphPin const* SourcePin, UFunction* InFunction, FInt32Range InRange) const
+{
+	TArray<int32> PinToCodeAssociations;
+	FindAllCodeLocationsFromSourcePin(SourcePin, InFunction, PinToCodeAssociations);
+
+	for (int32 i = 0; i < PinToCodeAssociations.Num(); ++i)
+	{
+		if (InRange.Contains(PinToCodeAssociations[i]))
+		{
+			return PinToCodeAssociations[i];
+		}
+	}
+
+	return INDEX_NONE;
+}
+
+void FBlueprintDebugData::FindAllCodeLocationsFromSourceNode(UEdGraphNode* SourceNode, UFunction* InFunction, TArray<int32>& OutNodeToCodeAssociations) const
+{
+	OutNodeToCodeAssociations.Empty();
+
+	if (const FDebuggingInfoForSingleFunction* pFuncInfo = PerFunctionLineNumbers.Find(InFunction))
+	{
+		for (auto CodeLocation : pFuncInfo->LineNumberToSourceNodeMap)
+		{
+			if (CodeLocation.Value == SourceNode)
+			{
+				OutNodeToCodeAssociations.Add(CodeLocation.Key);
+			}
+		}
+	}
+}
+
+FInt32Range FBlueprintDebugData::FindPureNodeScriptCodeRangeFromSourceNode(const UEdGraphNode* SourceNode, UFunction* InFunction) const
+{
+	FInt32Range Result = FInt32Range(INDEX_NONE);
+
+	if (const FDebuggingInfoForSingleFunction* DebugInfoPtr = PerFunctionLineNumbers.Find(InFunction))
+	{
+		if (const FInt32Range* ValuePtr = DebugInfoPtr->PureNodeScriptCodeRangeMap.Find(MakeWeakObjectPtr(const_cast<UEdGraphNode*>(SourceNode))))
+		{
+			Result = *ValuePtr;
+		}
+	}
+
+	return Result;
+}
+
+const TArray<TWeakObjectPtr<UEdGraphNode> >* FBlueprintDebugData::FindExpansionSourceNodesFromCodeLocation(UFunction* Function, int32 CodeOffset) const
+{
+	if (const FDebuggingInfoForSingleFunction* pFuncInfo = PerFunctionLineNumbers.Find(Function))
+	{
+		return pFuncInfo->LineNumberToTunnelInstanceSourceNodesMap.Find(CodeOffset);
+	}
+
+	return nullptr;
+}
+
+void FBlueprintDebugData::FindBreakpointInjectionSites(UEdGraphNode* Node, TArray<uint8*>& InstallSites) const
+{
+	TArray<int32> RecordIndices;
+	DebugNodeIndexLookup.MultiFind(Node, RecordIndices, true);
+	for(int i = 0; i < RecordIndices.Num(); ++i)
+	{
+		int32 RecordIndex = RecordIndices[i];
+		if (DebugNodeLineNumbers.IsValidIndex(RecordIndex))
+		{
+			const FNodeToCodeAssociation& Record = DebugNodeLineNumbers[RecordIndex];
+			if (UFunction* Scope = Record.Scope.Get())
+			{
+				if (Scope->Script.IsValidIndex(Record.Offset))
+				{
+					InstallSites.Add(&(Scope->Script[Record.Offset]));
+				}
+			}
+		}
+	}
+}
+
+FProperty* FBlueprintDebugData::FindClassPropertyForPin(const UEdGraphPin* Pin) const
+{
+	if (!Pin)
+	{
+		return nullptr;
+	}
+
+	TFieldPath<FProperty> PropertyPtr = DebugPinToPropertyMap.FindRef(Pin);
+	if ((PropertyPtr == nullptr) && (Pin->LinkedTo.Num() > 0))
+	{
+		// Try checking the other side of the connection
+		PropertyPtr = DebugPinToPropertyMap.FindRef(Pin->LinkedTo[0]);
+	}
+
+	return *PropertyPtr;
+}
+
+FProperty* FBlueprintDebugData::FindClassPropertyForNode(const UEdGraphNode* Node) const
+{
+	return *DebugObjectToPropertyMap.FindRef(MakeWeakObjectPtr(const_cast<UEdGraphNode*>(Node)));
+}
+
+void FBlueprintDebugData::RegisterNodeToCodeAssociation(UEdGraphNode* SourceNode, const TArray<TWeakObjectPtr<UEdGraphNode> >& ExpansionSourceNodes, UFunction* InFunction, int32 CodeOffset, bool bBreakpointSite)
+{
+	//@TODO: Nasty expansion behavior during compile time
+	if (bBreakpointSite)
+	{
+		DebugNodeLineNumbers.Emplace(SourceNode, InFunction, CodeOffset);
+		DebugNodeIndexLookup.Add(SourceNode, DebugNodeLineNumbers.Num() - 1);
+	}
+
+	FDebuggingInfoForSingleFunction& PerFuncInfo = PerFunctionLineNumbers.FindOrAdd(InFunction);
+	PerFuncInfo.LineNumberToSourceNodeMap.Add(CodeOffset, SourceNode);
+
+	if (ExpansionSourceNodes.Num() > 0)
+	{
+		PerFuncInfo.LineNumberToTunnelInstanceSourceNodesMap.Add(CodeOffset, ExpansionSourceNodes);
+	}
+}
+
+void FBlueprintDebugData::RegisterPureNodeScriptCodeRange(UEdGraphNode* SourceNode, UFunction* InFunction, FInt32Range InPureNodeScriptCodeRange)
+{
+	FDebuggingInfoForSingleFunction& PerFuncInfo = PerFunctionLineNumbers.FindOrAdd(InFunction);
+	PerFuncInfo.PureNodeScriptCodeRangeMap.Add(SourceNode, InPureNodeScriptCodeRange);
+}
+
+void FBlueprintDebugData::RegisterPinToCodeAssociation(UEdGraphPin const* SourcePin, UFunction* InFunction, int32 CodeOffset)
+{
+	FDebuggingInfoForSingleFunction& PerFuncInfo = PerFunctionLineNumbers.FindOrAdd(InFunction);
+	PerFuncInfo.LineNumberToSourcePinMap.Add(CodeOffset, SourcePin);
+	PerFuncInfo.SourcePinToLineNumbersMap.Add(SourcePin, CodeOffset);
+}
+
+const TMap<int32, FName>& FBlueprintDebugData::GetEntryPoints() const
+{
+	return EntryPoints;
+}
+
+bool FBlueprintDebugData::IsValidEntryPoint(const int32 LinkId) const
+{
+	return EntryPoints.Contains(LinkId);
+}
+
+void FBlueprintDebugData::RegisterEntryPoint(const int32 ScriptOffset, const FName FunctionName)
+{
+	EntryPoints.Add(ScriptOffset, FunctionName);
+}
+
+void FBlueprintDebugData::RegisterClassPropertyAssociation(class UObject* TrueSourceObject, class FProperty* AssociatedProperty)
+{
+	DebugObjectToPropertyMap.Add(TrueSourceObject, AssociatedProperty);
+}
+
+void FBlueprintDebugData::RegisterClassPropertyAssociation(const UEdGraphPin* TrueSourcePin, class FProperty* AssociatedProperty)
+{
+	if (TrueSourcePin)
+	{
+		DebugPinToPropertyMap.Add(TrueSourcePin, AssociatedProperty);
+	}
+}
+
+void FBlueprintDebugData::RegisterUUIDAssociation(UEdGraphNode* TrueSourceNode, int32 UUID)
+{
+	DebugNodesAllocatedUniqueIDsMap.Add(UUID, TrueSourceNode);
+}
+
+UObject* FBlueprintDebugData::FindObjectThatCreatedProperty(class FProperty* AssociatedProperty) const
+{
+	if (const TWeakObjectPtr<UObject>* pValue = DebugObjectToPropertyMap.FindKey(AssociatedProperty))
+	{
+		return pValue->Get();
+	}
+	else
+	{
+		return nullptr;
+	}
+}
+
+UEdGraphPin* FBlueprintDebugData::FindPinThatCreatedProperty(class FProperty* AssociatedProperty) const
+{
+	if (const FEdGraphPinReference* pValue = DebugPinToPropertyMap.FindKey(AssociatedProperty))
+	{
+		return pValue->Get();
+	}
+	else
+	{
+		return nullptr;
+	}
+}
+
+void FBlueprintDebugData::GenerateReversePropertyMap(TMap<FProperty*, UObject*>& PropertySourceMap)
+{
+	for (TMap<TWeakObjectPtr<UObject>, TFieldPath<FProperty>>::TIterator MapIt(DebugObjectToPropertyMap); MapIt; ++MapIt)
+	{
+		if (UObject* SourceObj = MapIt.Key().Get())
+		{
+			PropertySourceMap.Add(*MapIt.Value(), SourceObj);
+		}
+	}
+}
+
+#endif //if WITH_EDITORONLY_DATA
