Add Domain to default material assertion to narrow down quick launch failure.

#rnx

diff --git a/Engine/Source/Runtime/Engine/Private/Materials/Material.cpp b/Engine/Source/Runtime/Engine/Private/Materials/Material.cpp
--- a/Engine/Source/Runtime/Engine/Private/Materials/Material.cpp
+++ b/Engine/Source/Runtime/Engine/Private/Materials/Material.cpp
@@ -569,7 +569,7 @@ void UMaterialInterface::InitDefaultMaterials()
 					)
 				{
 					GDefaultMaterials[Domain] = LoadObject<UMaterial>(nullptr, *ResolvedPath, nullptr, LOAD_DisableDependencyPreloading, nullptr);
-					checkf(GDefaultMaterials[Domain] != nullptr, TEXT("Cannot load default material '%s' from path '%s'"), GDefaultMaterialNames[Domain], *ResolvedPath);
+					checkf(GDefaultMaterials[Domain] != nullptr, TEXT("Cannot load default material '%s' [Domain=%s] from path '%s'"), GDefaultMaterialNames[Domain], *MaterialDomainString((EMaterialDomain)Domain), *ResolvedPath);
 				}
 				if (GDefaultMaterials[Domain])
 				{
