Replace ensures with log messages in the network metrics database.

diff --git a/Engine/Source/Runtime/Engine/Private/Net/NetworkMetricsDatabase.cpp b/Engine/Source/Runtime/Engine/Private/Net/NetworkMetricsDatabase.cpp
--- a/Engine/Source/Runtime/Engine/Private/Net/NetworkMetricsDatabase.cpp
+++ b/Engine/Source/Runtime/Engine/Private/Net/NetworkMetricsDatabase.cpp
@@ -2,6 +2,7 @@
 
 #include "Net/NetworkMetricsDatabase.h"
 #include "EngineStats.h"
+#include "Logging/LogMacros.h"
 #include "ProfilingDebugging/CsvProfiler.h"
 
 #if USE_SERVER_PERF_COUNTERS
@@ -10,6 +11,8 @@
 
 #include UE_INLINE_GENERATED_CPP_BY_NAME(NetworkMetricsDatabase)
 
+DEFINE_LOG_CATEGORY_STATIC(LogNetworkMetrics, Error, All);
+
 void UNetworkMetricsDatabase::CreateFloat(const FName MetricName, float DefaultValue)
 {
 	if (ensureMsgf(!MetricTypes.Contains(MetricName), TEXT("Metric %s already exists in the database."), *MetricName.ToString()))
@@ -34,66 +37,108 @@ void UNetworkMetricsDatabase::CreateInt(const FName MetricName, int64 DefaultVal
 	}
 }
 
-void UNetworkMetricsDatabase::SetFloat(const FName MetricName, float Value)
+bool UNetworkMetricsDatabase::SetFloat(const FName MetricName, float Value)
 {
 	UE::Net::FNetworkMetric<float>* Metric = MetricFloats.Find(MetricName);
-	if (ensureMsgf(Metric, TEXT("Cannot find float metric %s."), *MetricName.ToString()))
+	if (Metric)
 	{
 		Metric->Value = Value;
+		return true;
+	}
+	else
+	{
+		UE_LOG(LogNetworkMetrics, Warning, TEXT("UNetworkMetricsDatabase::SetFloat Cannot find float metric %s."), *MetricName.ToString());
+		return false;
 	}
 }
 
-void UNetworkMetricsDatabase::SetMinFloat(const FName MetricName, float Value)
+bool UNetworkMetricsDatabase::SetMinFloat(const FName MetricName, float Value)
 {
 	UE::Net::FNetworkMetric<float>* Metric = MetricFloats.Find(MetricName);
-	if (ensureMsgf(Metric, TEXT("Cannot find float metric %s."), *MetricName.ToString()))
+	if (Metric)
 	{
 		Metric->Value = FMath::Min<float>(Value, Metric->Value);
+		return true;
+	}
+	else
+	{
+		UE_LOG(LogNetworkMetrics, Warning, TEXT("UNetworkMetricsDatabase::SetMinFloat Cannot find float metric %s."), *MetricName.ToString());
+		return false;
 	}
 }
 
-void UNetworkMetricsDatabase::SetMaxFloat(const FName MetricName, float Value)
+bool UNetworkMetricsDatabase::SetMaxFloat(const FName MetricName, float Value)
 {
 	UE::Net::FNetworkMetric<float>* Metric = MetricFloats.Find(MetricName);
-	if (ensureMsgf(Metric, TEXT("Cannot find float metric %s."), *MetricName.ToString()))
+	if (Metric)
 	{
 		Metric->Value = FMath::Max<float>(Value, Metric->Value);
+		return true;
+	}
+	else
+	{
+		UE_LOG(LogNetworkMetrics, Warning, TEXT("UNetworkMetricsDatabase::SetMaxFloat Cannot find float metric %s."), *MetricName.ToString());
+		return false;
 	}
 }
 
-void UNetworkMetricsDatabase::SetInt(const FName MetricName, int64 Value)
+bool UNetworkMetricsDatabase::SetInt(const FName MetricName, int64 Value)
 {
 	UE::Net::FNetworkMetric<int64>* Metric = MetricInts.Find(MetricName);
-	if (ensureMsgf(Metric, TEXT("Cannot find integer metric %s."), *MetricName.ToString()))
+	if (Metric)
 	{
 		Metric->Value = Value;
+		return true;
+	}
+	else
+	{
+		UE_LOG(LogNetworkMetrics, Warning, TEXT("UNetworkMetricsDatabase::SetInt Cannot find integer metric %s."), *MetricName.ToString());
+		return false;
 	}
 }
 
-void UNetworkMetricsDatabase::SetMinInt(const FName MetricName, int64 Value)
+bool UNetworkMetricsDatabase::SetMinInt(const FName MetricName, int64 Value)
 {
 	UE::Net::FNetworkMetric<int64>* Metric = MetricInts.Find(MetricName);
-	if (ensureMsgf(Metric, TEXT("Cannot find integer metric %s."), *MetricName.ToString()))
+	if (Metric)
 	{
 		Metric->Value = FMath::Min<int64>(Value, Metric->Value);
+		return true;
+	}
+	else
+	{
+		UE_LOG(LogNetworkMetrics, Warning, TEXT("UNetworkMetricsDatabase::SetMinInt Cannot find integer metric %s."), *MetricName.ToString());
+		return false;
 	}
 }
 
-void UNetworkMetricsDatabase::SetMaxInt(const FName MetricName, int64 Value)
+bool UNetworkMetricsDatabase::SetMaxInt(const FName MetricName, int64 Value)
 {
 	UE::Net::FNetworkMetric<int64>* Metric = MetricInts.Find(MetricName);
-	if (ensureMsgf(Metric, TEXT("Cannot find integer metric %s."), *MetricName.ToString()))
+	if (Metric)
 	{
 		Metric->Value = FMath::Max<int64>(Value, Metric->Value);
+		return true;
+	}
+	else
+	{
+		UE_LOG(LogNetworkMetrics, Warning, TEXT("UNetworkMetricsDatabase::SetMaxInt Cannot find integer metric %s."), *MetricName.ToString());
+		return false;
 	}
 }
 
-void UNetworkMetricsDatabase::IncrementInt(const FName MetricName, int64 Value)
+bool UNetworkMetricsDatabase::IncrementInt(const FName MetricName, int64 Value)
 {
 	UE::Net::FNetworkMetric<int64>* Metric = MetricInts.Find(MetricName);
-	if (ensureMsgf(Metric, TEXT("Cannot find integer metric %s."), *MetricName.ToString()))
+	if (Metric)
 	{
 		Metric->Value += Value;
+		return true;
+	}
+	else
+	{
+		UE_LOG(LogNetworkMetrics, Warning, TEXT("UNetworkMetricsDatabase::IncrementInt Cannot find integer metric %s."), *MetricName.ToString());
+		return false;
 	}
 }
 
diff --git a/Engine/Source/Runtime/Engine/Public/Net/NetworkMetricsDatabase.h b/Engine/Source/Runtime/Engine/Public/Net/NetworkMetricsDatabase.h
--- a/Engine/Source/Runtime/Engine/Public/Net/NetworkMetricsDatabase.h
+++ b/Engine/Source/Runtime/Engine/Public/Net/NetworkMetricsDatabase.h
@@ -40,19 +40,19 @@ public:
 	/* Add an integer metric. */
 	void CreateInt(const FName MetricName, int64 DefaultValue);
 	/* Set the value of an existing floating point metric. */
-	void SetFloat(const FName MetricName, float Value);
+	bool SetFloat(const FName MetricName, float Value);
 	/* Set the value of a floating point metric if it's smaller than the existing value. */
-	void SetMinFloat(const FName MetricName, float Value);
+	bool SetMinFloat(const FName MetricName, float Value);
 	/* Set the value of a floating point metric if it's bigger than the existing value. */
-	void SetMaxFloat(const FName MetricName, float Value);
+	bool SetMaxFloat(const FName MetricName, float Value);
 	/* Set the value of an existing integer metric. */
-	void SetInt(const FName MetricName, int64 Value);
+	bool SetInt(const FName MetricName, int64 Value);
 	/* Set the value of an integer metric if it's smaller than the existing value. */
-	void SetMinInt(const FName MetricName, int64 Value);
+	bool SetMinInt(const FName MetricName, int64 Value);
 	/* Set the value of an integer metric if it's bigger than the existing value. */
-	void SetMaxInt(const FName MetricName, int64 Value);
+	bool SetMaxInt(const FName MetricName, int64 Value);
 	/* Increment the value of an existing integer metric. */
-	void IncrementInt(const FName MetricName, int64 Value);
+	bool IncrementInt(const FName MetricName, int64 Value);
 	/* Returns true if a metric has been created in the database. */
 	bool Contains(const FName MetricName) const;
 
