DeletedObjectPlaceholder now keeps a WeakObjectPtr instead of a SoftObjectPath on the original object to allow deletion of the moved object without a warning about this placeholder object having a reference on the object to delete. This can happen if an actor is first moved from one package to another (like changing its External Data Layer) and then the moved actor is deleted.

diff --git a/Engine/Source/Runtime/Engine/Public/DeletedObjectPlaceholder.h b/Engine/Source/Runtime/Engine/Public/DeletedObjectPlaceholder.h
--- a/Engine/Source/Runtime/Engine/Public/DeletedObjectPlaceholder.h
+++ b/Engine/Source/Runtime/Engine/Public/DeletedObjectPlaceholder.h
@@ -25,7 +25,7 @@ public:
 	virtual bool IsAsset() const override;
 	virtual void PostEditUndo() override;
 	
-	const UObject* GetOriginalObject() const { FSoftObjectPtr ObjPtr(OriginalObject); return ObjPtr.Get(); }
+	const UObject* GetOriginalObject() const { return OriginalObject.Get(); }
 	const FString& GetDisplayName() const { return DisplayName; }
 	const FExternalDataLayerUID GetExternalDataLayerUID() const { return FExternalDataLayerUID(ExternalDataLayerUID); }
 
@@ -50,6 +50,6 @@ private:
 
 	/** Original object that is replaced by the placeholder */
 	UPROPERTY()
-	FSoftObjectPath OriginalObject;
+	TWeakObjectPtr<const UObject> OriginalObject;
 #endif
 };
\ No newline at end of file
