World Partition - HLOD: HLOD specific instanced static mesh components classes should use the "HideDropdown" class specifier rather than "Hidden"
* Hidden would prevent users from viewing any of their properties in their details panel

diff --git a/Engine/Source/Runtime/Engine/Public/WorldPartition/HLOD/HLODInstancedStaticMeshComponent.h b/Engine/Source/Runtime/Engine/Public/WorldPartition/HLOD/HLODInstancedStaticMeshComponent.h
--- a/Engine/Source/Runtime/Engine/Public/WorldPartition/HLOD/HLODInstancedStaticMeshComponent.h
+++ b/Engine/Source/Runtime/Engine/Public/WorldPartition/HLOD/HLODInstancedStaticMeshComponent.h
@@ -8,7 +8,7 @@
 #include "HLODInstancedStaticMeshComponent.generated.h"
 
 
-UCLASS(Hidden, NotPlaceable, MinimalAPI)
+UCLASS(HideDropDown, NotPlaceable, MinimalAPI)
 class UHLODInstancedStaticMeshComponent : public UInstancedStaticMeshComponent
 {
 	GENERATED_UCLASS_BODY()
