Fixed rare check that could happen when marking a Data Layer as Current quick enough after deleting it.

diff --git a/Engine/Source/Runtime/Engine/Private/WorldPartition/DataLayer/DataLayerInstance.cpp b/Engine/Source/Runtime/Engine/Private/WorldPartition/DataLayer/DataLayerInstance.cpp
--- a/Engine/Source/Runtime/Engine/Private/WorldPartition/DataLayer/DataLayerInstance.cpp
+++ b/Engine/Source/Runtime/Engine/Private/WorldPartition/DataLayer/DataLayerInstance.cpp
@@ -557,7 +557,7 @@ bool UDataLayerInstance::Validate(IStreamingGenerationErrorHandler* ErrorHandler
 
 bool UDataLayerInstance::CanBeInActorEditorContext() const
 {
-	return !IsReadOnly();
+	return !IsReadOnly() && GetOuterWorldDataLayers()->ContainsDataLayer(this);
 }
 
 bool UDataLayerInstance::IsActorEditorContextCurrentColorized() const
