Texture: fix bug causing check when saving size zero textures due to asking for pixelformat size on Invalid format.

diff --git a/Engine/Source/Runtime/Engine/Private/Texture.cpp b/Engine/Source/Runtime/Engine/Private/Texture.cpp
--- a/Engine/Source/Runtime/Engine/Private/Texture.cpp
+++ b/Engine/Source/Runtime/Engine/Private/Texture.cpp
@@ -2099,6 +2099,9 @@ void FTextureSource::InitLayered(
 		NewLayerFormat
 		);
 
+	// beware: IsValid() is still false now because BulkData is not yet set up
+	//	CalcLayerSize must not check IsValid()
+
 	int64 TotalBytes = 0;
 	for (int i = 0; i < NewNumLayers; ++i)
 	{
@@ -2114,7 +2117,7 @@ void FTextureSource::InitLayered(
 	else
 	{
 		BulkData.UpdatePayload(FUniqueBuffer::Alloc(TotalBytes).MoveToShared(), Owner);
-		// @todo Oodle : FIX ME ?? unitialized memory ??
+		// @todo Oodle : FIX ME ?? uninitialized memory ??
 		// ?? with no incoming data this is hashing garbage ???
 		// I think this also does FEditorBulkData::Register !?
 	}
@@ -2360,7 +2363,8 @@ void FTextureSource::Compress()
 
 	CheckTextureIsUnlocked(TEXT("Compress"));
 	
-	if ( CompressionFormat == TSCF_JPEG || CompressionFormat == TSCF_UEJPEG )
+	// !IsValid for size zero textures
+	if ( !IsValid() || CompressionFormat == TSCF_JPEG || CompressionFormat == TSCF_UEJPEG )
 	{
 		// leave JPEG data alone, and no need to apply LZ on top of it
 		BulkData.SetCompressionOptions(UE::Serialization::ECompressionOptions::Disabled);
@@ -2423,10 +2427,14 @@ FSharedBuffer FTextureSource::Decompress(IImageWrapperModule* ) const
 
 	// ImageWrapperModule argument ignored, not drilled through DecompressImage
 
-	int64 ExpectedTotalSize = CalcTotalSize();
-
 	FSharedBuffer Buffer;
 
+	if ( !IsValid() )
+	{
+		// size zero texture
+		return Buffer;
+	}
+
 	if (CompressionFormat != TSCF_None )
 	{
 		Buffer = TryDecompressData();
@@ -2436,6 +2444,8 @@ FSharedBuffer FTextureSource::Decompress(IImageWrapperModule* ) const
 		Buffer = BulkData.GetPayload().Get();
 	}
 	
+	int64 ExpectedTotalSize = CalcTotalSize();
+
 	// validate the size of the FSharedBuffer
 	if ( Buffer.GetSize() != ExpectedTotalSize )
 	{
@@ -2800,6 +2810,8 @@ int64 FTextureSource::CalcMipSize(int32 BlockIndex, int32 LayerIndex, int32 MipI
 
 int64 FTextureSource::GetBytesPerPixel(int32 LayerIndex) const
 {
+	// note: if !IsValid() this will check() because Format will be PF_Invalid
+
 	return GetBytesPerPixel(GetFormat(LayerIndex));
 }
 
@@ -3363,6 +3375,12 @@ void FTextureSource::RemoveSourceData()
 // total size in bytes including all blocks and layers
 int64 FTextureSource::CalcTotalSize() const
 {
+	if ( SizeX == 0 || SizeY == 0 || NumSlices == 0 || NumLayers == 0 || NumMips == 0 || Format == TSF_Invalid )
+	{
+		// size zero texture
+		return 0;
+	}
+
 	int NumBlocks = GetNumBlocks();
 	int64 TotalBytes = 0;
 	for (int i = 0; i < NumBlocks; ++i)
@@ -3398,6 +3416,12 @@ int64 FTextureSource::CalcBlockSize(const FTextureSourceBlock& Block) const
 
 int64 FTextureSource::CalcLayerSize(const FTextureSourceBlock& Block, int32 LayerIndex) const
 {
+	if ( SizeX == 0 || SizeY == 0 || NumSlices == 0 || NumLayers == 0 || NumMips == 0 || Format == TSF_Invalid )
+	{
+		// size zero texture
+		return 0;
+	}
+
 	int64 BytesPerPixel = GetBytesPerPixel(LayerIndex);
 
 	// This is used for memory allocation, so use FCheckedInt to rigorously check against overflow issues.
@@ -3417,6 +3441,13 @@ int64 FTextureSource::CalcLayerSize(const FTextureSourceBlock& Block, int32 Laye
 
 int64 FTextureSource::CalcMipOffset(int32 BlockIndex, int32 LayerIndex, int32 OffsetToMipIndex) const
 {
+#if 0
+	if ( BlockIndex == 0 && LayerIndex == 0 && OffsetToMipIndex == 0 )
+	{
+		return 0;
+	}
+#endif
+
 	FTextureSourceBlock Block;
 	GetBlock(BlockIndex, Block);
 	check(OffsetToMipIndex < Block.NumMips);
