Proper support of child actors for actors using External Data Layers

diff --git a/Engine/Source/Editor/DataLayerEditor/Private/DataLayer/DataLayerEditorSubsystem.cpp b/Engine/Source/Editor/DataLayerEditor/Private/DataLayer/DataLayerEditorSubsystem.cpp
--- a/Engine/Source/Editor/DataLayerEditor/Private/DataLayer/DataLayerEditorSubsystem.cpp
+++ b/Engine/Source/Editor/DataLayerEditor/Private/DataLayer/DataLayerEditorSubsystem.cpp
@@ -379,7 +379,7 @@ void UDataLayerEditorSubsystem::Deinitialize()
 
 const UExternalDataLayerInstance* UDataLayerEditorSubsystem::GetActorSpawningExternalDataLayerInstance(AActor* InActor) const
 {
-	if (GIsReinstancing || InActor->bIsEditorPreviewActor || !InActor->IsPackageExternal())
+	if (GIsReinstancing || InActor->bIsEditorPreviewActor || !InActor->IsPackageExternal() || !InActor->IsMainPackageActor())
 	{
 		return nullptr;
 	}
@@ -512,7 +512,7 @@ void UDataLayerEditorSubsystem::BeginDestroy()
 
 void UDataLayerEditorSubsystem::ApplyContext(AActor* InActor, bool bInForceTryApply)
 {
-	if (GIsReinstancing || InActor->bIsEditorPreviewActor || !InActor->IsPackageExternal() || (InActor->GetWorld() != GetWorld()))
+	if (GIsReinstancing || InActor->bIsEditorPreviewActor || !InActor->IsPackageExternal() || !InActor->IsMainPackageActor() || (InActor->GetWorld() != GetWorld()))
 	{
 		return;
 	}
diff --git a/Engine/Source/Runtime/Engine/Private/WorldPartition/DataLayer/ExternalDataLayerHelper.cpp b/Engine/Source/Runtime/Engine/Private/WorldPartition/DataLayer/ExternalDataLayerHelper.cpp
--- a/Engine/Source/Runtime/Engine/Private/WorldPartition/DataLayer/ExternalDataLayerHelper.cpp
+++ b/Engine/Source/Runtime/Engine/Private/WorldPartition/DataLayer/ExternalDataLayerHelper.cpp
@@ -154,6 +154,12 @@ bool FExternalDataLayerHelper::CanMoveActorsToExternalDataLayer(const TArray<AAc
 			return false;
 		}
 
+		if (!InActor->IsMainPackageActor())
+		{
+			OutFailureReason = FText::Format(LOCTEXT("CantMoveActorToEDL_ChildActorNotSupported", "Child Actor {0} cannot be moved to External Data Layer."), FText::FromString(InActor->GetName()));
+			return false;
+		}
+
 		if (!OldExternalDataLayerAsset && !NewExternalDataLayerAsset)
 		{
 			OutFailureReason = FText::Format(LOCTEXT("CantMoveActorToEDL_NoExternalDataLayer", "Actor {0} has already no External Data Layer."), FText::FromString(InActor->GetName()));
@@ -243,6 +249,7 @@ bool FExternalDataLayerHelper::MoveActorsToExternalDataLayer(const TArray<AActor
 {
 	auto MoveActorToExternalDataLayer = [](AActor* InActor, const UExternalDataLayerInstance* InExternalDataLayerInstance)
 	{
+		check(InActor->IsMainPackageActor());
 		const UPackage* OldActorPackage = InActor->GetExternalPackage();
 		const UExternalDataLayerAsset* NewExternalDataLayerAsset = InExternalDataLayerInstance ? InExternalDataLayerInstance->GetExternalDataLayerAsset() : nullptr;
 		const bool bShouldDirty = true;
@@ -258,7 +265,7 @@ bool FExternalDataLayerHelper::MoveActorsToExternalDataLayer(const TArray<AActor
 				DependantObjects.Add(Object);
 			}
 			return true;
-		}, false);
+		}, false, RF_NoFlags, EInternalObjectFlags::Garbage); // Skip garbage objects (like child actors destroyed when de-externalizing the actor)
 
 		// Clear Content Bundle Guid
 		FSetActorContentBundleGuid(InActor, FGuid());
