Fix MarkPackageDirty call in skeletalmesh to avoid doing so if outside of the gamethread.

#rnx

diff --git a/Engine/Source/Runtime/Engine/Private/SkeletalMesh.cpp b/Engine/Source/Runtime/Engine/Private/SkeletalMesh.cpp
--- a/Engine/Source/Runtime/Engine/Private/SkeletalMesh.cpp
+++ b/Engine/Source/Runtime/Engine/Private/SkeletalMesh.cpp
@@ -2495,7 +2495,7 @@ bool USkeletalMesh::CommitMeshDescription(
 	}
 	
 
-	if (InParams.bMarkPackageDirty)
+	if (IsInGameThread() && InParams.bMarkPackageDirty)
 	{
 		(void)MarkPackageDirty();
 	}
