Texture clean up and remove incorrect uses of BulkData.
Several APIs related to BulkData in texture were completely broken, such as IsBulkDataLoaded, deprecate them and remove uses.
Remove the TornOff AsyncSource copy as it was not actually helping anything, and was being used in the opposite situation of when we thought it was (it was being used when Payload was already in memory).

diff --git a/Engine/Plugins/MeshPainting/Source/MeshPaintingToolset/Private/MeshTexturePaintingTool.cpp b/Engine/Plugins/MeshPainting/Source/MeshPaintingToolset/Private/MeshTexturePaintingTool.cpp
--- a/Engine/Plugins/MeshPainting/Source/MeshPaintingToolset/Private/MeshTexturePaintingTool.cpp
+++ b/Engine/Plugins/MeshPainting/Source/MeshPaintingToolset/Private/MeshTexturePaintingTool.cpp
@@ -898,19 +898,17 @@ void UMeshTexturePaintingTool::StartPaintingTexture(UMeshComponent* InMeshCompon
 	// Check all the materials on the mesh to see if the user texture is there
 	int32 MaterialIndex = 0;
 	UMaterialInterface* MaterialToCheck = InMeshComponent->GetMaterial(MaterialIndex);
-	bool bIsSourceTextureStreamedIn = Texture2D->IsFullyStreamedIn() && Texture2D->Source.IsBulkDataLoaded() && !Texture2D->Source.HasHadBulkDataCleared();
+
+	Texture2D->BlockOnAnyAsyncBuild();
+	bool bIsSourceTextureStreamedIn = Texture2D->IsFullyStreamedIn() && !Texture2D->HasPendingInitOrStreaming();
 
 	// IMeshPaintComponentAdapter::DefaultQueryPaintableTextures already filters out un-used textures
 	if (!bIsSourceTextureStreamedIn)
 	{
-		if (!Texture2D->Source.IsBulkDataLoaded() || Texture2D->Source.HasHadBulkDataCleared())
-		{
-			Texture2D->Modify();
-		}
-
 		Texture2D->SetForceMipLevelsToBeResident(30.0f);
+		Texture2D->bForceMiplevelsToBeResident = true;
 		Texture2D->WaitForStreaming();
-		bIsSourceTextureStreamedIn = Texture2D->IsFullyStreamedIn();
+		bIsSourceTextureStreamedIn = Texture2D->IsFullyStreamedIn() && !Texture2D->HasPendingInitOrStreaming();
 	}
 
 	while (MaterialToCheck != nullptr)
@@ -920,7 +918,7 @@ void UMeshTexturePaintingTool::StartPaintingTexture(UMeshComponent* InMeshCompon
 			const int32 TextureWidth = Texture2D->Source.GetSizeX();
 			const int32 TextureHeight = Texture2D->Source.GetSizeY();
 
-			FTextureCompilingManager::Get().FinishCompilation({ Texture2D });
+			Texture2D->BlockOnAnyAsyncBuild();
 
 			if (TextureData == nullptr)
 			{
@@ -1425,6 +1423,7 @@ void UMeshTexturePaintingTool::FinishPaintingTexture()
 
 			{
 				// For undo
+				// @@ use PreEdit/PostEdit instead
 				TextureData->ScratchTexture->Modify();
 
 				const int32 NumPixels = TexturePixels.Num();
@@ -1662,10 +1661,12 @@ void UMeshTexturePaintingTool::CommitAllPaintedTextures()
 
 				{
 					// For undo
+					// @@ use PreEdit/PostEdit instead
 					TextureData->PaintingTexture2D->SetFlags(RF_Transactional);
 					TextureData->PaintingTexture2D->Modify();
 
 					// Store source art
+					// @@ use FMipLock.Image instead
 					FColor* Colors = (FColor*)TextureData->PaintingTexture2D->Source.LockMip(0);
 					check(TextureData->PaintingTexture2D->Source.CalcMipSize(0) == TexturePixels.Num() * sizeof(FColor));
 					FMemory::Memcpy(Colors, TexturePixels.GetData(), TexturePixels.Num() * sizeof(FColor));
diff --git a/Engine/Source/Runtime/Engine/Classes/Engine/Texture.h b/Engine/Source/Runtime/Engine/Classes/Engine/Texture.h
--- a/Engine/Source/Runtime/Engine/Classes/Engine/Texture.h
+++ b/Engine/Source/Runtime/Engine/Classes/Engine/Texture.h
@@ -265,7 +265,7 @@ struct FTextureSource
 
 	/** Make a copy with a torn-off BulkData that has the same Guid used for DDC as this->BulkData */
 	FTextureSource CopyTornOff() const;
-
+	
 	/** PNG Compresses the source art if possible or tells the bulk data to zlib compress when it saves out to disk. */
 	ENGINE_API void Compress();
 
@@ -382,18 +382,31 @@ struct FTextureSource
 	// returns volume depth, or 1 if not a volume
 	ENGINE_API int GetVolumeSizeZ() const;
 
+	//note: simple queries on BulkData are not taking the BulkDataLock , nor does BulkData internally mutex protect these
 	FORCEINLINE int64 GetSizeOnDisk() const { return BulkData.GetPayloadSize(); }
-	inline bool HasPayloadData() const { return BulkData.HasPayloadData(); }
-	/** Returns true if the texture's bulkdata payload is either already in memory or if the payload is 0 bytes in length. It will return false if the payload needs to load from disk */
-	FORCEINLINE bool IsBulkDataLoaded() const { return BulkData.DoesPayloadNeedLoading(); }
+	FORCEINLINE bool HasPayloadData() const { return BulkData.HasPayloadData(); } // this is the same as GetSizeOnDisk() != 0
+	
+	/** Reset the source to empty, frees all memory. */
+	ENGINE_API void Reset();
 
-	// Apply a visitor to the bulkdata :
+	// Apply a visitor to the bulkdata : prefer GetBulkDataPayload() instead.
 	ENGINE_API void OperateOnLoadedBulkData(TFunctionRef<void (const FSharedBuffer& BulkDataBuffer)> Operation);
 
+	// GetBulkDataPayload returns the raw bulkdata memory in compressed form, not decompressed (use LockMip or GetMipData for that)
+	//	the shared buffer should be treated as read only!
+	//	to modify, make a copy and call Init() on the source with the changed copy
+	ENGINE_API FSharedBuffer GetBulkDataPayload();
+		
+	/** Returns true if the texture's bulkdata payload is either already in memory or if the payload is 0 bytes in length. It will return false if the payload needs to load from disk */
+	//  broken function, did the opposite of its name (was == DoesPayloadNeedLoading, which is usually == HasPayloadData)
+	UE_DEPRECATED(5.5, "IsBulkDataLoaded was broken, do not use")
+	FORCEINLINE bool IsBulkDataLoaded() const { return true; }
+
 	UE_DEPRECATED(5.0, "There is no longer a need to call LoadBulkDataWithFileReader, FTextureSource::BulkData can now load the data on demand without it.")
 	FORCEINLINE bool LoadBulkDataWithFileReader() { return true; }
-
-	FORCEINLINE void RemoveBulkData() { BulkData.UnloadData(); }
+	
+	UE_DEPRECATED(5.5, "RemoveBulkData did not actually remove bulkdata; use ReleaseBulkDataCachedMemory if that's what you wanted")
+	FORCEINLINE void RemoveBulkData() { }
 	
 	/** Sets the GUID to use, and whether that GUID is actually a hash of some data. */
 	ENGINE_API void SetId(const FGuid& InId, bool bInGuidIsHash);
@@ -636,8 +649,6 @@ private:
 
 	/** Return true if the source art is not png compressed but could be. */
 	bool CanPNGCompress() const;
-	/** Removes source data. */
-	void RemoveSourceData();
 	/** Retrieve the size and offset for a source mip. The size includes all slices. */
 	int64 CalcMipOffset(int32 BlockIndex, int32 LayerIndex, int32 MipIndex) const;
 	
@@ -688,13 +699,15 @@ public:
 	This is automatically done by Init() and Mip Lock/Unlock.  New textures should always have the data hash as Id. */
 	ENGINE_API void UseHashAsGuid();
 
-	ENGINE_API void ReleaseSourceMemory(); // release the memory from the mips (does almost the same as remove source data except doesn't rebuild the guid)
-	FORCEINLINE bool HasHadBulkDataCleared() const { return bHasHadBulkDataCleared; }
-private:
-	/** Used while cooking to clear out unneeded memory after compression */
-	bool bHasHadBulkDataCleared;
+	UE_DEPRECATED(5.5, "Remove call to ReleaseSourceMemory.")
+	FORCEINLINE void ReleaseSourceMemory() { }
+
+	UE_DEPRECATED(5.5, "Remove call to HasHadBulkDataCleared.")
+	FORCEINLINE bool HasHadBulkDataCleared() const { return false; }
 #endif // WITH_EDITOR
 
+private:
+
 #if WITH_EDITORONLY_DATA
 	/** GUID used to track changes to the source data.
 	Typically with UseHashAsGuid , this "Id" is the hash of the BulkData.
diff --git a/Engine/Source/Runtime/Engine/Private/Texture.cpp b/Engine/Source/Runtime/Engine/Private/Texture.cpp
--- a/Engine/Source/Runtime/Engine/Private/Texture.cpp
+++ b/Engine/Source/Runtime/Engine/Private/Texture.cpp
@@ -1925,9 +1925,6 @@ FTextureSource::FTextureSource()
 #endif
 	  NumLockedMips(0u)
 	, LockState(ELockState::None)
-#if WITH_EDITOR
-	, bHasHadBulkDataCleared(false)
-#endif
 #if WITH_EDITORONLY_DATA
 	, BaseBlockX(0)
 	, BaseBlockY(0)
@@ -2314,6 +2311,9 @@ FTextureSource FTextureSource::CopyTornOff() const
 
 static bool ShouldUseUEDeltaForFormat(ETextureSourceFormat Format)
 {
+	// should have been detected earlier in Source.IsValid() check :
+	check( Format != TSF_Invalid && Format != TSF_MAX );
+
 	if ( Format == TSF_RGBA16F || Format == TSF_RGBA32F ||
 		Format == TSF_R16F || Format == TSF_R32F )
 	{
@@ -2444,6 +2444,9 @@ FSharedBuffer FTextureSource::Decompress(IImageWrapperModule* ) const
 		Buffer = BulkData.GetPayload().Get();
 	}
 	
+	// note: you could now do BulkData.UnloadData() , but it currently does not actually cache decompressed data
+	//	so that is usually a nop
+
 	int64 ExpectedTotalSize = CalcTotalSize();
 
 	// validate the size of the FSharedBuffer
@@ -3064,39 +3067,6 @@ FSharedBuffer FTextureSource::DoUEDeltaTransform(FSharedBuffer InBuffer,bool bFo
 				
 				check( Image.GetImageSizeBytes() == MipSize );
 				
-				#if 0
-				{
-					// debug dump
-					const FString DumpPath = TEXT("r:\\dump\\");
-					
-					IFileManager::Get().MakeDirectory(*DumpPath);
-
-					FString FileName = DumpPath;
-			
-					// get a texture name that's usable in a file name :
-					//FString Name = Owner->GetFullName();
-					FString Name = Owner ? Owner->GetName() : *TornOffOwnerName;
-					Name = FPaths::MakeValidFileName(Name, TEXT('_'));
-					Name.ReplaceCharInline(TEXT('.'),TEXT('_'));
-					Name.ReplaceCharInline(TEXT(' '),TEXT('_'));
-					Name.ReplaceInline(TEXT("__"),TEXT("_"));
-					Name.ReplaceInline(TEXT("__"),TEXT("_"));
-
-					FileName += Name;
-
-					// image save doesn't work on slices, so do them one by one
-					for(int64 SliceIndex=0;SliceIndex<Image.NumSlices;SliceIndex++)
-					{
-						FImageView Slice = Image.GetSlice(SliceIndex);
-
-						FileName += FString::Printf(TEXT("_%dx%dx%dx%d_%s"),BlockIndex,LayerIndex,MipIndex,SliceIndex,
-							bForward ? TEXT("fwd") : TEXT("rev"));
-					
-						FImageUtils::SaveImageAutoFormat(*FileName,Slice);
-					}
-				}
-				#endif
-
 				FImageCoreDelta::AddSplitStridedViewsForDelta( ImageViewPortions, Image);
 			}
 		}
@@ -3331,19 +3301,7 @@ void FTextureSource::ForceGenerateGuid()
 	bGuidIsHash = false;
 }
 
-void FTextureSource::ReleaseSourceMemory()
-{
-#if WITH_EDITOR
-	FScopeLock BulkDataExclusiveScope(&BulkDataLock.Get());
-#endif
-
-	check( LockState == ELockState::None && NumLockedMips == 0 );
-
-	bHasHadBulkDataCleared = true;
-	BulkData.UnloadData();
-}
-
-void FTextureSource::RemoveSourceData()
+void FTextureSource::Reset()
 {
 #if WITH_EDITOR
 	FScopeLock BulkDataExclusiveScope(&BulkDataLock.Get());
@@ -3441,12 +3399,11 @@ int64 FTextureSource::CalcLayerSize(const FTextureSourceBlock& Block, int32 Laye
 
 int64 FTextureSource::CalcMipOffset(int32 BlockIndex, int32 LayerIndex, int32 OffsetToMipIndex) const
 {
-#if 0
 	if ( BlockIndex == 0 && LayerIndex == 0 && OffsetToMipIndex == 0 )
 	{
+		// early out common case
 		return 0;
 	}
-#endif
 
 	FTextureSourceBlock Block;
 	GetBlock(BlockIndex, Block);
@@ -3554,8 +3511,25 @@ FGuid FTextureSource::GetId() const
 	return IdBuilder.Build();
 }
 
+FSharedBuffer FTextureSource::GetBulkDataPayload()
+{
+#if WITH_EDITOR
+	FScopeLock BulkDataExclusiveScope(&BulkDataLock.Get());
+#endif
+
+	FSharedBuffer Payload = BulkData.GetPayload().Get();
+	
+	// Payload has the Oodle LZ Decompress done, but not the TSCF compressor
+	//	(use Decompress() for that)
+
+	return Payload;
+}
+
 void FTextureSource::OperateOnLoadedBulkData(TFunctionRef<void(const FSharedBuffer& BulkDataBuffer)> Operation)
 {
+	// ?? why is this operation visitor necessary ? prefer to just return the FSharedBuffer.
+	//	most callers should just use GetBulkDataPayload instead.
+
 #if WITH_EDITOR
 	FScopeLock BulkDataExclusiveScope(&BulkDataLock.Get());
 #endif
@@ -3563,6 +3537,9 @@ void FTextureSource::OperateOnLoadedBulkData(TFunctionRef<void(const FSharedBuff
 	checkf(LockState == ELockState::None, TEXT("OperateOnLoadedBulkData shouldn't be called in-between LockMip/UnlockMip"));
 
 	FSharedBuffer Payload = BulkData.GetPayload().Get();
+
+	//  note: unlike LockMip, the BulkDataLock is held the entire time during this operation
+	//	  (for no reason AFAICT)
 	Operation(Payload);
 }
 
@@ -4736,7 +4713,7 @@ void FTextureSource::InitLayeredImpl(
 	int32 NewNumMips,
 	const ETextureSourceFormat* NewLayerFormat)
 {
-	RemoveSourceData();
+	Reset();
 	SizeX = NewSizeX;
 	SizeY = NewSizeY;
 	NumLayers = NewNumLayers;
@@ -4769,7 +4746,7 @@ void FTextureSource::InitBlockedImpl(const ETextureSourceFormat* InLayerFormats,
 	check(InNumBlocks > 0);
 	check(InNumLayers > 0);
 
-	RemoveSourceData();
+	Reset();
 
 	BaseBlockX = InBlocks[0].BlockX;
 	BaseBlockY = InBlocks[0].BlockY;
diff --git a/Engine/Source/Runtime/Engine/Private/TextureDerivedDataTask.cpp b/Engine/Source/Runtime/Engine/Private/TextureDerivedDataTask.cpp
--- a/Engine/Source/Runtime/Engine/Private/TextureDerivedDataTask.cpp
+++ b/Engine/Source/Runtime/Engine/Private/TextureDerivedDataTask.cpp
@@ -280,12 +280,6 @@ void FTextureSourceData::Init(UTexture& InTexture, TextureMipGenSettings InMipGe
 
 	TextureFullName = InTexture.GetFullName();
 
-	if (bAllowAsyncLoading && !InTexture.Source.IsBulkDataLoaded())
-	{
-		// Prepare the async source to be later able to load it from file if required.
-		AsyncSource = InTexture.Source.CopyTornOff(); // This copies information required to make a safe IO load async.
-	}
-
 	bValid = true;
 }
 
@@ -310,13 +304,6 @@ void FTextureSourceData::GetSourceMips(FTextureSource& Source, IImageWrapperModu
  			return;
 		}
 
-		if (Source.HasHadBulkDataCleared())
-		{	// don't do any work we can't reload this
-			UE_LOG(LogTexture, Error, TEXT("Unable to get texture source mips because its bulk data was released. %s"), *TextureFullName);
-			ReleaseMemory();
-			bValid = false;
-			return;
-		}
 		if (!Source.HasPayloadData())
 		{	// don't do any work we can't reload this
 			UE_LOG(LogTexture, Warning, TEXT("Unable to get texture source mips because its bulk data has no payload. This may happen if it was duplicated from cooked data. %s"), *TextureFullName);
@@ -416,18 +403,20 @@ void FTextureSourceData::GetSourceMips(FTextureSource& Source, IImageWrapperModu
 			ScopedMipData.ResetData();
 		}
 	}
-}
-
-
-void FTextureSourceData::GetAsyncSourceMips(IImageWrapperModule* InImageWrapper)
-{
-	if (bValid && !Blocks[0].MipsPerLayer[0].Num() && AsyncSource.HasPayloadData())
+	
+	#if 0
 	{
-		GetSourceMips(AsyncSource, InImageWrapper);
+		//we have got the source mip data and made a copy of it
+		//	no longer need the BulkData to be in memory
+		// note this is different than ScopedMipData.ResetData ; that frees the decompressed copy
+		//	this frees the compressed copy
+		// (note that BulkData does not cache decompressed payloads so this is usually a nop)
+		TRACE_CPUPROFILER_EVENT_SCOPE(FTextureSourceData::GetSourceMips_FreeBulkData);
+		Source.ReleaseBulkDataCachedMemory();
 	}
+	#endif
 }
 
-
 // When texture streaming is disabled, all of the mips are packed into a single FBulkData/FDerivedData
 // and "inlined", meaning they are saved and loaded as part of the serialized asset data.
 static bool GetBuildSettingsDisablesStreaming(const FTextureBuildSettings& InBuildSettings, const FTextureEngineParameters& InEngineParameters)
@@ -566,6 +555,7 @@ public:
 		{
 			Source.OperateOnLoadedBulkData([&Buffer](const FSharedBuffer& BulkDataBuffer)
 			{
+				// ?? what? this loads compressed data for the TextureSource, decompresses to Payload, and then recompresses it !?
 				Buffer = FCompressedBuffer::Compress(BulkDataBuffer);
 			});
 		}
@@ -1752,46 +1742,21 @@ static bool DDC1_LoadAndValidateTextureData(
 {
 	TRACE_CPUPROFILER_EVENT_SCOPE(Texture.DDC1_LoadAndValidateTextureData);
 
-	bool bHasTextureSourceMips = false;
-	bool bNeedsGetSourceMips;
+	// There can be a stall here waiting on the BulkData mutex if it is serializing to the undo buffer on the main thread.
 
-	{
-	// this can be a stall waiting on the BulkData mutex if it is serializing to the undo buffer on the main thread :
-	TRACE_CPUPROFILER_EVENT_SCOPE(Texture.IsBulkDataLoaded);
-	bNeedsGetSourceMips = TextureData.IsValid() && Texture.Source.IsBulkDataLoaded();
-	}
+	bool bNeedsGetSourceMips = TextureData.IsValid() && Texture.Source.HasPayloadData();
 
 	if ( bNeedsGetSourceMips )
 	{
 		TRACE_CPUPROFILER_EVENT_SCOPE(GetSourceMips);
 		TextureData.GetSourceMips(Texture.Source, ImageWrapper);
-		bHasTextureSourceMips = true;
 	}
 
-	bool bHasCompositeTextureSourceMips = false;
-	if (CompositeTextureData.IsValid() && Texture.GetCompositeTexture() && Texture.GetCompositeTexture()->Source.IsBulkDataLoaded())
+	if (CompositeTextureData.IsValid() && Texture.GetCompositeTexture() && Texture.GetCompositeTexture()->Source.HasPayloadData())
 	{
 		TRACE_CPUPROFILER_EVENT_SCOPE(GetCompositeSourceMips);
 		check( Texture.GetCompositeTexture()->Source.IsValid() );
 		CompositeTextureData.GetSourceMips(Texture.GetCompositeTexture()->Source, ImageWrapper);
-		bHasCompositeTextureSourceMips = true;
-	}
-
-	if (bAllowAsyncLoading)
-	{
-		if (!bHasTextureSourceMips)
-		{
-			TRACE_CPUPROFILER_EVENT_SCOPE(GetAsyncSourceMips);
-			TextureData.GetAsyncSourceMips(ImageWrapper);
-			TextureData.AsyncSource.RemoveBulkData();
-			}
-
-		if (!bHasCompositeTextureSourceMips)
-		{
-			TRACE_CPUPROFILER_EVENT_SCOPE(GetAsyncCompositeSourceMips);
-			CompositeTextureData.GetAsyncSourceMips(ImageWrapper);
-			CompositeTextureData.AsyncSource.RemoveBulkData();
-		}
 	}
 
 	return DDC1_IsTextureDataValid(TextureData, CompositeTextureData);
@@ -2160,7 +2125,7 @@ void FTextureCacheDerivedDataWorker::Finalize()
 
 	if (!bSucceeded)
 	{
-		if (!TextureData.HasPayload() && !Texture.Source.HasPayloadData())
+		if ( !Texture.Source.HasPayloadData() )
 		{
 			UE_LOG(LogTexture, Warning, TEXT("Unable to build texture source data, no available payload for %s. This may happen if it was duplicated from cooked data."), *TexturePathName);
 			return;
diff --git a/Engine/Source/Runtime/Engine/Private/TextureDerivedDataTask.h b/Engine/Source/Runtime/Engine/Private/TextureDerivedDataTask.h
--- a/Engine/Source/Runtime/Engine/Private/TextureDerivedDataTask.h
+++ b/Engine/Source/Runtime/Engine/Private/TextureDerivedDataTask.h
@@ -95,21 +95,13 @@ struct FTextureSourceData
 	void Init(UTexture& InTexture, TextureMipGenSettings InMipGenSettings, bool bInCubeMap, bool bInTextureArray, bool bInVolumeTexture, bool bAllowAsyncLoading);
 	bool IsValid() const { return bValid; }
 
-	bool HasPayload() const
-	{
-		return AsyncSource.HasPayloadData();
-	}
-	
 	// ImageWrapperModule is not used
 	void GetSourceMips(FTextureSource& Source, IImageWrapperModule* InImageWrapper = nullptr);
-	void GetAsyncSourceMips(IImageWrapperModule* InImageWrapper = nullptr);
 
 	void ReleaseMemory()
 	{
 		TRACE_CPUPROFILER_EVENT_SCOPE(Texture.ReleaseMemory);
 
-		// Unload BulkData loaded with LoadBulkDataWithFileReader
-		AsyncSource.RemoveBulkData();
 		Blocks.Empty();
 		Layers.Empty();
 		bValid = false;
@@ -118,7 +110,6 @@ struct FTextureSourceData
 	TArray<TPair<FLinearColor, FLinearColor>> LayerChannelMinMax; // Key == Min, Value == Max
 
 	FString TextureFullName;
-	FTextureSource AsyncSource;
 	TArray<FTextureSourceLayerData> Layers;
 	TArray<FTextureSourceBlockData> Blocks;
 	int32 SizeInBlocksX;
