Fix crash when launching the destruction map in ContentExamples
- this is caused by a particle having no geometry, triggering a check
- since the case is already handled, we can change it to an ensure

diff --git a/Engine/Source/Runtime/Engine/Private/PhysicsEngine/Experimental/PhysScene_Chaos.cpp b/Engine/Source/Runtime/Engine/Private/PhysicsEngine/Experimental/PhysScene_Chaos.cpp
--- a/Engine/Source/Runtime/Engine/Private/PhysicsEngine/Experimental/PhysScene_Chaos.cpp
+++ b/Engine/Source/Runtime/Engine/Private/PhysicsEngine/Experimental/PhysScene_Chaos.cpp
@@ -2180,7 +2180,7 @@ static void UpdateAccelerationStructureFromGeometryCollectionProxy(FGeometryColl
 	auto GetParticleWorldBounds = [](Chaos::FPBDRigidParticle& Particle) -> Chaos::FAABB3
 	{
 		const Chaos::FImplicitObjectRef Geometry = Particle.GetGeometry();
-		check(Geometry != nullptr)
+		ensure(Geometry != nullptr);
 
 		if ((Geometry != nullptr) && Geometry->HasBoundingBox())
 		{
