Fallback skeletal mesh rendering to dynamic relevance when various debug modes are enabled.

diff --git a/Engine/Source/Runtime/Engine/Private/SkeletalMesh.cpp b/Engine/Source/Runtime/Engine/Private/SkeletalMesh.cpp
--- a/Engine/Source/Runtime/Engine/Private/SkeletalMesh.cpp
+++ b/Engine/Source/Runtime/Engine/Private/SkeletalMesh.cpp
@@ -54,6 +54,7 @@
 #include "Algo/MaxElement.h"
 #include "UObject/FortniteMainBranchObjectVersion.h"
 #include "UObject/ReleaseObjectVersion.h"
+#include "GPUSkinCacheVisualizationData.h"
 
 #if WITH_EDITOR
 #include "AssetRegistry/IAssetRegistry.h"
@@ -7077,10 +7078,31 @@ FPrimitiveViewRelevance FSkeletalMeshSceneProxy::GetViewRelevance(const FSceneVi
 	// View relevance is updated once per frame per view across all views in the frame (including shadows) so we update the LOD level for next frame here.
 	MeshObject->UpdateMinDesiredLODLevel(View, GetBounds());
 
+	const auto& EngineShowFlags = View->Family->EngineShowFlags;
+
+	const auto IsDynamic = [&]
+	{
+#if !(UE_BUILD_SHIPPING) || WITH_EDITOR
+		return IsRichView(*View->Family)
+			|| EngineShowFlags.Bones
+			|| EngineShowFlags.Collision
+			|| EngineShowFlags.Bounds
+			|| EngineShowFlags.VertexColors
+			|| IsSelected()
+#if WITH_EDITORONLY_DATA
+			|| MeshObject->SelectedEditorMaterial != -1
+			|| MeshObject->SelectedEditorSection != -1
+#endif
+			|| GetGPUSkinCacheVisualizationData().IsActive();
+#else
+		return false;
+#endif
+	};
+
 	FPrimitiveViewRelevance Result;
-	Result.bDrawRelevance = IsShown(View) && View->Family->EngineShowFlags.SkeletalMeshes;
+	Result.bDrawRelevance = IsShown(View) && EngineShowFlags.SkeletalMeshes;
 	Result.bShadowRelevance = IsShadowCast(View);
-	Result.bStaticRelevance = (bRenderStatic || GSkeletalMeshUseCachedMDCs) && MeshObject->SupportsStaticRelevance() && !IsRichView(*View->Family);
+	Result.bStaticRelevance = (bRenderStatic || GSkeletalMeshUseCachedMDCs) && MeshObject->SupportsStaticRelevance() && !IsDynamic();
 	Result.bDynamicRelevance = !Result.bStaticRelevance;
 	Result.bRenderCustomDepth = ShouldRenderCustomDepth();
 	Result.bRenderInMainPass = ShouldRenderInMainPass();
@@ -7091,7 +7113,7 @@ FPrimitiveViewRelevance FSkeletalMeshSceneProxy::GetViewRelevance(const FSceneVi
 	MaterialRelevance.SetPrimitiveViewRelevance(Result);
 
 #if !UE_BUILD_SHIPPING
-	Result.bSeparateTranslucency |= View->Family->EngineShowFlags.Constraints;
+	Result.bSeparateTranslucency |= EngineShowFlags.Constraints;
 #endif
 
 #if WITH_EDITOR
