Disable parallel GDME on RHIs that do not support parallel rendering tasks

diff --git a/Engine/Source/Runtime/Engine/Private/PrimitiveSceneProxy.cpp b/Engine/Source/Runtime/Engine/Private/PrimitiveSceneProxy.cpp
--- a/Engine/Source/Runtime/Engine/Private/PrimitiveSceneProxy.cpp
+++ b/Engine/Source/Runtime/Engine/Private/PrimitiveSceneProxy.cpp
@@ -106,7 +106,9 @@ static FAutoConsoleVariableRef CVarParallelGatherDynamicMeshElements(
 
 bool IsParallelGatherDynamicMeshElementsEnabled()
 {
-	return GParallelGatherDynamicMeshElements;
+	return GParallelGatherDynamicMeshElements &&
+		// parallel GDME creates RHI resources on a task threads which is not supported by some RHIs (eg. OpenGL)
+		GSupportsParallelRenderingTasksWithSeparateRHIThread;
 }
 
 bool FPrimitiveSceneProxy::ShouldRenderCustomDepth() const
