Fix crashes when using ref pose overrides when using skeletons with virtual bones

diff --git a/Engine/Source/Runtime/Engine/Private/Components/SkinnedMeshComponent.cpp b/Engine/Source/Runtime/Engine/Private/Components/SkinnedMeshComponent.cpp
--- a/Engine/Source/Runtime/Engine/Private/Components/SkinnedMeshComponent.cpp
+++ b/Engine/Source/Runtime/Engine/Private/Components/SkinnedMeshComponent.cpp
@@ -3621,9 +3621,20 @@ void USkinnedMeshComponent::SetRefPoseOverride(const TArray<FTransform>& NewRefP
 	// previously allocated RefPoseOverride, if there was one, will potentially be used on other threads by BoneContainer for one frame
 	RefPoseOverride = MakeShared<FSkelMeshRefPoseOverride>();
 
+	const int32 NumBones = GetSkinnedAsset()->GetRefSkeleton().GetNum();
+
 	// Copy input transforms into override data
+	RefPoseOverride->RefBonePoses.Reserve(NumBones);
 	RefPoseOverride->RefBonePoses = NewRefPoseTransforms;
 
+	// Add additional space for any virtual bones & copy them in
+	if(NumBones > NumRealBones)
+	{
+		RefPoseOverride->RefBonePoses.SetNumUninitialized(NumBones);
+		const TArray<FTransform>& RefPose = GetSkinnedAsset()->GetRefSkeleton().GetRefBonePose();
+		FMemory::Memcpy(&RefPoseOverride->RefBonePoses[NumRealBones], &RefPose[NumRealBones], sizeof(FTransform) * (NumBones - NumRealBones));
+	}
+
 	// Allocate output inv matrices
 	RefPoseOverride->RefBasesInvMatrix.AddUninitialized(NumRealBones);
 
