Changing Boost Timelines to Tweens

There's an emissive material parameter that's enabled during boosting. This was controlled via a timeline and took around 70us to update. I moved this over to use the new Tween Manager plugin. This cut the time down to about 1/2th of the original cost.
After this, I decided to convert 2 more timelines over, all active during boosting. The first modifies the camera FoV. The second modifies the controller's rumble intensity. These two together were part of the same Blueprint logic and took 85us. Now using tweens, they are running 2.6x faster.

So we should see around 80us back when boosting.

#changelist validated
#virtualized

diff --git a/Engine/Source/Runtime/Engine/Classes/GameFramework/PlayerController.h b/Engine/Source/Runtime/Engine/Classes/GameFramework/PlayerController.h
--- a/Engine/Source/Runtime/Engine/Classes/GameFramework/PlayerController.h
+++ b/Engine/Source/Runtime/Engine/Classes/GameFramework/PlayerController.h
@@ -1300,7 +1300,7 @@ public:
 	UFUNCTION(reliable, client, BlueprintCallable, Category="Game|Feedback")
 	ENGINE_API void ClientStopForceFeedback(class UForceFeedbackEffect* ForceFeedbackEffect, FName Tag);
 
-private:
+public:
 	/** 
 	 * Latent action that controls the playing of force feedback 
 	 * Begins playing when Start is called.  Calling Update or Stop if the feedback is not active will have no effect.
@@ -1316,6 +1316,7 @@ private:
 	UFUNCTION(BlueprintCallable, meta=(Latent, LatentInfo="LatentInfo", ExpandEnumAsExecs="Action", Duration="-1", bAffectsLeftLarge="true", bAffectsLeftSmall="true", bAffectsRightLarge="true", bAffectsRightSmall="true", AdvancedDisplay="bAffectsLeftLarge,bAffectsLeftSmall,bAffectsRightLarge,bAffectsRightSmall"), Category="Game|Feedback")
 	ENGINE_API void PlayDynamicForceFeedback(float Intensity, float Duration, bool bAffectsLeftLarge, bool bAffectsLeftSmall, bool bAffectsRightLarge, bool bAffectsRightSmall, TEnumAsByte<EDynamicForceFeedbackAction::Type> Action, FLatentActionInfo LatentInfo);
 
+private:
 	//~ This method is purely for debugging purposes.
 	//~ It will trigger a ServerUpdateLevelVisibilityCall with the provided package name.
 	UFUNCTION(Exec)
