Light function atlas - added light function material that uses scene texture ids (e.g. gbuffer material data, normal, etc.) into the non compatible list of materials.

diff --git a/Engine/Shaders/Private/LightFunctionAtlas/LightFunctionAtlasRender.usf b/Engine/Shaders/Private/LightFunctionAtlas/LightFunctionAtlasRender.usf
--- a/Engine/Shaders/Private/LightFunctionAtlas/LightFunctionAtlasRender.usf
+++ b/Engine/Shaders/Private/LightFunctionAtlas/LightFunctionAtlasRender.usf
@@ -8,7 +8,7 @@
 #define LIGHT_FUNCTION_ATLAS 1
 
 // Used to trigger shader recompiles, guid should be regenerated on merge conflicts
-#pragma message("UESHADERMETADATA_VERSION 4847F0A3-06B4-4A6D-8C31-3F98BEEF9480")
+#pragma message("UESHADERMETADATA_VERSION 98DA9C13-D8B0-43F6-AD31-6198D1CE3A5C")
 
 #include "../Common.ush"
 #include "/Engine/Generated/Material.ush"
diff --git a/Engine/Shaders/Private/LightFunctionPixelShader.usf b/Engine/Shaders/Private/LightFunctionPixelShader.usf
--- a/Engine/Shaders/Private/LightFunctionPixelShader.usf
+++ b/Engine/Shaders/Private/LightFunctionPixelShader.usf
@@ -10,7 +10,7 @@
 #include "DynamicLightingCommon.ush"
 
 // Used to trigger shader recompiles, guid should be regenerated on merge conflicts
-#pragma message("UESHADERMETADATA_VERSION BF889CEA-E940-4C56-8555-E5BEB172F300")
+#pragma message("UESHADERMETADATA_VERSION EFF06B5F-BCC7-4DDE-AF88-28825B73EEEC")
 
 // used like this: mul(float4(SvPosition.xyz, 1), SvPositionToLight);
 float4x4 SvPositionToLight;
diff --git a/Engine/Source/Runtime/Engine/Private/Materials/HLSLMaterialTranslator.cpp b/Engine/Source/Runtime/Engine/Private/Materials/HLSLMaterialTranslator.cpp
--- a/Engine/Source/Runtime/Engine/Private/Materials/HLSLMaterialTranslator.cpp
+++ b/Engine/Source/Runtime/Engine/Private/Materials/HLSLMaterialTranslator.cpp
@@ -1707,7 +1707,7 @@ void FHLSLMaterialTranslator::DoTranslate()
 	bUsesAnisotropy = IsMaterialPropertyUsed(MP_Anisotropy, Chunk[MP_Anisotropy], FLinearColor(0, 0, 0, 0), 1);
 	MaterialCompilationOutput.bUsesAnisotropy = bUsesAnisotropy;
 
-	MaterialCompilationOutput.bIsLightFunctionAtlasCompatible = !bUsesVertexPosition && !bUsesSceneDepth;
+	MaterialCompilationOutput.bIsLightFunctionAtlasCompatible = !bUsesVertexPosition && !bUsesSceneDepth && !MaterialCompilationOutput.bNeedsSceneTextures;
 
 	EMaterialDecalResponse MDR = (EMaterialDecalResponse)Material->GetMaterialDecalResponse();
 	if (MDR == MDR_Color || MDR == MDR_ColorNormal || MDR == MDR_ColorRoughness || MDR == MDR_ColorNormalRoughness)
