* Fixed problem when building SVT from within the RVT component's Build button, which would actually disable SVT, because the texture wasn't waited upon while recreating the RVT component's render state
* Prevent mobile-specific SVT textures to be displayed if Virtual Textures on mobile is disabled in the project settings
* Various UX adjustments to make it clear to the user what to expect when using mobile-specific SVT
#tests editor

diff --git a/Engine/Source/Editor/VirtualTexturingEditor/Private/RuntimeVirtualTextureDetailsCustomization.cpp b/Engine/Source/Editor/VirtualTexturingEditor/Private/RuntimeVirtualTextureDetailsCustomization.cpp
--- a/Engine/Source/Editor/VirtualTexturingEditor/Private/RuntimeVirtualTextureDetailsCustomization.cpp
+++ b/Engine/Source/Editor/VirtualTexturingEditor/Private/RuntimeVirtualTextureDetailsCustomization.cpp
@@ -241,6 +241,10 @@ void FRuntimeVirtualTextureComponentDetailsCustomization::CustomizeDetails(IDeta
 		.IsEnabled(this, &FRuntimeVirtualTextureComponentDetailsCustomization::IsSetBoundsEnabled)
 	];
 
+	FText BuildButtonText = LOCTEXT("Button_Build_Tooltip", "Build the low mips as streaming virtual texture data. \n\
+		If \"Separate Texture For Mobile\" is enabled in the Streaming Texture, only the mobile version of the texture will be updated when hitting this button \
+		while the mobile preview mode is active (and only the desktop version otherwise).");
+
 	// Apply custom widget for BuildStreamingMips.
 	TSharedRef<IPropertyHandle> BuildStreamingMipsPropertyHandle = DetailBuilder.GetProperty(TEXT("bBuildStreamingMipsButton"));
 	DetailBuilder.EditDefaultProperty(BuildStreamingMipsPropertyHandle)->CustomWidget()
@@ -249,7 +253,7 @@ void FRuntimeVirtualTextureComponentDetailsCustomization::CustomizeDetails(IDeta
 		SNew(STextBlock)
 		.Font(IDetailLayoutBuilder::GetDetailFont())
 		.Text(LOCTEXT("Button_BuildStreamingTexture", "Build Streaming Texture"))
-		.ToolTipText(LOCTEXT("Button_Build_Tooltip", "Build the low mips as streaming virtual texture data"))
+		.ToolTipText(BuildButtonText)
 	]
 	.ValueContent()
 	[
@@ -262,6 +266,7 @@ void FRuntimeVirtualTextureComponentDetailsCustomization::CustomizeDetails(IDeta
 			.VAlign(VAlign_Center)
 			.ContentPadding(2)
 			.Text(LOCTEXT("Button_Build", "Build"))
+			.ToolTipText(BuildButtonText)
 			.OnClicked(this, &FRuntimeVirtualTextureComponentDetailsCustomization::BuildStreamedMips)
 			.IsEnabled(this, &FRuntimeVirtualTextureComponentDetailsCustomization::IsBuildStreamedMipsEnabled)
 		]
@@ -270,7 +275,7 @@ void FRuntimeVirtualTextureComponentDetailsCustomization::CustomizeDetails(IDeta
 		.VAlign(VAlign_Center)
 		[
 			SNew(SImage)
-			.Image(FAppStyle::GetBrush("Icons.Warning"))
+			.Image(FCoreStyle::Get().GetBrush("Icons.Warning"))
 			.Visibility(this, &FRuntimeVirtualTextureComponentDetailsCustomization::IsBuildWarningIconVisible)
 			.ToolTipText(LOCTEXT("Warning_Build_Tooltip", "The settings have changed since the Streaming Texture was last rebuilt. Streaming mips are disabled."))
 		]
diff --git a/Engine/Source/Runtime/Engine/Classes/VT/VirtualTextureBuilder.h b/Engine/Source/Runtime/Engine/Classes/VT/VirtualTextureBuilder.h
--- a/Engine/Source/Runtime/Engine/Classes/VT/VirtualTextureBuilder.h
+++ b/Engine/Source/Runtime/Engine/Classes/VT/VirtualTextureBuilder.h
@@ -49,12 +49,12 @@ public:
 	GENERATED_UCLASS_BODY()
 	ENGINE_API ~UVirtualTextureBuilder();
 
-	/** The UTexture object. */
+	/** The (embedded) texture asset. Use Build Virtual Textures in the Build menu, or the Build button in the Runtime Virtual Texture Component to create/update it. */
 	UPROPERTY(VisibleAnywhere, BlueprintReadOnly, Category = Texture)
 	TObjectPtr<class UVirtualTexture2D> Texture;
 
-	/** The UTexture object for Mobile rendering. */
-	UPROPERTY(VisibleAnywhere, BlueprintReadOnly, Category = Texture)
+	/** The (embedded) texture asset for mobile rendering, only if virtual texture support on Mobile is enabled and if RVT support on mobile is enabled in the project settings (see r.Mobile.VirtualTextures). Use Build Virtual Textures or the Build button in the Runtime Virtual Texture Component to create/update it. */
+	UPROPERTY(VisibleAnywhere, BlueprintReadOnly, Category = Texture, meta = (EditCondition=bSeparateTextureForMobile))
 	TObjectPtr<class UVirtualTexture2D> TextureMobile;
 
 	/** Some client defined hash of that defines how the Texture was built. */
@@ -64,18 +64,21 @@ public:
 	/** Virtual texture for a specific shading path */
 	ENGINE_API UVirtualTexture2D* GetVirtualTexture(EShadingPath ShadingPath) const;
 
-	/** Whether to use a separate texture for Mobile rendering. A separate texure will be built using mobile preview editor mode */
+	/** Whether to use a separate texture for Mobile rendering, only if virtual texture support on Mobile is enabled and if RVT support on mobile is enabled in the project settings (see r.Mobile.VirtualTextures). A separate texture will be built using mobile preview editor mode. Use this in case there is too much discrepancy between the RVT used for desktop vs. mobile. */
 	UPROPERTY(EditAnywhere, Category = Texture)
 	bool bSeparateTextureForMobile = false;
 
 #if WITH_EDITOR
 	/** Creates a new UVirtualTexture2D and stores it in the contained Texture. */
-	ENGINE_API void BuildTexture(EShadingPath ShadingPath, FVirtualTextureBuildDesc const& BuildDesc);
+	ENGINE_API void BuildTexture(EShadingPath ShadingPath, FVirtualTextureBuildDesc const& BuildDesc, bool bWaitForCompilation = false);
 #endif
 
 protected:
 	//~ Begin UObject Interface
-	ENGINE_API virtual void Serialize(FArchive& Ar) override;
-	ENGINE_API virtual void PostLoad() override;
+	virtual void Serialize(FArchive& Ar) override;
+	virtual void PostLoad() override;
+#if WITH_EDITOR
+	virtual bool CanEditChange(const FProperty* InProperty) const override;
+#endif // WITH_EDITOR
 	//~ End UObject Interface
 };
diff --git a/Engine/Source/Runtime/Engine/Private/Components/RuntimeVirtualTextureComponent.cpp b/Engine/Source/Runtime/Engine/Private/Components/RuntimeVirtualTextureComponent.cpp
--- a/Engine/Source/Runtime/Engine/Private/Components/RuntimeVirtualTextureComponent.cpp
+++ b/Engine/Source/Runtime/Engine/Private/Components/RuntimeVirtualTextureComponent.cpp
@@ -349,7 +349,7 @@ void URuntimeVirtualTextureComponent::InitializeStreamingTexture(EShadingPath Sh
 		BuildDesc.TileBorderSize = VirtualTexture->GetTileBorderSize();
 		BuildDesc.LODGroup = VirtualTexture->GetLODGroup();
 		BuildDesc.LossyCompressionAmount = GetLossyCompressionAmount();
-
+		
 		BuildDesc.LayerCount = VirtualTexture->GetLayerCount();
 		check(BuildDesc.LayerCount <= RuntimeVirtualTexture::MaxTextureLayers);
 		BuildDesc.LayerFormats.AddDefaulted(BuildDesc.LayerCount);
@@ -369,15 +369,20 @@ void URuntimeVirtualTextureComponent::InitializeStreamingTexture(EShadingPath Sh
 		BuildDesc.InSizeY = InSizeY;
 		BuildDesc.InData = InData;
 
-		StreamingTexture->BuildTexture(ShadingPath, BuildDesc);
+		// Make sure the streaming texture is fully built before marking the render state dirty, otherwise the scene proxy will be constructed thinking that it's not, which will prevent showing it in editor. 
+		//  It's a rarely-triggered, bake-time, editor-only function anyway, so the blocking wait is acceptable
+		constexpr bool bWaitForCompilation = true;
+		StreamingTexture->BuildTexture(ShadingPath, BuildDesc, bWaitForCompilation);
 		StreamingTexture->Modify();
 	}
+
+	MarkRenderStateDirty();
 }
 
 bool URuntimeVirtualTextureComponent::CanEditChange(const FProperty* InProperty) const
 {
 	bool bCanEdit = Super::CanEditChange(InProperty);
-	if (InProperty->GetFName() == TEXT("bUseStreamingLowMipsInEditor"))
+	if (InProperty->GetFName() == GET_MEMBER_NAME_CHECKED(URuntimeVirtualTextureComponent, bUseStreamingLowMipsInEditor))
 	{
 		bCanEdit &= GetVirtualTexture() != nullptr && GetStreamingTexture() != nullptr;
 	}
diff --git a/Engine/Source/Runtime/Engine/Private/VT/VirtualTextureBuilder.cpp b/Engine/Source/Runtime/Engine/Private/VT/VirtualTextureBuilder.cpp
--- a/Engine/Source/Runtime/Engine/Private/VT/VirtualTextureBuilder.cpp
+++ b/Engine/Source/Runtime/Engine/Private/VT/VirtualTextureBuilder.cpp
@@ -6,6 +6,7 @@
 #include "VT/VirtualTexture.h"
 #include "SceneInterface.h"
 #include "SceneUtils.h"
+#include "TextureCompiler.h"
 
 #if WITH_EDITOR
 #include "Interfaces/ITargetPlatform.h"
@@ -90,7 +91,26 @@ UVirtualTexture2D* UVirtualTextureBuilder::GetVirtualTexture(EShadingPath Shadin
 
 #if WITH_EDITOR
 
-static void BuildVirtualTexture2D(UVirtualTexture2D* Texture, FVirtualTextureBuildDesc const& BuildDesc)
+bool UVirtualTextureBuilder::CanEditChange(const FProperty* InProperty) const
+{
+	bool bCanEdit = Super::CanEditChange(InProperty);
+
+	static const TConsoleVariableData<int32>* CVarVirtualTextureOnMobile = IConsoleManager::Get().FindTConsoleVariableDataInt(TEXT("r.Mobile.VirtualTextures"));
+	check(CVarVirtualTextureOnMobile);
+	const bool bEnableVirtualTextureOnMobile = (CVarVirtualTextureOnMobile->GetValueOnGameThread() != 0);
+
+	if (InProperty->GetFName() == GET_MEMBER_NAME_CHECKED(UVirtualTextureBuilder, bSeparateTextureForMobile))
+	{
+		bCanEdit &= bEnableVirtualTextureOnMobile;
+	}
+	else if (InProperty->GetFName() == GET_MEMBER_NAME_CHECKED(UVirtualTextureBuilder, TextureMobile))
+	{
+		bCanEdit &= bEnableVirtualTextureOnMobile && bSeparateTextureForMobile;
+	}
+	return bCanEdit;
+}
+
+static void BuildVirtualTexture2D(UVirtualTexture2D* Texture, FVirtualTextureBuildDesc const& BuildDesc, bool bWaitForCompilation)
 {
 	Texture->VirtualTextureStreaming = true;
 	Texture->LODGroup = BuildDesc.LODGroup;
@@ -110,10 +130,15 @@ static void BuildVirtualTexture2D(UVirtualTexture2D* Texture, FVirtualTextureBui
 
 	Texture->Source.InitLayered(BuildDesc.InSizeX, BuildDesc.InSizeY, 1, BuildDesc.LayerCount, 1, BuildDesc.LayerFormats.GetData(), BuildDesc.InData);
 	Texture->PostEditChange();
+
+	if (bWaitForCompilation)
+	{
+		FTextureCompilingManager::Get().FinishCompilation({ Texture });
+	}
 }
 
 
-void UVirtualTextureBuilder::BuildTexture(EShadingPath ShadingPath, FVirtualTextureBuildDesc const& BuildDesc)
+void UVirtualTextureBuilder::BuildTexture(EShadingPath ShadingPath, FVirtualTextureBuildDesc const& BuildDesc, bool bWaitForCompilation)
 {
 	if (!bSeparateTextureForMobile)
 	{
@@ -130,13 +155,13 @@ void UVirtualTextureBuilder::BuildTexture(EShadingPath ShadingPath, FVirtualText
 		
 		BuildHash = BuildDesc.BuildHash;
 		TextureMobile = NewObject<UVirtualTexture2D>(this, TEXT("TextureMobile"));
-		BuildVirtualTexture2D(TextureMobile, BuildDesc);
+		BuildVirtualTexture2D(TextureMobile, BuildDesc, bWaitForCompilation);
 	}
 	else
 	{
 		BuildHash = BuildDesc.BuildHash;
 		Texture = NewObject<UVirtualTexture2D>(this, TEXT("Texture"));
-		BuildVirtualTexture2D(Texture, BuildDesc);
+		BuildVirtualTexture2D(Texture, BuildDesc, bWaitForCompilation);
 	}
 }
 
