Fix Texture2DArray and Cube Array were crashing in mimalloc due to deleting PlatformData twice in InvalidateTextureSource

diff --git a/Engine/Source/Runtime/Engine/Classes/Engine/Texture2DArray.cpp b/Engine/Source/Runtime/Engine/Classes/Engine/Texture2DArray.cpp
--- a/Engine/Source/Runtime/Engine/Classes/Engine/Texture2DArray.cpp
+++ b/Engine/Source/Runtime/Engine/Classes/Engine/Texture2DArray.cpp
@@ -495,11 +495,7 @@ ENGINE_API void UTexture2DArray::InvalidateTextureSource()
 {
 	Modify();
 
-	if (GetPlatformData())
-	{
-		delete GetPlatformData();
-		SetPlatformData(nullptr);
-	}
+	SetPlatformData(nullptr);
 
 	Source = FTextureSource();
 	Source.SetOwner(this);
diff --git a/Engine/Source/Runtime/Engine/Private/TextureCubeArray.cpp b/Engine/Source/Runtime/Engine/Private/TextureCubeArray.cpp
--- a/Engine/Source/Runtime/Engine/Private/TextureCubeArray.cpp
+++ b/Engine/Source/Runtime/Engine/Private/TextureCubeArray.cpp
@@ -752,16 +752,12 @@ ENGINE_API void UTextureCubeArray::InvalidateTextureSource()
 {
 	Modify();
 
-	if (GetPlatformData())
-	{
-		delete GetPlatformData();
-		SetPlatformData(nullptr);
-	}
+	SetPlatformData(nullptr);
 
 	Source = FTextureSource();
 	Source.SetOwner(this);
+
 	UpdateResource();
-	
 }
 #endif
 
