Make the "ShapeColor" on UShapeComponent BP Read-Write, instead of just BP Read.

#rnx

diff --git a/Engine/Source/Runtime/Engine/Classes/Components/ShapeComponent.h b/Engine/Source/Runtime/Engine/Classes/Components/ShapeComponent.h
--- a/Engine/Source/Runtime/Engine/Classes/Components/ShapeComponent.h
+++ b/Engine/Source/Runtime/Engine/Classes/Components/ShapeComponent.h
@@ -40,7 +40,7 @@ class UShapeComponent : public UPrimitiveComponent
 #endif // WITH_EDITORONLY_DATA
 
 	/** Color used to draw the shape. */
-	UPROPERTY(EditAnywhere, AdvancedDisplay, BlueprintReadOnly, Category=Shape)
+	UPROPERTY(EditAnywhere, AdvancedDisplay, BlueprintReadWrite, Category=Shape)
 	FColor ShapeColor;
 
 	/** Only show this component if the actor is selected */
