Fixed duplicate categories in the blueprint context menu when | seperated categories are declared with whitespace around the |'s
Make use of ParseTokens directly so that we can trim both leading and trailing whitespace from categories without too many intermediate allocations

diff --git a/Engine/Source/Runtime/Engine/Private/EdGraph/EdGraphSchema.cpp b/Engine/Source/Runtime/Engine/Private/EdGraph/EdGraphSchema.cpp
--- a/Engine/Source/Runtime/Engine/Private/EdGraph/EdGraphSchema.cpp
+++ b/Engine/Source/Runtime/Engine/Private/EdGraph/EdGraphSchema.cpp
@@ -11,6 +11,7 @@
 #include "ScopedTransaction.h"
 #include "EditorCategoryUtils.h"
 #include "Settings/EditorStyleSettings.h"
+#include "String/ParseTokens.h"
 #endif
 
 #define LOCTEXT_NAMESPACE "EdGraph"
@@ -221,22 +222,21 @@ void FGraphActionListBuilderBase::ActionGroup::Copy(const ActionGroup& Other)
 void FGraphActionListBuilderBase::ActionGroup::InitCategoryChain()
 {
 #if WITH_EDITOR
-	const TCHAR* CategoryDelim = TEXT("|");
-	FEditorCategoryUtils::GetCategoryDisplayString(RootCategory).ParseIntoArray(CategoryChain, CategoryDelim, true);
+	const FStringView CategoryDelim = FStringView(TEXT("|"));
+	const UE::String::EParseTokensOptions ParseOptions = 
+		UE::String::EParseTokensOptions::IgnoreCase |
+		UE::String::EParseTokensOptions::SkipEmpty |
+		UE::String::EParseTokensOptions::Trim;
+	UE::String::ParseTokens(FStringView(FEditorCategoryUtils::GetCategoryDisplayString(RootCategory)), CategoryDelim,
+		[&CategoryChain = CategoryChain](FStringView Token) { CategoryChain.Emplace(Token); },
+		ParseOptions);
 
 	if (Actions.Num() > 0)
 	{
-		TArray<FString> SubCategoryChain;
-
 		FString SubCategory = FEditorCategoryUtils::GetCategoryDisplayString(Actions[0]->GetCategory().ToString());
-		SubCategory.ParseIntoArray(SubCategoryChain, CategoryDelim, true);
-
-		CategoryChain.Append(SubCategoryChain);
-	}
-
-	for (FString& Category : CategoryChain)
-	{
-		Category.TrimStartInline();
+		UE::String::ParseTokens(FStringView(SubCategory), CategoryDelim,
+			[&CategoryChain = CategoryChain](FStringView Token) { CategoryChain.Emplace(Token); },
+			ParseOptions);
 	}
 #endif
 }
