Adding "Modulo" material node that maps to the modulo operator. Notably, this works with uint values.

diff --git a/Engine/Shaders/Private/LWCOperations.ush b/Engine/Shaders/Private/LWCOperations.ush
--- a/Engine/Shaders/Private/LWCOperations.ush
+++ b/Engine/Shaders/Private/LWCOperations.ush
@@ -121,6 +121,11 @@ FLWCType LWCLerp(FLWCType Lhs, FLWCType Rhs, FFloatType S)
 	return LWCConstructor(lerp(LWCGetTile(Lhs), LWCGetTile(Rhs), S), lerp(Lhs.Offset, Rhs.Offset, S));
 }
 
+FFloatType LWCModulo(FLWCType Lhs, FFloatType Rhs)
+{
+	return LWCToFloat(LWCSubtract(Lhs, LWCMultiply(LWCTrunc(LWCDivide(Lhs, Rhs)), Rhs)));
+}
+
 FFloatType LWCFmod(FLWCType Lhs, FFloatType Rhs)
 {
 	return LWCToFloat(LWCSubtract(Lhs, LWCMultiply(LWCTrunc(LWCDivide(Lhs, Rhs)), Rhs)));
diff --git a/Engine/Shaders/Private/WorldSpaceMath.ush b/Engine/Shaders/Private/WorldSpaceMath.ush
--- a/Engine/Shaders/Private/WorldSpaceMath.ush
+++ b/Engine/Shaders/Private/WorldSpaceMath.ush
@@ -119,6 +119,7 @@ FDFInverseMatrix WSToDF(FWSInverseMatrix In) { return DFFromTileOffset(In); }
 #define WSDivideByPow2 LWCDivide
 #define WSLerp LWCLerp
 #define WSFmodDemote LWCFmod
+#define WSModuloDemote LWCModulo
 
 #define WSSin LWCSin
 #define WSCos LWCCos
@@ -252,6 +253,7 @@ FDFInverseMatrix WSToDF(FWSInverseMatrix In) { return In; }
 #define WSDivideByPow2 DFDivideByPow2
 #define WSLerp DFLerp
 #define WSFmodDemote DFFmodDemote
+#define WSModuloDemote DFModuloDemote
 
 #define WSSin DFSin
 #define WSCos DFCos
diff --git a/Engine/Source/Runtime/Core/Private/UObject/DevObjectVersion.cpp b/Engine/Source/Runtime/Core/Private/UObject/DevObjectVersion.cpp
--- a/Engine/Source/Runtime/Core/Private/UObject/DevObjectVersion.cpp
+++ b/Engine/Source/Runtime/Core/Private/UObject/DevObjectVersion.cpp
@@ -128,14 +128,14 @@ FDevSystemGuids::FDevSystemGuids()
 	, GROOM_BINDING_DERIVED_DATA_VERSION(0x30769E53, 0x0C574C7B, 0xA15C56F2, 0x24A64E32)
 	, GROOM_DERIVED_DATA_VERSION(0x05A37379, 0xF8A049D5, 0x986824E9, 0xAAA83F41)
 	, LANDSCAPE_MOBILE_COOK_VERSION(0x0E9ADF72, 0xD6B64E0D, 0x81C4A92B, 0x081A37AB)
-	, MATERIALSHADERMAP_DERIVEDDATA_VER(0xD31091F1, 0x46F0456F, 0x9DA512EF, 0xA47EFB19)
+	, MATERIALSHADERMAP_DERIVEDDATA_VER(0x1C415798, 0x50BB46E5, 0x9AAF5307, 0xB2CE65CE)
 	, NANITE_DERIVEDDATA_VER(0xBEB0226A, 0x070E4ECA, 0x972C1E7D, 0xD8599E68)
 	, NIAGARASHADERMAP_DERIVEDDATA_VER(0x7BBD9913, 0xC1554D20, 0xADAE9F17, 0xB006299E)
 	, Niagara_LatestScriptCompileVersion(0x6D32B8EE, 0x909FCA7E, 0xA5CE4F17, 0x066A5F25)
 	, POSESEARCHDB_DERIVEDDATA_VER(0x389117E4, 0x807A4CC0, 0x9F37C2E6, 0xD808A78D)
 	, SkeletalMeshDerivedDataVersion(0x9B5F4544, 0x76D7481C, 0x9AD3F614, 0xA6C07904)
 	, STATICMESH_DERIVEDDATA_VER(0x2C6C400C, 0x2EDF47B3, 0x9BD36689, 0x7B77B208)
-	, MaterialTranslationDDCVersion(0xB1C6F9A0, 0x038E48A3, 0xA6D467A6, 0xC81A5849)
+	, MaterialTranslationDDCVersion(0xE6097D5A, 0x39794E1B, 0xA3FDDA91, 0x6E2221B6)
 {
 }
 
diff --git a/Engine/Source/Runtime/Core/Private/UObject/FortniteMainBranchObjectVersion.cpp b/Engine/Source/Runtime/Core/Private/UObject/FortniteMainBranchObjectVersion.cpp
--- a/Engine/Source/Runtime/Core/Private/UObject/FortniteMainBranchObjectVersion.cpp
+++ b/Engine/Source/Runtime/Core/Private/UObject/FortniteMainBranchObjectVersion.cpp
@@ -8,12 +8,12 @@ TMap<FGuid, FGuid> FFortniteMainBranchObjectVersion::GetSystemGuids()
 
 	SystemGuids.Add(DevGuids.GLOBALSHADERMAP_DERIVEDDATA_VER, FGuid("8D3A12924CDF4D9F84C0A900E9445CAD"));
 	SystemGuids.Add(DevGuids.LANDSCAPE_MOBILE_COOK_VERSION, FGuid("32D02EF867C74B71A0D4E0FA41392732"));
-	SystemGuids.Add(DevGuids.MATERIALSHADERMAP_DERIVEDDATA_VER, FGuid("E416BA6E4D48491994BB91E4B06E43AC"));
+	SystemGuids.Add(DevGuids.MATERIALSHADERMAP_DERIVEDDATA_VER, FGuid("BB56D95CD4FD48E38C4CD7869B44B987"));
 	SystemGuids.Add(DevGuids.NANITE_DERIVEDDATA_VER, FGuid("D28C003F825F4C178F698938FB0E47F7"));
 	SystemGuids.Add(DevGuids.NIAGARASHADERMAP_DERIVEDDATA_VER, FGuid("C417A50271B3427D98785388266E3B66"));
 	SystemGuids.Add(DevGuids.Niagara_LatestScriptCompileVersion, FGuid("E267F600682A4E94885E8AEC5CE1DC91"));
 	SystemGuids.Add(DevGuids.SkeletalMeshDerivedDataVersion, FGuid("B9AEAA3EE7AC4FCFB3860E8F3E06DD75"));
 	SystemGuids.Add(DevGuids.STATICMESH_DERIVEDDATA_VER, FGuid("543843B103794E0E9BA4BE80FB602F79"));
-	SystemGuids.Add(DevGuids.MaterialTranslationDDCVersion, FGuid("B1C6F9A0038E48A3A6D467A6C81A5849"));
+	SystemGuids.Add(DevGuids.MaterialTranslationDDCVersion, FGuid("05477E9B490F4FA9BE05B20BA3BF938C"));
 	return SystemGuids;
 }
diff --git a/Engine/Source/Runtime/Engine/Classes/Materials/MaterialExpressionModulo.h b/Engine/Source/Runtime/Engine/Classes/Materials/MaterialExpressionModulo.h
new file mode 100644
--- /dev/null
+++ b/Engine/Source/Runtime/Engine/Classes/Materials/MaterialExpressionModulo.h
@@ -0,0 +1,29 @@
+// Copyright Epic Games, Inc. All Rights Reserved.
+
+#pragma once
+
+#include "Materials/MaterialExpression.h"
+#include "MaterialExpressionModulo.generated.h"
+
+UCLASS(MinimalAPI, collapsecategories, hidecategories=Object)
+class UMaterialExpressionModulo : public UMaterialExpression
+{
+	GENERATED_UCLASS_BODY()
+
+	UPROPERTY()
+	FExpressionInput A;
+
+	UPROPERTY()
+	FExpressionInput B;
+
+	//~ Begin UMaterialExpression Interface
+#if WITH_EDITOR
+	virtual int32 Compile(class FMaterialCompiler* Compiler, int32 OutputIndex) override;
+	virtual uint32 GetInputType(int32 InputIndex) override;
+	virtual uint32 GetOutputType(int32 OutputIndex) override;
+	virtual void GetCaption(TArray<FString>& OutCaptions) const override;
+	virtual FText GetKeywords() const override;
+	virtual bool GenerateHLSLExpression(FMaterialHLSLGenerator& Generator, UE::HLSLTree::FScope& Scope, int32 OutputIndex, UE::HLSLTree::FExpression const*& OutExpression) const override;
+#endif // WITH_EDITOR
+	//~ End UMaterialExpression Interface
+};
diff --git a/Engine/Source/Runtime/Engine/Private/HLSLTree/HLSLTreeEmit.cpp b/Engine/Source/Runtime/Engine/Private/HLSLTree/HLSLTreeEmit.cpp
--- a/Engine/Source/Runtime/Engine/Private/HLSLTree/HLSLTreeEmit.cpp
+++ b/Engine/Source/Runtime/Engine/Private/HLSLTree/HLSLTreeEmit.cpp
@@ -830,10 +830,18 @@ int32 InternalFormatString(FStringBuilderBase* OutString, FEmitShaderDependencie
 	if (Format.Len() > 0)
 	{
 		check(OutString);
-		for (TCHAR Char : Format)
+		for (int32 Index=0; Index<Format.Len(); Index++)
 		{
+			const TCHAR Char = Format[Index];
 			if (Char == TEXT('%'))
 			{
+				if (Index + 1 < Format.Len() && Format[Index + 1] == TEXT('%'))
+				{
+					OutString->Append(TEXT("%%"));
+					Index++;
+					continue;
+				}
+
 				const FFormatArgVariant& Arg = ArgList[ArgIndex++];
 				switch (Arg.Type)
 				{
diff --git a/Engine/Source/Runtime/Engine/Private/HLSLTree/HLSLTreeOperation.cpp b/Engine/Source/Runtime/Engine/Private/HLSLTree/HLSLTreeOperation.cpp
--- a/Engine/Source/Runtime/Engine/Private/HLSLTree/HLSLTreeOperation.cpp
+++ b/Engine/Source/Runtime/Engine/Private/HLSLTree/HLSLTreeOperation.cpp
@@ -75,7 +75,8 @@ FOperationDescription GetOperationDescription(EOperation Op)
 	case EOperation::Sub: return FOperationDescription(TEXT("Subtract"), TEXT("-"), 2, Shader::EPreshaderOpcode::Sub); break;
 	case EOperation::Mul: return FOperationDescription(TEXT("Multiply"), TEXT("*"), 2, Shader::EPreshaderOpcode::Mul); break;
 	case EOperation::Div: return FOperationDescription(TEXT("Divide"), TEXT("/"), 2, Shader::EPreshaderOpcode::Div); break;
-	case EOperation::Fmod: return FOperationDescription(TEXT("Fmod"), TEXT("%"), 2, Shader::EPreshaderOpcode::Fmod); break;
+	case EOperation::Fmod: return FOperationDescription(TEXT("Fmod"), TEXT("fmod"), 2, Shader::EPreshaderOpcode::Fmod); break;
+	case EOperation::Modulo: return FOperationDescription(TEXT("Modulo"), TEXT("%"), 2, Shader::EPreshaderOpcode::Modulo); break;
 	case EOperation::Step: return FOperationDescription(TEXT("Step"), TEXT("step"), 2, Shader::EPreshaderOpcode::Nop); break;
 	case EOperation::PowPositiveClamped: return FOperationDescription(TEXT("PowPositiveClamped"), TEXT("PowPositiveClamped"), 2, Shader::EPreshaderOpcode::Nop); break;
 	case EOperation::Atan2: return FOperationDescription(TEXT("Atan2"), TEXT("atan2"), 2, Shader::EPreshaderOpcode::Atan2); break;
@@ -367,6 +368,10 @@ FOperationTypes GetOperationTypes(EOperation Op, TConstArrayView<FPreparedType>
 			Types.InputType[1] = Shader::MakeNonLWCType(Types.InputType[1]);
 			Types.ResultType = MakeNonLWCType(IntermediateType);
 			break;
+		case EOperation::Modulo:
+			Types.InputType[1] = Shader::MakeNonLWCType(Types.InputType[1]);
+			Types.ResultType = MakeNonLWCType(IntermediateType);
+			break;
 		case EOperation::PowPositiveClamped:
 		case EOperation::Atan2:
 		case EOperation::Atan2Fast:
@@ -621,6 +626,11 @@ void FExpressionOperation::ComputeAnalyticDerivatives(FTree& Tree, FExpressionDe
 		// We can't really do anything meaningful in the non-zero case.
 		OutResult = InputDerivatives[0];
 		break;
+	case EOperation::Modulo:
+		// Only valid when B derivatives are zero.
+		// We can't really do anything meaningful in the non-zero case.
+		OutResult = InputDerivatives[0];
+		break;
 	case EOperation::Min:
 	{
 		const FExpression* Cond = Tree.NewLess(Inputs[0], Inputs[1]);
@@ -820,6 +830,7 @@ void FExpressionOperation::EmitValueShader(FEmitContext& Context, FEmitScope& Sc
 	case EOperation::Mul: OutResult.Code = Context.EmitExpression(Scope, ResultType, Types.bIsLWC ? TEXT("WSMultiply(%, %)") : TEXT("(% * %)"), InputValue[0], InputValue[1]); break;
 	case EOperation::Div: OutResult.Code = Context.EmitExpression(Scope, ResultType, Types.bIsLWC ? TEXT("WSDivide(%, %)") : TEXT("(% / %)"), InputValue[0], InputValue[1]); break;
 	case EOperation::Fmod: OutResult.Code = Context.EmitExpression(Scope, ResultType, Types.bIsLWC ? TEXT("WSFmodDemote(%, %)") : TEXT("fmod(%, %)"), InputValue[0], InputValue[1]); break;
+	case EOperation::Modulo: OutResult.Code = Context.EmitExpression(Scope, ResultType, Types.bIsLWC ? TEXT("% %% %)") : TEXT("(% %% %)"), InputValue[0], InputValue[1]); break;
 	case EOperation::Step: OutResult.Code = Context.EmitExpression(Scope, ResultType, Types.bIsLWC ? TEXT("WSStep(%, %)") : TEXT("step(%, %)"), InputValue[0], InputValue[1]); break;
 	case EOperation::PowPositiveClamped: OutResult.Code = Context.EmitExpression(Scope, ResultType, TEXT("PositiveClampedPow(%, %)"), InputValue[0], InputValue[1]); break;
 	case EOperation::Atan2: OutResult.Code = Context.EmitExpression(Scope, ResultType, TEXT("atan2(%, %)"), InputValue[0], InputValue[1]); break;
diff --git a/Engine/Source/Runtime/Engine/Private/Materials/HLSLMaterialDerivativeAutogen.cpp b/Engine/Source/Runtime/Engine/Private/Materials/HLSLMaterialDerivativeAutogen.cpp
--- a/Engine/Source/Runtime/Engine/Private/Materials/HLSLMaterialDerivativeAutogen.cpp
+++ b/Engine/Source/Runtime/Engine/Private/Materials/HLSLMaterialDerivativeAutogen.cpp
@@ -384,9 +384,10 @@ FMaterialDerivativeAutogen::FOperationType2 FMaterialDerivativeAutogen::GetFunc2
 	case EFunc2::Max:
 	case EFunc2::Min:
 	case EFunc2::Fmod:
+	case EFunc2::Modulo:
 		// Operations that support LWC
 		// Any LWC type promotes result to LWC
-		if (IsLWCType(LhsType) || (IsLWCType(RhsType) && Op != EFunc2::Fmod))
+		if (IsLWCType(LhsType) || (IsLWCType(RhsType) && Op != EFunc2::Fmod && Op != EFunc2::Modulo))
 		{
 			const EDerivativeType LhsLWCType = MakeLWCType(LhsType);
 			const EDerivativeType RhsLWCType = MakeLWCType(RhsType);
@@ -408,6 +409,7 @@ FMaterialDerivativeAutogen::FOperationType2 FMaterialDerivativeAutogen::GetFunc2
 					// Make sure we generate LWCDivide(LWC, float) if possible, since have a LWC denominator will result in worse approximation
 					return FOperationType2(LWCType, LWCType, IsLWCType(RhsType) ? LWCType : NonLWCType);
 				case EFunc2::Fmod:
+				case EFunc2::Modulo:
 					// FMOD Rhs is non-LWC
 					return FOperationType2(NonLWCType, LWCType, NonLWCType);
 				case EFunc2::Min:
@@ -840,7 +842,7 @@ int32 FMaterialDerivativeAutogen::GenerateExpressionFunc2(FHLSLMaterialTranslato
 	const FDerivInfo LhsDerivInfo = Translator.GetDerivInfo(LhsCode);
 	const FDerivInfo RhsDerivInfo = Translator.GetDerivInfo(RhsCode);
 
-	if (Op == EFunc2::Fmod)
+	if (Op == EFunc2::Fmod || Op == EFunc2::Modulo)
 	{
 		check(RhsDerivInfo.DerivativeStatus == EDerivativeStatus::Zero);
 	}
@@ -997,6 +999,17 @@ int32 FMaterialDerivativeAutogen::GenerateExpressionFunc2(FHLSLMaterialTranslato
 				DstToken = TEXT("fmod(") + LhsToken + TEXT(",") + RhsToken + TEXT(")");
 			}
 			break;
+		case EFunc2::Modulo:
+			if (bIsLWC)
+			{
+				Translator.AddLWCFuncUsage(ELWCFunctionKind::Other);
+				DstToken = TEXT("WSModuloDemote(") + LhsToken + TEXT(",") + RhsToken + TEXT(")");
+			}
+			else
+			{
+				DstToken = TEXT("(") + LhsToken + TEXT("%") + RhsToken + TEXT(")");
+			}
+			break;
 		case EFunc2::Min:
 			if (bIsLWC)
 			{
@@ -1083,6 +1096,9 @@ int32 FMaterialDerivativeAutogen::GenerateExpressionFunc2(FHLSLMaterialTranslato
 		case EFunc2::Fmod:
 			DstToken = TEXT("FmodDeriv") + Suffix + TEXT("(") + LhsToken + TEXT(",") + RhsToken + TEXT(")");
 			break;
+		case EFunc2::Modulo:
+			DstToken = TEXT("ModuloDeriv") + Suffix + TEXT("(") + LhsToken + TEXT(",") + RhsToken + TEXT(")");
+			break;
 		case EFunc2::Dot:
 			DstToken = TEXT("DotDeriv") + Suffix + TEXT("(") + LhsToken + TEXT(",") + RhsToken + TEXT(")");
 			break;
@@ -1694,6 +1710,27 @@ FString FMaterialDerivativeAutogen::GenerateUsedFunctions(FHLSLMaterialTranslato
 						Ret += TEXT("}") HLSL_LINE_TERMINATOR;
 						Ret += TEXT("") HLSL_LINE_TERMINATOR;
 						break;
+					case EFunc2::Modulo:
+						// Only valid when B derivatives are zero.
+						// We can't really do anything meaningful in the non-zero case.
+						Ret += NonLWCBaseName + TEXT(" ModuloDeriv") + Suffix + TEXT("(") + LHSBaseName + TEXT(" A, ") + NonLWCBaseName + TEXT(" B)") HLSL_LINE_TERMINATOR;
+						Ret += TEXT("{") HLSL_LINE_TERMINATOR;
+						Ret += TEXT("\t") + NonLWCBaseName + TEXT(" Ret;") HLSL_LINE_TERMINATOR;
+						if (IsLWCType(DerivType))
+						{
+							Translator.AddLWCFuncUsage(ELWCFunctionKind::Other);
+							Ret += TEXT("\tRet.Value = WSModuloDemote(A.Value, B.Value);") HLSL_LINE_TERMINATOR;
+						}
+						else
+						{
+							Ret += TEXT("\tRet.Value = (A.Value % B.Value);") HLSL_LINE_TERMINATOR;
+						}
+						Ret += TEXT("\tRet.Ddx = A.Ddx;") HLSL_LINE_TERMINATOR;
+						Ret += TEXT("\tRet.Ddy = A.Ddy;") HLSL_LINE_TERMINATOR;
+						Ret += TEXT("\treturn Ret;") HLSL_LINE_TERMINATOR;
+						Ret += TEXT("}") HLSL_LINE_TERMINATOR;
+						Ret += TEXT("") HLSL_LINE_TERMINATOR;
+						break;
 					case EFunc2::Min:
 						Ret += BaseName + TEXT(" MinDeriv") + Suffix + TEXT("(") + LHSBaseName + TEXT(" A, ") + RHSBaseName + TEXT(" B)") HLSL_LINE_TERMINATOR;
 						Ret += TEXT("{") HLSL_LINE_TERMINATOR;
diff --git a/Engine/Source/Runtime/Engine/Private/Materials/HLSLMaterialDerivativeAutogen.h b/Engine/Source/Runtime/Engine/Private/Materials/HLSLMaterialDerivativeAutogen.h
--- a/Engine/Source/Runtime/Engine/Private/Materials/HLSLMaterialDerivativeAutogen.h
+++ b/Engine/Source/Runtime/Engine/Private/Materials/HLSLMaterialDerivativeAutogen.h
@@ -123,6 +123,7 @@ public:
 		Mul,
 		Div,
 		Fmod,
+		Modulo,
 		Max,
 		Min,
 		Dot,	// Depends on Add/Mul, so it must come after them
diff --git a/Engine/Source/Runtime/Engine/Private/Materials/HLSLMaterialTranslator.cpp b/Engine/Source/Runtime/Engine/Private/Materials/HLSLMaterialTranslator.cpp
--- a/Engine/Source/Runtime/Engine/Private/Materials/HLSLMaterialTranslator.cpp
+++ b/Engine/Source/Runtime/Engine/Private/Materials/HLSLMaterialTranslator.cpp
@@ -3999,9 +3999,9 @@ int32 FHLSLMaterialTranslator::AccessUniformExpression(int32 Index)
 	return AddInlinedCodeChunkZeroDeriv(CodeChunk.Type,FormattedCode.ToString());
 }
 
-FString FHLSLMaterialTranslator::CoerceValue(const FString& Code, EMaterialValueType SourceType, EMaterialValueType DestType)
+FString FHLSLMaterialTranslator::CoerceValue(const FString& Code, EMaterialValueType SourceType, EMaterialValueType DestType, EMaterialCastFlags AdditionalCastFlags)
 {
-	EMaterialCastFlags CastFlags = EMaterialCastFlags::ReplicateScalar;
+	EMaterialCastFlags CastFlags = EMaterialCastFlags::ReplicateScalar | AdditionalCastFlags;
 	if (DestType == MCT_Float || DestType == MCT_Float1 || DestType == MCT_LWCScalar)
 	{
 		// CoerceValue allows truncating to scalar types only
@@ -4011,9 +4011,9 @@ FString FHLSLMaterialTranslator::CoerceValue(const FString& Code, EMaterialValue
 }
 
 // CoerceParameter
-FString FHLSLMaterialTranslator::CoerceParameter(int32 Index,EMaterialValueType DestType)
+FString FHLSLMaterialTranslator::CoerceParameter(int32 Index,EMaterialValueType DestType, EMaterialCastFlags AdditionalCastFlags)
 {
-	return CoerceValue(GetParameterCode(Index), GetParameterType(Index), DestType);
+	return CoerceValue(GetParameterCode(Index), GetParameterType(Index), DestType, AdditionalCastFlags);
 }
 
 // GetParameterType
@@ -4720,22 +4720,23 @@ int32 FHLSLMaterialTranslator::CastToNonLWCIfDisabled(int32 Code)
 
 FString FHLSLMaterialTranslator::CastValue(const FString& Code, EMaterialValueType SourceType, EMaterialValueType DestType, EMaterialCastFlags Flags)
 {
-	const bool bAllowTruncate = EnumHasAnyFlags(Flags, EMaterialCastFlags::AllowTruncate);
-	const bool bAllowAppendZeroes = EnumHasAnyFlags(Flags, EMaterialCastFlags::AllowAppendZeroes);
-	bool bReplicateScalar = EnumHasAnyFlags(Flags, EMaterialCastFlags::ReplicateScalar);
-
 	if (SourceType == DestType)
 	{
 		return Code;
 	}
-	else if (IsFloatNumericType(SourceType) && IsFloatNumericType(DestType))
+	
+	const bool bAllowTruncate = EnumHasAnyFlags(Flags, EMaterialCastFlags::AllowTruncate);
+	const bool bAllowAppendZeroes = EnumHasAnyFlags(Flags, EMaterialCastFlags::AllowAppendZeroes);
+	const bool bAllowReplicateScalar = EnumHasAnyFlags(Flags, EMaterialCastFlags::ReplicateScalar);
+	const bool bAllowInteger = EnumHasAnyFlags(Flags, EMaterialCastFlags::AllowInteger);
+
+	const EMaterialValueType AllowedTypes = EMaterialValueType(MCT_Float | MCT_LWCType | (bAllowInteger ? MCT_UInt : 0));
+
+	if (IsMaterialValueType(SourceType, AllowedTypes) && IsMaterialValueType(DestType, AllowedTypes))
 	{
 		const uint32 NumSourceComponents = GetNumComponents(SourceType);
 		const uint32 NumDestComponents = GetNumComponents(DestType);
-		if (NumSourceComponents != 1)
-		{
-			bReplicateScalar = false;
-		}
+		const bool bReplicateScalar = bAllowReplicateScalar && (NumSourceComponents == 1);
 		if (!bReplicateScalar && !bAllowAppendZeroes && NumDestComponents > NumSourceComponents)
 		{
 			Errorf(TEXT("Cannot cast from smaller type %s to larger type %s."), DescribeType(SourceType), DescribeType(DestType));
@@ -4844,11 +4845,9 @@ FString FHLSLMaterialTranslator::CastValue(const FString& Code, EMaterialValueTy
 		check(NumComponents == NumDestComponents);
 		return Result;
 	}
-	else
-	{
-		Errorf(TEXT("Cannot cast between non-numeric types %s to %s."), DescribeType(SourceType), DescribeType(DestType));
-		return FString();
-	}
+
+	Errorf(TEXT("Cannot cast between non-numeric types %s to %s."), DescribeType(SourceType), DescribeType(DestType));
+	return FString();
 }
 
 /** Pushes a function onto the compiler's function stack, which indicates that compilation is entering a function. */
@@ -5603,6 +5602,30 @@ int32 FHLSLMaterialTranslator::Fmod(int32 A, int32 B)
 	}
 }
 
+int32 FHLSLMaterialTranslator::Modulo(int32 A, int32 B)
+{
+	if ((A == INDEX_NONE) || (B == INDEX_NONE))
+	{
+		return INDEX_NONE;
+	}
+
+	if (GetParameterUniformExpression(A) && GetParameterUniformExpression(B))
+	{
+		FAddUniformExpressionScope Scope(this);
+		return AddUniformExpression(Scope, new FMaterialUniformExpressionModulo(GetParameterUniformExpression(A),GetParameterUniformExpression(B)),
+			GetParameterType(A),TEXT("(%s %% %s)"),*GetParameterCode(A),*CoerceParameter(B,GetParameterType(A), EMaterialCastFlags::AllowInteger));
+	}
+
+	const FDerivInfo& BDerivInfo = GetDerivInfo(B, true);
+	if (IsAnalyticDerivEnabled() && BDerivInfo.DerivativeStatus == EDerivativeStatus::Zero)
+	{
+		// Analytic derivatives only make sense when RHS derivatives are zero.
+		return DerivativeAutogen.GenerateExpressionFunc2(*this, FMaterialDerivativeAutogen::EFunc2::Modulo, A, B);
+	}
+
+	return AddCodeChunk(GetParameterType(A), TEXT("(%s %% %s)"), *GetParameterCode(A), *CoerceParameter(B, GetParameterType(A), EMaterialCastFlags::AllowInteger));
+}
+
 /**
 * Creates the new shader code chunk needed for the Abs expression
 *
diff --git a/Engine/Source/Runtime/Engine/Private/Materials/HLSLMaterialTranslator.h b/Engine/Source/Runtime/Engine/Private/Materials/HLSLMaterialTranslator.h
--- a/Engine/Source/Runtime/Engine/Private/Materials/HLSLMaterialTranslator.h
+++ b/Engine/Source/Runtime/Engine/Private/Materials/HLSLMaterialTranslator.h
@@ -228,6 +228,7 @@ enum class EMaterialCastFlags : uint32
 	ReplicateScalar = (1u << 0),
 	AllowTruncate = (1u << 1),
 	AllowAppendZeroes = (1u << 2),
+	AllowInteger = (1u << 3),
 
 	ValidCast = ReplicateScalar | AllowTruncate,
 };
@@ -705,8 +706,8 @@ public:
 	FString CastValue(const FString& Code, EMaterialValueType SourceType, EMaterialValueType DestType, EMaterialCastFlags Flags);
 
 	// CoerceParameter
-	FString CoerceParameter(int32 Index, EMaterialValueType DestType);
-	FString CoerceValue(const FString& Code, EMaterialValueType SourceType, EMaterialValueType DestType);
+	FString CoerceParameter(int32 Index, EMaterialValueType DestType, EMaterialCastFlags AdditionalCastFlags = EMaterialCastFlags::None);
+	FString CoerceValue(const FString& Code, EMaterialValueType SourceType, EMaterialValueType DestType, EMaterialCastFlags AdditionalCastFlags = EMaterialCastFlags::None);
 
 	int32 CastToNonLWCIfDisabled(int32 Code);
 
@@ -925,6 +926,7 @@ protected:
 	virtual int32 Sign(int32 X) override;
 	virtual int32 Frac(int32 X) override;
 	virtual int32 Fmod(int32 A, int32 B) override;
+	virtual int32 Modulo(int32 A, int32 B) override;
 
 	/**
 	* Creates the new shader code chunk needed for the Abs expression
diff --git a/Engine/Source/Runtime/Engine/Private/Materials/MaterialExpressionHLSL.cpp b/Engine/Source/Runtime/Engine/Private/Materials/MaterialExpressionHLSL.cpp
--- a/Engine/Source/Runtime/Engine/Private/Materials/MaterialExpressionHLSL.cpp
+++ b/Engine/Source/Runtime/Engine/Private/Materials/MaterialExpressionHLSL.cpp
@@ -118,6 +118,7 @@
 #include "Materials/MaterialExpressionMaterialProxyReplace.h"
 #include "Materials/MaterialExpressionMax.h"
 #include "Materials/MaterialExpressionMin.h"
+#include "Materials/MaterialExpressionModulo.h"
 #include "Materials/MaterialExpressionMultiply.h"
 #include "Materials/MaterialExpressionNamedReroute.h"
 #include "Materials/MaterialExpressionNaniteReplace.h"
@@ -2894,6 +2895,18 @@ bool UMaterialExpressionFmod::GenerateHLSLExpression(FMaterialHLSLGenerator& Gen
 	return true;
 }
 
+bool UMaterialExpressionModulo::GenerateHLSLExpression(FMaterialHLSLGenerator& Generator, UE::HLSLTree::FScope& Scope, int32 OutputIndex, UE::HLSLTree::FExpression const*& OutExpression) const
+{
+	const UE::HLSLTree::FExpression* Lhs = A.AcquireHLSLExpression(Generator, Scope);
+	const UE::HLSLTree::FExpression* Rhs = B.AcquireHLSLExpression(Generator, Scope);
+	if (!Lhs || !Rhs)
+	{
+		return false;
+	}
+	OutExpression = Generator.GetTree().NewModulo(Lhs, Rhs);
+	return true;
+}
+
 bool UMaterialExpressionPower::GenerateHLSLExpression(FMaterialHLSLGenerator& Generator, UE::HLSLTree::FScope& Scope, int32 OutputIndex, UE::HLSLTree::FExpression const*& OutExpression) const
 {
 	using namespace UE::HLSLTree;
diff --git a/Engine/Source/Runtime/Engine/Private/Materials/MaterialExpressions.cpp b/Engine/Source/Runtime/Engine/Private/Materials/MaterialExpressions.cpp
--- a/Engine/Source/Runtime/Engine/Private/Materials/MaterialExpressions.cpp
+++ b/Engine/Source/Runtime/Engine/Private/Materials/MaterialExpressions.cpp
@@ -142,6 +142,7 @@
 #include "Materials/MaterialExpressionMaterialAttributeLayers.h"
 #include "Materials/MaterialExpressionMaterialProxyReplace.h"
 #include "Materials/MaterialExpressionMin.h"
+#include "Materials/MaterialExpressionModulo.h"
 #include "Materials/MaterialExpressionMultiply.h"
 #include "Materials/MaterialExpressionNaniteReplace.h"
 #include "Materials/MaterialExpressionNoise.h"
@@ -8534,6 +8535,64 @@ void UMaterialExpressionFmod::GetCaption(TArray<FString>& OutCaptions) const
 }
 #endif // WITH_EDITOR
 
+//
+//	UMaterialExpressionModulo
+//
+
+UMaterialExpressionModulo::UMaterialExpressionModulo(const FObjectInitializer& ObjectInitializer)
+	: Super(ObjectInitializer)
+{
+	// Structure to hold one-time initialization
+	struct FConstructorStatics
+	{
+		FText NAME_Math;
+		FConstructorStatics()
+			: NAME_Math(LOCTEXT( "Math", "Math" ))
+		{
+		}
+	};
+	static FConstructorStatics ConstructorStatics;
+
+#if WITH_EDITORONLY_DATA
+	MenuCategories.Add(ConstructorStatics.NAME_Math);
+#endif
+}
+
+#if WITH_EDITOR
+int32 UMaterialExpressionModulo::Compile(class FMaterialCompiler* Compiler, int32 OutputIndex)
+{
+	if (!A.GetTracedInput().Expression)
+	{
+		return Compiler->Errorf(TEXT("Missing Modulo input A"));
+	}
+	if (!B.GetTracedInput().Expression)
+	{
+		return Compiler->Errorf(TEXT("Missing Modulo input B"));
+	}
+	return Compiler->Modulo(A.Compile(Compiler), B.Compile(Compiler));
+}
+
+uint32 UMaterialExpressionModulo::GetInputType(int32 InputIndex)
+{
+	return MCT_Float | MCT_LWCType | MCT_UInt;
+}
+
+uint32 UMaterialExpressionModulo::GetOutputType(int32 OutputIndex)
+{
+	return MCT_Float | MCT_LWCType | MCT_UInt;
+}
+
+void UMaterialExpressionModulo::GetCaption(TArray<FString>& OutCaptions) const
+{
+	OutCaptions.Add(TEXT("Modulo"));
+}
+
+FText UMaterialExpressionModulo::GetKeywords() const
+{
+	return FText::FromString(TEXT("%"));
+}
+#endif // WITH_EDITOR
+
 UMaterialExpressionFrac::UMaterialExpressionFrac(const FObjectInitializer& ObjectInitializer)
 	: Super(ObjectInitializer)
 {
diff --git a/Engine/Source/Runtime/Engine/Private/Materials/MaterialUniformExpressions.cpp b/Engine/Source/Runtime/Engine/Private/Materials/MaterialUniformExpressions.cpp
--- a/Engine/Source/Runtime/Engine/Private/Materials/MaterialUniformExpressions.cpp
+++ b/Engine/Source/Runtime/Engine/Private/Materials/MaterialUniformExpressions.cpp
@@ -1808,6 +1808,7 @@ IMPLEMENT_MATERIALUNIFORMEXPRESSION_TYPE(FMaterialUniformExpressionFloor);
 IMPLEMENT_MATERIALUNIFORMEXPRESSION_TYPE(FMaterialUniformExpressionCeil);
 IMPLEMENT_MATERIALUNIFORMEXPRESSION_TYPE(FMaterialUniformExpressionFrac);
 IMPLEMENT_MATERIALUNIFORMEXPRESSION_TYPE(FMaterialUniformExpressionFmod);
+IMPLEMENT_MATERIALUNIFORMEXPRESSION_TYPE(FMaterialUniformExpressionModulo);
 IMPLEMENT_MATERIALUNIFORMEXPRESSION_TYPE(FMaterialUniformExpressionAbs);
 IMPLEMENT_MATERIALUNIFORMEXPRESSION_TYPE(FMaterialUniformExpressionTextureProperty);
 IMPLEMENT_MATERIALUNIFORMEXPRESSION_TYPE(FMaterialUniformExpressionTrigMath);
diff --git a/Engine/Source/Runtime/Engine/Private/Materials/MaterialUniformExpressions.h b/Engine/Source/Runtime/Engine/Private/Materials/MaterialUniformExpressions.h
--- a/Engine/Source/Runtime/Engine/Private/Materials/MaterialUniformExpressions.h
+++ b/Engine/Source/Runtime/Engine/Private/Materials/MaterialUniformExpressions.h
@@ -1702,6 +1702,50 @@ private:
 	TRefCountPtr<FMaterialUniformExpression> B;
 };
 
+/**
+ */
+class FMaterialUniformExpressionModulo : public FMaterialUniformExpression
+{
+	DECLARE_MATERIALUNIFORMEXPRESSION_TYPE(FMaterialUniformExpressionModulo);
+public:
+
+	FMaterialUniformExpressionModulo() {}
+	FMaterialUniformExpressionModulo(FMaterialUniformExpression* InA,FMaterialUniformExpression* InB):
+		A(InA),
+		B(InB)
+	{}
+
+	// FMaterialUniformExpression interface.
+	virtual void WriteNumberOpcodes(UE::Shader::FPreshaderData& OutData) const override
+	{
+		A->WriteNumberOpcodes(OutData);
+		B->WriteNumberOpcodes(OutData);
+		OutData.WriteOpcode(UE::Shader::EPreshaderOpcode::Modulo);
+	}
+	virtual bool IsConstant() const
+	{
+		return A->IsConstant() && B->IsConstant();
+	}
+	virtual bool IsIdentical(const FMaterialUniformExpression* OtherExpression) const
+	{
+		if (GetType() != OtherExpression->GetType())
+		{
+			return false;
+		}
+		const FMaterialUniformExpressionModulo* OtherMax = (const FMaterialUniformExpressionModulo*)OtherExpression;
+		return A->IsIdentical(OtherMax->A) && B->IsIdentical(OtherMax->B);
+	}
+
+	virtual TArrayView<const FMaterialUniformExpression*> GetChildren() const override
+	{
+		return TArrayView<const FMaterialUniformExpression*>((const FMaterialUniformExpression**)&A, 2);
+	}
+
+private:
+	TRefCountPtr<FMaterialUniformExpression> A;
+	TRefCountPtr<FMaterialUniformExpression> B;
+};
+
 /**
  * Absolute value evaluator for a given input expression
  */
diff --git a/Engine/Source/Runtime/Engine/Private/Shader/Preshader.cpp b/Engine/Source/Runtime/Engine/Private/Shader/Preshader.cpp
--- a/Engine/Source/Runtime/Engine/Private/Shader/Preshader.cpp
+++ b/Engine/Source/Runtime/Engine/Private/Shader/Preshader.cpp
@@ -1344,6 +1344,7 @@ FPreshaderValue EvaluatePreshader(const FUniformExpressionSet* UniformExpression
 		case EPreshaderOpcode::LessEqual: EvaluateBinaryOp(Stack, LessEqual); break;
 		case EPreshaderOpcode::GreaterEqual: EvaluateBinaryOp(Stack, GreaterEqual); break;
 		case EPreshaderOpcode::Fmod: EvaluateBinaryOpInPlace(Stack, FmodInPlace, Fmod); break;
+		case EPreshaderOpcode::Modulo: EvaluateBinaryOpInPlace(Stack, ModuloInPlace, Modulo); break;
 		case EPreshaderOpcode::Min: EvaluateBinaryOpInPlace(Stack, MinInPlace, Min); break;
 		case EPreshaderOpcode::Max: EvaluateBinaryOpInPlace(Stack, MaxInPlace, Max); break;
 		case EPreshaderOpcode::Clamp: EvaluateTernaryOp(Stack, Clamp); break;
@@ -1434,6 +1435,7 @@ static void EvaluatePreshaderDebug(const FUniformExpressionSet* UniformExpressio
 		case EPreshaderOpcode::LessEqual: EvaluateBinaryOp(Stack, LessEqual); break;
 		case EPreshaderOpcode::GreaterEqual: EvaluateBinaryOp(Stack, GreaterEqual); break;
 		case EPreshaderOpcode::Fmod: EvaluateBinaryOpInPlace(Stack, FmodInPlace, Fmod); break;
+		case EPreshaderOpcode::Modulo: EvaluateBinaryOpInPlace(Stack, ModuloInPlace, Modulo); break;
 		case EPreshaderOpcode::Min: EvaluateBinaryOpInPlace(Stack, MinInPlace, Min); break;
 		case EPreshaderOpcode::Max: EvaluateBinaryOpInPlace(Stack, MaxInPlace, Max); break;
 		case EPreshaderOpcode::Clamp: EvaluateTernaryOp(Stack, Clamp); break;
diff --git a/Engine/Source/Runtime/Engine/Private/Shader/ShaderValue.cpp b/Engine/Source/Runtime/Engine/Private/Shader/ShaderValue.cpp
--- a/Engine/Source/Runtime/Engine/Private/Shader/ShaderValue.cpp
+++ b/Engine/Source/Runtime/Engine/Private/Shader/ShaderValue.cpp
@@ -999,6 +999,17 @@ struct FOpDiv : public FOpBase { template<typename T> T operator()(T Lhs, T Rhs)
 struct FOpMin : public FOpBase { template<typename T> T operator()(T Lhs, T Rhs) const { return FMath::Min(Lhs, Rhs); } };
 struct FOpMax : public FOpBase { template<typename T> T operator()(T Lhs, T Rhs) const { return FMath::Max(Lhs, Rhs); } };
 struct FOpFmod : public FOpBaseNoInt { template<typename T> T operator()(T Lhs, T Rhs) const { return FMath::Fmod(Lhs, GetSafeDivisor(Rhs)); } };
+struct FOpModulo : public FOpBase
+{
+	static float  ModuloImpl(float Lhs, float Rhs)   { return FMath::Fmod(Lhs, Rhs); }
+	static double ModuloImpl(double Lhs, double Rhs) { return FMath::Fmod(Lhs, Rhs); }
+	static int32  ModuloImpl(int32 Lhs, int32 Rhs)   { return Lhs % Rhs;             }
+
+	template<typename T> T operator()(T Lhs, T Rhs) const
+	{
+		return ModuloImpl(Lhs, GetSafeDivisor(Rhs));
+	}
+};
 struct FOpAtan2 : public FOpBaseNoInt { template<typename T> T operator()(T Lhs, T Rhs) const { return FMath::Atan2(Lhs, Rhs); } };
 struct FOpLess : public FOpBase { template<typename T> bool operator()(T Lhs, T Rhs) const { return Lhs < Rhs; } };
 struct FOpGreater : public FOpBase { template<typename T> bool operator()(T Lhs, T Rhs) const { return Lhs > Rhs; } };
@@ -1545,6 +1556,11 @@ FValue Fmod(const FValue& Lhs, const FValue& Rhs)
 	return Private::BinaryOp(Private::FOpFmod(), Lhs, Rhs);
 }
 
+FValue Modulo(const FValue& Lhs, const FValue& Rhs)
+{
+	return Private::BinaryOp(Private::FOpModulo(), Lhs, Rhs);
+}
+
 FValue Atan2(const FValue& Lhs, const FValue& Rhs)
 {
 	return Private::BinaryOp(Private::FOpAtan2(), Lhs, Rhs);
@@ -1893,6 +1909,11 @@ EValueType FmodInPlace(EValueType LhsType, EValueType RhsType, TArrayView<FValue
 	return Private::BinaryOpInPlace(Private::FOpFmod(), LhsType, RhsType, Component, OutComponentsConsumed);
 }
 
+EValueType ModuloInPlace(EValueType LhsType, EValueType RhsType, TArrayView<FValueComponent> Component, int32& OutComponentsConsumed)
+{
+	return Private::BinaryOpInPlace(Private::FOpModulo(), LhsType, RhsType, Component, OutComponentsConsumed);
+}
+
 EValueType Atan2InPlace(EValueType LhsType, EValueType RhsType, TArrayView<FValueComponent> Component, int32& OutComponentsConsumed)
 {
 	return Private::BinaryOpInPlace(Private::FOpAtan2(), LhsType, RhsType, Component, OutComponentsConsumed);
diff --git a/Engine/Source/Runtime/Engine/Public/HLSLTree/HLSLTree.h b/Engine/Source/Runtime/Engine/Public/HLSLTree/HLSLTree.h
--- a/Engine/Source/Runtime/Engine/Public/HLSLTree/HLSLTree.h
+++ b/Engine/Source/Runtime/Engine/Public/HLSLTree/HLSLTree.h
@@ -664,6 +664,7 @@ public:
 	const FExpression* NewMul(const FExpression* Lhs, const FExpression* Rhs) { return NewBinaryOp(EOperation::Mul, Lhs, Rhs); }
 	const FExpression* NewDiv(const FExpression* Lhs, const FExpression* Rhs) { return NewBinaryOp(EOperation::Div, Lhs, Rhs); }
 	const FExpression* NewFmod(const FExpression* Lhs, const FExpression* Rhs) { return NewBinaryOp(EOperation::Fmod, Lhs, Rhs); }
+	const FExpression* NewModulo(const FExpression* Lhs, const FExpression* Rhs) { return NewBinaryOp(EOperation::Modulo, Lhs, Rhs); }
 	const FExpression* NewStep(const FExpression* Lhs, const FExpression* Rhs) { return NewBinaryOp(EOperation::Step, Lhs, Rhs); }
 	const FExpression* NewPowClamped(const FExpression* Lhs, const FExpression* Rhs) { return NewBinaryOp(EOperation::PowPositiveClamped, Lhs, Rhs); }
 	const FExpression* NewMin(const FExpression* Lhs, const FExpression* Rhs) { return NewBinaryOp(EOperation::Min, Lhs, Rhs); }
diff --git a/Engine/Source/Runtime/Engine/Public/HLSLTree/HLSLTreeTypes.h b/Engine/Source/Runtime/Engine/Public/HLSLTree/HLSLTreeTypes.h
--- a/Engine/Source/Runtime/Engine/Public/HLSLTree/HLSLTreeTypes.h
+++ b/Engine/Source/Runtime/Engine/Public/HLSLTree/HLSLTreeTypes.h
@@ -112,6 +112,7 @@ enum class EOperation : uint8
 	Mul,
 	Div,
 	Fmod,
+	Modulo,
 	Step,
 	PowPositiveClamped,
 	Atan2,
diff --git a/Engine/Source/Runtime/Engine/Public/MaterialCompiler.h b/Engine/Source/Runtime/Engine/Public/MaterialCompiler.h
--- a/Engine/Source/Runtime/Engine/Public/MaterialCompiler.h
+++ b/Engine/Source/Runtime/Engine/Public/MaterialCompiler.h
@@ -227,6 +227,7 @@ public:
 	virtual int32 Sign(int32 X) = 0;
 	virtual int32 Frac(int32 X) = 0;
 	virtual int32 Fmod(int32 A, int32 B) = 0;
+	virtual int32 Modulo(int32 A, int32 B) = 0;
 	virtual int32 Abs(int32 X) = 0;
 
 	virtual int32 ReflectionVector() = 0;
@@ -836,6 +837,7 @@ public:
 	virtual int32 Sign(int32 X) override { return Compiler->Sign(X); }
 	virtual int32 Frac(int32 X) override { return Compiler->Frac(X); }
 	virtual int32 Fmod(int32 A, int32 B) override { return Compiler->Fmod(A, B); }
+	virtual int32 Modulo(int32 A, int32 B) override { return Compiler->Modulo(A, B); }
 	virtual int32 Abs(int32 X) override { return Compiler->Abs(X); }
 
 	virtual int32 ReflectionVector() override { return Compiler->ReflectionVector(); }
diff --git a/Engine/Source/Runtime/Engine/Public/MaterialShared.h b/Engine/Source/Runtime/Engine/Public/MaterialShared.h
--- a/Engine/Source/Runtime/Engine/Public/MaterialShared.h
+++ b/Engine/Source/Runtime/Engine/Public/MaterialShared.h
@@ -241,25 +241,30 @@ inline uint32 GetNumComponents(EMaterialValueType Type)
 	}
 }
 
+inline bool IsMaterialValueType(EMaterialValueType InType, EMaterialValueType InCompare)
+{
+	return (InType & InCompare) != 0;
+}
+
 inline bool IsLWCType(EMaterialValueType InType)
 {
-	return (InType & MCT_LWCType);
+	return IsMaterialValueType(InType, MCT_LWCType);
 }
 
 inline bool IsFloatNumericType(EMaterialValueType InType)
 {
-	return (InType & MCT_Float) || IsLWCType(InType);
+	return IsMaterialValueType(InType, EMaterialValueType(MCT_Float | MCT_LWCType));
 }
 
 inline bool IsUIntNumericType(EMaterialValueType InType)
 {
-	return (InType & MCT_UInt);
+	return IsMaterialValueType(InType, MCT_UInt);
 }
 
 inline bool IsNumericType(EMaterialValueType InType)
 {
 	// 'ShadingModel' is considered an 'int' 
-	return IsFloatNumericType(InType) || InType == MCT_ShadingModel || IsUIntNumericType(InType);
+	return IsMaterialValueType(InType, EMaterialValueType(MCT_Float | MCT_LWCType | MCT_UInt | MCT_ShadingModel));
 }
 
 inline EMaterialValueType MakeNonLWCType(EMaterialValueType Type)
diff --git a/Engine/Source/Runtime/Engine/Public/Shader/Preshader.h b/Engine/Source/Runtime/Engine/Public/Shader/Preshader.h
--- a/Engine/Source/Runtime/Engine/Public/Shader/Preshader.h
+++ b/Engine/Source/Runtime/Engine/Public/Shader/Preshader.h
@@ -27,6 +27,7 @@ enum class EPreshaderOpcode : uint8
 	Mul,
 	Div,
 	Fmod,
+	Modulo,
 	Min,
 	Max,
 	Clamp,
diff --git a/Engine/Source/Runtime/Engine/Public/Shader/ShaderTypes.h b/Engine/Source/Runtime/Engine/Public/Shader/ShaderTypes.h
--- a/Engine/Source/Runtime/Engine/Public/Shader/ShaderTypes.h
+++ b/Engine/Source/Runtime/Engine/Public/Shader/ShaderTypes.h
@@ -569,6 +569,7 @@ ENGINE_API FValue Min(const FValue& Lhs, const FValue& Rhs);
 ENGINE_API FValue Max(const FValue& Lhs, const FValue& Rhs);
 ENGINE_API FValue Clamp(const FValue& Value, const FValue& Low, const FValue& High);
 ENGINE_API FValue Fmod(const FValue& Lhs, const FValue& Rhs);
+ENGINE_API FValue Modulo(const FValue& Lhs, const FValue& Rhs);
 ENGINE_API FValue Atan2(const FValue& Lhs, const FValue& Rhs);
 ENGINE_API FValue Dot(const FValue& Lhs, const FValue& Rhs);
 ENGINE_API FValue Cross(const FValue& Lhs, const FValue& Rhs);
@@ -609,6 +610,7 @@ ENGINE_API EValueType DivInPlace(EValueType LhsType, EValueType RhsType, TArrayV
 ENGINE_API EValueType MinInPlace(EValueType LhsType, EValueType RhsType, TArrayView<FValueComponent> Component, int32& OutComponentsConsumed);
 ENGINE_API EValueType MaxInPlace(EValueType LhsType, EValueType RhsType, TArrayView<FValueComponent> Component, int32& OutComponentsConsumed);
 ENGINE_API EValueType FmodInPlace(EValueType LhsType, EValueType RhsType, TArrayView<FValueComponent> Component, int32& OutComponentsConsumed);
+ENGINE_API EValueType ModuloInPlace(EValueType LhsType, EValueType RhsType, TArrayView<FValueComponent> Component, int32& OutComponentsConsumed);
 ENGINE_API EValueType Atan2InPlace(EValueType LhsType, EValueType RhsType, TArrayView<FValueComponent> Component, int32& OutComponentsConsumed);
 ENGINE_API EValueType AppendInPlace(EValueType LhsType, EValueType RhsType, TArrayView<FValueComponent> Component, int32& OutComponentsConsumed);
 
