Add support for analytic derivatives to local position node

diff --git a/Engine/Shaders/Private/MaterialTemplate.ush b/Engine/Shaders/Private/MaterialTemplate.ush
--- a/Engine/Shaders/Private/MaterialTemplate.ush
+++ b/Engine/Shaders/Private/MaterialTemplate.ush
@@ -972,6 +972,7 @@ FWSVector3 GetPrevParticleWorldPosition(FMaterialPixelParameters Parameters)
 
 float3 GetPositionPrimitiveSpace(FMaterialPixelParameters Parameters)
 {
+	// See also: derivative computation of local position in material translator.
 	return DFFastMultiplyDemote(Parameters.AbsoluteWorldPosition, GetPrimitiveData(Parameters).WorldToLocal);
 }
 
@@ -996,6 +997,7 @@ float3 GetPositionPrimitiveSpace_NoMaterialOffsets(FMaterialPixelParameters Para
 
 float3 GetPositionInstanceSpace(FMaterialPixelParameters Parameters)
 {
+	// See also: derivative computation of local position in material translator.
 	return DFFastMultiplyDemote(Parameters.AbsoluteWorldPosition, GetWorldToInstanceDF(Parameters));
 }
 
@@ -1729,13 +1731,13 @@ FMaterialLWCData MakeMaterialLWCData(FMaterialVertexParameters Parameters)
 	Result.AbsoluteWorldPosition = WSSubtract(Parameters.WorldPosition, Result.PreViewTranslation);
 
 	Result.InstanceToWorld = DFToWS(GetInstanceToWorldDF(Parameters));
-	Result.WorldToInstance = DFFastToWS(GetWorldToInstanceDF(Parameters));
+	Result.WorldToInstance = DFToWS(GetWorldToInstanceDF(Parameters));
 	Result.PreviousInstanceToWorld = DFToWS(GetPrevInstanceToWorldDF(Parameters));
 
 	Result.LocalToWorld = DFToWS(GetLocalToWorldDF(Parameters));
-	Result.WorldToLocal = DFFastToWS(GetWorldToLocalDF(Parameters));
+	Result.WorldToLocal = DFToWS(GetWorldToLocalDF(Parameters));
 	Result.PreviousLocalToWorld = DFToWS(GetPrevLocalToWorldDF(Parameters));
-	Result.PreviousWorldToLocal = DFFastToWS(GetPrevWorldToLocalDF(Parameters));
+	Result.PreviousWorldToLocal = DFToWS(GetPrevWorldToLocalDF(Parameters));
 
 #if DECAL_PRIMITIVE
 	/*
@@ -1751,7 +1753,7 @@ FMaterialLWCData MakeMaterialLWCData(FMaterialVertexParameters Parameters)
 	
 #if USE_INSTANCING || USE_INSTANCE_CULLING || IS_MESHPARTICLE_FACTORY
 #if MATERIAL_LWC_ENABLED
-	Result.ObjectWorldPosition = WSGetOrigin(DFFastToWS(Parameters.InstanceLocalToWorld));
+	Result.ObjectWorldPosition = WSGetOrigin(DFToWS(Parameters.InstanceLocalToWorld));
 #else //MATERIAL_LWC_ENABLED
 	Result.ObjectWorldPosition = WSPromote(Parameters.InstanceLocalToWorld.PostTranslation + Parameters.InstanceLocalToWorld.M[3].xyz);
 #endif //MATERIAL_LWC_ENABLED
diff --git a/Engine/Source/Runtime/Engine/Private/Materials/HLSLMaterialTranslator.cpp b/Engine/Source/Runtime/Engine/Private/Materials/HLSLMaterialTranslator.cpp
--- a/Engine/Source/Runtime/Engine/Private/Materials/HLSLMaterialTranslator.cpp
+++ b/Engine/Source/Runtime/Engine/Private/Materials/HLSLMaterialTranslator.cpp
@@ -6024,7 +6024,27 @@ int32 FHLSLMaterialTranslator::LocalPosition(EPositionIncludedOffsets IncludedOf
 	};
 	FString FiniteCode = FString::Format(TEXT("Get{0}Position{1}Space{2}(Parameters)"), FormatArgs);
 
-	return AddInlinedCodeChunk(MCT_Float3, *FiniteCode);
+	int32 Result = INDEX_NONE;
+	if (IsAnalyticDerivEnabled())
+	{
+		// Generate derivative based on WorldPosition_DDX
+		int32 WorldPositionCode = WorldPosition(IncludedOffsets == EPositionIncludedOffsets::IncludeOffsets ? WPT_CameraRelative : WPT_CameraRelativeNoOffsets);
+		int32 TransformedPositionCode = TransformPosition(MCB_TranslatedWorld, OriginType == ELocalPositionOrigin::Instance ? MCB_Instance : MCB_Local, WorldPositionCode);
+		FString TransformedPositionAnalyticCode = GetParameterCodeDeriv(TransformedPositionCode, CompiledPDV_Analytic);
+
+		FString AnalyticCode = DerivativeAutogen.ConstructDeriv(
+			*FiniteCode, 
+			*(TransformedPositionAnalyticCode + TEXT(".Ddx")), 
+			*(TransformedPositionAnalyticCode + TEXT(".Ddy")), 
+			GetDerivType(MCT_Float3));
+		Result = AddCodeChunkInnerDeriv(*FiniteCode, *AnalyticCode, MCT_Float3, false, EDerivativeStatus::Valid);
+	}
+	else
+	{
+		Result = AddInlinedCodeChunk(MCT_Float3, *FiniteCode);
+	}
+
+	return Result;
 }
 
 int32 FHLSLMaterialTranslator::WorldPosition(EWorldPositionIncludedOffsets WorldPositionIncludedOffsets)
@@ -10142,6 +10162,18 @@ static FString MultiplyMatrix(const TCHAR* Vector, const TCHAR* Matrix, int AWCo
 	}
 }
 
+static FString MultiplyTranslatedMatrix(const TCHAR* Vector, const TCHAR* MatrixPreTranslation, int AWComponent)
+{
+	if (AWComponent)
+	{
+		return FString::Printf(TEXT("mul(MaterialFloat4(%s, 1.0f), DFFastToTranslatedWorld(%s, ResolvedView.PreViewTranslation)).xyz"), Vector, MatrixPreTranslation);
+	}
+	else
+	{
+		return FString::Printf(TEXT("mul(%s, DFToFloat3x3(%s))"), Vector, MatrixPreTranslation);
+	}
+}
+
 static FString MultiplyTransposeMatrix(const TCHAR* Matrix, const TCHAR* Vector, int AWComponent)
 {
 	if (AWComponent)
@@ -10244,8 +10276,8 @@ int32 FHLSLMaterialTranslator::TransformBase(EMaterialCommonBasis SourceCoordBas
 			{
 				if (bIsPositionTransform)
 				{
-					CodeStr = MultiplyMatrix(TEXT("<A>"), TEXT("DFFastToTranslatedWorld(Get<PREV>LocalToWorldDF(Parameters), ResolvedView.PreViewTranslation)"), AWComponent);
-					CodeDerivStr = MultiplyMatrix(TEXT("<A>"), TEXT("DFFastToTranslatedWorld(Get<PREV>LocalToWorldDF(Parameters), ResolvedView.PreViewTranslation)"), 0);
+					CodeStr = MultiplyTranslatedMatrix(TEXT("<A>"), TEXT("Get<PREV>LocalToWorldDF(Parameters)"), AWComponent);
+					CodeDerivStr = MultiplyTranslatedMatrix(TEXT("<A>"), TEXT("Get<PREV>LocalToWorldDF(Parameters)"), 0);
 				}
 			}
 			// else use MCB_World as intermediary basis
@@ -10294,19 +10326,19 @@ int32 FHLSLMaterialTranslator::TransformBase(EMaterialCommonBasis SourceCoordBas
 				// TODO: inconsistent with TransformLocal<TO>World with instancing
 				// We have explicit options for "local" and "instance" spaces, but then GetLocalToWorld returns instance space, while GetWorldToLocal always returns primitive space. 
 				// It's inconsistent, but replacing either will break existing materials.
-				CodeStr = MultiplyMatrix(TEXT("<A>"), TEXT("DFFastToTranslatedWorld(Get<PREV>WorldToLocalDF(Parameters), ResolvedView.PreViewTranslation)"), AWComponent);
-				CodeDerivStr = MultiplyMatrix(TEXT("<A>"), TEXT("DFFastToTranslatedWorld(Get<PREV>WorldToLocalDF(Parameters), ResolvedView.PreViewTranslation)"), 0);
+				CodeStr = MultiplyTranslatedMatrix(TEXT("<A>"), TEXT("Get<PREV>WorldToLocalDF(Parameters)"), AWComponent);
+				CodeDerivStr = MultiplyTranslatedMatrix(TEXT("<A>"), TEXT("Get<PREV>WorldToLocalDF(Parameters)"), 0);
 			}
 			else if (DestCoordBasis == MCB_MeshParticle)
 			{
-				CodeStr = MultiplyMatrix(TEXT("<A>"), TEXT("DFFastToTranslatedWorld(Parameters.Particle.WorldToParticle, ResolvedView.PreViewTranslation)"), AWComponent);
-				CodeDerivStr = MultiplyMatrix(TEXT("<A>"), TEXT("DFFastToTranslatedWorld(Parameters.Particle.WorldToParticle, ResolvedView.PreViewTranslation)"), 0);
+				CodeStr = MultiplyTranslatedMatrix(TEXT("<A>"), TEXT("Parameters.Particle.WorldToParticle"), AWComponent);
+				CodeDerivStr = MultiplyTranslatedMatrix(TEXT("<A>"), TEXT("Parameters.Particle.WorldToParticle"), 0);
 				bUsesParticleWorldToLocal = true;
 			}
 			else if (DestCoordBasis == MCB_Instance)
 			{
-				CodeStr = MultiplyMatrix(TEXT("<A>"), TEXT("DFFastToTranslatedWorld(GetWorldToInstanceDF(Parameters), ResolvedView.PreViewTranslation)"), AWComponent);
-				CodeDerivStr = MultiplyMatrix(TEXT("<A>"), TEXT("DFFastToTranslatedWorld(GetWorldToInstanceDF(Parameters), ResolvedView.PreViewTranslation)"), 0);
+				CodeStr = MultiplyTranslatedMatrix(TEXT("<A>"), TEXT("GetWorldToInstanceDF(Parameters)"), AWComponent);
+				CodeDerivStr = MultiplyTranslatedMatrix(TEXT("<A>"), TEXT("GetWorldToInstanceDF(Parameters)"), 0);
 				bUsesInstanceWorldToLocalPS |= ShaderFrequency == SF_Pixel;
 			}
 			// else use MCB_World as intermediary basis
@@ -10400,8 +10432,8 @@ int32 FHLSLMaterialTranslator::TransformBase(EMaterialCommonBasis SourceCoordBas
 			}
 			else if (DestCoordBasis == MCB_TranslatedWorld)
 			{
-				CodeStr = MultiplyMatrix(TEXT("<A>"), TEXT("DFFastToTranslatedWorld(Parameters.Particle.ParticleToWorld, ResolvedView.PreViewTranslation)"), AWComponent);
-				CodeDerivStr = MultiplyMatrix(TEXT("<A>"), TEXT("DFFastToTranslatedWorld(Parameters.Particle.ParticleToWorld, ResolvedView.PreViewTranslation)"), 0);
+				CodeStr = MultiplyTranslatedMatrix(TEXT("<A>"), TEXT("Parameters.Particle.ParticleToWorld"), AWComponent);
+				CodeDerivStr = MultiplyTranslatedMatrix(TEXT("<A>"), TEXT("Parameters.Particle.ParticleToWorld"), 0);
 				bUsesParticleLocalToWorld = true;
 			}
 			// use World as an intermediary base
@@ -10418,8 +10450,8 @@ int32 FHLSLMaterialTranslator::TransformBase(EMaterialCommonBasis SourceCoordBas
 			}
 			else if (DestCoordBasis == MCB_TranslatedWorld)
 			{
-				CodeStr = MultiplyMatrix(TEXT("<A>"), TEXT("DFFastToTranslatedWorld(Get<PREV>InstanceToWorldDF(Parameters), ResolvedView.PreViewTranslation)"), AWComponent);
-				CodeDerivStr = MultiplyMatrix(TEXT("<A>"), TEXT("DFFastToTranslatedWorld(Get<PREV>InstanceToWorldDF(Parameters), ResolvedView.PreViewTranslation)"), 0);
+				CodeStr = MultiplyTranslatedMatrix(TEXT("<A>"), TEXT("Get<PREV>InstanceToWorldDF(Parameters)"), AWComponent);
+				CodeDerivStr = MultiplyTranslatedMatrix(TEXT("<A>"), TEXT("Get<PREV>InstanceToWorldDF(Parameters)"), 0);
 				bUsesInstanceLocalToWorldPS |= ShaderFrequency == SF_Pixel;
 			}
 			// use World as an intermediary base
