Light function manipulating UVs are are now marked as incompatible with the light function atlas for the deferred lighting passes.
Artists can still force a light function to go the fast path for such passes after ticking the visual by using the new bCompatibleWithLightFunctionAtlas bool on material.

diff --git a/Engine/Shaders/Private/LightFunctionAtlas/LightFunctionAtlasRender.usf b/Engine/Shaders/Private/LightFunctionAtlas/LightFunctionAtlasRender.usf
--- a/Engine/Shaders/Private/LightFunctionAtlas/LightFunctionAtlasRender.usf
+++ b/Engine/Shaders/Private/LightFunctionAtlas/LightFunctionAtlasRender.usf
@@ -8,7 +8,7 @@
 #define LIGHT_FUNCTION_ATLAS 1
 
 // Used to trigger shader recompiles, guid should be regenerated on merge conflicts
-#pragma message("UESHADERMETADATA_VERSION 98DA9C13-D8B0-43F6-AD31-6198D1CE3A5C")
+#pragma message("UESHADERMETADATA_VERSION D13B08DE-2E66-447C-AF33-27D896D280CA")
 
 #include "../Common.ush"
 #include "/Engine/Generated/Material.ush"
diff --git a/Engine/Shaders/Private/LightFunctionPixelShader.usf b/Engine/Shaders/Private/LightFunctionPixelShader.usf
--- a/Engine/Shaders/Private/LightFunctionPixelShader.usf
+++ b/Engine/Shaders/Private/LightFunctionPixelShader.usf
@@ -10,7 +10,7 @@
 #include "DynamicLightingCommon.ush"
 
 // Used to trigger shader recompiles, guid should be regenerated on merge conflicts
-#pragma message("UESHADERMETADATA_VERSION EFF06B5F-BCC7-4DDE-AF88-28825B73EEEC")
+#pragma message("UESHADERMETADATA_VERSION 02A826D0-F3F6-4726-A187-762FFB6D17CA")
 
 // used like this: mul(float4(SvPosition.xyz, 1), SvPositionToLight);
 float4x4 SvPositionToLight;
diff --git a/Engine/Source/Runtime/Engine/Classes/Materials/Material.h b/Engine/Source/Runtime/Engine/Classes/Materials/Material.h
--- a/Engine/Source/Runtime/Engine/Classes/Materials/Material.h
+++ b/Engine/Source/Runtime/Engine/Classes/Materials/Material.h
@@ -811,6 +811,14 @@ public:
 	UPROPERTY()
 	uint8 bUsedWithUI_DEPRECATED : 1;
 
+	/** 
+	 * Indicates that this material is safe to use with the light function atlas:
+	 *  - texture coordinates are manipulated in a way that works.
+	 *  - world position and depth are not used in the graph or do work with the atlas.
+	 */
+	UPROPERTY(EditAnywhere, Category = LightFunctionMaterial)
+	uint8 bCompatibleWithLightFunctionAtlas : 1;
+
 	/** 
 	 * Whether to automatically set usage flags based on what the material is applied to in the editor.
 	 * It can be useful to disable this on a base material with many instances, where adding another usage flag accidentally (eg bUsedWithSkeletalMeshes) can add a lot of shader permutations.
diff --git a/Engine/Source/Runtime/Engine/Private/Materials/HLSLMaterialTranslator.cpp b/Engine/Source/Runtime/Engine/Private/Materials/HLSLMaterialTranslator.cpp
--- a/Engine/Source/Runtime/Engine/Private/Materials/HLSLMaterialTranslator.cpp
+++ b/Engine/Source/Runtime/Engine/Private/Materials/HLSLMaterialTranslator.cpp
@@ -658,6 +658,7 @@ FHLSLMaterialTranslator::FHLSLMaterialTranslator(FMaterial* InMaterial,
 ,	bUsesInstanceWorldToLocalPS(false)
 ,	bUsesPerInstanceRandomPS(false)
 ,	bUsesVertexPosition(false)
+,	bPotentiallyManipulateTexCoords(false)
 ,	bUsesTransformVector(false)
 ,	bCompilingPreviousFrame(false)
 ,	bOutputsBasePassVelocities(true)
@@ -1707,7 +1708,16 @@ void FHLSLMaterialTranslator::DoTranslate()
 	bUsesAnisotropy = IsMaterialPropertyUsed(MP_Anisotropy, Chunk[MP_Anisotropy], FLinearColor(0, 0, 0, 0), 1);
 	MaterialCompilationOutput.bUsesAnisotropy = bUsesAnisotropy;
 
-	MaterialCompilationOutput.bIsLightFunctionAtlasCompatible = !bUsesVertexPosition && !bUsesSceneDepth && !MaterialCompilationOutput.bNeedsSceneTextures;
+	/**
+	 * A material is not compatible with the atlas if 
+	 *  - it samples world position of depth since we do not have such data when rendering an atlas tile.
+	 *  - it manipulates TexCoords for the following reasons:
+	 *		- Clamped textures will look different if texcoords are manipulated
+	 *		- Wrapped textures look good with any texcoords offsets (since tiles are repeatable and have correct HW filtering at edges)
+	 *		- Wrapped textures look different however if texcoords are scaled (looks correct for 1, 2, 3 but not for 1.5 or 1.1 for instance as UV space [0,1] will no longer align with the atlas tile edges)
+	 * But, an artist can specify and override the fact that a material is compatible with the light function atlas.
+	 */
+	MaterialCompilationOutput.bIsLightFunctionAtlasCompatible = (!bUsesVertexPosition && !bUsesSceneDepth && !MaterialCompilationOutput.bNeedsSceneTextures && !bPotentiallyManipulateTexCoords) || Material->IsCompatibleWithLightFunctionAtlas();
 
 	EMaterialDecalResponse MDR = (EMaterialDecalResponse)Material->GetMaterialDecalResponse();
 	if (MDR == MDR_Color || MDR == MDR_ColorNormal || MDR == MDR_ColorRoughness || MDR == MDR_ColorNormalRoughness)
@@ -6502,6 +6512,12 @@ int32 FHLSLMaterialTranslator::TextureCoordinate(uint32 CoordinateIndex, bool Un
 	}
 }
 
+void FHLSLMaterialTranslator::SetPotentiallyManipulateTexCoords()
+{
+	// This cannot be set from FHLSLMaterialTranslator::TextureCoordinate because it is called by texture sampling node for default UVs as coordinate index 0 for instance.
+	bPotentiallyManipulateTexCoords = true;
+}
+
 static const TCHAR* GetVTAddressMode(TextureAddress Address)
 {
 	switch (Address)
diff --git a/Engine/Source/Runtime/Engine/Private/Materials/HLSLMaterialTranslator.h b/Engine/Source/Runtime/Engine/Private/Materials/HLSLMaterialTranslator.h
--- a/Engine/Source/Runtime/Engine/Private/Materials/HLSLMaterialTranslator.h
+++ b/Engine/Source/Runtime/Engine/Private/Materials/HLSLMaterialTranslator.h
@@ -392,6 +392,8 @@ protected:
 
 	/** true if the material uses any type of vertex position */
 	uint32 bUsesVertexPosition : 1;
+	/** true if the material potentially manipulates any type of TexCoord for texture sampling */
+	uint32 bPotentiallyManipulateTexCoords : 1;
 
 	uint32 bUsesTransformVector : 1;
 	// True if the current property requires last frame's information
@@ -989,6 +991,7 @@ protected:
 #endif
 
 	virtual int32 TextureCoordinate(uint32 CoordinateIndex, bool UnMirrorU, bool UnMirrorV) override;
+	virtual void SetPotentiallyManipulateTexCoords() override;
 
 	uint32 AcquireVTStackIndex(
 		ETextureMipValueMode MipValueMode, 
diff --git a/Engine/Source/Runtime/Engine/Private/Materials/Material.cpp b/Engine/Source/Runtime/Engine/Private/Materials/Material.cpp
--- a/Engine/Source/Runtime/Engine/Private/Materials/Material.cpp
+++ b/Engine/Source/Runtime/Engine/Private/Materials/Material.cpp
@@ -1020,6 +1020,9 @@ UMaterial::UMaterial(const FObjectInitializer& ObjectInitializer)
 	bUsedWithVolumetricCloud = false;
 	bUsedWithHeterogeneousVolumes = false;
 
+	// We default to false=incompatible, and rely onto the automatic detection executed from the HLSLTranslator.
+	bCompatibleWithLightFunctionAtlas = false;
+
 	bUseEmissiveForDynamicAreaLighting = false;
 	RefractionDepthBias = 0.0f;
 	MaterialDecalResponse = MDR_ColorNormalRoughness;
diff --git a/Engine/Source/Runtime/Engine/Private/Materials/MaterialExpressions.cpp b/Engine/Source/Runtime/Engine/Private/Materials/MaterialExpressions.cpp
--- a/Engine/Source/Runtime/Engine/Private/Materials/MaterialExpressions.cpp
+++ b/Engine/Source/Runtime/Engine/Private/Materials/MaterialExpressions.cpp
@@ -5606,6 +5606,8 @@ UMaterialExpressionTextureCoordinate::UMaterialExpressionTextureCoordinate(const
 #if WITH_EDITOR
 int32 UMaterialExpressionTextureCoordinate::Compile(class FMaterialCompiler* Compiler, int32 OutputIndex)
 {
+	Compiler->SetPotentiallyManipulateTexCoords();
+
 	// Depending on whether we have U and V scale values that differ, we can perform a multiply by either
 	// a scalar or a float2.  These tiling values are baked right into the shader node, so they're always
 	// known at compile time.
diff --git a/Engine/Source/Runtime/Engine/Private/Materials/MaterialShared.cpp b/Engine/Source/Runtime/Engine/Private/Materials/MaterialShared.cpp
--- a/Engine/Source/Runtime/Engine/Private/Materials/MaterialShared.cpp
+++ b/Engine/Source/Runtime/Engine/Private/Materials/MaterialShared.cpp
@@ -1878,6 +1878,11 @@ bool FMaterialResource::IsFullyRough() const
 	return Material->bFullyRough;
 }
 
+bool FMaterialResource::IsCompatibleWithLightFunctionAtlas() const
+{
+	return Material->bCompatibleWithLightFunctionAtlas;
+}
+
 bool FMaterialResource::UseNormalCurvatureToRoughness() const
 {
 	return Material->bNormalCurvatureToRoughness;
diff --git a/Engine/Source/Runtime/Engine/Public/MaterialCompiler.h b/Engine/Source/Runtime/Engine/Public/MaterialCompiler.h
--- a/Engine/Source/Runtime/Engine/Public/MaterialCompiler.h
+++ b/Engine/Source/Runtime/Engine/Public/MaterialCompiler.h
@@ -273,6 +273,7 @@ public:
 	virtual int32 Switch(int32 SwitchValueInput, int32 DefaultInput, TArray<int32>& CompiledInputs) = 0;
 
 	virtual int32 TextureCoordinate(uint32 CoordinateIndex, bool UnMirrorU, bool UnMirrorV) = 0;
+	virtual void  SetPotentiallyManipulateTexCoords() = 0;
 	virtual int32 TextureSample(int32 Texture,int32 Coordinate,enum EMaterialSamplerType SamplerType,int32 MipValue0Index=INDEX_NONE,int32 MipValue1Index=INDEX_NONE,ETextureMipValueMode MipValueMode=TMVM_None,ESamplerSourceMode SamplerSource=SSM_FromTextureAsset,int32 TextureReferenceIndex=INDEX_NONE, bool AutomaticViewMipBias=false, bool AdaptiveVirtualTexture=false, bool EnableFeedback = true) = 0;
 	virtual int32 TextureProperty(int32 InTexture, EMaterialExposedTextureProperty Property) = 0;
 
@@ -881,6 +882,7 @@ public:
 	}
 
 	virtual int32 TextureCoordinate(uint32 CoordinateIndex, bool UnMirrorU, bool UnMirrorV) override { return Compiler->TextureCoordinate(CoordinateIndex, UnMirrorU, UnMirrorV); }
+	virtual void SetPotentiallyManipulateTexCoords() override { Compiler->SetPotentiallyManipulateTexCoords(); }
 
 	virtual int32 TextureDecalMipmapLevel(int32 TextureSizeInput) override { return Compiler->TextureDecalMipmapLevel(TextureSizeInput); }
 	virtual int32 TextureDecalDerivative(bool bDDY) override { return Compiler->TextureDecalDerivative(bDDY); }
diff --git a/Engine/Source/Runtime/Engine/Public/MaterialShared.h b/Engine/Source/Runtime/Engine/Public/MaterialShared.h
--- a/Engine/Source/Runtime/Engine/Public/MaterialShared.h
+++ b/Engine/Source/Runtime/Engine/Public/MaterialShared.h
@@ -2051,6 +2051,7 @@ public:
 	virtual bool IsUsedWithVirtualHeightfieldMesh() const { return false; }
 	virtual bool IsUsedWithNeuralNetworks() const { return false; }
 	virtual bool IsFullyRough() const { return false; }
+	virtual bool IsCompatibleWithLightFunctionAtlas() const { return false; }
 	virtual bool UseNormalCurvatureToRoughness() const { return false; }
 	virtual enum EMaterialFloatPrecisionMode GetMaterialFloatPrecisionMode() const { return EMaterialFloatPrecisionMode::MFPM_Default; };
 	virtual bool IsUsingAlphaToCoverage() const { return false; }
@@ -2811,6 +2812,7 @@ public:
 	ENGINE_API virtual bool IsUsedWithVolumetricCloud() const override;
 	ENGINE_API virtual bool IsUsedWithHeterogeneousVolumes() const override;
 	ENGINE_API virtual bool IsFullyRough() const override;
+	ENGINE_API virtual bool IsCompatibleWithLightFunctionAtlas() const override;
 	ENGINE_API virtual bool UseNormalCurvatureToRoughness() const override;
 	ENGINE_API virtual enum EMaterialFloatPrecisionMode GetMaterialFloatPrecisionMode() const override;
 	ENGINE_API virtual bool IsUsingAlphaToCoverage() const override;
