[Physics] AddForce and AddTorque support for FPhysicsObject through PrimitiveComponent on game thread

diff --git a/Engine/Source/Runtime/Engine/Private/PrimitiveComponentPhysics.cpp b/Engine/Source/Runtime/Engine/Private/PrimitiveComponentPhysics.cpp
--- a/Engine/Source/Runtime/Engine/Private/PrimitiveComponentPhysics.cpp
+++ b/Engine/Source/Runtime/Engine/Private/PrimitiveComponentPhysics.cpp
@@ -233,6 +233,19 @@ void UPrimitiveComponent::AddRadialImpulse(FVector Origin, float Radius, float S
 	{
 		BI->AddRadialImpulseToBody(Origin, Radius, Strength, Falloff, bVelChange);
 	}
+	else
+	{
+		TArray<Chaos::FPhysicsObjectHandle> PhysicsObjects = GetAllPhysicsObjects();
+		FLockedWritePhysicsObjectExternalInterface Interface = FPhysicsObjectExternalInterface::LockWrite(PhysicsObjects);
+
+		PhysicsObjects = PhysicsObjects.FilterByPredicate(
+			[&Interface](Chaos::FPhysicsObject* Object) {
+				return !Interface->AreAllDisabled({ &Object, 1 });
+			}
+		);
+
+		Interface->AddRadialImpulse(PhysicsObjects, Origin, Radius, Strength, Falloff, /*bApplyStrain*/true, /*bInvalidate*/true, bVelChange, /*MinValue*/0.f, /*MaxValue*/1.f);
+	}
 }
 
 
@@ -243,6 +256,40 @@ void UPrimitiveComponent::AddForce(FVector Force, FName BoneName, bool bAccelCha
 		WarnInvalidPhysicsOperations(LOCTEXT("AddForce", "AddForce"), BI, BoneName);
 		BI->AddForce(Force, true, bAccelChange);
 	}
+	else if (BoneName == NAME_None)
+	{
+		TArray<Chaos::FPhysicsObjectHandle> PhysicsObjects = GetAllPhysicsObjects();
+		FLockedWritePhysicsObjectExternalInterface Interface = FPhysicsObjectExternalInterface::LockWrite(PhysicsObjects);
+
+		PhysicsObjects = PhysicsObjects.FilterByPredicate(
+			[&Interface](Chaos::FPhysicsObject* Object) {
+				return !Interface->AreAllDisabled({ &Object, 1 });
+			}
+		);
+
+		if (bAccelChange)
+		{
+			const float Mass = Interface->GetMass(PhysicsObjects);
+			Interface->AddForce(PhysicsObjects, Force * Mass, /*bInvalidate*/true);
+		}
+		else
+		{
+			Interface->AddForce(PhysicsObjects, Force, /*bInvalidate*/true);
+		}
+	}
+	else if (Chaos::FPhysicsObject* Object = GetPhysicsObjectByName(BoneName))
+	{
+		FLockedWritePhysicsObjectExternalInterface Interface = FPhysicsObjectExternalInterface::LockWrite({ &Object, 1 });
+		if (bAccelChange)
+		{
+			const float Mass = Interface->GetMass({ &Object, 1 });
+			Interface->AddForce({ &Object, 1 }, Force * Mass, /*bInvalidate*/true);
+		}
+		else
+		{
+			Interface->AddForce({ &Object, 1 }, Force, /*bInvalidate*/true);
+		}
+	}
 }
 
 void UPrimitiveComponent::AddForceAtLocation(FVector Force, FVector Location, FName BoneName)
@@ -284,6 +331,40 @@ void UPrimitiveComponent::AddTorqueInRadians(FVector Torque, FName BoneName, boo
 		WarnInvalidPhysicsOperations(LOCTEXT("AddTorque", "AddTorque"), BI, BoneName);
 		BI->AddTorqueInRadians(Torque, true, bAccelChange);
 	}
+	else if (BoneName == NAME_None)
+	{
+		TArray<Chaos::FPhysicsObjectHandle> PhysicsObjects = GetAllPhysicsObjects();
+		FLockedWritePhysicsObjectExternalInterface Interface = FPhysicsObjectExternalInterface::LockWrite(PhysicsObjects);
+
+		PhysicsObjects = PhysicsObjects.FilterByPredicate(
+			[&Interface](Chaos::FPhysicsObject* Object) {
+				return !Interface->AreAllDisabled({ &Object, 1 });
+			}
+		);
+
+		if (bAccelChange)
+		{
+			const float Mass = Interface->GetMass(PhysicsObjects);
+			Interface->AddTorque(PhysicsObjects, Torque * Mass, /*bInvalidate*/true);
+		}
+		else
+		{
+			Interface->AddTorque(PhysicsObjects, Torque, /*bInvalidate*/true);
+		}
+	}
+	else if (Chaos::FPhysicsObject* Object = GetPhysicsObjectByName(BoneName))
+	{
+		FLockedWritePhysicsObjectExternalInterface Interface = FPhysicsObjectExternalInterface::LockWrite({ &Object, 1 });
+		if (bAccelChange)
+		{
+			const float Mass = Interface->GetMass({ &Object, 1 });
+			Interface->AddTorque({ &Object, 1 }, Torque * Mass, /*bInvalidate*/true);
+		}
+		else
+		{
+			Interface->AddTorque({ &Object, 1 }, Torque, /*bInvalidate*/true);
+		}
+	}
 }
 
 void UPrimitiveComponent::SetPhysicsLinearVelocity(FVector NewVel, bool bAddToCurrent, FName BoneName)
