Fixed FExternalDataLayerHelper::MoveActorsToExternalDataLayer to skip any actors in the package when transfering objects of the package to the new one.
This fixes a crash while replacing an actor by another one using External Data Layers (the old actor to be destroyed was re-outered to the new package).

diff --git a/Engine/Source/Runtime/Engine/Private/WorldPartition/DataLayer/ExternalDataLayerHelper.cpp b/Engine/Source/Runtime/Engine/Private/WorldPartition/DataLayer/ExternalDataLayerHelper.cpp
--- a/Engine/Source/Runtime/Engine/Private/WorldPartition/DataLayer/ExternalDataLayerHelper.cpp
+++ b/Engine/Source/Runtime/Engine/Private/WorldPartition/DataLayer/ExternalDataLayerHelper.cpp
@@ -260,7 +260,7 @@ bool FExternalDataLayerHelper::MoveActorsToExternalDataLayer(const TArray<AActor
 		TArray<UObject*> DependantObjects;
 		ForEachObjectWithPackage(OldActorPackage, [&DependantObjects](UObject* Object)
 		{
-			if (!Cast<UMetaData>(Object) && !Cast<UDeletedObjectPlaceholder>(Object))
+			if (!Cast<UMetaData>(Object) && !Cast<UDeletedObjectPlaceholder>(Object) && !Cast<AActor>(Object))
 			{
 				DependantObjects.Add(Object);
 			}
