Level Instance: Property Override Actor Coloration

#rnx

diff --git a/Engine/Source/Editor/LevelInstanceEditor/Private/LevelInstanceEditorModule.cpp b/Engine/Source/Editor/LevelInstanceEditor/Private/LevelInstanceEditorModule.cpp
--- a/Engine/Source/Editor/LevelInstanceEditor/Private/LevelInstanceEditorModule.cpp
+++ b/Engine/Source/Editor/LevelInstanceEditor/Private/LevelInstanceEditorModule.cpp
@@ -1029,10 +1029,14 @@ void FLevelInstanceEditorModule::StartupModule()
 	{
 		GLevelEditorModeTools().OnEditorModeIDChanged().AddRaw(this, &FLevelInstanceEditorModule::OnEditorModeIDChanged);
 	}
+
+	ULevelInstanceSubsystem::RegisterPrimitiveColorHandler();
 }
 
 void FLevelInstanceEditorModule::ShutdownModule()
 {
+	ULevelInstanceSubsystem::UnregisterPrimitiveColorHandler();
+
 	if (FModuleManager::Get().IsModuleLoaded("LevelEditor"))
 	{
 		FLevelEditorModule& LevelEditorModule = FModuleManager::GetModuleChecked<FLevelEditorModule>("LevelEditor");
diff --git a/Engine/Source/Runtime/Engine/Private/LevelInstance/LevelInstanceSubsystem.cpp b/Engine/Source/Runtime/Engine/Private/LevelInstance/LevelInstanceSubsystem.cpp
--- a/Engine/Source/Runtime/Engine/Private/LevelInstance/LevelInstanceSubsystem.cpp
+++ b/Engine/Source/Runtime/Engine/Private/LevelInstance/LevelInstanceSubsystem.cpp
@@ -12,6 +12,7 @@
 #include "EngineUtils.h"
 #include "LevelInstancePrivate.h"
 #include "LevelUtils.h"
+#include "GameFramework/ActorPrimitiveColorHandler.h"
 
 #if WITH_EDITOR
 #include "Settings/LevelEditorMiscSettings.h"
@@ -62,6 +63,10 @@
 
 DEFINE_LOG_CATEGORY(LogLevelInstance);
 
+#if WITH_EDITOR
+bool ULevelInstanceSubsystem::bPrimitiveColorHandlerRegistered = false;
+#endif
+
 ULevelInstanceSubsystem::ULevelInstanceSubsystem()
 	: UWorldSubsystem()
 #if WITH_EDITOR
@@ -3130,6 +3135,48 @@ void ULevelInstanceSubsystem::ResetPropertyOverrides(ILevelInstanceInterface* Le
 	GEditor->SelectNone(true, true);
 }
 
+void ULevelInstanceSubsystem::RegisterPrimitiveColorHandler()
+{
+#if ENABLE_ACTOR_PRIMITIVE_COLOR_HANDLER
+	if (!bPrimitiveColorHandlerRegistered && GetDefault<ULevelInstanceSettings>()->IsPropertyOverrideEnabled())
+	{
+		FActorPrimitiveColorHandler::Get().RegisterPrimitiveColorHandler(TEXT("LevelInstancePropertyOverride"), LOCTEXT("LevelInstancePropertyOverrideColorHandler", "Level Instance Property Override"), [](const UPrimitiveComponent* InPrimitiveComponent)
+		{
+			if (AActor* Actor = InPrimitiveComponent ? InPrimitiveComponent->GetOwner() : nullptr; Actor && Actor->IsInLevelInstance())
+			{
+				if (ULevelInstanceSubsystem* LevelInstanceSubsystem = UWorld::GetSubsystem<ULevelInstanceSubsystem>(Actor->GetWorld()))
+				{
+					if (ILevelInstanceInterface* OwningLevelInstance = LevelInstanceSubsystem->GetOwningLevelInstance(Actor->GetLevel()))
+					{
+						const FActorContainerID ContextContainerID = LevelInstanceSubsystem->GetLevelInstancePropertyOverridesContext(OwningLevelInstance);
+						TArray<const FActorPropertyOverride*> ActorPropertyOverrides;
+						if (LevelInstanceSubsystem->GetLevelInstancePropertyOverridesForActor(Actor, ContextContainerID, ActorPropertyOverrides))
+						{
+							return FLinearColor::Green;
+						}
+					}
+				}
+
+				return FLinearColor::Red;
+			}
+			return FLinearColor::White;
+		});
+		bPrimitiveColorHandlerRegistered = true;
+	}
+#endif
+}
+
+void ULevelInstanceSubsystem::UnregisterPrimitiveColorHandler()
+{
+#if ENABLE_ACTOR_PRIMITIVE_COLOR_HANDLER
+	if (bPrimitiveColorHandlerRegistered)
+	{
+		FActorPrimitiveColorHandler::Get().UnregisterPrimitiveColorHandler(TEXT("LevelInstancePropertyOverride"));
+		bPrimitiveColorHandlerRegistered = false;
+	}
+#endif
+}
+
 void ULevelInstanceSubsystem::UpdateLevelInstancesFromPropertyOverrideAsset(const TSoftObjectPtr<ULevelInstancePropertyOverrideAsset>& PreviousAssetPath, ULevelInstancePropertyOverrideAsset* NewAsset)
 {
 	if (PreviousAssetPath.IsNull())
diff --git a/Engine/Source/Runtime/Engine/Public/LevelInstance/LevelInstanceSubsystem.h b/Engine/Source/Runtime/Engine/Public/LevelInstance/LevelInstanceSubsystem.h
--- a/Engine/Source/Runtime/Engine/Public/LevelInstance/LevelInstanceSubsystem.h
+++ b/Engine/Source/Runtime/Engine/Public/LevelInstance/LevelInstanceSubsystem.h
@@ -188,6 +188,9 @@ private:
 	ENGINE_API void ResetPropertyOverridesForActor(AActor* Actor);
 	ENGINE_API bool CanResetPropertyOverrides(ILevelInstanceInterface* LevelInstance) const;
 	ENGINE_API void ResetPropertyOverrides(ILevelInstanceInterface* LevelInstance);
+
+	ENGINE_API static void RegisterPrimitiveColorHandler();
+	ENGINE_API static void UnregisterPrimitiveColorHandler();
 #endif
 
 private:
@@ -317,6 +320,7 @@ private:
 	FLevelInstanceChanged LevelInstanceChangedEvent;
 	FLevelInstancesUpdated LevelInstancesUpdatedEvent;
 
+	static bool bPrimitiveColorHandlerRegistered;
 #endif
 
 	struct FLevelInstance
