Fixed Undo of the operation of moving an actor from one External Data Layer to another.

diff --git a/Engine/Source/Runtime/Engine/Private/Components/ChildActorComponent.cpp b/Engine/Source/Runtime/Engine/Private/Components/ChildActorComponent.cpp
--- a/Engine/Source/Runtime/Engine/Private/Components/ChildActorComponent.cpp
+++ b/Engine/Source/Runtime/Engine/Private/Components/ChildActorComponent.cpp
@@ -193,6 +193,7 @@ void UChildActorComponent::Serialize(FArchive& Ar)
 #if WITH_EDITOR
 void UChildActorComponent::SetPackageExternal(bool bExternal, bool bShouldDirty)
 {
+	Modify(bShouldDirty);
 	DestroyChildActor();
 	CreateChildActor();
 }
