Fixed substrate crash with unlit when using coverage node (valide), any other operator (invalid) or node preview image.
Also fixed other errors due to missing early out of compiler.

diff --git a/Engine/Source/Runtime/Engine/Private/Materials/HLSLMaterialTranslator.cpp b/Engine/Source/Runtime/Engine/Private/Materials/HLSLMaterialTranslator.cpp
--- a/Engine/Source/Runtime/Engine/Private/Materials/HLSLMaterialTranslator.cpp
+++ b/Engine/Source/Runtime/Engine/Private/Materials/HLSLMaterialTranslator.cpp
@@ -1368,10 +1368,12 @@ void FHLSLMaterialTranslator::DoTranslate()
 			if (SubstrateCtx.bSubstrateTreeOutOfStackDepthOccurred)
 			{
 				Errorf(TEXT(" %s [%s]: Substrate - Cyclic graph detected when we only support acyclic graph."), *Material->GetDebugName(), *Material->GetAssetPath().ToString());
+				return;
 			}
 			if (!SubstrateCtx.SubstrateGenerateDerivedMaterialOperatorData(this))
 			{
 				Errorf(TEXT("Substrate material errors encountered."));
+				return;
 			}
 
 			bSubstrateWritesEmissive |= SubstrateCtx.bSubstrateWritesEmissive;
@@ -12196,8 +12198,10 @@ bool FHLSLMaterialTranslator::FSubstrateCompilationContext::SubstrateGenerateDer
 				const bool bRootOfParameterBlendingSubTree		= bCurrentOpRequestParameterBlending && !bInsideParameterBlendingSubTree;
 				const bool bUseParameterBlending				= bCurrentOpRequestParameterBlending || bInsideParameterBlendingSubTree;
 
-				if (CurrentOperator.BSDFType == SUBSTRATE_BSDF_TYPE_SLAB)
+				if (CurrentOperator.BSDFType == SUBSTRATE_BSDF_TYPE_SLAB || CurrentOperator.BSDFType == SUBSTRATE_BSDF_TYPE_UNLIT)
 				{
+					// Update the parameter blending data for BSDFs supporting operators.
+					// We also need to do this for UNLIT since it supports Coverage operator.
 					CurrentOperator.bUseParameterBlending = bUseParameterBlending;
 					CurrentOperator.bRootOfParameterBlendingSubTree = bRootOfParameterBlendingSubTree;
 				}
@@ -12267,23 +12271,26 @@ bool FHLSLMaterialTranslator::FSubstrateCompilationContext::SubstrateGenerateDer
 			}
 			bSubstrateMaterialIsUnlitNode = bHasUnlit;
 
-			if ((bHasUnlit || bHasVFogCloud || bHasHair || bHasEye || bHasSLW) && SubstrateMaterialEffectiveClosureCount > 1)
-			{
-				Compiler->Errorf(TEXT("Unlit, Fog/Cloud, Hair or SingleLayerWater must be used in isolation. See %s (asset: %s).\r\n"), *CompilerMaterial->GetDebugName(), *CompilerMaterial->GetAssetPath().ToString());
-				// Even though we could support Unlit with slab.
-			}
-
 			if (((bHasVFogCloud || bHasHair || bHasEye || bHasSLW) && bOperatorEncountered)
 				|| (bHasUnlit && bOperatorEncounteredButNotWeight)
 				)
 			{
-				Compiler->Errorf(TEXT("Unlit, Fog/Cloud, Hair or SingleLayerWater cannot be used with operators. See %s (asset: %s).\r\n"), *CompilerMaterial->GetDebugName(), *CompilerMaterial->GetAssetPath().ToString());
-				// This is because it will results in simpler lighting loops focusin on slab.
+				Compiler->Errorf(TEXT("Fog/Cloud, Hair or SingleLayerWater cannot be used with operators.\r\nUnlit can only be used with coverage operators. See %s (asset: %s).\r\n"), *CompilerMaterial->GetDebugName(), *CompilerMaterial->GetAssetPath().ToString());
+				// This is because it will results in lighting loops dedicated on slab and simply a few other isolated lighting cases.
+				return false;
+			}
+
+			if ((bHasUnlit || bHasVFogCloud || bHasHair || bHasEye || bHasSLW) && SubstrateMaterialEffectiveClosureCount > 1)
+			{
+				Compiler->Errorf(TEXT("Unlit, Fog/Cloud, Hair or SingleLayerWater must be used in isolation. See %s (asset: %s).\r\n"), *CompilerMaterial->GetDebugName(), *CompilerMaterial->GetAssetPath().ToString());
+				// Even though technically we could support Unlit combined with Slab.
+				return false;
 			}
 
 			if (SubstrateMaterialEffectiveClosureCount > SUBSTRATE_MAX_CLOSURE_COUNT)
 			{
-				Compiler->Errorf(TEXT("Material tries to register more BSDF than can be supproted (%d > %d). See %s (asset: %s).\r\n"), SubstrateMaterialEffectiveClosureCount, SUBSTRATE_MAX_CLOSURE_COUNT, *CompilerMaterial->GetDebugName(), *CompilerMaterial->GetAssetPath().ToString());
+				Compiler->Errorf(TEXT("Material tries to register more BSDF than can be supported (%d > %d). See %s (asset: %s).\r\n"), SubstrateMaterialEffectiveClosureCount, SUBSTRATE_MAX_CLOSURE_COUNT, *CompilerMaterial->GetDebugName(), *CompilerMaterial->GetAssetPath().ToString());
+				return false;
 			}
 		}
 
@@ -12641,7 +12648,8 @@ bool FHLSLMaterialTranslator::FSubstrateCompilationContext::SubstrateGenerateDer
 					}
 					case SUBSTRATE_BSDF_TYPE_SINGLELAYERWATER:
 					{
-						Compiler->Errorf(TEXT("Substrate error: single layer water should go through the its dedicated fast path in %s (asset: %s).\r\n"), *CompilerMaterial->GetDebugName(), *CompilerMaterial->GetAssetPath().ToString());
+						Compiler->Errorf(TEXT("Substrate error: single layer water should go through through a dedicated fast path in %s (asset: %s).\r\n"), *CompilerMaterial->GetDebugName(), *CompilerMaterial->GetAssetPath().ToString());
+						return false;
 						break;
 					}
 					case SUBSTRATE_BSDF_TYPE_UNLIT:
@@ -12652,6 +12660,7 @@ bool FHLSLMaterialTranslator::FSubstrateCompilationContext::SubstrateGenerateDer
 					default:
 					{
 						Compiler->Errorf(TEXT("Unkownd BSDF type encountered in %s (asset: %s).\r\n"), *CompilerMaterial->GetDebugName(), *CompilerMaterial->GetAssetPath().ToString());
+						return false;
 						break;
 					}
 					}
