Fixed crash in local position node when TransformPosition fails due to an unsupported material translation context.
Fixed missing pragma once.

diff --git a/Engine/Source/Runtime/Engine/Private/Materials/HLSLMaterialTranslator.cpp b/Engine/Source/Runtime/Engine/Private/Materials/HLSLMaterialTranslator.cpp
--- a/Engine/Source/Runtime/Engine/Private/Materials/HLSLMaterialTranslator.cpp
+++ b/Engine/Source/Runtime/Engine/Private/Materials/HLSLMaterialTranslator.cpp
@@ -6042,6 +6042,10 @@ int32 FHLSLMaterialTranslator::LocalPosition(EPositionIncludedOffsets IncludedOf
 		// Generate derivative based on WorldPosition_DDX
 		int32 WorldPositionCode = WorldPosition(IncludedOffsets == EPositionIncludedOffsets::IncludeOffsets ? WPT_CameraRelative : WPT_CameraRelativeNoOffsets);
 		int32 TransformedPositionCode = TransformPosition(MCB_TranslatedWorld, OriginType == ELocalPositionOrigin::Instance ? MCB_Instance : MCB_Local, WorldPositionCode);
+		if (TransformedPositionCode < 0) // TransformPosition failed
+		{
+			return TransformedPositionCode;
+		}
 		FString TransformedPositionAnalyticCode = GetParameterCodeDeriv(TransformedPositionCode, CompiledPDV_Analytic);
 
 		FString AnalyticCode = DerivativeAutogen.ConstructDeriv(
diff --git a/Engine/Source/Runtime/Renderer/Private/MeshEdgesRendering.h b/Engine/Source/Runtime/Renderer/Private/MeshEdgesRendering.h
--- a/Engine/Source/Runtime/Renderer/Private/MeshEdgesRendering.h
+++ b/Engine/Source/Runtime/Renderer/Private/MeshEdgesRendering.h
@@ -1,5 +1,7 @@
 // Copyright Epic Games, Inc. All Rights Reserved.
 
+#pragma once
+
 class FRDGBuilder;
 class FSceneView;
 class FViewInfo;
