Fixed bUseStreamingLowMipsInEditor ("View in Editor") for RVT volume component not being transient (by default, we always want the editor to not show the SVT)

diff --git a/Engine/Source/Runtime/Engine/Classes/Components/RuntimeVirtualTextureComponent.h b/Engine/Source/Runtime/Engine/Classes/Components/RuntimeVirtualTextureComponent.h
--- a/Engine/Source/Runtime/Engine/Classes/Components/RuntimeVirtualTextureComponent.h
+++ b/Engine/Source/Runtime/Engine/Classes/Components/RuntimeVirtualTextureComponent.h
@@ -90,7 +90,7 @@ protected:
 	FLinearColor StreamingMipsFixedColor;
 
 	/** Use streaming low mips when rendering in editor. Set true to view and debug the baked streaming low mips. */
-	UPROPERTY(EditAnywhere, AdvancedDisplay, Category = StreamingVirtualTexture, meta = (DisplayName = "View in Editor"))
+	UPROPERTY(EditAnywhere, Transient, AdvancedDisplay, Category = StreamingVirtualTexture, meta = (DisplayName = "View in Editor"))
 	bool bUseStreamingLowMipsInEditor = false;
 
 #if WITH_EDITOR
