overlaps against a geometry collection should return a proper ItemIndex depending on which particle was overlapped

diff --git a/Engine/Source/Runtime/Engine/Private/Collision/CollisionConversions.cpp b/Engine/Source/Runtime/Engine/Private/Collision/CollisionConversions.cpp
--- a/Engine/Source/Runtime/Engine/Private/Collision/CollisionConversions.cpp
+++ b/Engine/Source/Runtime/Engine/Private/Collision/CollisionConversions.cpp
@@ -630,6 +630,18 @@ void ConvertQueryOverlap(const FPhysicsShape& Shape, const FPhysicsActor& Actor,
 			OutOverlap.OverlapObjectHandle = FActorInstanceHandle::MakeActorHandleToResolve(PossibleOwner, INDEX_NONE);
 			OutOverlap.Component = PossibleOwner;
 			OutOverlap.ItemIndex = INDEX_NONE;
+
+			// If this is a geometry collection, we can do more to extract the properly ItemIndex.
+			if (const IPhysicsProxyBase* ActorProxy = Actor.GetProxy())
+			{
+				if (ActorProxy->GetType() == EPhysicsProxyType::GeometryCollectionType)
+				{
+					const FGeometryCollectionPhysicsProxy* ConcreteProxy = static_cast<const FGeometryCollectionPhysicsProxy*>(ActorProxy);
+					const FGeometryCollectionItemIndex ItemIndex = ConcreteProxy->GetItemIndexFromGTParticle_External(Actor.CastToRigidParticle());
+					OutOverlap.ItemIndex = ItemIndex.GetItemIndex();
+				}
+			}
+
 		}
 		else
 		{
