Add a new texture streaming csv stat PendingStreamInData

diff --git a/Engine/Source/Runtime/Engine/Private/Streaming/AsyncTextureStreaming.cpp b/Engine/Source/Runtime/Engine/Private/Streaming/AsyncTextureStreaming.cpp
--- a/Engine/Source/Runtime/Engine/Private/Streaming/AsyncTextureStreaming.cpp
+++ b/Engine/Source/Runtime/Engine/Private/Streaming/AsyncTextureStreaming.cpp
@@ -1091,6 +1091,7 @@ void FRenderAssetStreamingMipCalcTask::UpdateCSVOnlyStats_Async()
 	Stats.RequiredPool = 0;
 	Stats.CachedMips = 0;
 	Stats.WantedMips = 0;
+	Stats.PendingRequests = 0;
 
 	Stats.NumStreamedMeshes = 0;
 	Stats.AvgNumStreamedLODs = 0.f;
@@ -1123,6 +1124,11 @@ void FRenderAssetStreamingMipCalcTask::UpdateCSVOnlyStats_Async()
 		Stats.WantedMips += UsedSize;
 		Stats.CachedMips += FMath::Max<int64>(ResidentSize - UsedSize, 0);
 
+		if (StreamingRenderAsset.RequestedMips > StreamingRenderAsset.ResidentMips)
+		{
+			Stats.PendingRequests += StreamingRenderAsset.GetSize(StreamingRenderAsset.RequestedMips) - ResidentSize;
+		}
+
 		if (StreamingRenderAsset.IsMesh())
 		{
 			const int32 NumStreamedLODs = StreamingRenderAsset.MaxAllowedMips - StreamingRenderAsset.MinAllowedMips;
diff --git a/Engine/Source/Runtime/Engine/Private/Streaming/StreamingManagerTexture.cpp b/Engine/Source/Runtime/Engine/Private/Streaming/StreamingManagerTexture.cpp
--- a/Engine/Source/Runtime/Engine/Private/Streaming/StreamingManagerTexture.cpp
+++ b/Engine/Source/Runtime/Engine/Private/Streaming/StreamingManagerTexture.cpp
@@ -1766,6 +1766,7 @@ void FRenderAssetStreamingManager::UpdateResourceStreaming( float DeltaTime, boo
 	CSV_CUSTOM_STAT(TextureStreaming, ResidentMeshMem, ((float)DisplayedStats.ResidentMeshMem) / (1024.0f * 1024.0f), ECsvCustomStatOp::Set);
 	CSV_CUSTOM_STAT(TextureStreaming, StreamedMeshMem, ((float)DisplayedStats.StreamedMeshMem) / (1024.0f * 1024.0f), ECsvCustomStatOp::Set);
 	CSV_CUSTOM_STAT(TextureStreaming, NonStreamingMips, ((float)DisplayedStats.NonStreamingMips) / (1024.0f * 1024.0f), ECsvCustomStatOp::Set);
+	CSV_CUSTOM_STAT(TextureStreaming, PendingStreamInData, ((float)DisplayedStats.PendingRequests) / (1024.0f * 1024.0f), ECsvCustomStatOp::Set)
 
 	RenderAssetInstanceAsyncWork->EnsureCompletion();
 
