Fixed AWorldDataLayers acceleration table used by AWorldDataLayers::GetDataLayerInstanceFromAssetName and AWorldDataLayers::GetDataLayerInstance(const UDataLayerAsset*) in cooked builds.
This fixes AWorldDataLayers::GetDataLayerInstanceFromAssetName not finding the DataLayerInstance in cooked builds.

diff --git a/Engine/Source/Runtime/Engine/Private/WorldPartition/DataLayer/DataLayerManager.cpp b/Engine/Source/Runtime/Engine/Private/WorldPartition/DataLayer/DataLayerManager.cpp
--- a/Engine/Source/Runtime/Engine/Private/WorldPartition/DataLayer/DataLayerManager.cpp
+++ b/Engine/Source/Runtime/Engine/Private/WorldPartition/DataLayer/DataLayerManager.cpp
@@ -290,10 +290,10 @@ const UDataLayerInstance* UDataLayerManager::GetDataLayerInstanceFromName(const
 	return GetDataLayerInstance(InDataLayerInstanceName);
 }
 
-const UDataLayerInstance* UDataLayerManager::GetDataLayerInstanceFromAssetName(const FName& InDataLayerAssetFullName) const
+const UDataLayerInstance* UDataLayerManager::GetDataLayerInstanceFromAssetName(const FName& InDataLayerAssetPathName) const
 {
 	AWorldDataLayers* WorldDataLayers = GetWorldDataLayers();
-	return WorldDataLayers ? WorldDataLayers->GetDataLayerInstanceFromAssetName(InDataLayerAssetFullName) : nullptr;
+	return WorldDataLayers ? WorldDataLayers->GetDataLayerInstanceFromAssetName(InDataLayerAssetPathName) : nullptr;
 }
 
 bool UDataLayerManager::SetDataLayerInstanceRuntimeState(const UDataLayerInstance* InDataLayerInstance, EDataLayerRuntimeState InState, bool bInIsRecursive)
diff --git a/Engine/Source/Runtime/Engine/Private/WorldPartition/DataLayer/DataLayerSubsystem.cpp b/Engine/Source/Runtime/Engine/Private/WorldPartition/DataLayer/DataLayerSubsystem.cpp
--- a/Engine/Source/Runtime/Engine/Private/WorldPartition/DataLayer/DataLayerSubsystem.cpp
+++ b/Engine/Source/Runtime/Engine/Private/WorldPartition/DataLayer/DataLayerSubsystem.cpp
@@ -138,10 +138,10 @@ UDataLayerInstance* UDataLayerSubsystem::GetDataLayerInstanceFromAsset(const UDa
 	return DataLayerManager ? const_cast<UDataLayerInstance*>(DataLayerManager->GetDataLayerInstance(InDataLayerAsset)) : nullptr;
 }
 
-const UDataLayerInstance* UDataLayerSubsystem::GetDataLayerInstanceFromAssetName(const FName& InDataLayerAssetFullName) const
+const UDataLayerInstance* UDataLayerSubsystem::GetDataLayerInstanceFromAssetName(const FName& InDataLayerAssetPathName) const
 {
 	const UDataLayerManager* DataLayerManager = UDataLayerManager::GetDataLayerManager(GetWorld());
-	return DataLayerManager ? DataLayerManager->GetDataLayerInstanceFromAssetName(InDataLayerAssetFullName) : nullptr;
+	return DataLayerManager ? DataLayerManager->GetDataLayerInstanceFromAssetName(InDataLayerAssetPathName) : nullptr;
 }
 
 PRAGMA_DISABLE_DEPRECATION_WARNINGS
diff --git a/Engine/Source/Runtime/Engine/Private/WorldPartition/DataLayer/WorldDataLayers.cpp b/Engine/Source/Runtime/Engine/Private/WorldPartition/DataLayer/WorldDataLayers.cpp
--- a/Engine/Source/Runtime/Engine/Private/WorldPartition/DataLayer/WorldDataLayers.cpp
+++ b/Engine/Source/Runtime/Engine/Private/WorldPartition/DataLayer/WorldDataLayers.cpp
@@ -1068,13 +1068,13 @@ const UDataLayerInstance* AWorldDataLayers::GetDataLayerInstance(const FName& In
 	return nullptr;
 }
 
-const UDataLayerInstance* AWorldDataLayers::GetDataLayerInstanceFromAssetName(const FName& InDataLayerAssetFullName) const
+const UDataLayerInstance* AWorldDataLayers::GetDataLayerInstanceFromAssetName(const FName& InDataLayerAssetPathName) const
 {
 #if WITH_EDITOR	
 	const UDataLayerInstance* FoundDataLayerInstance = nullptr;
-	ForEachDataLayerInstance([InDataLayerAssetFullName, &FoundDataLayerInstance](UDataLayerInstance* DataLayerInstance)
+	ForEachDataLayerInstance([InDataLayerAssetPathName, &FoundDataLayerInstance](UDataLayerInstance* DataLayerInstance)
 	{
-		if (DataLayerInstance->GetDataLayerFullName().Equals(InDataLayerAssetFullName.ToString(), ESearchCase::IgnoreCase))
+		if (DataLayerInstance->GetDataLayerFullName().Equals(InDataLayerAssetPathName.ToString(), ESearchCase::IgnoreCase))
 		{
 			FoundDataLayerInstance = DataLayerInstance;
 			return false;
@@ -1083,7 +1083,7 @@ const UDataLayerInstance* AWorldDataLayers::GetDataLayerInstanceFromAssetName(co
 	});
 	return FoundDataLayerInstance;
 #else
-	if (const UDataLayerInstance* const* FoundDataLayerInstance = AssetNameToInstance.Find(InDataLayerAssetFullName.ToString()))
+	if (const UDataLayerInstance* const* FoundDataLayerInstance = AssetNameToInstance.Find(InDataLayerAssetPathName.ToString()))
 	{
 		return *FoundDataLayerInstance;
 	}
@@ -1111,7 +1111,7 @@ const UDataLayerInstance* AWorldDataLayers::GetDataLayerInstance(const UDataLaye
 	});
 	return FoundDataLayerInstance;
 #else
-	if (const UDataLayerInstance* const* FoundDataLayerInstance = AssetNameToInstance.Find(InDataLayerAsset->GetFullName()))
+	if (const UDataLayerInstance* const* FoundDataLayerInstance = AssetNameToInstance.Find(InDataLayerAsset->GetPathName()))
 	{
 		return *FoundDataLayerInstance;
 	}
@@ -1344,21 +1344,21 @@ void AWorldDataLayers::UpdateAccelerationTable(const UDataLayerInstance* DataLay
 	static_assert(DATALAYER_TO_INSTANCE_RUNTIME_CONVERSION_ENABLED, "Remove unnecessary cast. All DataLayerInstance now have assets");
 	if (const UDataLayerInstanceWithAsset* DataLayerInstanceWithAsset = Cast<UDataLayerInstanceWithAsset>(DataLayerInstance))
 	{
-		if (!DataLayerInstanceWithAsset->GetAsset())
-		{
-			UE_LOG(LogWorldPartition, Warning, TEXT("DataLayerWithAsset %s has null asset."), *DataLayerInstanceWithAsset->GetPathName());
-		}
-		else
+		if (const UDataLayerAsset* DataLayerAsset = DataLayerInstanceWithAsset->GetAsset())
 		{
 			if (bIsAdding)
 			{
-				AssetNameToInstance.Add(DataLayerInstanceWithAsset->GetAsset()->GetFullName(), DataLayerInstance);
+				AssetNameToInstance.Add(DataLayerAsset->GetPathName(), DataLayerInstance);
 			}
 			else
 			{
-				AssetNameToInstance.Remove(DataLayerInstanceWithAsset->GetAsset()->GetFullName());
+				AssetNameToInstance.Remove(DataLayerAsset->GetPathName());
 			}
 		}
+		else
+		{
+			UE_LOG(LogWorldPartition, Warning, TEXT("DataLayerWithAsset %s has null asset."), *DataLayerInstanceWithAsset->GetPathName());
+		}
 	}
 }
 #endif
diff --git a/Engine/Source/Runtime/Engine/Public/WorldPartition/DataLayer/DataLayerManager.h b/Engine/Source/Runtime/Engine/Public/WorldPartition/DataLayer/DataLayerManager.h
--- a/Engine/Source/Runtime/Engine/Public/WorldPartition/DataLayer/DataLayerManager.h
+++ b/Engine/Source/Runtime/Engine/Public/WorldPartition/DataLayer/DataLayerManager.h
@@ -121,7 +121,7 @@ private:
 	//~ End Verse support
 
 	ENGINE_API AWorldDataLayers* GetWorldDataLayers() const;
-	ENGINE_API const UDataLayerInstance* GetDataLayerInstanceFromAssetName(const FName& InDataLayerAssetFullName) const;
+	ENGINE_API const UDataLayerInstance* GetDataLayerInstanceFromAssetName(const FName& InDataLayerAssetPathName) const;
 
 	ENGINE_API void BroadcastOnDataLayerInstanceRuntimeStateChanged(const UDataLayerInstance* InDataLayer, EDataLayerRuntimeState InState);
 
diff --git a/Engine/Source/Runtime/Engine/Public/WorldPartition/DataLayer/DataLayerSubsystem.h b/Engine/Source/Runtime/Engine/Public/WorldPartition/DataLayer/DataLayerSubsystem.h
--- a/Engine/Source/Runtime/Engine/Public/WorldPartition/DataLayer/DataLayerSubsystem.h
+++ b/Engine/Source/Runtime/Engine/Public/WorldPartition/DataLayer/DataLayerSubsystem.h
@@ -72,7 +72,7 @@ public:
 	TArray<FName> GetDataLayerInstanceNames(const TArray<T>& InDataLayerIdentifiers, const ULevel* InLevelContext = nullptr) const;
 
 	UE_DEPRECATED(5.3, "This function has been deprecated.")
-	ENGINE_API const UDataLayerInstance* GetDataLayerInstanceFromAssetName(const FName& InDataLayerAssetFullName) const;
+	ENGINE_API const UDataLayerInstance* GetDataLayerInstanceFromAssetName(const FName& InDataLayerAssetPathName) const;
 
 	UE_DEPRECATED(5.3, "Use UDataLayerManager::SetDataLayerRuntimeState instead.")
 	ENGINE_API void SetDataLayerRuntimeState(const UDataLayerInstance* InDataLayerInstance, EDataLayerRuntimeState InState, bool bInIsRecursive = false);
diff --git a/Engine/Source/Runtime/Engine/Public/WorldPartition/DataLayer/WorldDataLayers.h b/Engine/Source/Runtime/Engine/Public/WorldPartition/DataLayer/WorldDataLayers.h
--- a/Engine/Source/Runtime/Engine/Public/WorldPartition/DataLayer/WorldDataLayers.h
+++ b/Engine/Source/Runtime/Engine/Public/WorldPartition/DataLayer/WorldDataLayers.h
@@ -129,7 +129,7 @@ public:
 	ENGINE_API bool ContainsDataLayer(const UDataLayerInstance* InDataLayer) const;
 	ENGINE_API const UDataLayerInstance* GetDataLayerInstance(const FName& InDataLayerInstanceName) const;
 	ENGINE_API const UDataLayerInstance* GetDataLayerInstance(const UDataLayerAsset* InDataLayerAsset) const;
-	ENGINE_API const UDataLayerInstance* GetDataLayerInstanceFromAssetName(const FName& InDataLayerAssetFullName) const;
+	ENGINE_API const UDataLayerInstance* GetDataLayerInstanceFromAssetName(const FName& InDataLayerAssetPathName) const;
 
 	UE_DEPRECATED(5.4, "Use ForEachDataLayerInstance() instead.")
 	ENGINE_API void ForEachDataLayer(TFunctionRef<bool(UDataLayerInstance*)> Func) { return ForEachDataLayerInstance(Func); }
