Texture : factor out common sequence into GetMipImageInfo helper ; refactor mip functions to start from ImageInfo to get size.

diff --git a/Engine/Source/Runtime/Engine/Classes/Engine/Texture.h b/Engine/Source/Runtime/Engine/Classes/Engine/Texture.h
--- a/Engine/Source/Runtime/Engine/Classes/Engine/Texture.h
+++ b/Engine/Source/Runtime/Engine/Classes/Engine/Texture.h
@@ -294,7 +294,7 @@ struct FTextureSource
 		return GetMipData(OutMipData, 0, 0, MipIndex, ImageWrapperModule);
 	}
 
-	/** Retrieve a copy of the MipData as an FImage */
+	/** Retrieve a copy of the MipData as an FImage.  Uses LockMip.  Allocates OutImage and copies into it. */
 	ENGINE_API bool GetMipImage(FImage & OutImage, int32 BlockIndex, int32 LayerIndex, int32 MipIndex);
 	
 	/** Retrieve a copy of the MipData as an FImage */
@@ -306,7 +306,12 @@ struct FTextureSource
 	/** Returns a FMipData structure that wraps around the entire mip chain for read only operations. This is more efficient than calling the above method once per mip. */
 	ENGINE_API FMipData GetMipData(class IImageWrapperModule* ImageWrapperModule);
 
-	/** Computes the size of a single mip in bytes */
+	/** Get ImageInfo of one mip.  Does not require locking the mips, or accessing the payload data.
+	Returns false if the requested indexes were out of bounds, true if successful and fills OutImageInfo. **/
+	ENGINE_API bool GetMipImageInfo(FImageInfo & OutImageInfo, int32 BlockIndex, int32 LayerIndex, int32 MipIndex) const;
+
+	/** Computes the size of a single mip in bytes.  Same size as the image in GetMipImageInfo.
+	See also: CalcMipOffset. */
 	ENGINE_API int64 CalcMipSize(int32 BlockIndex, int32 LayerIndex, int32 MipIndex) const;
 
 	/** Computes the number of bytes per-pixel. */
@@ -635,7 +640,7 @@ private:
 	FMipAllocation LockedMipData;
 
 	// Internal implementation for locking the mip data, called by LockMipReadOnly or LockMip.
-	FMutableMemoryView LockMipInternal(int32 BlockIndex, int32 LayerIndex, int32 MipIndex, ELockState RequestedLockState);
+	FMutableMemoryView LockMipInternal(int32 BlockIndex, int32 LayerIndex, int32 MipIndex, ELockState RequestedLockState, FImageInfo & OutImageInfo);
 		
 	/** Returns the source data fully decompressed */
 	// ImageWrapperModule is not used
@@ -647,8 +652,6 @@ private:
 	*/
 	FSharedBuffer DoUEDeltaTransform(FSharedBuffer InBuffer,bool bForward) const;
 
-	/** Return true if the source art is not png compressed but could be. */
-	bool CanPNGCompress() const;
 	/** Retrieve the size and offset for a source mip. The size includes all slices. */
 	int64 CalcMipOffset(int32 BlockIndex, int32 LayerIndex, int32 MipIndex) const;
 	
diff --git a/Engine/Source/Runtime/Engine/Private/Texture.cpp b/Engine/Source/Runtime/Engine/Private/Texture.cpp
--- a/Engine/Source/Runtime/Engine/Private/Texture.cpp
+++ b/Engine/Source/Runtime/Engine/Private/Texture.cpp
@@ -475,7 +475,7 @@ void UTexture::BlockOnAnyAsyncBuild()
 	if (!IsAsyncCacheComplete())
 	{
 		FinishCachePlatformData();
-	}	
+	}
 	
 	if (IsDefaultTexture())
 	{
@@ -1331,12 +1331,6 @@ void UTexture::PreSave(FObjectPreSaveContext ObjectSaveContext)
 	Super::PreSave(ObjectSaveContext);
 
 #if WITH_EDITOR
-	if (DeferCompression)
-	{
-		GWarn->StatusUpdate( 0, 0, FText::Format( NSLOCTEXT("UnrealEd", "SavingPackage_CompressingTexture", "Compressing texture:  {0}"), FText::FromString(GetName()) ) );
-		DeferCompression = false;
-		UpdateResource();
-	}
 
 	// Ensure that compilation has finished before saving the package
 	// otherwise async compilation might try to read the bulkdata
@@ -1345,13 +1339,20 @@ void UTexture::PreSave(FObjectPreSaveContext ObjectSaveContext)
 	// because it invalidates the texture build due to source hash change
 	// and could cause another build to be triggered during PostCompilation
 	// causing reentrancy problems.
-	FTextureCompilingManager::Get().FinishCompilation({ this });
+	BlockOnAnyAsyncBuild();
 
 	if (!GEngine->IsAutosaving() && !ObjectSaveContext.IsProceduralSave())
 	{
 		GWarn->StatusUpdate(0, 0, FText::Format(NSLOCTEXT("UnrealEd", "SavingPackage_CompressingSourceArt", "Compressing source art for texture:  {0}"), FText::FromString(GetName())));
 		Source.Compress();
 	}
+	
+	if (DeferCompression)
+	{
+		GWarn->StatusUpdate( 0, 0, FText::Format( NSLOCTEXT("UnrealEd", "SavingPackage_CompressingTexture", "Compressing texture:  {0}"), FText::FromString(GetName()) ) );
+		DeferCompression = false;
+		UpdateResource();
+	}
 #endif // #if WITH_EDITOR
 }
 
@@ -2479,38 +2480,17 @@ FTextureSource::FMipLock::FMipLock(ELockState InLockState,FTextureSource * InTex
 	LayerIndex(InLayerIndex),
 	MipIndex(InMipIndex)
 {
-	FMutableMemoryView Locked = TextureSource->LockMipInternal(BlockIndex, LayerIndex, MipIndex, LockState);
-	if ( !Locked.IsEmpty() )
-	{	
-		FTextureSourceBlock Block;
-		TextureSource->GetBlock(BlockIndex, Block);
-		check(MipIndex < Block.NumMips);
-
-		Image.RawData = (uint8*)Locked.GetData();
-		Image.SizeX = FMath::Max(Block.SizeX >> MipIndex, 1);
-		Image.SizeY = FMath::Max(Block.SizeY >> MipIndex, 1);
-		Image.NumSlices = TextureSource->GetMippedNumSlices(Block.NumSlices,MipIndex);
-		Image.Format = FImageCoreUtils::ConvertToRawImageFormat(TextureSource->GetFormat(LayerIndex));
-		Image.GammaSpace = TextureSource->GetGammaSpace(LayerIndex);
-		
-		const int64 MipSizeBytes = TextureSource->CalcMipSize(BlockIndex, LayerIndex, MipIndex);
-
-		if (Image.GetImageSizeBytes() != Locked.GetSize())
-		{
-			// Don't just check on this one since it's actually potential OOB.
-			UE_LOG(LogTexture, Error, TEXT("Locked mip %d / block %d / layer %d has a format expecting %llu bytes but locked data is %llu, failing to lock!"),
-				InMipIndex, InBlockIndex, InLayerIndex, Image.GetImageSizeBytes(), Locked.GetSize());
-			Image = FImage();
-			LockState = ELockState::None;
-			return;
-		}
-
-		check( Image.GetImageSizeBytes() == MipSizeBytes );
-		check( IsValid() );
+	FMutableMemoryView Locked = TextureSource->LockMipInternal(BlockIndex, LayerIndex, MipIndex, LockState, Image);
+	if ( Locked.IsEmpty() )
+	{
+		Image = FImageView();
+		LockState = ELockState::None;
+		check( !IsValid() );
 	}
 	else
 	{
-		LockState = ELockState::None;
+		Image.RawData = (uint8*)Locked.GetData();
+		check( IsValid() );
 	}
 }
 
@@ -2546,15 +2526,19 @@ FTextureSource::FMipLock::~FMipLock()
 
 const uint8* FTextureSource::LockMipReadOnly(int32 BlockIndex, int32 LayerIndex, int32 MipIndex)
 {
-	return (const uint8*)LockMipInternal(BlockIndex, LayerIndex, MipIndex, ELockState::ReadOnly).GetData();
+	FImageInfo Info;
+	FMutableMemoryView View = LockMipInternal(BlockIndex, LayerIndex, MipIndex, ELockState::ReadOnly, Info);
+	return (const uint8*) View.GetData();
 }
 
 uint8* FTextureSource::LockMip(int32 BlockIndex, int32 LayerIndex, int32 MipIndex)
 {
-	return (uint8*)LockMipInternal(BlockIndex, LayerIndex, MipIndex, ELockState::ReadWrite).GetData();
+	FImageInfo Info;
+	FMutableMemoryView View = LockMipInternal(BlockIndex, LayerIndex, MipIndex, ELockState::ReadWrite, Info);
+	return (uint8*) View.GetData();
 }
 
-FMutableMemoryView FTextureSource::LockMipInternal(int32 BlockIndex, int32 LayerIndex, int32 MipIndex, ELockState RequestedLockState)
+FMutableMemoryView FTextureSource::LockMipInternal(int32 BlockIndex, int32 LayerIndex, int32 MipIndex, ELockState RequestedLockState, FImageInfo & OutImageInfo)
 {
 	TRACE_CPUPROFILER_EVENT_SCOPE(FTextureSource::LockMip);
 	
@@ -2592,83 +2576,85 @@ FMutableMemoryView FTextureSource::LockMipInternal(int32 BlockIndex, int32 Layer
 	//
 	// note: if you are using this to access mips one at a time, that is very inefficient unless you hold one lock ref throughout
 
-	FMutableMemoryView MipView;
-
-	// @@ overall NumMips or Block.NumMips ?
-	if (BlockIndex < GetNumBlocks() && LayerIndex < NumLayers && MipIndex < NumMips)
+	if ( ! GetMipImageInfo(OutImageInfo,BlockIndex,LayerIndex,MipIndex) )
 	{
-		if (NumLockedMips > 0 && RequestedLockState != LockState)
-		{
-			// previously locked, and requested lock is not the same as previous
-			if ( LockState == ELockState::ReadWrite )
-			{
-				// previous lock was for write
-				// we're requesting a read
-				//	allow it, promote our request to write
-				// this must be happening due to recursive locking, NOT from different threads
-				// if anyone has a write lock, texture multi-threading is not allowed
-				RequestedLockState = ELockState::ReadWrite;
-			}
-			else
-			{
-				// was previously locked for read, now wants to write
-				// that is not allowed, will fail
-				check( LockState == ELockState::ReadOnly );
-				check( RequestedLockState == ELockState::ReadWrite );
-
-				UE_LOG(LogTexture,Error, TEXT("LockMip cannot lock for write when previously locked for read [%s]"), 
-					Owner ? *Owner->GetFullName() : *TornOffOwnerName);	
-				
-				// no data, you did not get the lock, do not call Unlock
-				return MipView;
-			}
-		}
-
-		if (LockedMipData.IsNull())
-		{
-			checkf(NumLockedMips == 0, TEXT("Texture mips are locked but the LockedMipData is missing"));
-			LockedMipData = Decompress(nullptr);
-		}
-
-		if (RequestedLockState == ELockState::ReadOnly)
+		// failed, did not get lock, do not call Unlock
+		return FMutableMemoryView();
+	}
+	
+	if (NumLockedMips > 0 && RequestedLockState != LockState)
+	{
+		// previously locked, and requested lock is not the same as previous
+		if ( LockState == ELockState::ReadWrite )
 		{
-			// We cast away the const as the ReadOnly wrapper will put it back.
-			FSharedBuffer ReadOnlyMip = LockedMipData.GetDataReadOnly();
-			MipView = FMutableMemoryView((void*)ReadOnlyMip.GetData(), ReadOnlyMip.GetSize());
+			// previous lock was for write
+			// we're requesting a read
+			//	allow it, promote our request to write
+			// this must be happening due to recursive locking, NOT from different threads
+			// if anyone has a write lock, texture multi-threading is not allowed
+			RequestedLockState = ELockState::ReadWrite;
 		}
 		else
 		{
-			MipView = LockedMipData.GetDataReadWriteView();
-		}
+			// was previously locked for read, now wants to write
+			// that is not allowed, will fail
+			check( LockState == ELockState::ReadOnly );
+			check( RequestedLockState == ELockState::ReadWrite );
 
-		if ( MipView.IsEmpty() )
-		{
+			UE_LOG(LogTexture,Error, TEXT("LockMip cannot lock for write when previously locked for read [%s]"), 
+				Owner ? *Owner->GetFullName() : *TornOffOwnerName);	
+				
 			// no data, you did not get the lock, do not call Unlock
-			return MipView;
+			return FMutableMemoryView();
 		}
-		
-		int64 MipOffset = CalcMipOffset(BlockIndex, LayerIndex, MipIndex);
-		int64 MipSize = CalcMipSize(BlockIndex,LayerIndex,MipIndex);
+	}
 
-		MipView.MidInline(MipOffset, MipSize);
-		if (MipView.GetSize() != MipSize)
-		{
-			UE_LOG(LogTexture,Error,TEXT("Mip Data is too small : %lld < %lld+%lld"), LockedMipData.GetSize(),MipOffset,MipSize); 
-			LockedMipData.Reset();
-			return MipView;
-		}
+	if (LockedMipData.IsNull())
+	{
+		checkf(NumLockedMips == 0, TEXT("Texture mips are locked but the LockedMipData is missing"));
+		LockedMipData = Decompress(nullptr);
+	}
+	
+	FMutableMemoryView MipView;
 
-		if (NumLockedMips == 0)
-		{
-			LockState = RequestedLockState;
-		}
-		else
-		{
-			checkf(LockState == RequestedLockState, TEXT("Cannot change the lock type until UnlockMip is called"));
-		}
+	if (RequestedLockState == ELockState::ReadOnly)
+	{
+		// We cast away the const as the ReadOnly wrapper will put it back.
+		FSharedBuffer ReadOnlyMip = LockedMipData.GetDataReadOnly();
+		MipView = FMutableMemoryView((void*)ReadOnlyMip.GetData(), ReadOnlyMip.GetSize());
+	}
+	else
+	{
+		MipView = LockedMipData.GetDataReadWriteView();
+	}
 
-		++NumLockedMips;
+	if ( MipView.IsEmpty() )
+	{
+		// no data, you did not get the lock, do not call Unlock
+		return FMutableMemoryView();
 	}
+		
+	int64 MipOffset = CalcMipOffset(BlockIndex, LayerIndex, MipIndex);
+	int64 MipSize = OutImageInfo.GetImageSizeBytes();
+
+	MipView.MidInline(MipOffset, MipSize);
+	if (MipView.GetSize() != MipSize)
+	{
+		UE_LOG(LogTexture,Error,TEXT("Mip Data is too small : %lld < %lld+%lld"), LockedMipData.GetSize(),MipOffset,MipSize); 
+		LockedMipData.Reset();
+		return FMutableMemoryView();
+	}
+
+	if (NumLockedMips == 0)
+	{
+		LockState = RequestedLockState;
+	}
+	else
+	{
+		checkf(LockState == RequestedLockState, TEXT("Cannot change the lock type until UnlockMip is called"));
+	}
+
+	++NumLockedMips;
 
 	return MipView;
 }
@@ -2722,26 +2708,17 @@ void FTextureSource::UnlockMip(int32 BlockIndex, int32 LayerIndex, int32 MipInde
 
 bool FTextureSource::GetMipImage(FImage & OutImage, int32 BlockIndex, int32 LayerIndex, int32 MipIndex)
 {
-	TArray64<uint8> MipData;
-	if ( ! GetMipData(MipData,BlockIndex,LayerIndex,MipIndex) )
+
+	FMipLock MipLock(FTextureSource::ELockState::ReadOnly,this,BlockIndex,LayerIndex,MipIndex);
+
+	if( ! MipLock.IsValid() )
 	{
 		return false;
 	}
-	
-	// code dupe of FMipLock ; consider deleting this and use FMipLock.Image instead
-	
-	FTextureSourceBlock Block;
-	GetBlock(BlockIndex, Block);
-	check(MipIndex < Block.NumMips);
-
-	OutImage.RawData = MoveTemp(MipData);
-	OutImage.SizeX = FMath::Max(Block.SizeX >> MipIndex, 1);
-	OutImage.SizeY = FMath::Max(Block.SizeY >> MipIndex, 1);
-	OutImage.NumSlices = GetMippedNumSlices(Block.NumSlices,MipIndex);
-	OutImage.Format = FImageCoreUtils::ConvertToRawImageFormat(GetFormat(LayerIndex));
-	OutImage.GammaSpace = GetGammaSpace(LayerIndex);
 
-	check( OutImage.GetImageSizeBytes() == OutImage.RawData.Num() );
+	// MipLock.Image points into the lock sharedbuffer
+	//	allocate memory in destination and memcpy it out
+	MipLock.Image.CopyTo(OutImage);
 
 	return true;
 }
@@ -2750,28 +2727,18 @@ bool FTextureSource::GetMipData(TArray64<uint8>& OutMipData, int32 BlockIndex, i
 {
 	TRACE_CPUPROFILER_EVENT_SCOPE(FTextureSource::GetMipData (TArray64));
 	
-	// used by GetMipImage()
-
 	// note: do not use this to get all mips by calling GetMipData repeatedly, it's very inefficient, as it may decompress the source each time
 	//	instead use the GetMipData that returns all mips in one call
 
-	FMipLock MipLock(FTextureSource::ELockState::ReadOnly,this,BlockIndex,LayerIndex,MipIndex);
-
-	if( ! MipLock.IsValid() )
+	FImage MipImage;
+	if ( ! GetMipImage(MipImage,BlockIndex,LayerIndex,MipIndex) )
 	{
 		return false;
 	}
 
-	const void * MipData = MipLock.Image.RawData;
-	const int64 MipSize = MipLock.Image.GetImageSizeBytes();
-
-	check( MipSize > 0 );
+	OutMipData = MoveTemp(MipImage.RawData);
 
-	OutMipData.SetNumUninitialized(MipSize);
-	FMemory::Memcpy(
-		OutMipData.GetData(),
-		MipData,
-		MipSize);
+	check( OutMipData.Num() == MipImage.GetImageSizeBytes() );
 
 	return true;
 }
@@ -2797,18 +2764,43 @@ FTextureSource::FMipData FTextureSource::GetMipData(IImageWrapperModule* )
 	}
 }
 
-int64 FTextureSource::CalcMipSize(int32 BlockIndex, int32 LayerIndex, int32 MipIndex) const
+bool FTextureSource::GetMipImageInfo(FImageInfo & OutImage, int32 BlockIndex, int32 LayerIndex, int32 MipIndex) const
 {
+	if ( BlockIndex < 0 || BlockIndex >= GetNumBlocks() )
+	{
+		return false;
+	}
+	if ( LayerIndex < 0 || LayerIndex >= GetNumLayers() )
+	{
+		return false;
+	}
+
 	FTextureSourceBlock Block;
-	GetBlock(BlockIndex, Block);
-	check(MipIndex < Block.NumMips);
+	GetBlock(BlockIndex,Block);
 
-	const int64 MipSizeX = FMath::Max(Block.SizeX >> MipIndex, 1);
-	const int64 MipSizeY = FMath::Max(Block.SizeY >> MipIndex, 1);
-	const int64 MipSlices = GetMippedNumSlices(Block.NumSlices,MipIndex);
+	if ( MipIndex < 0 || MipIndex >= Block.NumMips )
+	{
+		return false;
+	}
+
+	OutImage.SizeX = FMath::Max(Block.SizeX >> MipIndex, 1);
+	OutImage.SizeY = FMath::Max(Block.SizeY >> MipIndex, 1);
+	OutImage.NumSlices = GetMippedNumSlices(Block.NumSlices,MipIndex);
+	OutImage.Format = FImageCoreUtils::ConvertToRawImageFormat(GetFormat(LayerIndex));
+	OutImage.GammaSpace = GetGammaSpace(LayerIndex);
+
+	return true;
+}
+
+int64 FTextureSource::CalcMipSize(int32 BlockIndex, int32 LayerIndex, int32 MipIndex) const
+{
+	FImageInfo Image;
+	if ( ! GetMipImageInfo(Image,BlockIndex,LayerIndex,MipIndex) )
+	{
+		return 0;
+	}
 
-	const int64 BytesPerPixel = GetBytesPerPixel(LayerIndex);
-	return MipSizeX * MipSizeY * MipSlices * BytesPerPixel;
+	return Image.GetImageSizeBytes();
 }
 
 int64 FTextureSource::GetBytesPerPixel(int32 LayerIndex) const
@@ -2840,6 +2832,9 @@ bool FTextureSource::AreAllBlocksPowerOfTwo() const
 
 bool FTextureSource::IsValid() const
 {
+	// note: the check of HasPayloadData() means that during Init() we are not yet IsValid() until the BulkData is set
+	// a zero size TextureSource is considered not valid
+
 	return SizeX > 0 && SizeY > 0 && NumSlices > 0 && NumLayers > 0 && NumMips > 0 &&
 		Format != TSF_Invalid && HasPayloadData();
 }
@@ -3050,22 +3045,17 @@ FSharedBuffer FTextureSource::DoUEDeltaTransform(FSharedBuffer InBuffer,bool bFo
 		{
 			for(int64 MipIndex=0;MipIndex<Block.NumMips;MipIndex++)
 			{
-				// @@ this is code-duped a lot, factor it out :
 				int64 MipOffset = CalcMipOffset(BlockIndex, LayerIndex, MipIndex);
-				int64 MipSize = CalcMipSize(BlockIndex,LayerIndex,MipIndex);
+				
+				FImageView Image;
+				verify( GetMipImageInfo(Image,BlockIndex,LayerIndex,MipIndex) );
+
+				int64 MipSize = Image.GetImageSizeBytes();
 				check( MipOffset+MipSize <= ImageSize );
 
 				const uint8 * InPtr = InData + MipOffset;
 
-				FImageView Image;
 				Image.RawData = (void *)InPtr;
-				Image.SizeX = FMath::Max(Block.SizeX >> MipIndex, 1);
-				Image.SizeY = FMath::Max(Block.SizeY >> MipIndex, 1);
-				Image.NumSlices = GetMippedNumSlices(Block.NumSlices,MipIndex);
-				Image.Format = FImageCoreUtils::ConvertToRawImageFormat(GetFormat(LayerIndex));
-				Image.GammaSpace = GetGammaSpace(LayerIndex);
-				
-				check( Image.GetImageSizeBytes() == MipSize );
 				
 				FImageCoreDelta::AddSplitStridedViewsForDelta( ImageViewPortions, Image);
 			}
@@ -3275,26 +3265,6 @@ void FTextureSource::ImportCustomProperties(const TCHAR* SourceText, FFeedbackCo
 	}
 }
 
-bool FTextureSource::CanPNGCompress() const
-{
-	bool bCanPngCompressFormat = (Format == TSF_G8 || Format == TSF_G16 || Format == TSF_BGRA8 || Format == TSF_RGBA16 );
-
-	if (
-		NumLayers == 1 &&
-		NumMips == 1 &&
-		NumSlices == 1 &&
-		Blocks.Num() == 0 &&
-		SizeX > 4 &&
-		SizeY > 4 &&
-		HasPayloadData() &&
-		bCanPngCompressFormat &&
-		CompressionFormat == TSCF_None)
-	{
-		return true;
-	}
-	return false;
-}
-
 void FTextureSource::ForceGenerateGuid()
 {
 	Id = FGuid::NewGuid();
@@ -3382,10 +3352,11 @@ int64 FTextureSource::CalcLayerSize(const FTextureSourceBlock& Block, int32 Laye
 
 	int64 BytesPerPixel = GetBytesPerPixel(LayerIndex);
 
-	// This is used for memory allocation, so use FCheckedInt to rigorously check against overflow issues.
+	// This is used for memory allocation, so use FGuardedInt64 to rigorously check against overflow issues.
 	FGuardedInt64 TotalSize(0);
 	for (int32 MipIndex = 0; MipIndex < Block.NumMips; ++MipIndex)
 	{
+		// == CalcMipSize
 		int32 MipSizeX = FMath::Max<int32>(Block.SizeX >> MipIndex, 1);
 		int32 MipSizeY = FMath::Max<int32>(Block.SizeY >> MipIndex, 1);
 		int32 MipSizeZ = GetMippedNumSlices(Block.NumSlices,MipIndex);
@@ -3405,11 +3376,30 @@ int64 FTextureSource::CalcMipOffset(int32 BlockIndex, int32 LayerIndex, int32 Of
 		return 0;
 	}
 
+	/*************
+
+	Memory layout :
+
+	[Block 0            ][Block 1       ]
+	[[layer     ][layer]][[layer][layer]]
+	[[[mip][mip]][[mip]]][...
+
+	Block start positions are cached in BlockDataOffsets[]
+	then you step over whole layers
+	then step into all mips on a layer
+
+	note these are the mips in the *source*, not the number of mips generated
+
+	*************/
+
+	check( BlockIndex < GetNumBlocks() );
+	check( LayerIndex < GetNumLayers() );
+
 	FTextureSourceBlock Block;
 	GetBlock(BlockIndex, Block);
 	check(OffsetToMipIndex < Block.NumMips);
 
-	// This is used for memory indexing, so use FCheckedInt to rigorously check against overflow issues.
+	// This is used for memory indexing, so use FGuardedInt64 to rigorously check against overflow issues.
 	FGuardedInt64 MipOffset(BlockDataOffsets[BlockIndex]);
 
 	// Skip over the initial layers within the tile
@@ -3422,6 +3412,7 @@ int64 FTextureSource::CalcMipOffset(int32 BlockIndex, int32 LayerIndex, int32 Of
 
 	for (int32 MipIndex = 0; MipIndex < OffsetToMipIndex; ++MipIndex)
 	{
+		// == CalcMipSize
 		int32 MipSizeX = FMath::Max<int32>(Block.SizeX >> MipIndex, 1);
 		int32 MipSizeY = FMath::Max<int32>(Block.SizeY >> MipIndex, 1);
 		int32 MipSizeZ = GetMippedNumSlices(Block.NumSlices,MipIndex);
@@ -4412,8 +4403,14 @@ bool FTextureSource::FMipData::GetMipData(TArray64<uint8>& OutMipData, int32 Blo
 	{
 		const int64 MipOffset = TextureSource.CalcMipOffset(BlockIndex, LayerIndex, MipIndex);
 		const int64 MipSize = TextureSource.CalcMipSize(BlockIndex, LayerIndex, MipIndex);
+		
+		FGuardedInt64 GuardedMipEnd(MipOffset);
+		GuardedMipEnd += MipSize;
+		uint64 MipEnd = (uint64) GuardedMipEnd.Get(-1);
+
+		check( MipEnd <= MipData.GetSize() );
 
-		if ((int64)MipData.GetSize() >= MipOffset + MipSize)
+		if ( MipEnd <= MipData.GetSize() )
 		{
 			OutMipData.Empty(MipSize);
 			OutMipData.AddUninitialized(MipSize);
@@ -4430,14 +4427,26 @@ bool FTextureSource::FMipData::GetMipData(TArray64<uint8>& OutMipData, int32 Blo
 	return false;
 }
 
-FSharedBuffer FTextureSource::FMipData::GetMipData(int32 BlockIndex, int32 LayerIndex, int32 MipIndex) const
+// FSharedBuffer returned is a subview and doesn't allocate a smaller buffer - but will also hold a ref to the full allocation!
+FSharedBuffer FTextureSource::FMipData::GetMipDataWithInfo(int32 BlockIndex, int32 LayerIndex, int32 MipIndex, FImageInfo& OutImageInfo) const
 {
-	if (BlockIndex < TextureSource.GetNumBlocks() && LayerIndex < TextureSource.GetNumLayers() && MipIndex < TextureSource.GetNumMips() && !MipData.IsNull())
+	if ( MipData.IsNull() )
+	{
+		return MipData;
+	}
+
+	if ( TextureSource.GetMipImageInfo(OutImageInfo,BlockIndex,LayerIndex,MipIndex) )
 	{
 		const int64 MipOffset = TextureSource.CalcMipOffset(BlockIndex, LayerIndex, MipIndex);
-		const int64 MipSize = TextureSource.CalcMipSize(BlockIndex, LayerIndex, MipIndex);
+		const int64 MipSize = OutImageInfo.GetImageSizeBytes();
+
+		FGuardedInt64 GuardedMipEnd(MipOffset);
+		GuardedMipEnd += MipSize;
+		uint64 MipEnd = (uint64) GuardedMipEnd.Get(-1);
 
-		if ((int64)MipData.GetSize() >= MipOffset + MipSize)
+		check( MipEnd <= MipData.GetSize() );
+
+		if ( MipEnd <= MipData.GetSize() )
 		{
 			return FSharedBuffer::MakeView((const uint8*)MipData.GetData() + MipOffset, MipSize, MipData);
 		}
@@ -4446,24 +4455,10 @@ FSharedBuffer FTextureSource::FMipData::GetMipData(int32 BlockIndex, int32 Layer
 	return FSharedBuffer();
 }
 
-FSharedBuffer FTextureSource::FMipData::GetMipDataWithInfo(int32 InBlockIndex, int32 InLayerIndex, int32 InMipIndex, FImageInfo& OutImageInfo) const
+FSharedBuffer FTextureSource::FMipData::GetMipData(int32 BlockIndex, int32 LayerIndex, int32 MipIndex) const
 {
-	// This is a subview and doesn't allocate a smaller buffer - but will also hold the full allocation!
-	FSharedBuffer MipDataView = GetMipData(InBlockIndex, InLayerIndex, InMipIndex);
-	if (MipDataView.IsNull())
-	{
-		return MipDataView;
-	}
-
-	FTextureSourceBlock Block;
-	TextureSource.GetBlock(InBlockIndex, Block);
-
-	OutImageInfo.SizeX = FMath::Max(Block.SizeX >> InMipIndex, 1);
-	OutImageInfo.SizeY = FMath::Max(Block.SizeY >> InMipIndex, 1);
-	OutImageInfo.NumSlices = TextureSource.GetMippedNumSlices(Block.NumSlices, InMipIndex);
-	OutImageInfo.Format = FImageCoreUtils::ConvertToRawImageFormat(TextureSource.GetFormat(InLayerIndex));
-	OutImageInfo.GammaSpace = TextureSource.GetGammaSpace(InLayerIndex);
-	return MipDataView;
+	FImageInfo Info;
+	return GetMipDataWithInfo(BlockIndex,LayerIndex,MipIndex,Info);
 }
 
 #endif //WITH_EDITOR
@@ -4556,7 +4551,9 @@ bool FTextureSource::ComputeChannelLinearMinMax(int32 InLayerIndex, FLinearColor
 		// The data is already present and locked from the mip0 lock above, this just gets us the imageview
 		// Note we only look at mip 0 ; it is possible that other mips go out of the MinMax bound we find.
 		// -> should probably fix this
-		FTextureSource::FMipLock LockedBlock(ELockState::ReadOnly, const_cast<FTextureSource*>(this), BlockIndex, InLayerIndex, 0);
+		int32 MipIndex = 0;
+
+		FTextureSource::FMipLock LockedBlock(ELockState::ReadOnly, const_cast<FTextureSource*>(this), BlockIndex, InLayerIndex, MipIndex);
 		check(LockedBlock.IsValid()); // should be same as validity check above!!
 
 		FLinearColor MinColor, MaxColor;
@@ -4625,25 +4622,21 @@ bool FTextureSource::UpdateChannelMinMaxFromIncomingTextureData(FMemoryView InNe
 
 		for (int32 BlockIndex = 0; BlockIndex < GetNumBlocks(); BlockIndex++)
 		{
-			FTextureSourceBlock Block;
-			GetBlock(BlockIndex, Block);
-
 			// note: only does mip0 of each layer/block !!
 			// -> should probably fix this
-			int64 MipOffset = CalcMipOffset(BlockIndex, LayerIndex, 0);
-			int64 MipSize = CalcMipSize(BlockIndex, LayerIndex, 0);
+			int32 MipIndex = 0;
+			
+			FImageView Image;
+			verify( GetMipImageInfo(Image,BlockIndex,LayerIndex,MipIndex) );
+
+			int64 MipOffset = CalcMipOffset(BlockIndex, LayerIndex, MipIndex);
+			int64 MipSize = Image.GetImageSizeBytes();
 
 			FMemoryView MipView = InNewTextureData.Mid(MipOffset, MipSize);
 
 			if (MipView.GetSize() == MipSize)
 			{
-				FImageView Image;
 				Image.RawData = (void*)MipView.GetData();
-				Image.SizeX = FMath::Max(Block.SizeX, 1);
-				Image.SizeY = FMath::Max(Block.SizeY, 1);
-				Image.NumSlices = GetMippedNumSlices(Block.NumSlices, 0);
-				Image.Format = FImageCoreUtils::ConvertToRawImageFormat(GetFormat(LayerIndex));
-				Image.GammaSpace = GetGammaSpace(LayerIndex);
 
 				FLinearColor MinColor, MaxColor;
 				FImageCore::ComputeChannelLinearMinMax(Image, MinColor, MaxColor);
