Use current LOD when updating deformer instance.
Without this, if we only set the deformer on some LODs then we can end up updating an invalid LOD.
Also ensure that we call AllocateResources/ReleaseResources for all valid LOD deformer instances.

diff --git a/Engine/Source/Runtime/Engine/Private/Components/SkinnedMeshComponent.cpp b/Engine/Source/Runtime/Engine/Private/Components/SkinnedMeshComponent.cpp
--- a/Engine/Source/Runtime/Engine/Private/Components/SkinnedMeshComponent.cpp
+++ b/Engine/Source/Runtime/Engine/Private/Components/SkinnedMeshComponent.cpp
@@ -1071,6 +1071,8 @@ void USkinnedMeshComponent::CreateRenderState_Concurrent(FRegisterComponentConte
 
 	Super::CreateRenderState_Concurrent(Context);
 
+	int32 UseLOD = 0;
+
 	if (GetSkinnedAsset())
 	{
 		BoneTransformUpdateMethodQueue.Reset();
@@ -1084,22 +1086,22 @@ void USkinnedMeshComponent::CreateRenderState_Concurrent(FRegisterComponentConte
 			//	without animated, causing random skinning issues
 			// This can happen if your MinLOD is not valid anymore after loading
 			// which causes meshes to be invisible
-			int32 ModifiedLODLevel = GetPredictedLODLevel();
+			UseLOD = GetPredictedLODLevel();
 			{
 				int32 MinLodIndex = ComputeMinLOD();
 				int32 MaxLODIndex = MeshObject->GetSkeletalMeshRenderData().LODRenderData.Num() - 1;
-				ModifiedLODLevel = FMath::Clamp(ModifiedLODLevel, MinLodIndex, MaxLODIndex);
+				UseLOD = FMath::Clamp(UseLOD, MinLodIndex, MaxLODIndex);
 			}
 
 			// Clamp to loaded streaming data if available
 			if ((GetSkinnedAsset()->IsStreamable() || !IStreamingManager::Get().IsRenderAssetStreamingEnabled(EStreamableRenderAssetType::SkeletalMesh)) && MeshObject)
 			{
-				ModifiedLODLevel = FMath::Max<int32>(ModifiedLODLevel, MeshObject->GetSkeletalMeshRenderData().PendingFirstLODIdx);
+				UseLOD = FMath::Max<int32>(UseLOD, MeshObject->GetSkeletalMeshRenderData().PendingFirstLODIdx);
 			}
 
 			// If we have a valid LOD, set up required data, during reimport we may try to create data before we have all the LODs
 			// imported, in that case we skip until we have all the LODs
-			if(GetSkinnedAsset()->IsValidLODIndex(ModifiedLODLevel))
+			if(GetSkinnedAsset()->IsValidLODIndex(UseLOD))
 			{
 				const bool bMorphTargetsAllowed = CVarEnableMorphTargets.GetValueOnAnyThread(true) != 0;
 
@@ -1111,7 +1113,7 @@ void USkinnedMeshComponent::CreateRenderState_Concurrent(FRegisterComponentConte
 
 				constexpr bool bZeroOldWeights = false;
 				RefreshExternalMorphTargetWeights(bZeroOldWeights);
-				MeshObject->Update(ModifiedLODLevel, this, ActiveMorphTargets, MorphTargetWeights, EPreviousBoneTransformUpdateMode::UpdatePrevious, GetExternalMorphWeights(ModifiedLODLevel));  // send to rendering thread
+				MeshObject->Update(UseLOD, this, ActiveMorphTargets, MorphTargetWeights, EPreviousBoneTransformUpdateMode::UpdatePrevious, GetExternalMorphWeights(UseLOD));  // send to rendering thread
 			}
 		}
 
@@ -1119,10 +1121,16 @@ void USkinnedMeshComponent::CreateRenderState_Concurrent(FRegisterComponentConte
 		UpdateMorphMaterialUsageOnProxy();
 	}
 
- 	if (UMeshDeformerInstance* MeshDeformerInstance = GetMeshDeformerInstance())
+	for (UMeshDeformerInstance* MeshDeformerInstance : MeshDeformerInstances.DeformerInstances)
+	{
+		if (MeshDeformerInstance != nullptr)
+		{
+			MeshDeformerInstance->AllocateResources();
+		}
+	}
+
+	if (UMeshDeformerInstance* MeshDeformerInstance = GetMeshDeformerInstanceForLOD(UseLOD))
  	{
-		MeshDeformerInstance->AllocateResources();
-		
 		// Enqueue immediate execution of work here to ensure that we have some deformer outputs written for the next frame.
 		UMeshDeformerInstance::FEnqueueWorkDesc Desc;
 		Desc.Scene = GetScene();
@@ -1136,9 +1144,12 @@ void USkinnedMeshComponent::DestroyRenderState_Concurrent()
 {
 	Super::DestroyRenderState_Concurrent();
 
-	if (UMeshDeformerInstance* MeshDeformerInstance = GetMeshDeformerInstance())
+	for (UMeshDeformerInstance* MeshDeformerInstance : MeshDeformerInstances.DeformerInstances)
 	{
-		MeshDeformerInstance->ReleaseResources();
+		if (MeshDeformerInstance != nullptr)
+		{
+			MeshDeformerInstance->ReleaseResources();
+		}
 	}
 
 	if(MeshObject)
@@ -2136,7 +2147,7 @@ bool USkinnedMeshComponent::IsSkinCacheAllowed(int32 LodIdx) const
 	static const IConsoleVariable* CVarDefaultGPUSkinCacheBehavior = IConsoleManager::Get().FindConsoleVariable(TEXT("r.SkinCache.DefaultBehavior"));
 	const bool bGlobalDefault = CVarDefaultGPUSkinCacheBehavior && ESkinCacheDefaultBehavior(CVarDefaultGPUSkinCacheBehavior->GetInt()) == ESkinCacheDefaultBehavior::Inclusive;
 
-	if (GetMeshDeformerInstance() != nullptr)
+	if (GetMeshDeformerInstanceForLOD(LodIdx) != nullptr)
 	{
 		// Disable skin cache if a mesh deformer is in use.
 		// Any animation buffers are expected to be owned by the MeshDeformer.
