USkinnedMeshComponent (min LOD validation);

Adding a validation step to the SetMinLOD(), OverrideMinLOD(), and ComputeMinLOD() function calls so we aren't setting to a value that doesn't have render data. This was seen on lower end platforms where the mesh would be invisible because there was no render data for the min LOD value being set.

-- Code --
SkinnedMeshComponent:
GetValidMinLOD() - Validate the min LOD value of the mesh component by iterating over render data to make sure we get something usable.
SetMinLOD() - Calling GetValidMinLOD().
OverrideMinLOD() - Calling GetValidMinLOD().
ComputeMinLOD() - Calling GetValidMinLOD().

#tests PIE - since this was on low end platforms only, I changed the skeletal mesh to only let us do LOD3, I then changed my graphics to the lowest quality, I then spawned multiple actors that are using the mesh to make sure it didn't turn invisible.

diff --git a/Engine/Source/Runtime/Engine/Classes/Components/SkinnedMeshComponent.h b/Engine/Source/Runtime/Engine/Classes/Components/SkinnedMeshComponent.h
--- a/Engine/Source/Runtime/Engine/Classes/Components/SkinnedMeshComponent.h
+++ b/Engine/Source/Runtime/Engine/Classes/Components/SkinnedMeshComponent.h
@@ -1116,6 +1116,13 @@ public:
 	 */
 	ENGINE_API int32 ComputeMinLOD() const;
 
+	/**
+	 * Validate the min LOD value of the mesh component by iterating over render data to make sure we get something usable.
+	 *
+	 * @param	InMinLOD	The starting LOD index to iterate from to make sure the render data is valid. Range from [0, Max Number of LOD - 1].
+	 */
+	int32 GetValidMinLOD(const int32 InMinLODIndex) const;
+
 public:
 	//~ Begin UObject Interface
 	ENGINE_API virtual void BeginDestroy() override;
diff --git a/Engine/Source/Runtime/Engine/Private/Components/SkinnedMeshComponent.cpp b/Engine/Source/Runtime/Engine/Private/Components/SkinnedMeshComponent.cpp
--- a/Engine/Source/Runtime/Engine/Private/Components/SkinnedMeshComponent.cpp
+++ b/Engine/Source/Runtime/Engine/Private/Components/SkinnedMeshComponent.cpp
@@ -3949,30 +3949,46 @@ int32 USkinnedMeshComponent::GetNumLODs() const
 	return NumLODs;
 }
 
-
 void USkinnedMeshComponent::SetMinLOD(int32 InNewMinLOD)
 {
-	int32 MaxLODIndex = GetNumLODs() - 1;
-	MinLodModel = FMath::Clamp(InNewMinLOD, 0, MaxLODIndex);
+	MinLodModel = GetValidMinLOD(InNewMinLOD);
 }
 
 void USkinnedMeshComponent::OverrideMinLOD(int32 InNewMinLOD)
 {
-	int32 MaxLODIndex = GetNumLODs() - 1;
-	MinLodModel = FMath::Clamp(InNewMinLOD, 0, MaxLODIndex);
+	MinLodModel = GetValidMinLOD(InNewMinLOD);
 	bOverrideMinLod = true;
 }
 
 int32 USkinnedMeshComponent::ComputeMinLOD() const
 {
-	int32 AssetMinLod = GetSkinnedAsset()->GetMinLodIdx();
+	const int32 AssetMinLOD = GetSkinnedAsset()->GetMinLodIdx();
 	// overriden MinLOD can't be higher than asset MinLOD
-	int32 MinLodIndex = bOverrideMinLod ? FMath::Max(MinLodModel, AssetMinLod) : AssetMinLod;
-	int32 NumLODs = GetNumLODs();
-	// want to make sure MinLOD stays within the valid range
-	MinLodIndex = FMath::Min(MinLodIndex, NumLODs - 1);
-	MinLodIndex = FMath::Max(MinLodIndex, 0);
-	return MinLodIndex;
+	int32 MinLODIndex = bOverrideMinLod ? FMath::Max(MinLodModel, AssetMinLOD) : AssetMinLOD;
+	MinLODIndex = GetValidMinLOD(MinLODIndex);
+	return MinLODIndex;
+}
+
+int32 USkinnedMeshComponent::GetValidMinLOD(const int32 InMinLODIndex) const
+{
+	// Iterate the render data to validate that our min LOD has data that can be used.
+	const int32 MaxLODIndex = GetNumLODs() - 1;
+	const FSkeletalMeshRenderData* RenderData = GetSkeletalMeshRenderData();
+	const int32 FirstValidLODIndex = RenderData != nullptr ? RenderData->GetFirstValidLODIdx(InMinLODIndex) : INDEX_NONE;
+	UE_CLOG(FirstValidLODIndex == INDEX_NONE
+		, LogSkinnedMeshComp, Log
+		, TEXT("%hs We're trying to override min LOD for [%s] with no valid LOD index from [%i] to [%i]")
+		, __FUNCTION__
+		, *GetName()
+		, InMinLODIndex
+		, MaxLODIndex
+	);
+
+	// Return the first LOD that has render data that can be used.
+	/** NOTE: We're logging if the index is not valid in the render data but we still want to return a valid value from 0 to max. 
+	  * Render data could be invalid if we're still loading/streaming in the asset.
+	  */
+	return FMath::Clamp<int32>(FirstValidLODIndex, 0, MaxLODIndex);
 }
 
 #if WITH_EDITOR
