Fixed crash in skeletal mesh if source buffer does not exist. Also fixed skeletal mesh for base vertex factory to actually use cached mesh commands.

diff --git a/Engine/Source/Runtime/Engine/Private/GPUSkinVertexFactory.cpp b/Engine/Source/Runtime/Engine/Private/GPUSkinVertexFactory.cpp
--- a/Engine/Source/Runtime/Engine/Private/GPUSkinVertexFactory.cpp
+++ b/Engine/Source/Runtime/Engine/Private/GPUSkinVertexFactory.cpp
@@ -1002,6 +1002,7 @@ IMPLEMENT_GPUSKINNING_VERTEX_FACTORY_TYPE(TGPUSkinVertexFactory, "/Engine/Privat
 	  EVertexFactoryFlags::UsedWithMaterials 
 	| EVertexFactoryFlags::SupportsDynamicLighting
 	| EVertexFactoryFlags::SupportsPSOPrecaching
+	| EVertexFactoryFlags::SupportsCachingMeshDrawCommands
 );
 
 /*-----------------------------------------------------------------------------
@@ -1252,6 +1253,7 @@ IMPLEMENT_GPUSKINNING_VERTEX_FACTORY_TYPE(TGPUSkinAPEXClothVertexFactory, "/Engi
 	  EVertexFactoryFlags::UsedWithMaterials
 	| EVertexFactoryFlags::SupportsDynamicLighting
 	| EVertexFactoryFlags::SupportsPSOPrecaching
+	| EVertexFactoryFlags::SupportsCachingMeshDrawCommands
 );
 
 
@@ -1292,10 +1294,16 @@ void FGPUSkinPassthroughVertexFactory::InitRHI(FRHICommandListBase& RHICmdList)
 		return bSupportsManualVertexFetch && EnumHasAnyFlags(Component.VertexStreamUsage, EVertexStreamUsage::ManualFetch);
 	};
 
+	const auto GetVertexBufferRHI = [] (const FVertexBuffer* VertexBuffer) -> FRHIBuffer*
+	{
+		return VertexBuffer ? VertexBuffer->GetRHI() : GNullVertexBuffer.GetRHI();
+	};
+
 	if (EnumHasAnyFlags(VertexAttributesRequested, EVertexAttributeFlags::Position))
 	{
-		SourceStreamBuffers[EVertexAttribute::VertexPosition] = Data.PositionComponent.VertexBuffer->GetRHI();
-		StreamSourceSlots[EVertexAttribute::VertexPosition] = FRHIStreamSourceSlot::Create(Data.PositionComponent.VertexBuffer->GetRHI());
+		FRHIBuffer* Buffer = GetVertexBufferRHI(Data.PositionComponent.VertexBuffer);
+		SourceStreamBuffers[EVertexAttribute::VertexPosition] = Buffer;
+		StreamSourceSlots[EVertexAttribute::VertexPosition] = FRHIStreamSourceSlot::Create(Buffer);
 		Data.PositionComponent.Offset = 0;
 		Data.PositionComponent.VertexStreamUsage |= EVertexStreamUsage::Overridden;
 		Data.PositionComponent.Stride = 3 * sizeof(float);
@@ -1306,8 +1314,9 @@ void FGPUSkinPassthroughVertexFactory::InitRHI(FRHICommandListBase& RHICmdList)
 	{
 		if (!IsManualVertexFetch(Data.ColorComponent))
 		{
-			SourceStreamBuffers[EVertexAttribute::VertexColor] = Data.ColorComponent.VertexBuffer->GetRHI();
-			StreamSourceSlots[EVertexAttribute::VertexColor] = FRHIStreamSourceSlot::Create(Data.ColorComponent.VertexBuffer->GetRHI());
+			FRHIBuffer* Buffer = GetVertexBufferRHI(Data.ColorComponent.VertexBuffer);
+			SourceStreamBuffers[EVertexAttribute::VertexColor] = Buffer;
+			StreamSourceSlots[EVertexAttribute::VertexColor] = FRHIStreamSourceSlot::Create(Buffer);
 			Data.ColorComponent.Offset = 0;
 			Data.ColorComponent.Type = VET_Color;
 			Data.ColorComponent.VertexStreamUsage |= EVertexStreamUsage::Overridden;
@@ -1322,8 +1331,9 @@ void FGPUSkinPassthroughVertexFactory::InitRHI(FRHICommandListBase& RHICmdList)
 
 	if (EnumHasAnyFlags(VertexAttributesRequested, EVertexAttributeFlags::Tangent) && !IsManualVertexFetch(Data.TangentBasisComponents[0]))
 	{
-		SourceStreamBuffers[EVertexAttribute::VertexTangent] = Data.TangentBasisComponents[0].VertexBuffer->GetRHI();
-		StreamSourceSlots[EVertexAttribute::VertexTangent] = FRHIStreamSourceSlot::Create(Data.TangentBasisComponents[0].VertexBuffer->GetRHI());
+		FRHIBuffer* Buffer = GetVertexBufferRHI(Data.TangentBasisComponents[0].VertexBuffer);
+		SourceStreamBuffers[EVertexAttribute::VertexTangent] = Buffer;
+		StreamSourceSlots[EVertexAttribute::VertexTangent] = FRHIStreamSourceSlot::Create(Buffer);
 		Data.TangentBasisComponents[0].VertexStreamUsage |= EVertexStreamUsage::Overridden;
 		Data.TangentBasisComponents[0].Offset = 0;
 		Data.TangentBasisComponents[0].Type = VET_Short4N;
diff --git a/Engine/Source/Runtime/Engine/Private/SkeletalMesh.cpp b/Engine/Source/Runtime/Engine/Private/SkeletalMesh.cpp
--- a/Engine/Source/Runtime/Engine/Private/SkeletalMesh.cpp
+++ b/Engine/Source/Runtime/Engine/Private/SkeletalMesh.cpp
@@ -7144,14 +7144,13 @@ FPrimitiveViewRelevance FSkeletalMeshSceneProxy::GetViewRelevance(const FSceneVi
 		return IsRichView(*View->Family)
 			|| EngineShowFlags.Bones
 			|| EngineShowFlags.Collision
-			|| EngineShowFlags.Bounds
-			|| EngineShowFlags.VertexColors
-			|| IsSelected()
+			|| EngineShowFlags.VisualizeGPUSkinCache
+			|| (IsSelected() && (EngineShowFlags.VertexColors || EngineShowFlags.Bounds))
 #if WITH_EDITORONLY_DATA
 			|| MeshObject->SelectedEditorMaterial != -1
 			|| MeshObject->SelectedEditorSection != -1
 #endif
-			|| GetGPUSkinCacheVisualizationData().IsActive();
+			;
 #else
 		return false;
 #endif
