TextureDefines.h compression format descriptions remove pre-dx11 formats

diff --git a/Engine/Source/Runtime/Engine/Classes/Engine/TextureDefines.h b/Engine/Source/Runtime/Engine/Classes/Engine/TextureDefines.h
--- a/Engine/Source/Runtime/Engine/Classes/Engine/TextureDefines.h
+++ b/Engine/Source/Runtime/Engine/Classes/Engine/TextureDefines.h
@@ -374,18 +374,18 @@ extern ENGINE_API FTextureSourceFormatInfo GTextureSourceFormats[TSF_MAX];		// M
 UENUM()
 enum TextureCompressionSettings : int
 {
-	TC_Default					UMETA(DisplayName = "Default (DXT1/5, BC1/3 on DX11)"),
-	TC_Normalmap				UMETA(DisplayName = "Normalmap (DXT5, BC5 on DX11)"),
+	TC_Default					UMETA(DisplayName = "Default (BC1 or BC3 with A)"),
+	TC_Normalmap				UMETA(DisplayName = "Normalmap (BC5)"),
 	TC_Masks					UMETA(DisplayName = "Masks (no sRGB)"),
 	TC_Grayscale				UMETA(DisplayName = "Grayscale (G8/16, RGB8 sRGB)"),
 	TC_Displacementmap			UMETA(DisplayName = "Displacementmap (G8/16)"),
 	TC_VectorDisplacementmap	UMETA(DisplayName = "VectorDisplacementmap (RGBA8)"),
 	TC_HDR						UMETA(DisplayName = "HDR (RGBA16F, no sRGB)"),
-	TC_EditorIcon				UMETA(DisplayName = "UserInterface2D (RGBA)"),
-	TC_Alpha					UMETA(DisplayName = "Alpha (no sRGB, BC4 on DX11)"),
+	TC_EditorIcon				UMETA(DisplayName = "UserInterface2D (RGBA8)"),
+	TC_Alpha					UMETA(DisplayName = "Alpha (no sRGB, BC4)"),
 	TC_DistanceFieldFont		UMETA(DisplayName = "DistanceFieldFont (G8)"),
-	TC_HDR_Compressed			UMETA(DisplayName = "HDR Compressed (RGB, BC6H, DX11)"),
-	TC_BC7						UMETA(DisplayName = "BC7 (DX11, optional A)"),
+	TC_HDR_Compressed			UMETA(DisplayName = "HDR Compressed (RGB, BC6H)"),
+	TC_BC7						UMETA(DisplayName = "BC7 (RGBA)"),
 	TC_HalfFloat				UMETA(DisplayName = "Half Float (R16F)"),
 	TC_LQ				        UMETA(Hidden, DisplayName = "Low Quality (BGR565/BGR555A1)", ToolTip = "BGR565/BGR555A1, fallback to DXT1/DXT5 on Mac platform"),
 	TC_EncodedReflectionCapture	UMETA(Hidden), 
