Fixes CharacterMovementComponent applying an anim montage's root motion with incorrect scale on last tick when the montage is triggered by AbilityTask_PlayMontageAndWait and the skel mesh is not recently rendered.
CharacterMovementComponent can tick skeletal mesh poses, which can internally end anim montages and dispatch animation events like montage blend out. AbilityTask_PlayerMontageAndWait resets the root motion scale on that event, so previously CMC then proceeded to consume accumulated root motion with the latest, but already reset root motion translation scale. Now it caches the scale prior to ticking the skel mesh pose, so any accumulated root motion is applied with the appropriate scale.

diff --git a/Engine/Source/Runtime/Engine/Private/Components/CharacterMovementComponent.cpp b/Engine/Source/Runtime/Engine/Private/Components/CharacterMovementComponent.cpp
--- a/Engine/Source/Runtime/Engine/Private/Components/CharacterMovementComponent.cpp
+++ b/Engine/Source/Runtime/Engine/Private/Components/CharacterMovementComponent.cpp
@@ -1551,7 +1551,11 @@ void UCharacterMovementComponent::TickComponent(float DeltaTime, enum ELevelTick
 		USkeletalMeshComponent* CharacterMesh = CharacterOwner->GetMesh();
 		if (CharacterMesh->ShouldTickPose())
 		{
+			// Keep track of if we're playing root motion, just in case the root motion montage ends this frame.
+			// Also cache the root motion translation scale, in case the root motion ends in TickPose and
+			// translation scale is reset by a blend out listener.
 			const bool bWasPlayingRootMotion = CharacterOwner->IsPlayingRootMotion();
+			const float RootMotionTranslationScale = CharacterOwner->GetAnimRootMotionTranslationScale();
 
 			CharacterMesh->TickPose(DeltaTime, true);
 			// We are simulating character movement on physics thread, do not tick movement.
@@ -1561,7 +1565,7 @@ void UCharacterMovementComponent::TickComponent(float DeltaTime, enum ELevelTick
 				FRootMotionMovementParams RootMotion = CharacterMesh->ConsumeRootMotion();
 				if (RootMotion.bHasRootMotion)
 				{
-					RootMotion.ScaleRootMotionTranslation(CharacterOwner->GetAnimRootMotionTranslationScale());
+					RootMotion.ScaleRootMotionTranslation(RootMotionTranslationScale);
 					RootMotionParams.Accumulate(RootMotion);
 				}
 			}
@@ -11577,7 +11581,10 @@ void UCharacterMovementComponent::TickCharacterPose(float DeltaTime)
 	if (CharacterMesh->ShouldTickPose())
 	{
 		// Keep track of if we're playing root motion, just in case the root motion montage ends this frame.
+		// Also cache the root motion translation scale, in case the root motion ends in TickPose and
+		// translation scale is reset by a blend out listener.
 		const bool bWasPlayingRootMotion = CharacterOwner->IsPlayingRootMotion();
+		const float RootMotionTranslationScale = CharacterOwner->GetAnimRootMotionTranslationScale();
 
 		CharacterMesh->TickPose(DeltaTime, true);
 
@@ -11587,7 +11594,7 @@ void UCharacterMovementComponent::TickCharacterPose(float DeltaTime)
 			FRootMotionMovementParams RootMotion = CharacterMesh->ConsumeRootMotion();
 			if (RootMotion.bHasRootMotion)
 			{
-				RootMotion.ScaleRootMotionTranslation(CharacterOwner->GetAnimRootMotionTranslationScale());
+				RootMotion.ScaleRootMotionTranslation(RootMotionTranslationScale);
 				RootMotionParams.Accumulate(RootMotion);
 			}
 
