Added Lit Wireframe viewmode, which combines MeshEdges with the Lit viewmode

diff --git a/Engine/Plugins/Runtime/nDisplay/Source/DisplayClusterConfigurator/Private/Views/Viewport/DisplayClusterConfiguratorSCSEditorViewport.cpp b/Engine/Plugins/Runtime/nDisplay/Source/DisplayClusterConfigurator/Private/Views/Viewport/DisplayClusterConfiguratorSCSEditorViewport.cpp
--- a/Engine/Plugins/Runtime/nDisplay/Source/DisplayClusterConfigurator/Private/Views/Viewport/DisplayClusterConfiguratorSCSEditorViewport.cpp
+++ b/Engine/Plugins/Runtime/nDisplay/Source/DisplayClusterConfigurator/Private/Views/Viewport/DisplayClusterConfiguratorSCSEditorViewport.cpp
@@ -211,6 +211,7 @@ public:
 		static FName LitModeIconName("EditorViewport.LitMode");
 		static FName UnlitModeIconName("EditorViewport.UnlitMode");
 		static FName WireframeModeIconName("EditorViewport.WireframeMode");
+		static FName LitWireframeModeIconName("EditorViewport.LitWireframeMode");
 
 		FName Icon = NAME_None;
 
@@ -229,6 +230,10 @@ public:
 			case VMI_BrushWireframe:
 				Icon = WireframeModeIconName;
 				break;
+
+			case VMI_Lit_Wireframe:
+				Icon = LitWireframeModeIconName;
+				break;
 			}
 		}
 
diff --git a/Engine/Source/Editor/EditorStyle/Private/SlateEditorStyle.cpp b/Engine/Source/Editor/EditorStyle/Private/SlateEditorStyle.cpp
--- a/Engine/Source/Editor/EditorStyle/Private/SlateEditorStyle.cpp
+++ b/Engine/Source/Editor/EditorStyle/Private/SlateEditorStyle.cpp
@@ -4761,6 +4761,7 @@ void FSlateEditorStyle::FStyle::SetupLevelEditorStyle()
 		Set( "EditorViewport.LitMode", new IMAGE_BRUSH( "Icons/icon_ViewMode_Lit_16px", Icon16x16 ) );
 		Set( "EditorViewport.UnlitMode", new IMAGE_BRUSH( "Icons/icon_ViewMode_Unlit_16px", Icon16x16 ) );
 		Set( "EditorViewport.WireframeMode", new IMAGE_BRUSH( "Icons/icon_ViewMode_BrushWireframe_16px", Icon16x16 ) );
+		Set( "EditorViewport.LitWireframeMode", new IMAGE_BRUSH("Icons/icon_ViewMode_BrushWireframe_16px", Icon16x16));
 		Set( "EditorViewport.DetailLightingMode", new IMAGE_BRUSH( "Icons/icon_ViewMode_DetailLighting_16px", Icon16x16 ) );
 		Set( "EditorViewport.LightingOnlyMode", new IMAGE_BRUSH( "Icons/icon_ViewMode_LightingOnly_16px", Icon16x16 ) );
 		Set( "EditorViewport.LightComplexityMode", new IMAGE_BRUSH( "Icons/icon_ViewMode_LightComplexity_16px", Icon16x16 ) );
diff --git a/Engine/Source/Editor/EditorStyle/Private/StarshipStyle.cpp b/Engine/Source/Editor/EditorStyle/Private/StarshipStyle.cpp
--- a/Engine/Source/Editor/EditorStyle/Private/StarshipStyle.cpp
+++ b/Engine/Source/Editor/EditorStyle/Private/StarshipStyle.cpp
@@ -4234,6 +4234,7 @@ void FStarshipEditorStyle::FStyle::SetupLevelEditorStyle()
 		Set( "EditorViewport.LitMode",            	  new IMAGE_BRUSH_SVG("Starship/Common/LitCube", Icon16x16 ) );
 		Set( "EditorViewport.UnlitMode",          	  new IMAGE_BRUSH_SVG("Starship/Common/UnlitCube", Icon16x16 ) );
 		Set( "EditorViewport.WireframeMode",      	  new IMAGE_BRUSH_SVG("Starship/Common/BrushWireframe", Icon16x16 ) );
+		Set( "EditorViewport.LitWireframeMode",       new IMAGE_BRUSH_SVG("Starship/Common/BrushWireframe", Icon16x16) );
 		Set( "EditorViewport.DetailLightingMode", 	  new IMAGE_BRUSH_SVG("Starship/Common/DetailLighting", Icon16x16 ) );
 		Set( "EditorViewport.LightingOnlyMode",   	  new IMAGE_BRUSH_SVG("Starship/Common/LightBulb", Icon16x16 ) );
 		
diff --git a/Engine/Source/Editor/UnrealEd/Private/EditorViewportClient.cpp b/Engine/Source/Editor/UnrealEd/Private/EditorViewportClient.cpp
--- a/Engine/Source/Editor/UnrealEd/Private/EditorViewportClient.cpp
+++ b/Engine/Source/Editor/UnrealEd/Private/EditorViewportClient.cpp
@@ -2870,6 +2870,7 @@ bool FEditorViewportClient::SupportsPreviewResolutionFraction() const
 	{
 	case VMI_BrushWireframe:
 	case VMI_Wireframe:
+	case VMI_Lit_Wireframe:
 	case VMI_LightComplexity:
 	case VMI_LightmapDensity:
 	case VMI_LitLightmapDensity:
diff --git a/Engine/Source/Editor/UnrealEd/Private/EditorViewportCommands.cpp b/Engine/Source/Editor/UnrealEd/Private/EditorViewportCommands.cpp
--- a/Engine/Source/Editor/UnrealEd/Private/EditorViewportCommands.cpp
+++ b/Engine/Source/Editor/UnrealEd/Private/EditorViewportCommands.cpp
@@ -64,6 +64,7 @@ void FEditorViewportCommands::RegisterCommands()
 	UI_COMMAND( WireframeMode, "Brush Wireframe View Mode", "Renders the scene in brush wireframe", EUserInterfaceActionType::RadioButton, FInputChord( EModifierKey::Alt, EKeys::Two ) );
 	UI_COMMAND( UnlitMode, "Unlit View Mode", "Renders the scene with no lights", EUserInterfaceActionType::RadioButton, FInputChord( EModifierKey::Alt, EKeys::Three ) );
 	UI_COMMAND( LitMode, "Lit View Mode", "Renders the scene with normal lighting", EUserInterfaceActionType::RadioButton, FInputChord( EModifierKey::Alt, EKeys::Four ) );
+	UI_COMMAND( LitWireframeMode, "Lit Wireframe View Mode", "Renders the scene with normal lighting and a wireframe overlay", EUserInterfaceActionType::RadioButton, FInputChord(EModifierKey::Alt, EKeys::Nine));
 #if RHI_RAYTRACING
 	UI_COMMAND( PathTracingMode, "PathTracing View Mode", "Renders the scene using a path tracer", EUserInterfaceActionType::RadioButton, FInputChord());
 	UI_COMMAND( RayTracingDebugMode, "Ray tracing Debug View Mode", "Renders the scene in ray tracing debug mode", EUserInterfaceActionType::RadioButton, FInputChord());
diff --git a/Engine/Source/Editor/UnrealEd/Private/SEditorViewport.cpp b/Engine/Source/Editor/UnrealEd/Private/SEditorViewport.cpp
--- a/Engine/Source/Editor/UnrealEd/Private/SEditorViewport.cpp
+++ b/Engine/Source/Editor/UnrealEd/Private/SEditorViewport.cpp
@@ -415,6 +415,7 @@ void SEditorViewport::BindCommands()
 	MAP_VIEWMODE_ACTION( Commands.WireframeMode, VMI_BrushWireframe );
 	MAP_VIEWMODE_ACTION( Commands.UnlitMode, VMI_Unlit );
 	MAP_VIEWMODE_ACTION( Commands.LitMode, VMI_Lit );
+	MAP_VIEWMODE_ACTION( Commands.LitWireframeMode, VMI_Lit_Wireframe);
 #if RHI_RAYTRACING
 	if (IsRayTracingAllowed())
 	{
diff --git a/Engine/Source/Editor/UnrealEd/Private/SEditorViewportViewMenu.cpp b/Engine/Source/Editor/UnrealEd/Private/SEditorViewportViewMenu.cpp
--- a/Engine/Source/Editor/UnrealEd/Private/SEditorViewportViewMenu.cpp
+++ b/Engine/Source/Editor/UnrealEd/Private/SEditorViewportViewMenu.cpp
@@ -136,6 +136,7 @@ void SEditorViewportViewMenu::FillViewMenu(UToolMenu* Menu) const
 				Section.AddMenuEntry(BaseViewportActions.LitMode, UViewModeUtils::GetViewModeDisplayName(VMI_Lit));
 				Section.AddMenuEntry(BaseViewportActions.UnlitMode, UViewModeUtils::GetViewModeDisplayName(VMI_Unlit));
 				Section.AddMenuEntry(BaseViewportActions.WireframeMode, UViewModeUtils::GetViewModeDisplayName(VMI_BrushWireframe));
+				Section.AddMenuEntry(BaseViewportActions.LitWireframeMode, UViewModeUtils::GetViewModeDisplayName(VMI_Lit_Wireframe));
 				Section.AddMenuEntry(BaseViewportActions.DetailLightingMode, UViewModeUtils::GetViewModeDisplayName(VMI_Lit_DetailLighting));
 				Section.AddMenuEntry(BaseViewportActions.LightingOnlyMode, UViewModeUtils::GetViewModeDisplayName(VMI_LightingOnly));
 				Section.AddMenuEntry(BaseViewportActions.ReflectionOverrideMode, UViewModeUtils::GetViewModeDisplayName(VMI_ReflectionOverride));
diff --git a/Engine/Source/Editor/UnrealEd/Public/EditorViewportCommands.h b/Engine/Source/Editor/UnrealEd/Public/EditorViewportCommands.h
--- a/Engine/Source/Editor/UnrealEd/Public/EditorViewportCommands.h
+++ b/Engine/Source/Editor/UnrealEd/Public/EditorViewportCommands.h
@@ -68,6 +68,9 @@ public:
 	/** Changes the viewport to lit mode */
 	TSharedPtr< FUICommandInfo > LitMode;
 
+	/** Changes the viewport to lit wireframe mode */
+	TSharedPtr< FUICommandInfo > LitWireframeMode;
+
 #if RHI_RAYTRACING
 	TSharedPtr< FUICommandInfo > PathTracingMode;
 	TSharedPtr< FUICommandInfo > RayTracingDebugMode;
diff --git a/Engine/Source/Runtime/Engine/Classes/Engine/EngineBaseTypes.h b/Engine/Source/Runtime/Engine/Classes/Engine/EngineBaseTypes.h
--- a/Engine/Source/Runtime/Engine/Classes/Engine/EngineBaseTypes.h
+++ b/Engine/Source/Runtime/Engine/Classes/Engine/EngineBaseTypes.h
@@ -928,7 +928,11 @@ enum EViewModeIndex : int
 	/** Visualize Groom debug views */
 	VMI_VisualizeGroom = 35 UMETA(DisplayName = "Groom Visualization"),
 
+
 	VMI_LWCComplexity = 36 UMETA(DisplayName = "Material LWC Function Usage"),
+	
+	/** Lit Wireframe. */
+	VMI_Lit_Wireframe = 37 UMETA(DisplayName = "Lit Wireframe"),
 
 	VMI_Max UMETA(Hidden),
 
diff --git a/Engine/Source/Runtime/Engine/Private/ShowFlags.cpp b/Engine/Source/Runtime/Engine/Private/ShowFlags.cpp
--- a/Engine/Source/Runtime/Engine/Private/ShowFlags.cpp
+++ b/Engine/Source/Runtime/Engine/Private/ShowFlags.cpp
@@ -311,6 +311,9 @@ void ApplyViewMode(EViewModeIndex ViewModeIndex, bool bPerspective, FEngineShowF
 		case VMI_Lit_DetailLighting:
 			bPostProcessing = true;
 			break;
+		case VMI_Lit_Wireframe:
+			bPostProcessing = true;
+			break;
 		case VMI_LightingOnly:
 			bPostProcessing = true;
 			break;
@@ -414,6 +417,7 @@ void ApplyViewMode(EViewModeIndex ViewModeIndex, bool bPerspective, FEngineShowF
 	EngineShowFlags.SetPathTracing(ViewModeIndex == VMI_PathTracing);
 	EngineShowFlags.SetVisualizeGPUSkinCache(ViewModeIndex == VMI_VisualizeGPUSkinCache);
 	EngineShowFlags.SetVisualizeLWCComplexity(ViewModeIndex == VMI_LWCComplexity);
+	EngineShowFlags.SetMeshEdges(ViewModeIndex == VMI_Lit_Wireframe);
 }
 
 void EngineShowFlagOverride(EShowFlagInitMode ShowFlagInitMode, EViewModeIndex ViewModeIndex, FEngineShowFlags& EngineShowFlags, bool bCanDisableTonemapper)
@@ -549,6 +553,7 @@ void EngineShowFlagOverride(EShowFlagInitMode ShowFlagInitMode, EViewModeIndex V
 		}
 
 		if( ViewModeIndex == VMI_Lit ||
+			ViewModeIndex == VMI_Lit_Wireframe ||
 			ViewModeIndex == VMI_LightingOnly ||
 			ViewModeIndex == VMI_LitLightmapDensity)
 		{
@@ -800,6 +805,10 @@ EViewModeIndex FindViewMode(const FEngineShowFlags& EngineShowFlags)
 	{
 		return VMI_ReflectionOverride;
 	}
+	else if (EngineShowFlags.MeshEdges)
+	{
+		return VMI_Lit_Wireframe;
+	}
 	else if (EngineShowFlags.Wireframe)
 	{
 		if (EngineShowFlags.Brushes)
@@ -861,6 +870,7 @@ const TCHAR* GetViewModeName(EViewModeIndex ViewModeIndex)
 		case VMI_Unlit:						return TEXT("Unlit");
 		case VMI_Lit:						return TEXT("Lit");
 		case VMI_Lit_DetailLighting:		return TEXT("Lit_DetailLighting");
+		case VMI_Lit_Wireframe:				return TEXT("Lit_Wireframe");
 		case VMI_LightingOnly:				return TEXT("LightingOnly");
 		case VMI_LightComplexity:			return TEXT("LightComplexity");
 		case VMI_ShaderComplexity:			return TEXT("ShaderComplexity");
diff --git a/Engine/Source/Runtime/Engine/Private/ViewModeNames.cpp b/Engine/Source/Runtime/Engine/Private/ViewModeNames.cpp
--- a/Engine/Source/Runtime/Engine/Private/ViewModeNames.cpp
+++ b/Engine/Source/Runtime/Engine/Private/ViewModeNames.cpp
@@ -40,6 +40,10 @@ TArray<FText> FillViewModeDisplayNames()
 		{
 			ViewModeDisplayNames.Emplace(LOCTEXT("UViewModeUtils_VMI_Lit_DetailLighting", "Detail Lighting"));
 		}
+		else if (ViewModeIndex == VMI_Lit_Wireframe)
+		{
+			ViewModeDisplayNames.Emplace(LOCTEXT("UViewModeUtils_VMI_Lit_Wireframe", "Lit Wireframe"));
+		}
 		// Lit wo/ materials
 		else if (ViewModeIndex == VMI_LightingOnly)
 		{
@@ -244,6 +248,10 @@ TArray<const FSlateBrush*> FillViewModeDisplayIcons()
 		{
 			ViewModeDisplayIcons.Emplace(FAppStyle::Get().GetBrush("EditorViewport.DetailLightingMode"));
 		}
+		else if (ViewModeIndex == VMI_Lit_Wireframe)
+		{
+			ViewModeDisplayIcons.Emplace(FAppStyle::Get().GetBrush("EditorViewport.LitWireframeMode"));
+		}
 		// Lit wo/ materials
 		else if (ViewModeIndex == VMI_LightingOnly)
 		{
