Make Texcoord SRV available if using GPUSkinPassthrough.
Fixes case where in DeformerGraph where UVs weren't being passed to deformer.
Bug only happened when other conditions weren't met. That is in perf mode when RHISupportsManualVertexFetch is false.

diff --git a/Engine/Source/Runtime/Engine/Private/Rendering/StaticMeshVertexBuffer.cpp b/Engine/Source/Runtime/Engine/Private/Rendering/StaticMeshVertexBuffer.cpp
--- a/Engine/Source/Runtime/Engine/Private/Rendering/StaticMeshVertexBuffer.cpp
+++ b/Engine/Source/Runtime/Engine/Private/Rendering/StaticMeshVertexBuffer.cpp
@@ -322,7 +322,7 @@ void FStaticMeshVertexBuffer::InitRHI(FRHICommandListBase& RHICmdList)
 	const bool bHadTexCoordData = TexcoordData != nullptr;
 	const bool bCreateTexCoordSRV = bHadTexCoordData && TexcoordData->GetAllowCPUAccess();
 	TexCoordVertexBuffer.VertexBufferRHI = CreateTexCoordRHIBuffer(RHICmdList);
-	if (TexCoordVertexBuffer.VertexBufferRHI && (bCreateTexCoordSRV || RHISupportsManualVertexFetch(GMaxRHIShaderPlatform)))
+	if (TexCoordVertexBuffer.VertexBufferRHI && (bCreateTexCoordSRV || RHISupportsManualVertexFetch(GMaxRHIShaderPlatform) || IsGPUSkinPassThroughSupported(GMaxRHIShaderPlatform)))
 	{
 		uint32       Stride = GetUseFullPrecisionUVs() ? 8 : 4;
 		EPixelFormat Format = GetUseFullPrecisionUVs() ? PF_G32R32F : PF_G16R16F;
