DeadBlending: Fixed shadowed variable in lambda warning

diff --git a/Engine/Source/Runtime/Engine/Private/Animation/AnimNode_DeadBlending.cpp b/Engine/Source/Runtime/Engine/Private/Animation/AnimNode_DeadBlending.cpp
--- a/Engine/Source/Runtime/Engine/Private/Animation/AnimNode_DeadBlending.cpp
+++ b/Engine/Source/Runtime/Engine/Private/Animation/AnimNode_DeadBlending.cpp
@@ -435,14 +435,14 @@ void FAnimNode_DeadBlending::InitFrom(const FCompactPose& InPose, const FBlended
 	if (SrcPoseCurr.DeltaTime > UE_SMALL_NUMBER)
 	{
 		UE::Anim::FNamedValueArrayUtils::Union(CurveData, SrcPosePrev.Curves.BlendedCurve,
-			[this, DeltaTime = SrcPoseCurr.DeltaTime](FDeadBlendingCurveElement& OutResultElement, const UE::Anim::FCurveElement& InElement1, UE::Anim::ENamedValueUnionFlags InFlags)
+			[this, SrcDeltaTime = SrcPoseCurr.DeltaTime](FDeadBlendingCurveElement& OutResultElement, const UE::Anim::FCurveElement& InElement1, UE::Anim::ENamedValueUnionFlags InFlags)
 			{
 				bool bSrcCurrValid = (bool)(InFlags & UE::Anim::ENamedValueUnionFlags::ValidArg0);
 				bool bSrcPrevValid = (bool)(InFlags & UE::Anim::ENamedValueUnionFlags::ValidArg1);
 
 				if (bSrcCurrValid && bSrcPrevValid)
 				{
-					OutResultElement.Velocity = (OutResultElement.Value - InElement1.Value) / DeltaTime;
+					OutResultElement.Velocity = (OutResultElement.Value - InElement1.Value) / SrcDeltaTime;
 					OutResultElement.Velocity = FMath::Clamp(OutResultElement.Velocity, -MaximumCurveVelocity, MaximumCurveVelocity);
 				}
 			});
