SDK (1): Changes in common code needed to support SDK changes on some platforms

diff --git a/Engine/Plugins/Experimental/Web/WebAPI/Source/WebAPI/Public/WebAPIHttpRequest.h b/Engine/Plugins/Experimental/Web/WebAPI/Source/WebAPI/Public/WebAPIHttpRequest.h
--- a/Engine/Plugins/Experimental/Web/WebAPI/Source/WebAPI/Public/WebAPIHttpRequest.h
+++ b/Engine/Plugins/Experimental/Web/WebAPI/Source/WebAPI/Public/WebAPIHttpRequest.h
@@ -248,7 +248,7 @@ TRequest<PayloadType>& TRequest<PayloadType>::BindHeaderReceivedCallback(TFuncti
 template <typename PayloadType>
 TRequest<PayloadType>& TRequest<PayloadType>::SetPayloadData(const PayloadType& InPayloadData)
 {
-	static_assert(TModels_V<CStaticStructProvider, PayloadType>, TEXT("PayloadType must be a UStruct"));
+	static_assert(TModels_V<CStaticStructProvider, PayloadType>, "PayloadType must be a UStruct");
 
 	FString PayloadAsString = PayloadToString_Internal(InPayloadData).GetValue();
 	PayloadAsString.Append(TEXT("]")).InsertAt(0, TEXT("["));
@@ -273,7 +273,7 @@ bool TRequest<PayloadType>::HasPayload() const
 template <typename PayloadType>
 TValueOrError<FString, FString> TRequest<PayloadType>::PayloadToString_Internal(const PayloadType& InContent) const
 {
-	static_assert(TModels_V<CStaticStructProvider, PayloadType>, TEXT("PayloadType must be a UStruct"));
+	static_assert(TModels_V<CStaticStructProvider, PayloadType>, "PayloadType must be a UStruct");
 
 	FString StringifiedContent;
 	if (!FJsonObjectConverter::UStructToJsonObjectString(InContent, StringifiedContent))
diff --git a/Engine/Source/Developer/ShaderCompilerCommon/Private/ShaderCompilerCommon.cpp b/Engine/Source/Developer/ShaderCompilerCommon/Private/ShaderCompilerCommon.cpp
--- a/Engine/Source/Developer/ShaderCompilerCommon/Private/ShaderCompilerCommon.cpp
+++ b/Engine/Source/Developer/ShaderCompilerCommon/Private/ShaderCompilerCommon.cpp
@@ -1424,6 +1424,8 @@ FString CreateShaderCompilerWorkerDirectCommandLine(const FShaderCompilerInput&
 	Text += TEXT("\" -shaderPlatformName=");
 	Text += Input.ShaderPlatformName.GetPlainNameString();
 
+	Text += FString::Printf(TEXT(" -supportedHardwareMask=%u"), Input.SupportedHardwareMask);
+
 	switch (Input.Target.Frequency)
 	{
 	case SF_Vertex:			Text += TEXT(" -vs"); break;
diff --git a/Engine/Source/Developer/TargetPlatform/Private/TargetPlatformBase.cpp b/Engine/Source/Developer/TargetPlatform/Private/TargetPlatformBase.cpp
--- a/Engine/Source/Developer/TargetPlatform/Private/TargetPlatformBase.cpp
+++ b/Engine/Source/Developer/TargetPlatform/Private/TargetPlatformBase.cpp
@@ -59,6 +59,11 @@ bool FTargetPlatformBase::UsesRayTracing() const
 	return CVar ? (CVar->GetInt() != 0) : false;
 }
 
+uint32 FTargetPlatformBase::GetSupportedHardwareMask() const
+{
+	return 0;
+}
+
 EOfflineBVHMode FTargetPlatformBase::GetStaticMeshOfflineBVHMode() const
 {
 	return EOfflineBVHMode::Disabled;
diff --git a/Engine/Source/Developer/TargetPlatform/Private/TargetPlatformSettingsBase.cpp b/Engine/Source/Developer/TargetPlatform/Private/TargetPlatformSettingsBase.cpp
--- a/Engine/Source/Developer/TargetPlatform/Private/TargetPlatformSettingsBase.cpp
+++ b/Engine/Source/Developer/TargetPlatform/Private/TargetPlatformSettingsBase.cpp
@@ -45,6 +45,11 @@ bool FTargetPlatformSettingsBase::UsesRayTracing() const
 	return CVar ? (CVar->GetInt() != 0) : false;
 }
 
+uint32 FTargetPlatformSettingsBase::GetSupportedHardwareMask() const
+{
+	return 0;
+}
+
 EOfflineBVHMode FTargetPlatformSettingsBase::GetStaticMeshOfflineBVHMode() const
 {
 	return EOfflineBVHMode::Disabled;
diff --git a/Engine/Source/Developer/TargetPlatform/Public/Common/TargetPlatformBase.h b/Engine/Source/Developer/TargetPlatform/Public/Common/TargetPlatformBase.h
--- a/Engine/Source/Developer/TargetPlatform/Public/Common/TargetPlatformBase.h
+++ b/Engine/Source/Developer/TargetPlatform/Public/Common/TargetPlatformBase.h
@@ -90,6 +90,8 @@ public:
 
 	TARGETPLATFORM_API virtual bool UsesRayTracing() const override;
 
+	TARGETPLATFORM_API virtual uint32 GetSupportedHardwareMask() const override;
+
 	TARGETPLATFORM_API virtual EOfflineBVHMode GetStaticMeshOfflineBVHMode() const override;
 
 	TARGETPLATFORM_API virtual bool GetStaticMeshOfflineBVHCompression() const override;
@@ -802,6 +804,10 @@ public:
 	{
 		return TargetPlatformSettings->UsesRayTracing();
 	}
+	virtual uint32 GetSupportedHardwareMask() const override
+	{
+		return TargetPlatformSettings->GetSupportedHardwareMask();
+	}
 	virtual EOfflineBVHMode GetStaticMeshOfflineBVHMode() const override 
 	{
 		return TargetPlatformSettings->GetStaticMeshOfflineBVHMode();
diff --git a/Engine/Source/Developer/TargetPlatform/Public/Common/TargetPlatformSettingsBase.h b/Engine/Source/Developer/TargetPlatform/Public/Common/TargetPlatformSettingsBase.h
--- a/Engine/Source/Developer/TargetPlatform/Public/Common/TargetPlatformSettingsBase.h
+++ b/Engine/Source/Developer/TargetPlatform/Public/Common/TargetPlatformSettingsBase.h
@@ -27,6 +27,7 @@ public:
 	TARGETPLATFORM_API virtual bool UsesSelectiveBasePassOutputs() const override;
 	TARGETPLATFORM_API virtual bool UsesDistanceFields() const override;
 	TARGETPLATFORM_API virtual bool UsesRayTracing() const override;
+	TARGETPLATFORM_API virtual uint32 GetSupportedHardwareMask() const override;
 	TARGETPLATFORM_API virtual EOfflineBVHMode GetStaticMeshOfflineBVHMode() const override;
 	TARGETPLATFORM_API virtual bool GetStaticMeshOfflineBVHCompression() const override;
 	TARGETPLATFORM_API virtual bool ForcesSimpleSkyDiffuse() const override;
diff --git a/Engine/Source/Developer/TargetPlatform/Public/Interfaces/ITargetPlatformSettings.h b/Engine/Source/Developer/TargetPlatform/Public/Interfaces/ITargetPlatformSettings.h
--- a/Engine/Source/Developer/TargetPlatform/Public/Interfaces/ITargetPlatformSettings.h
+++ b/Engine/Source/Developer/TargetPlatform/Public/Interfaces/ITargetPlatformSettings.h
@@ -163,6 +163,14 @@ public:
 	*/
 	virtual bool UsesRayTracing() const = 0;
 
+	/**
+	 * Gets a platform-dependent bitfield describing which hardware generations are supported.
+	 * Applies to platforms with multiple iterations of the device, which may vary in performance.
+	 * Each bit represents a supported device version.
+	 * Return 0 if information is not provided by the platform.
+	 */
+	virtual uint32 GetSupportedHardwareMask() const = 0;
+
 	/**
 	* Gets static mesh offline BVH mode
 	*/
diff --git a/Engine/Source/Programs/ShaderCompileWorker/Private/ShaderCompileWorker.cpp b/Engine/Source/Programs/ShaderCompileWorker/Private/ShaderCompileWorker.cpp
--- a/Engine/Source/Programs/ShaderCompileWorker/Private/ShaderCompileWorker.cpp
+++ b/Engine/Source/Programs/ShaderCompileWorker/Private/ShaderCompileWorker.cpp
@@ -881,6 +881,7 @@ static void DirectCompile(const TArray<const class IShaderFormat*>& ShaderFormat
 	FName FormatName;
 	FName ShaderPlatformName;
 	FString Entry = TEXT("Main");
+	uint32 SupportedHardwareMask = 0;
 	FString DumpDebugInfoPath;
 	bool bPipeline = false;
 	EShaderFrequency Frequency = SF_Pixel;
@@ -909,6 +910,10 @@ static void DirectCompile(const TArray<const class IShaderFormat*>& ShaderFormat
 			{
 				ShaderPlatformName = FName(*Token.RightChop(19));
 			}
+			else if (Token.StartsWith(TEXT("supportedHardwareMask=")))
+			{
+				SupportedHardwareMask = (uint32)FCString::Atoi64(*Token.RightChop(22));
+			}
 			else if (Token.StartsWith(TEXT("DebugInfoPath=")))
 			{
 				DumpDebugInfoPath = Token.RightChop(14);
@@ -987,6 +992,7 @@ static void DirectCompile(const TArray<const class IShaderFormat*>& ShaderFormat
 	Job.Input.EntryPointName = Entry;
 	Job.Input.ShaderFormat = FormatName;
 	Job.Input.ShaderPlatformName = ShaderPlatformName;
+	Job.Input.SupportedHardwareMask = SupportedHardwareMask;
 	Job.Input.VirtualSourceFilePath = InputFile;
 	Job.Input.Target.Platform =  ShaderFormatNameToShaderPlatform(FormatName);
 	Job.Input.Target.Frequency = Frequency;
diff --git a/Engine/Source/Runtime/Engine/Classes/Engine/LocalPlayer.h b/Engine/Source/Runtime/Engine/Classes/Engine/LocalPlayer.h
--- a/Engine/Source/Runtime/Engine/Classes/Engine/LocalPlayer.h
+++ b/Engine/Source/Runtime/Engine/Classes/Engine/LocalPlayer.h
@@ -238,7 +238,6 @@ private:
 
 public:
 	// UObject interface
-	ENGINE_API virtual void FinishDestroy() override;
 	static ENGINE_API void AddReferencedObjects(UObject* InThis, FReferenceCollector& Collector);
 	// End of UObject interface
 
diff --git a/Engine/Source/Runtime/Engine/Private/LocalPlayer.cpp b/Engine/Source/Runtime/Engine/Private/LocalPlayer.cpp
--- a/Engine/Source/Runtime/Engine/Private/LocalPlayer.cpp
+++ b/Engine/Source/Runtime/Engine/Private/LocalPlayer.cpp
@@ -254,6 +254,14 @@ void ULocalPlayer::InitOnlineSession()
 void ULocalPlayer::PlayerRemoved()
 {
 	SubsystemCollection.Deinitialize();
+
+	if (!IsTemplate())
+	{
+		for (FSceneViewStateReference& ViewState : ViewStates)
+		{
+			ViewState.Destroy();
+		}
+	}
 }
 
 bool ULocalPlayer::SpawnPlayActor(const FString& URL,FString& OutError, UWorld* InWorld)
@@ -374,18 +382,6 @@ void ULocalPlayer::SendSplitJoin(TArray<FString>& Options)
 	}
 }
 
-void ULocalPlayer::FinishDestroy()
-{
-	if ( !IsTemplate() )
-	{
-		for (FSceneViewStateReference& ViewState : ViewStates)
-		{
-			ViewState.Destroy();
-		}
-	}
-	Super::FinishDestroy();
-}
-
 /**
  * Singleton managing saved locked views and the current per-player state.
  */
diff --git a/Engine/Source/Runtime/Engine/Private/ShaderCompiler/ShaderCompiler.cpp b/Engine/Source/Runtime/Engine/Private/ShaderCompiler/ShaderCompiler.cpp
--- a/Engine/Source/Runtime/Engine/Private/ShaderCompiler/ShaderCompiler.cpp
+++ b/Engine/Source/Runtime/Engine/Private/ShaderCompiler/ShaderCompiler.cpp
@@ -7802,6 +7802,8 @@ void GlobalBeginCompileShader(
 	const EShaderPlatform ShaderPlatform = EShaderPlatform(Target.Platform);
 	const FName ShaderFormatName = LegacyShaderPlatformToShaderFormat(ShaderPlatform);
 
+	ITargetPlatform* TargetPlatform = GetTargetPlatformManager()->FindTargetPlatformWithSupport(TEXT("ShaderFormat"), ShaderFormatName);
+
 	FShaderCompileUtilities::GenerateBrdfHeaders(ShaderPlatform);
 
 	// NOTE:  Input.bCompilingForShaderPipeline is initialized by the constructor for single versus pipeline jobs, do not initialize again here!
@@ -7809,6 +7811,7 @@ void GlobalBeginCompileShader(
 	Input.Target = Target;
 	Input.ShaderPlatformName = FDataDrivenShaderPlatformInfo::GetName(ShaderPlatform);
 	Input.ShaderFormat = ShaderFormatName;
+	Input.SupportedHardwareMask = TargetPlatform ? TargetPlatform->GetSupportedHardwareMask() : 0;
 	Input.CompressionFormat = GetShaderCompressionFormat();
 	GetShaderCompressionOodleSettings(Input.OodleCompressor, Input.OodleLevel);
 	Input.VirtualSourceFilePath = SourceFilename;
@@ -8386,7 +8389,6 @@ void GlobalBeginCompileShader(
 		SET_SHADER_DEFINE(Input.Environment, SUPPORTS_INDEPENDENT_SAMPLERS, bSupportsIndependentSamplers ? 1 : 0);
 	}
 
-	ITargetPlatform* TargetPlatform = GetTargetPlatformManager()->FindTargetPlatformWithSupport(TEXT("ShaderFormat"), ShaderFormatName);
 	bool bForwardShading = false;
 	{
 		if (TargetPlatform)
diff --git a/Engine/Source/Runtime/GeometryCore/Public/Curve/DynamicGraph2.h b/Engine/Source/Runtime/GeometryCore/Public/Curve/DynamicGraph2.h
--- a/Engine/Source/Runtime/GeometryCore/Public/Curve/DynamicGraph2.h
+++ b/Engine/Source/Runtime/GeometryCore/Public/Curve/DynamicGraph2.h
@@ -139,7 +139,7 @@ public:
 	value_iteration<TVector2<T>> Vertices() const
 	{
 		return vertices_refcount.MappedIndices<TVector2<T>>(
-			[&](int vid) {
+			[&, this](int vid) {
 				return vertices.template AsVector2<T>(vid);
 			});
 	}
diff --git a/Engine/Source/Runtime/GeometryCore/Public/Curve/DynamicGraph3.h b/Engine/Source/Runtime/GeometryCore/Public/Curve/DynamicGraph3.h
--- a/Engine/Source/Runtime/GeometryCore/Public/Curve/DynamicGraph3.h
+++ b/Engine/Source/Runtime/GeometryCore/Public/Curve/DynamicGraph3.h
@@ -111,7 +111,7 @@ public:
 	value_iteration<TVector<T>> VerticesItr() const
 	{
 		return vertices_refcount.MappedIndices<TVector<T>>(
-			[&](int vid) {
+			[&, this](int vid) {
 				return Vertices.template AsVector3<T>(vid);
 			});
 	}
diff --git a/Engine/Source/Runtime/RHI/Public/RHIValidationContext.h b/Engine/Source/Runtime/RHI/Public/RHIValidationContext.h
--- a/Engine/Source/Runtime/RHI/Public/RHIValidationContext.h
+++ b/Engine/Source/Runtime/RHI/Public/RHIValidationContext.h
@@ -549,6 +549,11 @@ public:
 		RHIContext->RHIResummarizeHTile(DepthTexture);
 	}
 
+	virtual void* RHIGetNativeCommandBuffer() override final
+	{
+		return RHIContext->RHIGetNativeCommandBuffer();
+	}
+
 	virtual void RHIBeginTransitions(TArrayView<const FRHITransition*> Transitions) override final
 	{
 		ensureMsgf(!State.bInsideBeginRenderPass, TEXT("Beginning a transition within a renderpass is not supported!"));
diff --git a/Engine/Source/Runtime/RenderCore/Private/Shader.cpp b/Engine/Source/Runtime/RenderCore/Private/Shader.cpp
--- a/Engine/Source/Runtime/RenderCore/Private/Shader.cpp
+++ b/Engine/Source/Runtime/RenderCore/Private/Shader.cpp
@@ -2002,7 +2002,16 @@ void ShaderMapAppendKey(EShaderPlatform Platform, FShaderKeyGenerator& KeyGen)
 		ShaderFormat->AppendToKeyString(ShaderFormatExtraData);
 		KeyGen.Append(ShaderFormatExtraData);
 	}
-	
+
+	ITargetPlatform* TargetPlatform = GetTargetPlatformManagerRef().FindTargetPlatformWithSupport(TEXT("ShaderFormat"), ShaderFormatName);
+
+	uint32 SupportedHardwareMask = TargetPlatform ? TargetPlatform->GetSupportedHardwareMask() : 0;
+
+	if (SupportedHardwareMask != 0)
+	{
+		KeyString += FString::Printf(TEXT("_SHM_%X"), SupportedHardwareMask);
+	}
+
 	// Encode the Metal standard into the shader compile options so that they recompile if the settings change.
 	if (IsMetalPlatform(Platform))
 	{
@@ -2122,8 +2131,6 @@ void ShaderMapAppendKey(EShaderPlatform Platform, FShaderKeyGenerator& KeyGen)
 		}
 	}
 
-	ITargetPlatform* TargetPlatform = GetTargetPlatformManager()->FindTargetPlatformWithSupport(TEXT("ShaderFormat"), ShaderFormatName);
-
 	{
 		bool bForwardShading = false;
 		if (TargetPlatform)
diff --git a/Engine/Source/Runtime/RenderCore/Private/ShaderCore.cpp b/Engine/Source/Runtime/RenderCore/Private/ShaderCore.cpp
--- a/Engine/Source/Runtime/RenderCore/Private/ShaderCore.cpp
+++ b/Engine/Source/Runtime/RenderCore/Private/ShaderCore.cpp
@@ -3689,6 +3689,7 @@ FArchive& operator<<(FArchive& Ar, FShaderCompilerInput& Input)
 	Ar << Input.VirtualSourceFilePath;
 	Ar << Input.EntryPointName;
 	Ar << Input.ShaderName;
+	Ar << Input.SupportedHardwareMask;
 	Ar << Input.bCompilingForShaderPipeline;
 	Ar << Input.bIncludeUsedOutputs;
 	Ar << Input.UsedOutputs;
diff --git a/Engine/Source/Runtime/RenderCore/Public/ShaderCompilerCore.h b/Engine/Source/Runtime/RenderCore/Public/ShaderCompilerCore.h
--- a/Engine/Source/Runtime/RenderCore/Public/ShaderCompilerCore.h
+++ b/Engine/Source/Runtime/RenderCore/Public/ShaderCompilerCore.h
@@ -249,6 +249,8 @@ struct FShaderCompilerInput
 	FString EntryPointName;
 	FString ShaderName;
 
+	uint32 SupportedHardwareMask = 0;
+
 	// Skips the preprocessor and instead loads the usf file directly
 	UE_DEPRECATED(5.4, "bSkipPreprocessedCache member is deprecated; set EShaderDebugInfoFlags::CompileFromDebugUSF on DebugInfoFlags instead.")
 	bool bSkipPreprocessedCache = false;
diff --git a/Engine/Source/Runtime/RenderCore/Public/ShaderParameterStructDeclaration.h b/Engine/Source/Runtime/RenderCore/Public/ShaderParameterStructDeclaration.h
--- a/Engine/Source/Runtime/RenderCore/Public/ShaderParameterStructDeclaration.h
+++ b/Engine/Source/Runtime/RenderCore/Public/ShaderParameterStructDeclaration.h
@@ -19,7 +19,7 @@ template<typename T>
 const FShaderParametersMetadata* GetForwardDeclaredShaderParametersStructMetadata(const T*)
 {
 	static_assert(!std::is_same<T, T>::value /* true */,
-		TEXT("Partial uniform buffer struct declaration. Use `DECLARE_UNIFORM_BUFFER_STRUCT()` instead of `class T;`.")
+		"Partial uniform buffer struct declaration. Use `DECLARE_UNIFORM_BUFFER_STRUCT()` instead of `class T;`."
 	);
 	return nullptr;
 }
