Fix inconsistency between how FInputAxisBinding and FInputAxisKeyBinding determine their axis value when the input keys have been previously consumed

diff --git a/Engine/Source/Runtime/Engine/Private/UserInterface/PlayerInput.cpp b/Engine/Source/Runtime/Engine/Private/UserInterface/PlayerInput.cpp
--- a/Engine/Source/Runtime/Engine/Private/UserInterface/PlayerInput.cpp
+++ b/Engine/Source/Runtime/Engine/Private/UserInterface/PlayerInput.cpp
@@ -50,6 +50,11 @@ namespace UE
 		static FAutoConsoleVariableRef CVarAutoReconcilePressedEventsOnFirstRepeat(TEXT("Input.AutoReconcilePressedEventsOnFirstRepeat"),
 			bAutoReconcilePressedEventsOnFirstRepeat,
 			TEXT("If true, then we will automatically mark a IE_Pressed event if we receive an IE_Repeat event but have not received a pressed event first.\nNote: This option will be removed in a future update."));
+
+		static bool bClearAxisValueIfConsumed = true;
+		static FAutoConsoleVariableRef CVarClearAxisValueIfConsumed(TEXT("Input.ClearAxisValueIfConsumed"),
+			bClearAxisValueIfConsumed,
+			TEXT("If true, we will clear the value of any FInputAxisKeyBinding whose FKey has been previously consumed.\nNote: This option will be removed in a future update."));
 	}
 }
 
@@ -1438,6 +1443,10 @@ void UPlayerInput::EvaluateInputDelegates(const TArray<UInputComponent*>& InputC
 						KeysToConsume.AddUnique(AxisKeyBinding.AxisKey);
 					}
 				}
+				else if(UE::Input::bClearAxisValueIfConsumed)
+				{
+					AxisKeyBinding.AxisValue = 0.f;
+				}
 
 				if (AxisKeyBinding.AxisDelegate.IsBound())
 				{
