Implemented support for different ETextureMipValueMode values in the SparseVolumeTexture sample node. MipValueModes requiring implicit mip levels manually calculate a mip level based on derivatives computed with the DDX() and DDY() nodes/functions. In contexts where hardware ddx/ddy intrinsics aren't supported, analytical derivatives can be used. Also fixed up some input nodes that weren't aware of analytical derivatives and which are likely to be used to compute SVT UVs.
This CL is standalone, but to take advantage of analytical derivatives, the context calling into the material (HeterogeneousVolumes, VolumetricClouds) needs to be compiled with analytical derivatives taken into account.

#rnx

diff --git a/Engine/Shaders/Private/MaterialTemplate.ush b/Engine/Shaders/Private/MaterialTemplate.ush
--- a/Engine/Shaders/Private/MaterialTemplate.ush
+++ b/Engine/Shaders/Private/MaterialTemplate.ush
@@ -42,7 +42,7 @@
 
 // Update this GUID to force all material shaders to recompile
 // Merge conflicts on this line should be resolved by generating a new GUID
-#pragma message("UESHADERMETADATA_VERSION 381E7C4C-6229-4F2E-8F0E-EE692C2785A3")
+#pragma message("UESHADERMETADATA_VERSION 838540DA-BB92-4973-80CB-D9E993EBA40D")
 
 #if USES_SPEEDTREE
 	#include "/Engine/Private/SpeedTreeCommon.ush"
diff --git a/Engine/Shaders/Private/SparseVolumeTexture/SparseVolumeTextureCommon.ush b/Engine/Shaders/Private/SparseVolumeTexture/SparseVolumeTextureCommon.ush
--- a/Engine/Shaders/Private/SparseVolumeTexture/SparseVolumeTextureCommon.ush
+++ b/Engine/Shaders/Private/SparseVolumeTexture/SparseVolumeTextureCommon.ush
@@ -64,6 +64,19 @@ FSparseVolumeTextureUniforms SparseVolumeTextureUnpackUniforms(const uint4 Packe
 	return Result;
 }
 
+float SparseVolumeTextureCalculateMipLevel(FSparseVolumeTextureUniforms Uniforms, float3 UVWDDX, float3 UVWDDY, float MipBias)
+{
+	const float3 Resolution = Uniforms.VolumePageResolution * float(SPARSE_VOLUME_TILE_RES);
+	const float3 TexCoordDDX = UVWDDX * Resolution;
+	const float3 TexCoordDDY = UVWDDY * Resolution;
+	const float DDXLengthSquared = dot(TexCoordDDX, TexCoordDDX);
+	const float DDYLengthSquared = dot(TexCoordDDY, TexCoordDDY);
+	const float MaxLengthSquared = max(DDXLengthSquared, DDYLengthSquared);
+	const float MipLevel = 0.5f * log2(max(MaxLengthSquared, 1e-8f));
+	const float MipLevelBiased = max(0.0f, MipLevel + MipBias);
+	return MipLevelBiased;
+}
+
 float3 SparseVolumeTextureGetVoxelCoord(Texture3D<uint> PageTable, FSparseVolumeTextureUniforms Uniforms, float3 PageTableCoord, int MipLevel)
 {
 	// Load from page table texture
diff --git a/Engine/Source/Runtime/Engine/Classes/Materials/MaterialExpressionSparseVolumeTextureBase.h b/Engine/Source/Runtime/Engine/Classes/Materials/MaterialExpressionSparseVolumeTextureBase.h
--- a/Engine/Source/Runtime/Engine/Classes/Materials/MaterialExpressionSparseVolumeTextureBase.h
+++ b/Engine/Source/Runtime/Engine/Classes/Materials/MaterialExpressionSparseVolumeTextureBase.h
@@ -17,7 +17,7 @@ class UMaterialExpressionSparseVolumeTextureBase : public UMaterialExpression
 {
 	GENERATED_UCLASS_BODY()
 
-	/** The Sparse Virtual Texture to sample. */
+	/** The Sparse Volume Texture to sample. */
 	UPROPERTY(EditAnywhere, BlueprintReadWrite, Category = SparseVolumeTexture)
 	TObjectPtr<USparseVolumeTexture> SparseVolumeTexture;
 
diff --git a/Engine/Source/Runtime/Engine/Classes/Materials/MaterialExpressionSparseVolumeTextureSample.h b/Engine/Source/Runtime/Engine/Classes/Materials/MaterialExpressionSparseVolumeTextureSample.h
--- a/Engine/Source/Runtime/Engine/Classes/Materials/MaterialExpressionSparseVolumeTextureSample.h
+++ b/Engine/Source/Runtime/Engine/Classes/Materials/MaterialExpressionSparseVolumeTextureSample.h
@@ -26,8 +26,21 @@ class UMaterialExpressionSparseVolumeTextureSample : public UMaterialExpressionS
 	UPROPERTY(meta = (RequiredInput = "false", ToolTip = "Defaults to 'SparseVolumeTexture' if not specified"))
 	FExpressionInput TextureObject;
 
+	/** Meaning depends on MipValueMode, a single unit is one mip level  */
 	UPROPERTY(meta = (RequiredInput = "false", ToolTip = "Defaults to 0 if not specified"))
-	FExpressionInput MipLevel;
+	FExpressionInput MipValue;
+
+	/** Enabled only if MipValueMode == TMVM_Derivative */
+	UPROPERTY(meta = (RequiredInput = "true", ToolTip = "Coordinates derivative over the X axis"))
+	FExpressionInput CoordinatesDX;
+
+	/** Enabled only if MipValueMode == TMVM_Derivative */
+	UPROPERTY(meta = (RequiredInput = "true", ToolTip = "Coordinates derivative over the Y axis"))
+	FExpressionInput CoordinatesDY;
+
+	/** Defines how the MipValue property is applied to the texture lookup */
+	UPROPERTY(EditAnywhere, Category = MaterialExpressionTextureSample, meta = (DisplayName = "MipValueMode", ShowAsInputPin = "Advanced"))
+	TEnumAsByte<ETextureMipValueMode> MipValueMode;
 
 	/**
 	 * Controls where the sampler for this texture lookup will come from.
@@ -38,10 +51,18 @@ class UMaterialExpressionSparseVolumeTextureSample : public UMaterialExpressionS
 	UPROPERTY(EditAnywhere, Category = MaterialExpressionTextureSample, Meta = (ShowAsInputPin = "Advanced"))
 	TEnumAsByte<ESamplerSourceMode> SamplerSource;
 
+	/** Only used if MipValue is not hooked up */
+	UPROPERTY(EditAnywhere, Category = MaterialExpressionTextureSample)
+	int32 ConstMipValue;
+
 protected:
 
 #if WITH_EDITOR
+	ENGINE_API virtual bool CanEditChange(const FProperty* InProperty) const override;
 	ENGINE_API virtual void PostLoad() override;
+	ENGINE_API virtual TArrayView<FExpressionInput*> GetInputsView() override;
+	ENGINE_API virtual FExpressionInput* GetInput(int32 InputIndex) override;
+	ENGINE_API virtual FName GetInputName(int32 InputIndex) const override;
 	ENGINE_API virtual int32 Compile(class FMaterialCompiler* Compiler, int32 OutputIndex) override;
 	ENGINE_API virtual void GetCaption(TArray<FString>& OutCaptions) const override;
 	ENGINE_API virtual uint32 GetOutputType(int32 OutputIndex) override;
diff --git a/Engine/Source/Runtime/Engine/Private/Materials/HLSLMaterialTranslator.cpp b/Engine/Source/Runtime/Engine/Private/Materials/HLSLMaterialTranslator.cpp
--- a/Engine/Source/Runtime/Engine/Private/Materials/HLSLMaterialTranslator.cpp
+++ b/Engine/Source/Runtime/Engine/Private/Materials/HLSLMaterialTranslator.cpp
@@ -6165,7 +6165,7 @@ int32 FHLSLMaterialTranslator::ObjectRadius()
 
 int32 FHLSLMaterialTranslator::ObjectBounds()
 {
-	return AddInlinedCodeChunk(MCT_Float3, TEXT("float3(GetPrimitiveData(Parameters).ObjectBoundsX, GetPrimitiveData(Parameters).ObjectBoundsY, GetPrimitiveData(Parameters).ObjectBoundsZ)"));
+	return AddInlinedCodeChunkZeroDeriv(MCT_Float3, TEXT("float3(GetPrimitiveData(Parameters).ObjectBoundsX, GetPrimitiveData(Parameters).ObjectBoundsY, GetPrimitiveData(Parameters).ObjectBoundsZ)"));
 }
 
 int32 FHLSLMaterialTranslator::ObjectLocalBounds(int32 OutputIndex)
@@ -6173,9 +6173,9 @@ int32 FHLSLMaterialTranslator::ObjectLocalBounds(int32 OutputIndex)
 	switch (OutputIndex)
 	{
 	case 0: // Half extents
-		return AddInlinedCodeChunk(MCT_Float3, TEXT("((GetPrimitiveData(Parameters).LocalObjectBoundsMax - GetPrimitiveData(Parameters).LocalObjectBoundsMin) / 2.0f)"));
+		return AddInlinedCodeChunkZeroDeriv(MCT_Float3, TEXT("((GetPrimitiveData(Parameters).LocalObjectBoundsMax - GetPrimitiveData(Parameters).LocalObjectBoundsMin) / 2.0f)"));
 	case 1: // Full extents
-		return AddInlinedCodeChunk(MCT_Float3, TEXT("(GetPrimitiveData(Parameters).LocalObjectBoundsMax - GetPrimitiveData(Parameters).LocalObjectBoundsMin)"));
+		return AddInlinedCodeChunkZeroDeriv(MCT_Float3, TEXT("(GetPrimitiveData(Parameters).LocalObjectBoundsMax - GetPrimitiveData(Parameters).LocalObjectBoundsMin)"));
 	case 2: // Min point
 		return GetPrimitiveProperty(MCT_Float3, TEXT("ObjectLocalBounds"), TEXT("LocalObjectBoundsMin"));
 	case 3: // Max point
@@ -6192,13 +6192,13 @@ int32 FHLSLMaterialTranslator::InstanceLocalBounds(int32 OutputIndex)
 	switch (OutputIndex)
 	{
 	case 0: // Half extents
-		return AddInlinedCodeChunk(MCT_Float3, TEXT("(GetPrimitiveData(Parameters).InstanceLocalBoundsExtent)"));
+		return AddInlinedCodeChunkZeroDeriv(MCT_Float3, TEXT("(GetPrimitiveData(Parameters).InstanceLocalBoundsExtent)"));
 	case 1: // Full extents
-		return AddInlinedCodeChunk(MCT_Float3, TEXT("(GetPrimitiveData(Parameters).InstanceLocalBoundsExtent * 2.0f)"));
+		return AddInlinedCodeChunkZeroDeriv(MCT_Float3, TEXT("(GetPrimitiveData(Parameters).InstanceLocalBoundsExtent * 2.0f)"));
 	case 2: // Min point
-		return AddInlinedCodeChunk(MCT_Float3, TEXT("(GetPrimitiveData(Parameters).InstanceLocalBoundsCenter - GetPrimitiveData(Parameters).InstanceLocalBoundsExtent)"));
+		return AddInlinedCodeChunkZeroDeriv(MCT_Float3, TEXT("(GetPrimitiveData(Parameters).InstanceLocalBoundsCenter - GetPrimitiveData(Parameters).InstanceLocalBoundsExtent)"));
 	case 3: // Max point
-		return AddInlinedCodeChunk(MCT_Float3, TEXT("(GetPrimitiveData(Parameters).InstanceLocalBoundsCenter + GetPrimitiveData(Parameters).InstanceLocalBoundsExtent)"));
+		return AddInlinedCodeChunkZeroDeriv(MCT_Float3, TEXT("(GetPrimitiveData(Parameters).InstanceLocalBoundsCenter + GetPrimitiveData(Parameters).InstanceLocalBoundsExtent)"));
 	default:
 		check(false);
 	}
@@ -6211,9 +6211,9 @@ int32 FHLSLMaterialTranslator::PreSkinnedLocalBounds(int32 OutputIndex)
 	switch (OutputIndex)
 	{
 	case 0: // Half extents
-		return AddInlinedCodeChunk(MCT_Float3, TEXT("((GetPrimitiveData(Parameters).PreSkinnedLocalBoundsMax - GetPrimitiveData(Parameters).PreSkinnedLocalBoundsMin) / 2.0f)"));
+		return AddInlinedCodeChunkZeroDeriv(MCT_Float3, TEXT("((GetPrimitiveData(Parameters).PreSkinnedLocalBoundsMax - GetPrimitiveData(Parameters).PreSkinnedLocalBoundsMin) / 2.0f)"));
 	case 1: // Full extents
-		return AddInlinedCodeChunk(MCT_Float3, TEXT("(GetPrimitiveData(Parameters).PreSkinnedLocalBoundsMax - GetPrimitiveData(Parameters).PreSkinnedLocalBoundsMin)"));
+		return AddInlinedCodeChunkZeroDeriv(MCT_Float3, TEXT("(GetPrimitiveData(Parameters).PreSkinnedLocalBoundsMax - GetPrimitiveData(Parameters).PreSkinnedLocalBoundsMin)"));
 	case 2: // Min point
 		return GetPrimitiveProperty(MCT_Float3, TEXT("PreSkinnedLocalBounds"), TEXT("PreSkinnedLocalBoundsMin"));
 	case 3: // Max point
@@ -15309,5 +15309,84 @@ int32 FHLSLMaterialTranslator::SparseVolumeTextureSamplePhysicalTileData(int32 S
 	return AddCodeChunk(MCT_Float4, *SampleCode);
 }
 
+int32 FHLSLMaterialTranslator::SparseVolumeTextureSample(int32 SparseVolumeTextureIndex, int32 UVWIndex, int32 MipValue0Index, int32 MipValue1Index, int32 PhysicalTileDataIdxIndex, ETextureMipValueMode MipValueMode, ESamplerSourceMode SamplerSource)
+{
+	if (MipValueMode == TMVM_Derivative)
+	{
+		if (MipValue0Index == INDEX_NONE)
+		{
+			return Errorf(TEXT("Missing DDX(UVs) parameter"));
+		}
+		else if (MipValue1Index == INDEX_NONE)
+		{
+			return Errorf(TEXT("Missing DDY(UVs) parameter"));
+		}
+		else if (!IsFloatNumericType(GetParameterType(MipValue0Index)))
+		{
+			return Errorf(TEXT("Invalid DDX(UVs) parameter"));
+		}
+		else if (!IsFloatNumericType(GetParameterType(MipValue1Index)))
+		{
+			return Errorf(TEXT("Invalid DDY(UVs) parameter"));
+		}
+	}
+	else if (MipValueMode != TMVM_None && MipValue0Index != INDEX_NONE && !IsFloatNumericType(GetParameterType(MipValue0Index)))
+	{
+		return Errorf(TEXT("Invalid mip map parameter"));
+	}
+
+	int32 MipLevelIndex = INDEX_NONE;
+	int32 MipLevelBiasIndex = Constant(0.0f);
+	int32 DDXIndex = INDEX_NONE;
+	int32 DDYIndex = INDEX_NONE;
+	const bool bDeriveMipLevel = (MipValueMode == TMVM_None || MipValueMode == TMVM_MipBias || MipValueMode == TMVM_Derivative);
+
+	if (MipValueMode == TMVM_None || MipValueMode == TMVM_MipBias)
+	{
+		if (MipValueMode == TMVM_MipBias)
+		{
+			MipLevelBiasIndex = MipValue0Index;
+		}
+
+		DDXIndex = DDX(UVWIndex);
+		DDYIndex = DDY(UVWIndex);
+	}
+	else if (MipValueMode == TMVM_MipLevel)
+	{
+		MipLevelIndex = MipValue0Index;
+	}
+	else if (MipValueMode == TMVM_Derivative)
+	{
+		DDXIndex = MipValue0Index;
+		DDYIndex = MipValue1Index;
+	}
+	else
+	{
+		checkNoEntry();
+	}
+
+	if (bDeriveMipLevel)
+	{
+		FString MipLevelCode = FString::Printf(TEXT("SparseVolumeTextureCalculateMipLevel(%s, %s, %s, %s)"), *GetParameterCode(SparseVolumeTextureIndex), *GetParameterCode(DDXIndex), *GetParameterCode(DDYIndex), *GetParameterCode(MipLevelBiasIndex));
+		MipLevelIndex = AddCodeChunk(MCT_Float1, *MipLevelCode);
+	}
+
+	// Sample the first mip
+	int32 MipLevel0Index = Floor(MipLevelIndex);
+	int32 VoxelCoordMip0Index = SparseVolumeTextureSamplePageTable(SparseVolumeTextureIndex, UVWIndex, MipLevel0Index, SamplerSource);
+	int32 Mip0SampleIndex = SparseVolumeTextureSamplePhysicalTileData(SparseVolumeTextureIndex, VoxelCoordMip0Index, PhysicalTileDataIdxIndex);
+
+	// Sample the second mip
+	// SVT_TODO: Try to optimize out this second sample if LerpAlpha == 0. Might need to do that in HLSL.
+	int32 MipLevel1Index = Ceil(MipLevelIndex);
+	int32 VoxelCoordMip1Index = SparseVolumeTextureSamplePageTable(SparseVolumeTextureIndex, UVWIndex, MipLevel1Index, SamplerSource);
+	int32 Mip1SampleIndex = SparseVolumeTextureSamplePhysicalTileData(SparseVolumeTextureIndex, VoxelCoordMip1Index, PhysicalTileDataIdxIndex);
+
+	// Lerp
+	int32 LerpAlphaIndex = Frac(MipLevelIndex);
+	int32 LerpedResultIndex = Lerp(Mip0SampleIndex, Mip1SampleIndex, LerpAlphaIndex);
+
+	return LerpedResultIndex;
+}
 
 #endif // WITH_EDITORONLY_DATA
diff --git a/Engine/Source/Runtime/Engine/Private/Materials/HLSLMaterialTranslator.h b/Engine/Source/Runtime/Engine/Private/Materials/HLSLMaterialTranslator.h
--- a/Engine/Source/Runtime/Engine/Private/Materials/HLSLMaterialTranslator.h
+++ b/Engine/Source/Runtime/Engine/Private/Materials/HLSLMaterialTranslator.h
@@ -1073,6 +1073,7 @@ protected:
 	virtual int32 SparseVolumeTextureUniformParameter(FName ParameterName, int32 TextureIndex, int32 VectorIndex, UE::Shader::EValueType Type) override;
 	virtual int32 SparseVolumeTextureSamplePageTable(int32 SparseVolumeTextureIndex, int32 UVWIndex, int32 MipLevelIndex, ESamplerSourceMode SamplerSource) override;
 	virtual int32 SparseVolumeTextureSamplePhysicalTileData(int32 SparseVolumeTextureIndex, int32 VoxelCoordIndex, int32 PhysicalTileDataIdxIndex) override;
+	virtual int32 SparseVolumeTextureSample(int32 SparseVolumeTextureIndex, int32 UVWIndex, int32 MipValue0Index, int32 MipValue1Index, int32 PhysicalTileDataIdxIndex, ETextureMipValueMode MipValueMode, ESamplerSourceMode SamplerSource) override;
 
 	virtual UObject* GetReferencedTexture(int32 Index);
 
diff --git a/Engine/Source/Runtime/Engine/Private/Materials/MaterialExpressions.cpp b/Engine/Source/Runtime/Engine/Private/Materials/MaterialExpressions.cpp
--- a/Engine/Source/Runtime/Engine/Private/Materials/MaterialExpressions.cpp
+++ b/Engine/Source/Runtime/Engine/Private/Materials/MaterialExpressions.cpp
@@ -28380,11 +28380,43 @@ UMaterialExpressionSparseVolumeTextureSample::UMaterialExpressionSparseVolumeTex
 
 #if WITH_EDITOR
 
+bool UMaterialExpressionSparseVolumeTextureSample::CanEditChange(const FProperty* InProperty) const
+{
+	bool bIsEditable = Super::CanEditChange(InProperty);
+	if (bIsEditable && InProperty != nullptr)
+	{
+		FName PropertyFName = InProperty->GetFName();
+
+		if (PropertyFName == GET_MEMBER_NAME_CHECKED(UMaterialExpressionSparseVolumeTextureSample, ConstMipValue))
+		{
+			bIsEditable = MipValueMode == TMVM_MipLevel || MipValueMode == TMVM_MipBias;
+		}
+		else if (PropertyFName == GET_MEMBER_NAME_CHECKED(UMaterialExpressionSparseVolumeTextureSample, SparseVolumeTexture))
+		{
+			// The Texture property is overridden by a connection to TextureObject
+			bIsEditable = TextureObject.GetTracedInput().Expression == nullptr;
+		}
+	}
+
+	return bIsEditable;
+}
+
 void UMaterialExpressionSparseVolumeTextureSample::PostLoad()
 {
 	Super::PostLoad();
 }
 
+TArrayView<FExpressionInput*> UMaterialExpressionSparseVolumeTextureSample::GetInputsView()
+{
+	CachedInputs.Empty();
+	uint32 InputIndex = 0;
+	while (FExpressionInput* Ptr = GetInput(InputIndex++))
+	{
+		CachedInputs.Add(Ptr);
+	}
+	return CachedInputs;
+}
+
 void UMaterialExpressionSparseVolumeTextureSample::PostEditChangeProperty(FPropertyChangedEvent& PropertyChangedEvent)
 {
 	// Update what needs to be when the texture is changed
@@ -28396,6 +28428,14 @@ void UMaterialExpressionSparseVolumeTextureSample::PostEditChangeProperty(FPrope
 		}
 	}
 
+	if (PropertyChangedEvent.GetMemberPropertyName() == GET_MEMBER_NAME_CHECKED(ThisClass, MipValueMode))
+	{
+		if (GraphNode)
+		{
+			GraphNode->ReconstructNode();
+		}
+	}
+
 	Super::PostEditChangeProperty(PropertyChangedEvent);
 }
 
@@ -28413,20 +28453,86 @@ uint32 UMaterialExpressionSparseVolumeTextureSample::GetOutputType(int32 OutputI
 	return MCT_Float1;
 }
 
+// this define is only used for the following function
+#define IF_INPUT_RETURN(Item) if(!InputIndex) return &Item; --InputIndex
+FExpressionInput* UMaterialExpressionSparseVolumeTextureSample::GetInput(int32 InputIndex)
+{
+	IF_INPUT_RETURN(Coordinates);
+
+	IF_INPUT_RETURN(TextureObject);
+
+	if (MipValueMode == TMVM_Derivative)
+	{
+		IF_INPUT_RETURN(CoordinatesDX);
+		IF_INPUT_RETURN(CoordinatesDY);
+	}
+	else if (MipValueMode != TMVM_None)
+	{
+		IF_INPUT_RETURN(MipValue);
+	}
+
+	return nullptr;
+}
+#undef IF_INPUT_RETURN
+
+// this define is only used for the following function
+#define IF_INPUT_RETURN(Name) if(!InputIndex) return Name; --InputIndex
+FName UMaterialExpressionSparseVolumeTextureSample::GetInputName(int32 InputIndex) const
+{
+	// Coordinates
+	IF_INPUT_RETURN(TEXT("Coordinates"));
+
+	// TextureObject
+	IF_INPUT_RETURN(TEXT("TextureObject"));
+
+	if (MipValueMode == TMVM_MipLevel)
+	{
+		// MipValue
+		IF_INPUT_RETURN(TEXT("MipLevel"));
+	}
+	else if (MipValueMode == TMVM_MipBias)
+	{
+		// MipValue
+		IF_INPUT_RETURN(TEXT("MipBias"));
+	}
+	else if (MipValueMode == TMVM_Derivative)
+	{
+		// CoordinatesDX
+		IF_INPUT_RETURN(TEXT("DDX(UVs)"));
+		// CoordinatesDY
+		IF_INPUT_RETURN(TEXT("DDY(UVs)"));
+	}
+
+	return TEXT("");
+}
+#undef IF_INPUT_RETURN
+
+// this define is only used for the following function
+#define IF_INPUT_RETURN(Type) if(!InputIndex) return (Type); --InputIndex
 uint32 UMaterialExpressionSparseVolumeTextureSample::GetInputType(int32 InputIndex)
 {
-	switch (InputIndex)
+	// Coordinates
+	IF_INPUT_RETURN(MCT_Float3);
+
+	// TextureObject
+	IF_INPUT_RETURN(MCT_SparseVolumeTexture);
+
+	if (MipValueMode == TMVM_MipLevel || MipValueMode == TMVM_MipBias)
 	{
-	case 0:
-		return MCT_Float3;
-	case 1:
-		return MCT_SparseVolumeTexture;
-	case 2:
-		return MCT_Float1;
-	default:
-		return MCT_Unknown;
+		// MipValue
+		IF_INPUT_RETURN(MCT_Float);
 	}
+	else if (MipValueMode == TMVM_Derivative)
+	{
+		// CoordinatesDX
+		IF_INPUT_RETURN(MCT_Float);
+		// CoordinatesDY
+		IF_INPUT_RETURN(MCT_Float);
+	}
+
+	return MCT_Unknown;
 }
+#undef IF_INPUT_RETURN
 
 int32 UMaterialExpressionSparseVolumeTextureSample::Compile(class FMaterialCompiler* Compiler, int32 OutputIndex)
 {
@@ -28478,38 +28584,35 @@ int32 UMaterialExpressionSparseVolumeTextureSample::Compile(class FMaterialCompi
 
 			if (CoordinateIndex == INDEX_NONE)
 			{
-				CompilerError(Compiler, TEXT("Failed to generate fallback UVW input for sparse volume texture"));
+				return CompilerError(Compiler, TEXT("Failed to generate fallback UVW input for sparse volume texture"));
 			}
 		}
 
-		UMaterialExpression* MipLevelExpression = MipLevel.GetTracedInput().Expression;
-
-		// Shared inputs for both potential samples
-		int32 PhysicalTileDataIdxIndex = Compiler->Constant(OutputIndex);
-		int32 MipLevelInputIndex = MipLevelExpression ? MipLevel.Compile(Compiler) : INDEX_NONE;
-		
-		// Sample the first mip
-		int32 MipLevel0Index = MipLevelExpression ? Compiler->Floor(MipLevelInputIndex) : Compiler->Constant(0.0f);
-		int32 VoxelCoordMip0Index = Compiler->SparseVolumeTextureSamplePageTable(SparseVolumeTextureIndex, CoordinateIndex, MipLevel0Index, SamplerSource);
-		int32 Mip0SampleIndex = Compiler->SparseVolumeTextureSamplePhysicalTileData(SparseVolumeTextureIndex, VoxelCoordMip0Index, PhysicalTileDataIdxIndex);
-
-		if (MipLevelExpression)
+		int32 MipValue0Index = INDEX_NONE;
+		int32 MipValue1Index = INDEX_NONE;
+		if (MipValueMode == TMVM_Derivative)
 		{
-			// Sample the second mip
-			// SVT_TODO: Try to optimize out this second sample if LerpAlpha == 0. Might need to do that in HLSL.
-			int32 MipLevel1Index = Compiler->Ceil(MipLevelInputIndex);
-			int32 VoxelCoordMip1Index = Compiler->SparseVolumeTextureSamplePageTable(SparseVolumeTextureIndex, CoordinateIndex, MipLevel1Index, SamplerSource);
-			int32 Mip1SampleIndex = Compiler->SparseVolumeTextureSamplePhysicalTileData(SparseVolumeTextureIndex, VoxelCoordMip1Index, PhysicalTileDataIdxIndex);
-
-			// Lerp
-			int32 LerpAlphaIndex = Compiler->Frac(MipLevelInputIndex);
-			int32 LerpedResultIndex = Compiler->Lerp(Mip0SampleIndex, Mip1SampleIndex, LerpAlphaIndex);
-			return LerpedResultIndex;
+			if (CoordinatesDX.GetTracedInput().IsConnected())
+			{
+				MipValue0Index = CoordinatesDX.Compile(Compiler);
+			}
+			if (CoordinatesDY.GetTracedInput().IsConnected())
+			{
+				MipValue1Index = CoordinatesDY.Compile(Compiler);
+			}
+		}
+		else if (MipValue.GetTracedInput().IsConnected())
+		{
+			MipValue0Index = MipValue.Compile(Compiler);
 		}
 		else
 		{
-			return Mip0SampleIndex;
+			MipValue0Index = Compiler->Constant(ConstMipValue);
 		}
+
+		int32 PhysicalTileDataIdxIndex = Compiler->Constant(OutputIndex);
+
+		return Compiler->SparseVolumeTextureSample(SparseVolumeTextureIndex, CoordinateIndex, MipValue0Index, MipValue1Index, PhysicalTileDataIdxIndex, MipValueMode, SamplerSource);
 	}
 	else
 	{
diff --git a/Engine/Source/Runtime/Engine/Public/MaterialCompiler.h b/Engine/Source/Runtime/Engine/Public/MaterialCompiler.h
--- a/Engine/Source/Runtime/Engine/Public/MaterialCompiler.h
+++ b/Engine/Source/Runtime/Engine/Public/MaterialCompiler.h
@@ -357,6 +357,8 @@ public:
 	 */
 	virtual int32 SparseVolumeTextureSamplePhysicalTileData(int32 SparseVolumeTextureIndex, int32 VoxelCoordIndex, int32 PhysicalTileDataIdxIndex) = 0;
 
+	virtual int32 SparseVolumeTextureSample(int32 SparseVolumeTextureIndex, int32 UVWIndex, int32 MipValue0Index, int32 MipValue1Index, int32 PhysicalTileDataIdxIndex, ETextureMipValueMode MipValueMode, ESamplerSourceMode SamplerSource) = 0;
+
 	virtual UObject* GetReferencedTexture(int32 Index) { return nullptr; }
 
 	int32 Texture(UTexture* InTexture, EMaterialSamplerType SamplerType, ESamplerSourceMode SamplerSource=SSM_FromTextureAsset)
@@ -920,6 +922,7 @@ public:
 	virtual int32 SparseVolumeTextureUniformParameter(FName ParameterName, int32 TextureIndex, int32 VectorIndex, UE::Shader::EValueType Type) override														{ return Compiler->SparseVolumeTextureUniformParameter(ParameterName, TextureIndex, VectorIndex, Type); }
 	virtual int32 SparseVolumeTextureSamplePageTable(int32 SparseVolumeTextureIndex, int32 UVWIndex, int32 MipLevelIndex, ESamplerSourceMode SamplerSource) override										{ return Compiler->SparseVolumeTextureSamplePageTable(SparseVolumeTextureIndex, UVWIndex, MipLevelIndex, SamplerSource); }
 	virtual int32 SparseVolumeTextureSamplePhysicalTileData(int32 SparseVolumeTextureIndex, int32 VoxelCoordIndex, int32 PhysicalTileDataIdxIndex) override													{ return Compiler->SparseVolumeTextureSamplePhysicalTileData(SparseVolumeTextureIndex, VoxelCoordIndex, PhysicalTileDataIdxIndex); }
+	virtual int32 SparseVolumeTextureSample(int32 SparseVolumeTextureIndex, int32 UVWIndex, int32 MipValue0Index, int32 MipValue1Index, int32 PhysicalTileDataIdxIndex, ETextureMipValueMode MipValueMode, ESamplerSourceMode SamplerSource) { return Compiler->SparseVolumeTextureSample(SparseVolumeTextureIndex, UVWIndex, MipValue0Index, MipValue1Index, PhysicalTileDataIdxIndex, MipValueMode, SamplerSource); }
 
 	virtual UObject* GetReferencedTexture(int32 Index) override { return Compiler->GetReferencedTexture(Index); }
 
