Removed assert that checks for non-empty skeletal mesh vertex factories, as certain LODs can have no actual mesh data.

diff --git a/Engine/Source/Runtime/Engine/Private/SkeletalRenderGPUSkin.cpp b/Engine/Source/Runtime/Engine/Private/SkeletalRenderGPUSkin.cpp
--- a/Engine/Source/Runtime/Engine/Private/SkeletalRenderGPUSkin.cpp
+++ b/Engine/Source/Runtime/Engine/Private/SkeletalRenderGPUSkin.cpp
@@ -626,8 +626,6 @@ void FSkeletalMeshObjectGPUSkin::ProcessUpdatedDynamicData(
 	const FSkeletalMeshLODRenderData& LODData = SkeletalMeshRenderData->LODRenderData[LODIndex];
 	const TArray<FSkelMeshRenderSection>& Sections = GetRenderSections(LODIndex);
 
-	check(!VertexFactoryData.VertexFactories.IsEmpty());
-
 	const bool bHasWeightedActiveMorphs = DynamicData->NumWeightedActiveMorphTargets > 0;
 	const bool bHasExternalMorphs = DynamicData->ExternalMorphWeightData.HasActiveMorphs() && !DynamicData->ExternalMorphWeightData.MorphSets.IsEmpty();
 	
