Deprecate constructors and factory functions related to lists of FEdGraphSchemaActions, which have long been unused

diff --git a/Engine/Source/Runtime/Engine/Classes/EdGraph/EdGraphSchema.h b/Engine/Source/Runtime/Engine/Classes/EdGraph/EdGraphSchema.h
--- a/Engine/Source/Runtime/Engine/Classes/EdGraph/EdGraphSchema.h
+++ b/Engine/Source/Runtime/Engine/Classes/EdGraph/EdGraphSchema.h
@@ -462,6 +462,7 @@ public:
 		ENGINE_API ActionGroup( TSharedPtr<FEdGraphSchemaAction> InAction, FString RootCategory = FString());
 
 		/** Constructor accepting multiple actions */
+		UE_DEPRECATED(5.5, "Collections of Action lists were ambiguous and unused - make a compound FEdGraphSchemaAction and construct the ActionGroup with it")
 		ENGINE_API ActionGroup( const TArray< TSharedPtr<FEdGraphSchemaAction> >& InActions, FString RootCategory = FString());
 
 		/** Move constructor and move assignment operator */
@@ -484,6 +485,7 @@ public:
 		 * @param FromPins Optional pins that the action was dragged from.
 		 * @param Location The position on the graph to place new nodes.
 		 */
+		UE_DEPRECATED(5.5, "Invoke PerformAction on the ActionGroup::Actions directly")
 		ENGINE_API void PerformAction( class UEdGraph* ParentGraph, TArray<UEdGraphPin*>& FromPins, const FVector2D Location );
 		
 		/**
@@ -547,6 +549,7 @@ public:
 	ENGINE_API virtual void AddAction( const TSharedPtr<FEdGraphSchemaAction>& NewAction, FString const& Category = FString() );
 
 	/** Adds an action entry containing multiple actions */
+	UE_DEPRECATED(5.5, "Collections of Action lists were ambiguous and unused - make a compound FEdGraphSchemaAction and Add it")
 	ENGINE_API virtual void AddActionList( const TArray<TSharedPtr<FEdGraphSchemaAction> >& NewActions, FString const& Category = FString() );
 
 	/** Appends all the action entries from a different graph action builder */
@@ -585,6 +588,7 @@ public:
 
 	// FGraphActionListBuilderBase Interface
 	ENGINE_API virtual void AddAction(const TSharedPtr<FEdGraphSchemaAction>& NewAction, FString const& Category = FString() ) override;
+	UE_DEPRECATED(5.5, "Collections of Action lists were ambiguous and unused - make a compound FEdGraphSchemaAction and Add it")
 	ENGINE_API virtual void AddActionList(const TArray<TSharedPtr<FEdGraphSchemaAction> >& NewActions, FString const& Category = FString()) override;
 	// End of FGraphActionListBuilderBase Interface
 
diff --git a/Engine/Source/Runtime/Engine/Private/EdGraph/EdGraphSchema.cpp b/Engine/Source/Runtime/Engine/Private/EdGraph/EdGraphSchema.cpp
--- a/Engine/Source/Runtime/Engine/Private/EdGraph/EdGraphSchema.cpp
+++ b/Engine/Source/Runtime/Engine/Private/EdGraph/EdGraphSchema.cpp
@@ -111,7 +111,9 @@ void FGraphActionListBuilderBase::AddAction( const TSharedPtr<FEdGraphSchemaActi
 
 void FGraphActionListBuilderBase::AddActionList( const TArray<TSharedPtr<FEdGraphSchemaAction> >& NewActions, FString const& Category)
 {
+	PRAGMA_DISABLE_DEPRECATION_WARNINGS
 	Entries.Add( ActionGroup( NewActions, Category ) );
+	PRAGMA_ENABLE_DEPRECATION_WARNINGS
 }
 
 void FGraphActionListBuilderBase::Append( FGraphActionListBuilderBase& Other )
@@ -268,7 +270,9 @@ void FCategorizedGraphActionListBuilder::AddAction(TSharedPtr<FEdGraphSchemaActi
 
 void FCategorizedGraphActionListBuilder::AddActionList(TArray<TSharedPtr<FEdGraphSchemaAction> > const& NewActions, FString const& CategoryIn)
 {
+	PRAGMA_DISABLE_DEPRECATION_WARNINGS
 	FGraphActionListBuilderBase::AddActionList(NewActions, ConcatCategories(Category, CategoryIn));
+	PRAGMA_ENABLE_DEPRECATION_WARNINGS
 }
 
 /////////////////////////////////////////////////////
