[Perf] Improve Audio Thread PRS Stats

1. Added csv scopes to track audio task run/wait times.
2. Enabled csv category Audio/AudioMetasound to track AudioMixerRenderThread timings.
3. Added a "GameThread Audio" summary table to summarize relevant audio perf counters.

diff --git a/Engine/Source/Runtime/Engine/Private/AudioThread.cpp b/Engine/Source/Runtime/Engine/Private/AudioThread.cpp
--- a/Engine/Source/Runtime/Engine/Private/AudioThread.cpp
+++ b/Engine/Source/Runtime/Engine/Private/AudioThread.cpp
@@ -255,6 +255,8 @@ TUniqueFunction<void()> FAudioThread::GetCommandWrapper(TUniqueFunction<void()>
 	{
 		return [Function = MoveTemp(InFunction), InStatId]()
 		{
+			CSV_SCOPED_TIMING_STAT_EXCLUSIVE(Audio);
+
 			FScopeCycleCounter ScopeCycleCounter(InStatId);
 			FAudioThread::SetCurrentAudioThreadStatId(InStatId);
 
@@ -276,6 +278,8 @@ TUniqueFunction<void()> FAudioThread::GetCommandWrapper(TUniqueFunction<void()>
 	{
 		return [Function = MoveTemp(InFunction), InStatId]()
 		{
+			CSV_SCOPED_TIMING_STAT_EXCLUSIVE(Audio);
+
 			FScopeCycleCounter ScopeCycleCounter(InStatId);
 			Function();
 		};
@@ -441,6 +445,7 @@ void FAudioThread::StopAudioThread()
 FAudioCommandFence::~FAudioCommandFence()
 {
 	check(IsInGameThread());
+	CSV_SCOPED_SET_WAIT_STAT(Audio);
 	Fence.Wait();
 }
 
@@ -453,7 +458,11 @@ void FAudioCommandFence::BeginFence()
 		return;
 	}
 
-	Fence.Wait();
+	{
+		CSV_SCOPED_SET_WAIT_STAT(Audio);
+		Fence.Wait();
+	}
+	
 	FAudioThread::ProcessAllCommands();
 	Fence = GAudioAsyncBatcher.LastBatch;
 }
@@ -470,6 +479,11 @@ bool FAudioCommandFence::IsFenceComplete() const
 void FAudioCommandFence::Wait(bool bProcessGameThreadTasks) const
 {
 	check(IsInGameThread());
-	Fence.Wait();
+
+	{
+		CSV_SCOPED_SET_WAIT_STAT(Audio);
+		Fence.Wait();
+	}
+
 	Fence = UE::Tasks::FTask{}; // release the task as it can hold some references
 }
