Removed code that made Mobile-only properties in UVirtualTextureBuilder readonly when Enable virtual texture support for Mobile project settings is disabled. This turns out to be impractical since this settings maps to a per-platform CVar (r.Mobile.VirtualTextures), which means the property could be disabled in the project settings but actually enabled on some mobile platforms. This is a valid use case so we now let the user access the property regardless of that project settings value

diff --git a/Engine/Source/Runtime/Engine/Classes/VT/VirtualTextureBuilder.h b/Engine/Source/Runtime/Engine/Classes/VT/VirtualTextureBuilder.h
--- a/Engine/Source/Runtime/Engine/Classes/VT/VirtualTextureBuilder.h
+++ b/Engine/Source/Runtime/Engine/Classes/VT/VirtualTextureBuilder.h
@@ -77,8 +77,5 @@ protected:
 	//~ Begin UObject Interface
 	virtual void Serialize(FArchive& Ar) override;
 	virtual void PostLoad() override;
-#if WITH_EDITOR
-	virtual bool CanEditChange(const FProperty* InProperty) const override;
-#endif // WITH_EDITOR
 	//~ End UObject Interface
 };
diff --git a/Engine/Source/Runtime/Engine/Private/VT/VirtualTextureBuilder.cpp b/Engine/Source/Runtime/Engine/Private/VT/VirtualTextureBuilder.cpp
--- a/Engine/Source/Runtime/Engine/Private/VT/VirtualTextureBuilder.cpp
+++ b/Engine/Source/Runtime/Engine/Private/VT/VirtualTextureBuilder.cpp
@@ -91,25 +91,6 @@ UVirtualTexture2D* UVirtualTextureBuilder::GetVirtualTexture(EShadingPath Shadin
 
 #if WITH_EDITOR
 
-bool UVirtualTextureBuilder::CanEditChange(const FProperty* InProperty) const
-{
-	bool bCanEdit = Super::CanEditChange(InProperty);
-
-	static const TConsoleVariableData<int32>* CVarVirtualTextureOnMobile = IConsoleManager::Get().FindTConsoleVariableDataInt(TEXT("r.Mobile.VirtualTextures"));
-	check(CVarVirtualTextureOnMobile);
-	const bool bEnableVirtualTextureOnMobile = (CVarVirtualTextureOnMobile->GetValueOnGameThread() != 0);
-
-	if (InProperty->GetFName() == GET_MEMBER_NAME_CHECKED(UVirtualTextureBuilder, bSeparateTextureForMobile))
-	{
-		bCanEdit &= bEnableVirtualTextureOnMobile;
-	}
-	else if (InProperty->GetFName() == GET_MEMBER_NAME_CHECKED(UVirtualTextureBuilder, TextureMobile))
-	{
-		bCanEdit &= bEnableVirtualTextureOnMobile && bSeparateTextureForMobile;
-	}
-	return bCanEdit;
-}
-
 static void BuildVirtualTexture2D(UVirtualTexture2D* Texture, FVirtualTextureBuildDesc const& BuildDesc, bool bWaitForCompilation)
 {
 	Texture->VirtualTextureStreaming = true;
