Check mip data provider next tick state to avoid executing wrong steps when multiple providers in the same FTextureStreamIn have different tick states

diff --git a/Engine/Source/Runtime/Engine/Private/Streaming/TextureStreamIn.cpp b/Engine/Source/Runtime/Engine/Private/Streaming/TextureStreamIn.cpp
--- a/Engine/Source/Runtime/Engine/Private/Streaming/TextureStreamIn.cpp
+++ b/Engine/Source/Runtime/Engine/Private/Streaming/TextureStreamIn.cpp
@@ -168,7 +168,8 @@ void FTextureStreamIn::DoInitMipDataProviders(const FContext& Context)
 	for (TUniquePtr<FTextureMipDataProvider>& MipDataProvider : MipDataProviders)
 	{
 		check(MipDataProvider);
-		if (IsSameThread(MipDataProvider->GetNextTickThread(), Context.CurrentThread))
+		if (MipDataProvider->GetNextTickState() == FTextureMipDataProvider::ETickState::Init
+			&& IsSameThread(MipDataProvider->GetNextTickThread(), Context.CurrentThread))
 		{
 			MipDataProvider->Init(Context, SyncOptions);
 		}
@@ -186,7 +187,8 @@ void FTextureStreamIn::DoGetMipData(const FContext& Context)
 	for (TUniquePtr<FTextureMipDataProvider>& MipDataProvider : MipDataProviders)
 	{
 		check(MipDataProvider);
-		if (IsSameThread(MipDataProvider->GetNextTickThread(), Context.CurrentThread))
+		if (MipDataProvider->GetNextTickState() == FTextureMipDataProvider::ETickState::GetMips
+			&& IsSameThread(MipDataProvider->GetNextTickThread(), Context.CurrentThread))
 		{
 			StartingMipIndex = MipDataProvider->GetMips(Context, StartingMipIndex, MipInfos, SyncOptions);
 		}
@@ -198,7 +200,8 @@ bool FTextureStreamIn::DoPollMipData(const FContext& Context)
 	for (TUniquePtr<FTextureMipDataProvider>& MipDataProvider : MipDataProviders)
 	{
 		check(MipDataProvider);
-		if (IsSameThread(MipDataProvider->GetNextTickThread(), Context.CurrentThread))
+		if (MipDataProvider->GetNextTickState() == FTextureMipDataProvider::ETickState::PollMips
+			&& IsSameThread(MipDataProvider->GetNextTickThread(), Context.CurrentThread))
 		{
 			if (!MipDataProvider->PollMips(SyncOptions))
 			{
@@ -220,7 +223,8 @@ void FTextureStreamIn::DoCleanUpMipDataProviders(const FContext& Context)
 	for (TUniquePtr<FTextureMipDataProvider>& MipDataProvider : MipDataProviders)
 	{
 		check(MipDataProvider);
-		if (IsSameThread(MipDataProvider->GetNextTickThread(), Context.CurrentThread))
+		if (MipDataProvider->GetNextTickState() == FTextureMipDataProvider::ETickState::CleanUp
+			&& IsSameThread(MipDataProvider->GetNextTickThread(), Context.CurrentThread))
 		{
 			MipDataProvider->CleanUp(SyncOptions);
 		}
