World Partition - HLOD: Updated IsHLODRelevant() to ensure that class descriptors have proper values for that flag
* Removed IsTemplate() test, actor class descriptors are built out of CDOs
* Removed arbitrary bounds size check (squared size needed to be 0.1)

diff --git a/Engine/Source/Runtime/Engine/Private/Actor.cpp b/Engine/Source/Runtime/Engine/Private/Actor.cpp
--- a/Engine/Source/Runtime/Engine/Private/Actor.cpp
+++ b/Engine/Source/Runtime/Engine/Private/Actor.cpp
@@ -6040,11 +6040,6 @@ bool AActor::IsHLODRelevant() const
 		return false;
 	}
 
-	if (IsTemplate())
-	{
-		return false;
-	}
-
 	if (IsHidden())
 	{
 		return false;
@@ -6068,13 +6063,6 @@ bool AActor::IsHLODRelevant() const
 	}
 #endif
 
-	FVector Origin, Extent;
-	GetActorBounds(false, Origin, Extent);
-	if (Extent.SizeSquared() <= 0.1)
-	{
-		return false;
-	}
-
 	return HasHLODRelevantComponents();
 }
 
