[WorldPartition] Override loading ranges
- Added generic command wp.Runtime.OverrideRuntimeLoadingRange to replace per-hash implementations.
- Added functions to override grids loading ranges in the world partition subsystem, not available in shipping builds.

diff --git a/Engine/Source/Runtime/Engine/Private/WorldPartition/RuntimeHashSet/WorldPartitionRuntimeHashSet.cpp b/Engine/Source/Runtime/Engine/Private/WorldPartition/RuntimeHashSet/WorldPartitionRuntimeHashSet.cpp
--- a/Engine/Source/Runtime/Engine/Private/WorldPartition/RuntimeHashSet/WorldPartitionRuntimeHashSet.cpp
+++ b/Engine/Source/Runtime/Engine/Private/WorldPartition/RuntimeHashSet/WorldPartitionRuntimeHashSet.cpp
@@ -8,66 +8,12 @@
 #include "WorldPartition/HLOD/HLODLayer.h"
 #include "WorldPartition/ContentBundle/ContentBundleDescriptor.h"
 #include "WorldPartition/DataLayer/DataLayersID.h"
+#include "WorldPartition/WorldPartitionSubsystem.h"
 #include "Algo/RemoveIf.h"
 #include "Algo/Transform.h"
 #include "Misc/HashBuilder.h"
 #include "Misc/ArchiveMD5.h"
 
-#if !UE_BUILD_SHIPPING
-#include "Engine/Engine.h"
-#include "WorldPartition/WorldPartitionHelpers.h"
-#include "WorldPartition/WorldPartitionSubsystem.h"
-
-TMap<FName, int32> FRuntimePartitionStreamingData::OverriddenLoadingRanges;
-static const TCHAR* GOverrideHashSetLoadingRangeCommandName = TEXT("wp.Runtime.OverrideRuntimeHashSetLoadingRange");
-static FDelegateHandle OnWorldPartitionSubsystemDeinitializedFDelegateHandle;
-FAutoConsoleCommand FRuntimePartitionStreamingData::OverrideLoadingRangeCommand(
-	GOverrideHashSetLoadingRangeCommandName,
-	TEXT("Sets runtime loading range. Args -partition=[Name] -range=[Range]"),
-	FConsoleCommandWithArgsDelegate::CreateLambda([](const TArray<FString>& InArgs)
-	{
-		FString ArgString = FString::Join(InArgs, TEXT(" "));
-		FName OverrideGridName;
-		int32 OverrideLoadingRange = -1;
-		FParse::Value(*ArgString, TEXT("partition="), OverrideGridName);
-		FParse::Value(*ArgString, TEXT("range="), OverrideLoadingRange);
-
-		if (!OnWorldPartitionSubsystemDeinitializedFDelegateHandle.IsValid())
-		{
-			OnWorldPartitionSubsystemDeinitializedFDelegateHandle = UWorldPartitionSubsystem::OnWorldPartitionSubsystemDeinitialized.AddLambda([](UWorldPartitionSubsystem* InWorldPartitionSubsystem, UWorld* InWorld)
-			{
-				if (InWorld && InWorld->IsGameWorld())
-				{
-					OverriddenLoadingRanges.Reset();
-				}
-			});
-		}
-
-		for (const FWorldContext& Context : GEngine->GetWorldContexts())
-		{
-			UWorld* World = Context.World();
-			if (World && World->IsGameWorld())
-			{
-				if (UWorld::HasSubsystem<UWorldPartitionSubsystem>(World))
-				{
-					FWorldPartitionHelpers::ServerExecConsoleCommand(World, GOverrideHashSetLoadingRangeCommandName, InArgs);
-
-					if (OverrideLoadingRange >= 0)
-					{
-						OverriddenLoadingRanges.Add(OverrideGridName, OverrideLoadingRange);
-					}
-					else
-					{
-						OverriddenLoadingRanges.Remove(OverrideGridName);
-					}
-					break;
-				}
-			}
-		}
-	})
-);
-#endif
-
 void FRuntimePartitionStreamingData::CreatePartitionsSpatialIndex() const
 {
 	if (!SpatialIndex)
@@ -106,12 +52,12 @@ void FRuntimePartitionStreamingData::DestroyPartitionsSpatialIndex() const
 int32 FRuntimePartitionStreamingData::GetLoadingRange() const
 {
 #if !UE_BUILD_SHIPPING
-	if (int32* OverriddenLoadingRange = OverriddenLoadingRanges.Find(Name))
+	int32 OverriddenLoadingRange;
+	if (UWorldPartitionSubsystem::GetOverrideLoadingRange(Name, OverriddenLoadingRange))
 	{
-		return *OverriddenLoadingRange;
+		return OverriddenLoadingRange;
 	}
 #endif
-
 	return LoadingRange;
 }
 
@@ -618,11 +564,11 @@ void UWorldPartitionRuntimeHashSet::ForEachStreamingCellsSources(const TArray<FW
 uint32 UWorldPartitionRuntimeHashSet::ComputeUpdateStreamingHash() const
 {
 	FHashBuilder HashBuilder(Super::ComputeUpdateStreamingHash());
-
-#if !UE_BUILD_SHIPPING
-	HashBuilder << FRuntimePartitionStreamingData::OverriddenLoadingRanges;
-#endif
-
+	ForEachStreamingData([&HashBuilder](const FRuntimePartitionStreamingData& StreamingData)
+	{
+		HashBuilder << StreamingData.GetLoadingRange();
+		return true;
+	});
 	return HashBuilder.GetHash();
 }
 
diff --git a/Engine/Source/Runtime/Engine/Private/WorldPartition/WorldPartitionRuntimeSpatialHash.cpp b/Engine/Source/Runtime/Engine/Private/WorldPartition/WorldPartitionRuntimeSpatialHash.cpp
--- a/Engine/Source/Runtime/Engine/Private/WorldPartition/WorldPartitionRuntimeSpatialHash.cpp
+++ b/Engine/Source/Runtime/Engine/Private/WorldPartition/WorldPartitionRuntimeSpatialHash.cpp
@@ -125,11 +125,11 @@ const FSquare2DGridHelper& FSpatialHashStreamingGrid::GetGridHelper() const
 	return *GridHelper;
 }
 
-static const FString GOverrideLoadingRangeCommandName(TEXT("wp.Runtime.OverrideRuntimeSpatialHashLoadingRange"));
+static const FString GOverrideRuntimeSpatialHashLoadingRangeName(TEXT("wp.Runtime.OverrideRuntimeSpatialHashLoadingRange"));
 bool FSpatialHashStreamingGrid::bAddedWorldPartitionSubsystemDeinitializedCallback = false;
 TMap<int32, float> FSpatialHashStreamingGrid::OverriddenLoadingRanges;
 FAutoConsoleCommand FSpatialHashStreamingGrid::OverrideLoadingRangeCommand(
-	*GOverrideLoadingRangeCommandName,
+	*GOverrideRuntimeSpatialHashLoadingRangeName,
 	TEXT("Sets runtime loading range. Args -grid=[index] -range=[override_loading_range]"),
 	FConsoleCommandWithArgsDelegate::CreateLambda([](const TArray<FString>& InArgs)
 	{
@@ -156,7 +156,7 @@ FAutoConsoleCommand FSpatialHashStreamingGrid::OverrideLoadingRangeCommand(
 			UWorld* World = Context.World();
 			if (World && World->IsGameWorld())
 			{
-				FWorldPartitionHelpers::ServerExecConsoleCommand(World, GOverrideLoadingRangeCommandName, InArgs);
+				FWorldPartitionHelpers::ServerExecConsoleCommand(World, GOverrideRuntimeSpatialHashLoadingRangeName, InArgs);
 				if (UWorld::HasSubsystem<UWorldPartitionSubsystem>(World))
 				{
 					if (OverrideLoadingRange >= 0.f)
@@ -176,11 +176,20 @@ FAutoConsoleCommand FSpatialHashStreamingGrid::OverrideLoadingRangeCommand(
 
 float FSpatialHashStreamingGrid::GetLoadingRange() const
 {
+#if !UE_BUILD_SHIPPING
+	int32 OverriddenLoadingRange;
+	if (UWorldPartitionSubsystem::GetOverrideLoadingRange(GridName, OverriddenLoadingRange))
+	{
+		return OverriddenLoadingRange;
+	}
+#endif
+
 	if (float* OverrideLoadingRange = FSpatialHashStreamingGrid::OverriddenLoadingRanges.Find(GridIndex))
 	{
 		check(*OverrideLoadingRange >= 0.f);
 		return *OverrideLoadingRange;
 	}
+
 	return LoadingRange;
 }
 
diff --git a/Engine/Source/Runtime/Engine/Private/WorldPartition/WorldPartitionSubsystem.cpp b/Engine/Source/Runtime/Engine/Private/WorldPartition/WorldPartitionSubsystem.cpp
--- a/Engine/Source/Runtime/Engine/Private/WorldPartition/WorldPartitionSubsystem.cpp
+++ b/Engine/Source/Runtime/Engine/Private/WorldPartition/WorldPartitionSubsystem.cpp
@@ -42,6 +42,12 @@
 #include "WorldPartition/WorldPartitionPropertyOverride.h"
 #endif
 
+#if !UE_BUILD_SHIPPING
+#include "Engine/Engine.h"
+#include "WorldPartition/WorldPartitionHelpers.h"
+#include "WorldPartition/WorldPartitionSubsystem.h"
+#endif
+
 #include UE_INLINE_GENERATED_CPP_BY_NAME(WorldPartitionSubsystem)
 
 DECLARE_STATS_GROUP(TEXT("World Partition"), STATGROUP_WorldPartition, STATCAT_Advanced);
@@ -189,6 +195,49 @@ static FAutoConsoleVariableRef CVarUdateStreamingStateTimeLimit(
 	ECVF_Default
 );
 
+#if !UE_BUILD_SHIPPING
+TMap<FName, int32> UWorldPartitionSubsystem::OverriddenLoadingRanges;
+static const TCHAR* GOverrideLoadingRangeCommandName = TEXT("wp.Runtime.OverrideRuntimeLoadingRange");
+static FDelegateHandle OnWorldPartitionSubsystemDeinitializedFDelegateHandle;
+FAutoConsoleCommand UWorldPartitionSubsystem::OverrideLoadingRangeCommand(
+	GOverrideLoadingRangeCommandName,
+	TEXT("Sets runtime loading range. Args -grid=[Name] -range=[Range]"),
+	FConsoleCommandWithArgsDelegate::CreateLambda([](const TArray<FString>& InArgs)
+	{
+		FString ArgString = FString::Join(InArgs, TEXT(" "));
+		FName OverrideGridName;
+		int32 OverrideLoadingRange = -1;
+		FParse::Value(*ArgString, TEXT("grid="), OverrideGridName);
+		FParse::Value(*ArgString, TEXT("range="), OverrideLoadingRange);
+
+		if (!OnWorldPartitionSubsystemDeinitializedFDelegateHandle.IsValid())
+		{
+			OnWorldPartitionSubsystemDeinitializedFDelegateHandle = UWorldPartitionSubsystem::OnWorldPartitionSubsystemDeinitialized.AddLambda([](UWorldPartitionSubsystem* InWorldPartitionSubsystem, UWorld* InWorld)
+			{
+				if (InWorld && InWorld->IsGameWorld())
+				{
+					OverriddenLoadingRanges.Reset();
+				}
+			});
+		}
+
+		for (const FWorldContext& Context : GEngine->GetWorldContexts())
+		{
+			UWorld* World = Context.World();
+			if (World && World->IsGameWorld())
+			{
+				if (UWorld::HasSubsystem<UWorldPartitionSubsystem>(World))
+				{
+					FWorldPartitionHelpers::ServerExecConsoleCommand(World, GOverrideLoadingRangeCommandName, InArgs);
+					UWorldPartitionSubsystem::SetOverrideLoadingRange(OverrideGridName, OverrideLoadingRange);
+					break;
+				}
+			}
+		}
+	})
+);
+#endif
+
 TMulticastDelegate<void(UWorldPartitionSubsystem*, UWorld*)> UWorldPartitionSubsystem::OnWorldPartitionSubsystemInitialized;
 TMulticastDelegate<void(UWorldPartitionSubsystem*, UWorld*)> UWorldPartitionSubsystem::OnWorldPartitionSubsystemDeinitialized;
 
@@ -470,6 +519,30 @@ void UWorldPartitionSubsystem::ForEachWorldPartition(TFunctionRef<bool(UWorldPar
 	}
 }
 
+#if !UE_BUILD_SHIPPING
+void UWorldPartitionSubsystem::SetOverrideLoadingRange(FName Name, int32 LoadingRange)
+{
+	if (LoadingRange >= 0)
+	{
+		OverriddenLoadingRanges.Add(Name, LoadingRange);
+	}
+	else
+	{
+		OverriddenLoadingRanges.Remove(Name);
+	}
+}
+
+bool UWorldPartitionSubsystem::GetOverrideLoadingRange(FName Name, int32& LoadingRange)
+{
+	if (int32* OverriddenLoadingRange = OverriddenLoadingRanges.Find(Name))
+	{
+		LoadingRange = *OverriddenLoadingRange;
+		return true;
+	}
+	return false;
+}
+#endif
+
 void UWorldPartitionSubsystem::OnWorldPartitionInitialized(UWorldPartition* InWorldPartition)
 {
 	if (RegisteredWorldPartitions.IsEmpty())
diff --git a/Engine/Source/Runtime/Engine/Public/WorldPartition/RuntimeHashSet/WorldPartitionRuntimeHashSet.h b/Engine/Source/Runtime/Engine/Public/WorldPartition/RuntimeHashSet/WorldPartitionRuntimeHashSet.h
--- a/Engine/Source/Runtime/Engine/Public/WorldPartition/RuntimeHashSet/WorldPartitionRuntimeHashSet.h
+++ b/Engine/Source/Runtime/Engine/Public/WorldPartition/RuntimeHashSet/WorldPartitionRuntimeHashSet.h
@@ -96,11 +96,6 @@ protected:
 	// Transient
 	mutable TUniquePtr<FStaticSpatialIndexType> SpatialIndex;
 	mutable TUniquePtr<FStaticSpatialIndexType> SpatialIndex2D;
-
-#if !UE_BUILD_SHIPPING
-	static TMap<FName, int32> OverriddenLoadingRanges;
-	static class FAutoConsoleCommand OverrideLoadingRangeCommand;
-#endif
 };
 
 template<>
diff --git a/Engine/Source/Runtime/Engine/Public/WorldPartition/WorldPartitionSubsystem.h b/Engine/Source/Runtime/Engine/Public/WorldPartition/WorldPartitionSubsystem.h
--- a/Engine/Source/Runtime/Engine/Public/WorldPartition/WorldPartitionSubsystem.h
+++ b/Engine/Source/Runtime/Engine/Public/WorldPartition/WorldPartitionSubsystem.h
@@ -101,6 +101,11 @@ public:
 
 	ENGINE_API void ForEachWorldPartition(TFunctionRef<bool(UWorldPartition*)> Func);
 
+#if !UE_BUILD_SHIPPING
+	static ENGINE_API void SetOverrideLoadingRange(FName Name, int32 LoadingRange);
+	static ENGINE_API bool GetOverrideLoadingRange(FName Name, int32& LoadingRange);
+#endif
+
 #if WITH_EDITOR
 	ENGINE_API FWorldPartitionActorFilter GetWorldPartitionActorFilter(const FString& InWorldPackage, EWorldPartitionActorFilterType InFilterTypes = EWorldPartitionActorFilterType::Loading) const;
 	ENGINE_API TMap<FActorContainerID, TSet<FGuid>> GetFilteredActorsPerContainer(const FActorContainerID& InContainerID, const FString& InWorldPackage, const FWorldPartitionActorFilter& InActorFilter, EWorldPartitionActorFilterType InFilterTypes = EWorldPartitionActorFilterType::Loading);
@@ -127,6 +132,10 @@ protected:
 	//~ End USubsystem Interface.
 
 private:
+#if !UE_BUILD_SHIPPING
+	static TMap<FName, int32> OverriddenLoadingRanges;
+	static class FAutoConsoleCommand OverrideLoadingRangeCommand;
+#endif
 
 	// Streaming Sources
 	void UpdateStreamingSources();
