Level Instance: change bIsInEditLevelInstanceHierarchy into a ENUM_CLASS_FLAGS and add a new flag to tag actors with property overrides

#rnx

diff --git a/Engine/Source/Runtime/Engine/Classes/GameFramework/Actor.h b/Engine/Source/Runtime/Engine/Classes/GameFramework/Actor.h
--- a/Engine/Source/Runtime/Engine/Classes/GameFramework/Actor.h
+++ b/Engine/Source/Runtime/Engine/Classes/GameFramework/Actor.h
@@ -102,6 +102,16 @@ enum class ELevelInstanceType : uint8
 	LevelInstancePropertyOverride
 };
 
+UENUM()
+enum class ELevelInstanceFlags : uint8
+{
+	None = 0,
+	IsInEditHierarchy = 1,
+	HasPropertyOverrides = 2,
+};
+
+ENUM_CLASS_FLAGS(ELevelInstanceFlags)
+
 #if WITH_EDITORONLY_DATA
 /** Enum defining how actor will be placed in the partition */
 UENUM()
@@ -295,11 +305,12 @@ private:
 	UPROPERTY(Transient)
 	ELevelInstanceType LevelInstanceType;
 
-	/** Whether this actor belongs to a level instance in a level instance hierarchy currently being edited. Itself or its parent level instances. */
+	/** Flags related to level instances for this actor. */
 	UPROPERTY(Transient)
-	uint8 bIsInEditLevelInstanceHierarchy:1;
+	ELevelInstanceFlags LevelInstanceFlags;
 
 	friend struct FSetActorIsInLevelInstance;
+	friend struct FAddActorLevelInstanceFlags;
 
 public:
 	UPROPERTY(EditAnywhere, AdvancedDisplay, Category = LevelInstance, meta = (Tooltip = "If checked, this Actor will only get loaded in a main world (persistent level), it will not be loaded through Level Instances."))
@@ -355,7 +366,13 @@ public:
 	/** If true, the actor belongs to a level instance which is currently being edited or a parent level instance being edited. */
 	bool IsInEditLevelInstanceHierarchy() const
 	{
-		return bIsInEditLevelInstanceHierarchy;
+		return EnumHasAnyFlags(LevelInstanceFlags, ELevelInstanceFlags::IsInEditHierarchy);
+	}
+
+	/** iF true, this actor or one of its components has property overrides applied. */
+	bool HasLevelInstancePropertyOverrides() const
+	{
+		return EnumHasAnyFlags(LevelInstanceFlags, ELevelInstanceFlags::HasPropertyOverrides);
 	}
 
 	/** If true, the actor belongs to a level instance. */
@@ -4587,6 +4604,19 @@ private:
 	friend class ULevelStreamingLevelInstanceEditor;
 	friend class ULevelStreamingLevelInstanceEditorPropertyOverride;
 };
+
+struct FAddActorLevelInstanceFlags
+{
+private:
+	FAddActorLevelInstanceFlags(AActor* InActor, ELevelInstanceFlags InFlagsToAdd)
+	{
+		EnumAddFlags(InActor->LevelInstanceFlags, InFlagsToAdd);
+	}
+
+	friend class ULevelStreamingLevelInstance;
+	friend class ULevelStreamingLevelInstanceEditor;
+	friend class ULevelStreamingLevelInstanceEditorPropertyOverride;
+};
 #endif
 
 /** Helper function for executing tick functions based on the normal conditions previous found in UActorComponent::ConditionalTick */
diff --git a/Engine/Source/Runtime/Engine/Private/Actor.cpp b/Engine/Source/Runtime/Engine/Private/Actor.cpp
--- a/Engine/Source/Runtime/Engine/Private/Actor.cpp
+++ b/Engine/Source/Runtime/Engine/Private/Actor.cpp
@@ -215,7 +215,7 @@ void AActor::InitializeDefaults()
 	bHasRegisteredAllComponents = false;
 
 #if WITH_EDITORONLY_DATA
-	bIsInEditLevelInstanceHierarchy = false;
+	LevelInstanceFlags = ELevelInstanceFlags::None;
 	LevelInstanceType = ELevelInstanceType::None;
 	PivotOffset = FVector::ZeroVector;
 #endif
@@ -5275,13 +5275,20 @@ void AActor::PushLevelInstanceEditingStateToProxies(bool bInEditingState)
 	TInlineComponentArray<UPrimitiveComponent*> PrimComponents;
 	GetComponents(PrimComponents);
 
-	bIsInEditLevelInstanceHierarchy = bInEditingState;
+	if (bInEditingState)
+	{
+		EnumAddFlags(LevelInstanceFlags, ELevelInstanceFlags::IsInEditHierarchy);
+	}
+	else
+	{
+		EnumRemoveFlags(LevelInstanceFlags, ELevelInstanceFlags::IsInEditHierarchy);
+	}
 
 	for (const auto& PrimComponent : PrimComponents)
 	{
 		if (PrimComponent->IsRegistered())
 		{
-			PrimComponent->PushLevelInstanceEditingStateToProxy(bIsInEditLevelInstanceHierarchy);
+			PrimComponent->PushLevelInstanceEditingStateToProxy(bInEditingState);
 		}
 	}
 
diff --git a/Engine/Source/Runtime/Engine/Private/LevelInstance/LevelInstanceEditorPropertyOverrideLevelStreaming.cpp b/Engine/Source/Runtime/Engine/Private/LevelInstance/LevelInstanceEditorPropertyOverrideLevelStreaming.cpp
--- a/Engine/Source/Runtime/Engine/Private/LevelInstance/LevelInstanceEditorPropertyOverrideLevelStreaming.cpp
+++ b/Engine/Source/Runtime/Engine/Private/LevelInstance/LevelInstanceEditorPropertyOverrideLevelStreaming.cpp
@@ -205,6 +205,9 @@ void ULevelStreamingLevelInstanceEditorPropertyOverride::ApplyPropertyOverrides(
 						Actor->GetRootComponent()->UpdateComponentToWorld();
 						Actor->MarkComponentsRenderStateDirty();
 					}
+
+					// Flag actor as being overriden
+					FAddActorLevelInstanceFlags AddFlags(Actor, ELevelInstanceFlags::HasPropertyOverrides);
 				}
 			}
 		}
diff --git a/Engine/Source/Runtime/Engine/Private/LevelInstance/LevelInstanceLevelStreaming.cpp b/Engine/Source/Runtime/Engine/Private/LevelInstance/LevelInstanceLevelStreaming.cpp
--- a/Engine/Source/Runtime/Engine/Private/LevelInstance/LevelInstanceLevelStreaming.cpp
+++ b/Engine/Source/Runtime/Engine/Private/LevelInstance/LevelInstanceLevelStreaming.cpp
@@ -210,6 +210,9 @@ void ULevelStreamingLevelInstance::ApplyPropertyOverrides(const TArray<AActor*>&
 							Actor->GetRootComponent()->UpdateComponentToWorld();
 							Actor->MarkComponentsRenderStateDirty();
 						}
+
+						// Flag actor as being overriden
+						FAddActorLevelInstanceFlags AddFlags(Actor, ELevelInstanceFlags::HasPropertyOverrides);
 					}
 				}
 			}
diff --git a/Engine/Source/Runtime/Engine/Private/LevelInstance/LevelInstanceSubsystem.cpp b/Engine/Source/Runtime/Engine/Private/LevelInstance/LevelInstanceSubsystem.cpp
--- a/Engine/Source/Runtime/Engine/Private/LevelInstance/LevelInstanceSubsystem.cpp
+++ b/Engine/Source/Runtime/Engine/Private/LevelInstance/LevelInstanceSubsystem.cpp
@@ -3144,20 +3144,7 @@ void ULevelInstanceSubsystem::RegisterPrimitiveColorHandler()
 		{
 			if (AActor* Actor = InPrimitiveComponent ? InPrimitiveComponent->GetOwner() : nullptr; Actor && Actor->IsInLevelInstance())
 			{
-				if (ULevelInstanceSubsystem* LevelInstanceSubsystem = UWorld::GetSubsystem<ULevelInstanceSubsystem>(Actor->GetWorld()))
-				{
-					if (ILevelInstanceInterface* OwningLevelInstance = LevelInstanceSubsystem->GetOwningLevelInstance(Actor->GetLevel()))
-					{
-						const FActorContainerID ContextContainerID = LevelInstanceSubsystem->GetLevelInstancePropertyOverridesContext(OwningLevelInstance);
-						TArray<const FActorPropertyOverride*> ActorPropertyOverrides;
-						if (LevelInstanceSubsystem->GetLevelInstancePropertyOverridesForActor(Actor, ContextContainerID, ActorPropertyOverrides))
-						{
-							return FLinearColor::Green;
-						}
-					}
-				}
-
-				return FLinearColor::Red;
+				return Actor->HasLevelInstancePropertyOverrides() ? FLinearColor::Green : FLinearColor::Red;
 			}
 			return FLinearColor::White;
 		});
