*Added 'deprecated' to the Material Palatte for PreSkinnedPosition + tooltip
*Added tooltip for the Local Position node
*Expanded the tooltip for the Instance(Pre-skinning) Local Origin option

#tests Open Material Editor and look at tooltip for Local Position node and the Pre-Skinned Position node (PreSkinnedPosition should now say (deprecated) in the palette too)

diff --git a/Engine/Source/Runtime/Engine/Classes/Materials/MaterialExpressionLocalPosition.h b/Engine/Source/Runtime/Engine/Classes/Materials/MaterialExpressionLocalPosition.h
--- a/Engine/Source/Runtime/Engine/Classes/Materials/MaterialExpressionLocalPosition.h
+++ b/Engine/Source/Runtime/Engine/Classes/Materials/MaterialExpressionLocalPosition.h
@@ -25,9 +25,9 @@ enum class ELocalPositionOrigin
 	/** Position relative to instance */
 	Instance UMETA(DisplayName="Instance"),
 
-	/** Position relative to instance, before skinning is applied. */
-	/** Returns the instance position for non-skeletal meshes. Incompatible with GPU skin cache feature. */
-	InstancePreSkinning UMETA(DisplayName="Instance (pre-skinning)"),
+	/** Returns pre-skinned local position for skeletal meshes, usable in vertex shader only.
+	Returns the instance position for non-skeletal meshes. Incompatible with GPU skin cache feature.*/
+	InstancePreSkinning UMETA(DisplayName="Pre-Skinned Instance"),
 
 	/** Position relative to primitive actor component */
 	Primitive UMETA(DisplayName="Component")
@@ -49,6 +49,7 @@ class UMaterialExpressionLocalPosition : public UMaterialExpression
 	bool GenerateHLSLExpression(FMaterialHLSLGenerator& Generator, UE::HLSLTree::FScope& Scope, int32 OutputIndex, UE::HLSLTree::FExpression const*& OutExpression) const;
 	virtual int32 Compile(class FMaterialCompiler* Compiler, int32 OutputIndex) override;
 	virtual void GetCaption(TArray<FString>& OutCaptions) const override;
+	virtual void GetExpressionToolTip(TArray<FString>& OutToolTip) override;
 	virtual FText GetKeywords() const override {return FText::FromString(TEXT("position preskinned local instance primitive"));}
 #endif
 	//~ End UMaterialExpression Interface
diff --git a/Engine/Source/Runtime/Engine/Classes/Materials/MaterialExpressionPreSkinnedPosition.h b/Engine/Source/Runtime/Engine/Classes/Materials/MaterialExpressionPreSkinnedPosition.h
--- a/Engine/Source/Runtime/Engine/Classes/Materials/MaterialExpressionPreSkinnedPosition.h
+++ b/Engine/Source/Runtime/Engine/Classes/Materials/MaterialExpressionPreSkinnedPosition.h
@@ -19,6 +19,8 @@ class UMaterialExpressionPreSkinnedPosition : public UMaterialExpression
 	virtual void GetExpressionToolTip(TArray<FString>& OutToolTip) override;
 	virtual FText GetKeywords() const override {return FText::FromString(TEXT("position"));}
 	virtual bool GenerateHLSLExpression(FMaterialHLSLGenerator& Generator, UE::HLSLTree::FScope& Scope, int32 OutputIndex, UE::HLSLTree::FExpression const*& OutExpression) const override;
+	virtual FText GetCreationDescription() const override;
+	virtual FText GetCreationName() const override;
 #endif
 	//~ End UMaterialExpression Interface
 };
diff --git a/Engine/Source/Runtime/Engine/Private/Materials/MaterialExpressions.cpp b/Engine/Source/Runtime/Engine/Private/Materials/MaterialExpressions.cpp
--- a/Engine/Source/Runtime/Engine/Private/Materials/MaterialExpressions.cpp
+++ b/Engine/Source/Runtime/Engine/Private/Materials/MaterialExpressions.cpp
@@ -13340,6 +13340,11 @@ void UMaterialExpressionLocalPosition::GetCaption(TArray<FString>& OutCaptions)
 		checkNoEntry();
 	}
 }
+
+void UMaterialExpressionLocalPosition::GetExpressionToolTip(TArray<FString>& OutToolTip) 
+{
+	ConvertToMultilineToolTip(TEXT("Gets the local position of the mesh, based on the selected Local Origin"), 40, OutToolTip);
+}
 #endif // WITH_EDITOR
 
 ///////////////////////////////////////////////////////////////////////////////
@@ -22849,6 +22854,17 @@ void UMaterialExpressionPreSkinnedPosition::GetExpressionToolTip(TArray<FString>
 	                               "Returns pre-skinned local position for skeletal meshes, usable in vertex shader only."
 	                               "Returns the local position for non-skeletal meshes. Incompatible with GPU skin cache feature."), 40, OutToolTip);
 }
+
+FText UMaterialExpressionPreSkinnedPosition::GetCreationDescription() const
+{
+	return LOCTEXT("PreSkinnedPositionCreationDesc", "Deprecated, has been merged into 'Local Position' node.");
+}
+
+FText UMaterialExpressionPreSkinnedPosition::GetCreationName() const
+{
+	return LOCTEXT("PreSkinnedPositionCreationName", "PreSkinnedPosition (Deprecated)");
+}
+
 #endif // WITH_EDITOR
 
 ///////////////////////////////////////////////////////////////////////////////
