Adding Texture Collection
- UTextureCollection data asset that just holds an array of Texture assets.
- Uses the new RHIResourceCollection RHI resource type which puts all the texture bindless indices into its own buffer.
- Adding TextureCollection material nodes for accessing and sampling from Texture Collections.

diff --git a/Engine/Source/Developer/MaterialBaking/Private/ExportMaterialProxy.h b/Engine/Source/Developer/MaterialBaking/Private/ExportMaterialProxy.h
--- a/Engine/Source/Developer/MaterialBaking/Private/ExportMaterialProxy.h
+++ b/Engine/Source/Developer/MaterialBaking/Private/ExportMaterialProxy.h
@@ -224,6 +224,7 @@ public:
 		SetQualityLevelProperties(GMaxRHIFeatureLevel);
 		Material = InMaterialInterface->GetMaterial();
 		ReferencedTextures = InMaterialInterface->GetReferencedTextures();
+		ReferencedTextureCollections = InMaterialInterface->GetReferencedTextureCollections();
 
 		const FMaterialResource* Resource = InMaterialInterface->GetMaterialResource(GMaxRHIFeatureLevel);
 
@@ -301,6 +302,11 @@ public:
 		return ReferencedTextures;
 	}
 
+	virtual TConstArrayView<TObjectPtr<UTextureCollection>> GetReferencedTextureCollections() const override
+	{
+		return ReferencedTextureCollections;
+	}
+
 	virtual void GetStaticParameterSet(EShaderPlatform Platform, FStaticParameterSet& OutSet) const override
 	{
 		if (const FMaterialResource* Resource = MaterialInterface->GetMaterialResource(GMaxRHIFeatureLevel))
@@ -633,6 +639,7 @@ private:
 	UMaterialInterface* MaterialInterface;
 	UMaterial* Material;
 	TArray<TObjectPtr<UObject>> ReferencedTextures;
+	TArray<TObjectPtr<UTextureCollection>> ReferencedTextureCollections;
 	/** The property to compile for rendering the sample */
 	EMaterialProperty PropertyToCompile;
 	/** Stores which exported attribute this proxy is compiling for. */
diff --git a/Engine/Source/Editor/MaterialEditor/Private/MaterialEditor.cpp b/Engine/Source/Editor/MaterialEditor/Private/MaterialEditor.cpp
--- a/Engine/Source/Editor/MaterialEditor/Private/MaterialEditor.cpp
+++ b/Engine/Source/Editor/MaterialEditor/Private/MaterialEditor.cpp
@@ -42,6 +42,7 @@
 #include "Engine/TextureCube.h"
 #include "Engine/Texture2DArray.h"
 #include "Engine/TextureCubeArray.h"
+#include "Engine/TextureCollection.h"
 #include "SparseVolumeTexture/SparseVolumeTexture.h"
 #include "Dialogs/Dialogs.h"
 #include "UnrealEdGlobals.h"
@@ -236,6 +237,7 @@ FMatExpressionPreview::FMatExpressionPreview(UMaterialExpression* InExpression)
 	else
 	{
 		ReferencedTextures = InExpression->Material->GetReferencedTextures();
+		ReferencedTextureCollections = InExpression->Material->GetReferencedTextureCollections();
 	}
 }
 
@@ -246,6 +248,7 @@ FMatExpressionPreview::~FMatExpressionPreview()
 void FMatExpressionPreview::AddReferencedObjects(FReferenceCollector& Collector)
 {
 	Collector.AddReferencedObjects(ReferencedTextures);
+	Collector.AddReferencedObjects(ReferencedTextureCollections);
 	if (CachedExpressionData)
 	{
 		CachedExpressionData->AddReferencedObjects(Collector);
@@ -388,6 +391,18 @@ TArrayView<const TObjectPtr<UObject>> FMatExpressionPreview::GetReferencedTextur
 	return MakeArrayView(ReferencedTextures);
 }
 
+TConstArrayView<TObjectPtr<UTextureCollection>> FMatExpressionPreview::GetReferencedTextureCollections() const
+{
+	if (CachedExpressionData)
+	{
+		// Path for new HLSL translator
+		return MakeArrayView(CachedExpressionData->ReferencedTextureCollections);
+	}
+
+	// Legacy path
+	return MakeArrayView(ReferencedTextureCollections);
+}
+
 const FMaterialCachedHLSLTree* FMatExpressionPreview::GetCachedHLSLTree() const
 {
 	return CachedHLSLTree.Get();
diff --git a/Engine/Source/Editor/MaterialEditor/Private/MaterialEditor.h b/Engine/Source/Editor/MaterialEditor/Private/MaterialEditor.h
--- a/Engine/Source/Editor/MaterialEditor/Private/MaterialEditor.h
+++ b/Engine/Source/Editor/MaterialEditor/Private/MaterialEditor.h
@@ -88,6 +88,7 @@ public:
 	virtual bool ShouldCache(EShaderPlatform Platform, const FShaderType* ShaderType, const FVertexFactoryType* VertexFactoryType) const override;
 
 	virtual TArrayView<const TObjectPtr<UObject>> GetReferencedTextures() const override;
+	virtual TConstArrayView<TObjectPtr<UTextureCollection>> GetReferencedTextureCollections() const override;
 
 	////////////////
 	// FMaterialRenderProxy interface.
@@ -184,6 +185,7 @@ private:
 	TUniquePtr<FMaterialCachedHLSLTree> CachedHLSLTree;
 	TWeakObjectPtr<UMaterialExpression> Expression;
 	TArray<TObjectPtr<UObject>> ReferencedTextures;
+	TArray<TObjectPtr<UTextureCollection>> ReferencedTextureCollections;
 	FGuid Id;
 };
 
diff --git a/Engine/Source/Editor/MaterialEditor/Private/MaterialEditorUtilities.cpp b/Engine/Source/Editor/MaterialEditor/Private/MaterialEditorUtilities.cpp
--- a/Engine/Source/Editor/MaterialEditor/Private/MaterialEditorUtilities.cpp
+++ b/Engine/Source/Editor/MaterialEditor/Private/MaterialEditorUtilities.cpp
@@ -16,6 +16,7 @@
 #include "Materials/MaterialExpressionStaticSwitch.h"
 #include "Materials/MaterialExpressionComment.h"
 #include "Materials/MaterialExpressionComposite.h"
+#include "Materials/MaterialExpressionTextureCollectionParameter.h"
 #include "Materials/MaterialExpressionTextureSample.h"
 #include "Materials/MaterialExpressionTextureSampleParameter.h"
 #include "Materials/MaterialExpressionRuntimeVirtualTextureSampleParameter.h"
@@ -471,6 +472,7 @@ void FMaterialEditorUtilities::GetVisibleMaterialParametersFromExpression(
 	// If it's a material parameter it must be visible so add it to the list
 	UMaterialExpressionParameter* Param = Cast<UMaterialExpressionParameter>( MaterialExpressionKey.Expression );
 	UMaterialExpressionTextureSampleParameter* TexParam = Cast<UMaterialExpressionTextureSampleParameter>( MaterialExpressionKey.Expression );
+	UMaterialExpressionTextureCollectionParameter* TextureCollectionParam = Cast<UMaterialExpressionTextureCollectionParameter>( MaterialExpressionKey.Expression );
 	UMaterialExpressionRuntimeVirtualTextureSampleParameter* RuntimeVirtualTexParam = Cast<UMaterialExpressionRuntimeVirtualTextureSampleParameter>(MaterialExpressionKey.Expression);
 	UMaterialExpressionSparseVolumeTextureSampleParameter* SparseVolumeTexParam = Cast<UMaterialExpressionSparseVolumeTextureSampleParameter>(MaterialExpressionKey.Expression);
 	UMaterialExpressionFontSampleParameter* FontParam = Cast<UMaterialExpressionFontSampleParameter>( MaterialExpressionKey.Expression );
@@ -483,6 +485,10 @@ void FMaterialEditorUtilities::GetVisibleMaterialParametersFromExpression(
 	{
 		ParameterInfo.Name = TexParam->ParameterName;
 	}
+	else if (TextureCollectionParam)
+	{
+		ParameterInfo.Name = TextureCollectionParam->ParameterName;
+	}
 	else if (RuntimeVirtualTexParam)
 	{
 		ParameterInfo.Name = RuntimeVirtualTexParam->ParameterName;
@@ -496,7 +502,7 @@ void FMaterialEditorUtilities::GetVisibleMaterialParametersFromExpression(
 		ParameterInfo.Name = FontParam->ParameterName;
 	}
 		
-	if (Param || TexParam || FontParam || RuntimeVirtualTexParam || SparseVolumeTexParam)
+	if (Param || TexParam || TextureCollectionParam || FontParam || RuntimeVirtualTexParam || SparseVolumeTexParam)
 	{
 		VisibleExpressions.AddUnique(ParameterInfo);
 	}
diff --git a/Engine/Source/Editor/MaterialEditor/Private/MaterialInstanceEditor.cpp b/Engine/Source/Editor/MaterialEditor/Private/MaterialInstanceEditor.cpp
--- a/Engine/Source/Editor/MaterialEditor/Private/MaterialInstanceEditor.cpp
+++ b/Engine/Source/Editor/MaterialEditor/Private/MaterialInstanceEditor.cpp
@@ -530,6 +530,7 @@ void FMaterialInstanceEditor::ReInitMaterialFunctionProxies()
 		TArray<FVectorParameterValue> VectorParameterValues = FunctionInstanceProxy->VectorParameterValues;
 		TArray<FDoubleVectorParameterValue> DoubleVectorParameterValues = FunctionInstanceProxy->DoubleVectorParameterValues;
 		TArray<FTextureParameterValue> TextureParameterValues = FunctionInstanceProxy->TextureParameterValues;
+		TArray<FTextureCollectionParameterValue> TextureCollectionParameterValues = FunctionInstanceProxy->TextureCollectionParameterValues;
 		TArray<FRuntimeVirtualTextureParameterValue> RuntimeVirtualTextureParameterValues = FunctionInstanceProxy->RuntimeVirtualTextureParameterValues;
 		TArray<FSparseVolumeTextureParameterValue> SparseVolumeTextureParameterValues = FunctionInstanceProxy->SparseVolumeTextureParameterValues;
 		TArray<FFontParameterValue> FontParameterValues = FunctionInstanceProxy->FontParameterValues;
@@ -595,6 +596,18 @@ void FMaterialInstanceEditor::ReInitMaterialFunctionProxies()
 			}
 		}
 
+		FunctionInstanceProxy->GetAllTextureCollectionParameterInfo(OutParameterInfo, Guids);
+		FunctionInstanceProxy->TextureCollectionParameterValues.Empty();
+		for (FTextureCollectionParameterValue& TextureCollectionParameter : TextureCollectionParameterValues)
+		{
+			int32 Index = Guids.Find(TextureCollectionParameter.ExpressionGUID);
+			if (Index != INDEX_NONE)
+			{
+				FunctionInstanceProxy->TextureCollectionParameterValues.Add(TextureCollectionParameter);
+				FunctionInstanceProxy->TextureCollectionParameterValues.Last().ParameterInfo = OutParameterInfo[Index];
+			}
+		}
+
 		FunctionInstanceProxy->GetAllRuntimeVirtualTextureParameterInfo(OutParameterInfo, Guids);
 		FunctionInstanceProxy->RuntimeVirtualTextureParameterValues.Empty();
 		for (FRuntimeVirtualTextureParameterValue& RuntimeVirtualTextureParameter : RuntimeVirtualTextureParameterValues)
diff --git a/Engine/Source/Editor/MaterialEditor/Private/MaterialPropertyHelpers.cpp b/Engine/Source/Editor/MaterialEditor/Private/MaterialPropertyHelpers.cpp
--- a/Engine/Source/Editor/MaterialEditor/Private/MaterialPropertyHelpers.cpp
+++ b/Engine/Source/Editor/MaterialEditor/Private/MaterialPropertyHelpers.cpp
@@ -425,6 +425,7 @@ FReply FMaterialPropertyHelpers::OnClickedSaveNewFunctionInstance(class UMateria
 					ChildInstance->VectorParameterValues = EditedInstance->VectorParameterValues;
 					ChildInstance->DoubleVectorParameterValues = EditedInstance->DoubleVectorParameterValues;
 					ChildInstance->TextureParameterValues = EditedInstance->TextureParameterValues;
+					ChildInstance->TextureCollectionParameterValues = EditedInstance->TextureCollectionParameterValues;
 					ChildInstance->RuntimeVirtualTextureParameterValues = EditedInstance->RuntimeVirtualTextureParameterValues;
 					ChildInstance->SparseVolumeTextureParameterValues = EditedInstance->SparseVolumeTextureParameterValues;
 					ChildInstance->FontParameterValues = EditedInstance->FontParameterValues;
@@ -518,6 +519,7 @@ FReply FMaterialPropertyHelpers::OnClickedSaveNewLayerInstance(class UMaterialFu
 					ChildInstance->VectorParameterValues = EditedInstance->VectorParameterValues;
 					ChildInstance->DoubleVectorParameterValues = EditedInstance->DoubleVectorParameterValues;
 					ChildInstance->TextureParameterValues = EditedInstance->TextureParameterValues;
+					ChildInstance->TextureCollectionParameterValues = EditedInstance->TextureCollectionParameterValues;
 					ChildInstance->RuntimeVirtualTextureParameterValues = EditedInstance->RuntimeVirtualTextureParameterValues;
 					ChildInstance->SparseVolumeTextureParameterValues = EditedInstance->SparseVolumeTextureParameterValues;
 					ChildInstance->FontParameterValues = EditedInstance->FontParameterValues;
diff --git a/Engine/Source/Editor/UnrealEd/Classes/MaterialEditor/DEditorTextureCollectionParameterValue.h b/Engine/Source/Editor/UnrealEd/Classes/MaterialEditor/DEditorTextureCollectionParameterValue.h
new file mode 100644
--- /dev/null
+++ b/Engine/Source/Editor/UnrealEd/Classes/MaterialEditor/DEditorTextureCollectionParameterValue.h
@@ -0,0 +1,37 @@
+// Copyright Epic Games, Inc. All Rights Reserved.
+
+#pragma once
+
+#include "MaterialEditor/DEditorParameterValue.h"
+#include "DEditorTextureCollectionParameterValue.generated.h"
+
+class UTextureCollection;
+
+UCLASS(hidecategories=Object, collapsecategories, MinimalAPI)
+class UDEditorTextureCollectionParameterValue : public UDEditorParameterValue
+{
+	GENERATED_UCLASS_BODY()
+
+	UPROPERTY(EditAnywhere, Category=DEditorTextureCollectionParameterValue)
+	TObjectPtr<UTextureCollection> ParameterValue;
+
+	virtual FName GetDefaultGroupName() const override { return TEXT("Texture Parameter Values"); }
+
+	virtual bool GetValue(FMaterialParameterMetadata& OutResult) const override
+	{
+		UDEditorParameterValue::GetValue(OutResult);
+		OutResult.Value = ParameterValue;
+		return true;
+	}
+
+	virtual bool SetValue(const FMaterialParameterValue& Value) override
+	{
+		if (Value.Type == EMaterialParameterType::TextureCollection)
+		{
+			ParameterValue = Value.TextureCollection;
+			return true;
+		}
+		return false;
+	}
+};
+
diff --git a/Engine/Source/Editor/UnrealEd/Private/EditorParameters.cpp b/Engine/Source/Editor/UnrealEd/Private/EditorParameters.cpp
--- a/Engine/Source/Editor/UnrealEd/Private/EditorParameters.cpp
+++ b/Engine/Source/Editor/UnrealEd/Private/EditorParameters.cpp
@@ -11,6 +11,7 @@
 #include "MaterialEditor/DEditorStaticComponentMaskParameterValue.h"
 #include "MaterialEditor/DEditorStaticSwitchParameterValue.h"
 #include "MaterialEditor/DEditorTextureParameterValue.h"
+#include "MaterialEditor/DEditorTextureCollectionParameterValue.h"
 #include "MaterialEditor/DEditorVectorParameterValue.h"
 #include "MaterialTypes.h"
 #include "Math/Color.h"
@@ -57,6 +58,11 @@ UDEditorTextureParameterValue::UDEditorTextureParameterValue(const FObjectInitia
 {
 }
 
+UDEditorTextureCollectionParameterValue::UDEditorTextureCollectionParameterValue(const FObjectInitializer& ObjectInitializer)
+	: Super(ObjectInitializer)
+{
+}
+
 UDEditorVectorParameterValue::UDEditorVectorParameterValue(const FObjectInitializer& ObjectInitializer)
 	: Super(ObjectInitializer)
 {
@@ -125,6 +131,16 @@ static UDEditorParameterValue* CreateParameter_Texture(UObject* Owner, const FMa
 	return Parameter;
 }
 
+static UDEditorParameterValue* CreateParameter_TextureCollection(UObject* Owner, const FMaterialParameterMetadata& Meta)
+{
+	UDEditorTextureCollectionParameterValue* Parameter = NewObject<UDEditorTextureCollectionParameterValue>(Owner);
+	if (Meta.Value.Type == EMaterialParameterType::TextureCollection)
+	{
+		Parameter->ParameterValue = Meta.Value.TextureCollection;
+	}
+	return Parameter;
+}
+
 static UDEditorParameterValue* CreateParameter_RuntimeVirtualTexture(UObject* Owner, const FMaterialParameterMetadata& Meta)
 {
 	UDEditorRuntimeVirtualTextureParameterValue* Parameter = NewObject<UDEditorRuntimeVirtualTextureParameterValue>(Owner);
@@ -191,6 +207,7 @@ UDEditorParameterValue* UDEditorParameterValue::Create(UObject* Owner,
 	case EMaterialParameterType::Vector: Parameter = CreateParameter_Vector(Owner, Meta); break;
 	case EMaterialParameterType::DoubleVector: Parameter = CreateParameter_DoubleVector(Owner, Meta); break;
 	case EMaterialParameterType::Texture: Parameter = CreateParameter_Texture(Owner, Meta); break;
+	case EMaterialParameterType::TextureCollection: Parameter = CreateParameter_TextureCollection(Owner, Meta); break;
 	case EMaterialParameterType::RuntimeVirtualTexture: Parameter = CreateParameter_RuntimeVirtualTexture(Owner, Meta); break;
 	case EMaterialParameterType::SparseVolumeTexture: Parameter = CreateParameter_SparseVolumeTexture(Owner, Meta); break;
 	case EMaterialParameterType::Font: Parameter = CreateParameter_Font(Owner, Meta); break;
diff --git a/Engine/Source/Editor/UnrealEd/Private/Lightmass/LightmassRender.cpp b/Engine/Source/Editor/UnrealEd/Private/Lightmass/LightmassRender.cpp
--- a/Engine/Source/Editor/UnrealEd/Private/Lightmass/LightmassRender.cpp
+++ b/Engine/Source/Editor/UnrealEd/Private/Lightmass/LightmassRender.cpp
@@ -267,6 +267,7 @@ public:
 			Usage = InUsage;
 
 			ReferencedTextures = MaterialInterface->GetReferencedTextures();
+			ReferencedTextureCollections = MaterialInterface->GetReferencedTextureCollections();
 
 			FMaterialResource* Resource = InMaterialInterface->GetMaterialResource(GMaxRHIFeatureLevel);
 			if (Resource)
@@ -297,6 +298,11 @@ public:
 		return ReferencedTextures;
 	}
 
+	virtual TConstArrayView<TObjectPtr<UTextureCollection>> GetReferencedTextureCollections() const override
+	{
+		return ReferencedTextureCollections;
+	}
+
 	/**
 	 * Should the shader for this material with the given platform, shader type and vertex 
 	 * factory type combination be compiled
@@ -999,6 +1005,7 @@ private:
 	UMaterialInterface* MaterialInterface;
 	UMaterial* Material;
 	TArray<TObjectPtr<UObject>> ReferencedTextures;
+	TArray<TObjectPtr<UTextureCollection>> ReferencedTextureCollections;
 	/** The property to compile for rendering the sample */
 	EMaterialProperty PropertyToCompile;
 	/** Stores which exported attribute this proxy is compiling for. */
diff --git a/Engine/Source/Editor/UnrealEd/Private/PreviewMaterial.cpp b/Engine/Source/Editor/UnrealEd/Private/PreviewMaterial.cpp
--- a/Engine/Source/Editor/UnrealEd/Private/PreviewMaterial.cpp
+++ b/Engine/Source/Editor/UnrealEd/Private/PreviewMaterial.cpp
@@ -882,6 +882,7 @@ void UMaterialEditorInstanceConstant::ApplySourceFunctionChanges()
 		SourceFunction->VectorParameterValues = SourceInstance->VectorParameterValues;
 		SourceFunction->DoubleVectorParameterValues = SourceInstance->DoubleVectorParameterValues;
 		SourceFunction->TextureParameterValues = SourceInstance->TextureParameterValues;
+		SourceFunction->TextureCollectionParameterValues = SourceInstance->TextureCollectionParameterValues;
 		SourceFunction->RuntimeVirtualTextureParameterValues = SourceInstance->RuntimeVirtualTextureParameterValues;
 		SourceFunction->SparseVolumeTextureParameterValues = SourceInstance->SparseVolumeTextureParameterValues;
 		SourceFunction->FontParameterValues = SourceInstance->FontParameterValues;
diff --git a/Engine/Source/Runtime/Engine/Classes/Engine/TextureCollection.h b/Engine/Source/Runtime/Engine/Classes/Engine/TextureCollection.h
new file mode 100644
--- /dev/null
+++ b/Engine/Source/Runtime/Engine/Classes/Engine/TextureCollection.h
@@ -0,0 +1,61 @@
+// Copyright Epic Games, Inc. All Rights Reserved.
+
+#pragma once
+
+#include "Engine/DataAsset.h"
+#include "RenderResource.h"
+#include "TextureCollection.generated.h"
+
+class FTextureResource;
+class UTexture;
+class UTextureCollection;
+
+struct FTextureCollectionResource : public FRenderResource
+{
+	FTextureCollectionResource(UTextureCollection* InParent);
+
+	// FRenderResource Interface
+	virtual void InitRHI(FRHICommandListBase& RHICmdList) final;
+	virtual void ReleaseRHI() final;
+	// ~FRenderResource
+
+	FRHIResourceCollection* GetRHI() const { return ResourceCollectionRHI.GetReference(); }
+
+private:
+	TArray<FTextureResource*>       InputTextureResources;
+	TArray<FTextureReferenceRHIRef> InputTextures;
+	FRHIResourceCollectionRef       ResourceCollectionRHI;
+};
+
+UCLASS(MinimalAPI)
+class UTextureCollection : public UDataAsset
+{
+	GENERATED_BODY()
+
+public:
+	UPROPERTY(EditAnywhere, Category=TextureCollection)
+	TArray<TObjectPtr<UTexture>> Textures;
+
+	//~ Begin UObject Interface.
+#if WITH_EDITOR
+	virtual void PostEditChangeProperty(FPropertyChangedEvent& PropertyChangedEvent) override;
+#endif // WITH_EDITOR
+	virtual void PostLoad() override;
+	//~ End UObject Interface.
+
+	ENGINE_API void SetResource(FTextureCollectionResource* InResource);
+	ENGINE_API FTextureCollectionResource* GetResource() const;
+	ENGINE_API FTextureCollectionResource* GetResource();
+
+	ENGINE_API FTextureCollectionResource* CreateResource();
+	ENGINE_API void ReleaseResource();
+	ENGINE_API void UpdateResource();
+
+protected:
+#if WITH_EDITOR
+	void NotifyMaterials();
+#endif
+
+	FTextureCollectionResource* PrivateResource = nullptr;
+	FTextureCollectionResource* PrivateResourceRenderThread = nullptr;
+};
diff --git a/Engine/Source/Runtime/Engine/Classes/Materials/MaterialExpression.h b/Engine/Source/Runtime/Engine/Classes/Materials/MaterialExpression.h
--- a/Engine/Source/Runtime/Engine/Classes/Materials/MaterialExpression.h
+++ b/Engine/Source/Runtime/Engine/Classes/Materials/MaterialExpression.h
@@ -17,6 +17,7 @@
 class UEdGraphNode;
 class UMaterial;
 class UTexture;
+class UTextureCollection;
 struct FPropertyChangedEvent;
 struct FMaterialParameterMetadata;
 struct FMaterialShadingModelField;
@@ -344,6 +345,8 @@ class UMaterialExpression : public UObject
 	/** Returns true if GetReferencedTexture() / GetReferencedTextures() can ever return a valid pointer(s). */
 	virtual bool CanReferenceTexture() const { return false; }
 
+	virtual UTextureCollection* GetReferencedTextureCollection() const { return nullptr; }
+
 #if WITH_EDITOR
 	/**
 	 *	Get the outputs supported by this expression.
diff --git a/Engine/Source/Runtime/Engine/Classes/Materials/MaterialExpressionTextureBase.h b/Engine/Source/Runtime/Engine/Classes/Materials/MaterialExpressionTextureBase.h
--- a/Engine/Source/Runtime/Engine/Classes/Materials/MaterialExpressionTextureBase.h
+++ b/Engine/Source/Runtime/Engine/Classes/Materials/MaterialExpressionTextureBase.h
@@ -67,6 +67,13 @@ class UMaterialExpressionTextureBase : public UMaterialExpression
 	 */
 	static ENGINE_API EMaterialSamplerType GetSamplerTypeForTexture( const UTexture* Texture, bool ForceNoVT = false );
 
+	static ENGINE_API bool VerifySamplerType(
+		const FString& TexturePathName,
+		EMaterialSamplerType CorrectSamplerType,
+		bool bSRGB,
+		EMaterialSamplerType SamplerType,
+		FString& OutErrorMessage);
+
 	/**
 	 * Verify that the texture and sampler type. Generates a compiler waring if
 	 * they do not.
diff --git a/Engine/Source/Runtime/Engine/Classes/Materials/MaterialExpressionTextureCollection.h b/Engine/Source/Runtime/Engine/Classes/Materials/MaterialExpressionTextureCollection.h
new file mode 100644
--- /dev/null
+++ b/Engine/Source/Runtime/Engine/Classes/Materials/MaterialExpressionTextureCollection.h
@@ -0,0 +1,27 @@
+// Copyright Epic Games, Inc. All Rights Reserved.
+
+#pragma once
+
+#include "Engine/TextureCollection.h"
+#include "Materials/MaterialExpression.h"
+#include "MaterialExpressionTextureCollection.generated.h"
+
+UCLASS(MinimalAPI)
+class UMaterialExpressionTextureCollection : public UMaterialExpression
+{
+	GENERATED_UCLASS_BODY()
+
+	UPROPERTY(EditAnywhere, BlueprintReadWrite, Category=TextureCollection)
+	TObjectPtr<UTextureCollection> TextureCollection;
+
+	//~ Begin UMaterialExpression Interface
+#if WITH_EDITOR
+	virtual int32 Compile(class FMaterialCompiler* Compiler, int32 OutputIndex) override;
+	virtual void GetCaption(TArray<FString>& OutCaptions) const override;
+	virtual uint32 GetOutputType(int32 OutputIndex) override;
+	virtual bool GenerateHLSLExpression(FMaterialHLSLGenerator& Generator, UE::HLSLTree::FScope& Scope, int32 OutputIndex, UE::HLSLTree::FExpression const*& OutExpression) const override;
+
+	virtual UTextureCollection* GetReferencedTextureCollection() const { return TextureCollection; }
+#endif
+	//~ End UMaterialExpression Interface
+};
diff --git a/Engine/Source/Runtime/Engine/Classes/Materials/MaterialExpressionTextureCollectionParameter.h b/Engine/Source/Runtime/Engine/Classes/Materials/MaterialExpressionTextureCollectionParameter.h
new file mode 100644
--- /dev/null
+++ b/Engine/Source/Runtime/Engine/Classes/Materials/MaterialExpressionTextureCollectionParameter.h
@@ -0,0 +1,54 @@
+// Copyright Epic Games, Inc. All Rights Reserved.
+
+#pragma once
+
+#include "Materials/MaterialExpressionTextureCollection.h"
+#include "MaterialExpressionTextureCollectionParameter.generated.h"
+
+UCLASS(MinimalAPI)
+class UMaterialExpressionTextureCollectionParameter : public UMaterialExpressionTextureCollection
+{
+	GENERATED_UCLASS_BODY()
+
+	/** The name of the parameter */
+	UPROPERTY(EditAnywhere, Category= MaterialExpressionTextureCollectionParameter)
+	FName ParameterName;
+
+	/** GUID that should be unique within the material, this is used for parameter renaming. */
+	UPROPERTY()
+	FGuid ExpressionGUID;
+
+	/** The name of the parameter Group to display in MaterialInstance Editor. Default is None group */
+	UPROPERTY(EditAnywhere, Category= MaterialExpressionTextureCollectionParameter)
+	FName Group;
+
+	/** Controls where the this parameter is displayed in a material instance parameter list.  The lower the number the higher up in the parameter list. */
+	UPROPERTY(EditAnywhere, Category= MaterialExpressionTextureCollectionParameter)
+	int32 SortPriority = 32;
+
+	//~ Begin UMaterialExpression Interface
+#if WITH_EDITOR
+	virtual void  GetCaption(TArray<FString>& OutCaptions) const override;
+	virtual int32 Compile(class FMaterialCompiler* Compiler, int32 OutputIndex) override;
+	virtual bool  GenerateHLSLExpression(FMaterialHLSLGenerator& Generator, UE::HLSLTree::FScope& Scope, int32 OutputIndex, UE::HLSLTree::FExpression const*& OutExpression) const override;
+
+	virtual bool    CanRenameNode() const override;
+	virtual FString GetEditableName() const override;
+	virtual void    SetEditableName(const FString& NewName) override;
+
+	virtual bool  HasAParameterName() const override;
+	virtual FName GetParameterName() const override;
+	virtual void  SetParameterName(const FName& Name) override;
+	virtual void  ValidateParameterName(const bool bAllowDuplicateName) override;
+
+	virtual bool GetParameterValue(FMaterialParameterMetadata& OutMeta) const override;
+	virtual bool SetParameterValue(const FName& Name, const FMaterialParameterMetadata& Meta, EMaterialExpressionSetParameterValueFlags Flags) override;
+#endif
+
+	virtual FGuid& GetParameterExpressionId() override;
+	//~ End UMaterialExpression Interface
+
+#if WITH_EDITOR
+	bool SetParameterValue(const FName& InParameterName, UTextureCollection* InValue, EMaterialExpressionSetParameterValueFlags Flags = EMaterialExpressionSetParameterValueFlags::None);
+#endif
+};
diff --git a/Engine/Source/Runtime/Engine/Classes/Materials/MaterialExpressionTextureObjectFromCollection.h b/Engine/Source/Runtime/Engine/Classes/Materials/MaterialExpressionTextureObjectFromCollection.h
new file mode 100644
--- /dev/null
+++ b/Engine/Source/Runtime/Engine/Classes/Materials/MaterialExpressionTextureObjectFromCollection.h
@@ -0,0 +1,67 @@
+// Copyright Epic Games, Inc. All Rights Reserved.
+
+#pragma once
+
+#include "Engine/TextureCollection.h"
+#include "Materials/MaterialExpression.h"
+#include "MaterialValueType.h"
+#include "MaterialExpressionTextureObjectFromCollection.generated.h"
+
+/** Describes how textures are sampled for materials */
+UENUM(BlueprintType)
+enum class ETextureCollectionMemberType : uint8
+{
+	Texture2D        UMETA(DisplayName = "Texture 2D"),
+	TextureCube      UMETA(DisplayName = "Texture Cube"),
+	Texture2DArray   UMETA(DisplayName = "Texture 2D Array"),
+	TextureCubeArray UMETA(DisplayName = "Texture Cube Array"),
+	TextureVolume    UMETA(DisplayName = "Volume Texture"),
+
+	Max
+};
+
+inline EMaterialValueType MaterialValueTypeFromTextureCollectionMemberType(ETextureCollectionMemberType InType)
+{
+	switch (InType)
+	{
+	case ETextureCollectionMemberType::Texture2D:        return MCT_Texture2D;
+	case ETextureCollectionMemberType::TextureCube:      return MCT_TextureCube;
+	case ETextureCollectionMemberType::Texture2DArray:   return MCT_Texture2DArray;
+	case ETextureCollectionMemberType::TextureCubeArray: return MCT_TextureCubeArray;
+	case ETextureCollectionMemberType::TextureVolume:    return MCT_VolumeTexture;
+	}
+	return MCT_Texture2D;
+}
+
+UCLASS(MinimalAPI)
+class UMaterialExpressionTextureObjectFromCollection : public UMaterialExpression
+{
+	GENERATED_UCLASS_BODY()
+
+	UPROPERTY(meta = (RequiredInput = "false"))
+	FExpressionInput TextureCollection;
+
+	UPROPERTY(EditAnywhere, BlueprintReadWrite, Category=TextureCollection)
+	TObjectPtr<UTextureCollection> TextureCollectionObject;
+
+	UPROPERTY(meta = (RequiredInput = "false"))
+	FExpressionInput CollectionIndex;
+
+	UPROPERTY(EditAnywhere, BlueprintReadWrite, Category=TextureCollection, meta = (OverridingInputProperty = "CollectionIndex"))
+	int32 ConstCollectionIndex = 0;
+
+	UPROPERTY(EditAnywhere, BlueprintReadWrite, Category=TextureCollection)
+	ETextureCollectionMemberType TextureType = ETextureCollectionMemberType::Texture2D;
+
+	//~ Begin UMaterialExpression Interface
+#if WITH_EDITOR
+	virtual int32 Compile(class FMaterialCompiler* Compiler, int32 OutputIndex) override;
+	virtual uint32 GetInputType(int32 InputIndex) override;
+	virtual uint32 GetOutputType(int32 OutputIndex) override;
+	virtual void GetCaption(TArray<FString>& OutCaptions) const override;
+	virtual bool GenerateHLSLExpression(FMaterialHLSLGenerator& Generator, UE::HLSLTree::FScope& Scope, int32 OutputIndex, UE::HLSLTree::FExpression const*& OutExpression) const override;
+
+	virtual UTextureCollection* GetReferencedTextureCollection() const { return TextureCollectionObject; }
+#endif
+	//~ End UMaterialExpression Interface
+};
diff --git a/Engine/Source/Runtime/Engine/Classes/Materials/MaterialFunctionInstance.h b/Engine/Source/Runtime/Engine/Classes/Materials/MaterialFunctionInstance.h
--- a/Engine/Source/Runtime/Engine/Classes/Materials/MaterialFunctionInstance.h
+++ b/Engine/Source/Runtime/Engine/Classes/Materials/MaterialFunctionInstance.h
@@ -45,6 +45,10 @@ class UMaterialFunctionInstance : public UMaterialFunctionInterface
 	UPROPERTY(EditAnywhere, Category=MaterialFunctionInstance)
 	TArray<struct FTextureParameterValue> TextureParameterValues;
 
+	/** Texture Collection parameters. */
+	UPROPERTY(EditAnywhere, Category=MaterialFunctionInstance)
+	TArray<struct FTextureCollectionParameterValue> TextureCollectionParameterValues;
+
 	/** Font parameters. */
 	UPROPERTY(EditAnywhere, Category=MaterialFunctionInstance)
 	TArray<struct FFontParameterValue> FontParameterValues;
diff --git a/Engine/Source/Runtime/Engine/Classes/Materials/MaterialFunctionInterface.h b/Engine/Source/Runtime/Engine/Classes/Materials/MaterialFunctionInterface.h
--- a/Engine/Source/Runtime/Engine/Classes/Materials/MaterialFunctionInterface.h
+++ b/Engine/Source/Runtime/Engine/Classes/Materials/MaterialFunctionInterface.h
@@ -341,6 +341,7 @@ public:
 	ENGINE_API bool OverrideNamedScalarParameter(const FHashedMaterialParameterInfo& ParameterInfo, float& OutValue);
 	ENGINE_API bool OverrideNamedVectorParameter(const FHashedMaterialParameterInfo& ParameterInfo, FLinearColor& OutValue);
 	ENGINE_API bool OverrideNamedTextureParameter(const FHashedMaterialParameterInfo& ParameterInfo, class UTexture*& OutValue);
+	ENGINE_API bool OverrideNamedTextureCollectionParameter(const FHashedMaterialParameterInfo& ParameterInfo, class UTextureCollection*& OutValue);
 	ENGINE_API bool OverrideNamedRuntimeVirtualTextureParameter(const FHashedMaterialParameterInfo& ParameterInfo, class URuntimeVirtualTexture*& OutValue);
 	ENGINE_API bool OverrideNamedSparseVolumeTextureParameter(const FHashedMaterialParameterInfo& ParameterInfo, class USparseVolumeTexture*& OutValue);
 	ENGINE_API bool OverrideNamedFontParameter(const FHashedMaterialParameterInfo& ParameterInfo, class UFont*& OutFontValue, int32& OutFontPage);
diff --git a/Engine/Source/Runtime/Engine/Classes/Materials/MaterialInstance.h b/Engine/Source/Runtime/Engine/Classes/Materials/MaterialInstance.h
--- a/Engine/Source/Runtime/Engine/Classes/Materials/MaterialInstance.h
+++ b/Engine/Source/Runtime/Engine/Classes/Materials/MaterialInstance.h
@@ -23,6 +23,7 @@ class ITargetPlatform;
 class UPhysicalMaterial;
 class USubsurfaceProfile;
 class UTexture;
+class UTextureCollection;
 
 //
 // Forward declarations.
@@ -282,6 +283,54 @@ struct FTextureParameterValue
 	}
 };
 
+/** Editable texture collection parameter. */
+USTRUCT(BlueprintType)
+struct FTextureCollectionParameterValue
+{
+	GENERATED_USTRUCT_BODY()
+
+	UPROPERTY(EditAnywhere, BlueprintReadWrite, Category=TextureCollectionParameterValue)
+	FMaterialParameterInfo ParameterInfo;
+
+	UPROPERTY(EditAnywhere, BlueprintReadWrite, Category=TextureCollectionParameterValue)
+	TObjectPtr<UTextureCollection> ParameterValue;
+
+	UPROPERTY()
+	FGuid ExpressionGUID;
+
+	explicit FTextureCollectionParameterValue(const FMaterialParameterInfo& InParameterInfo = FMaterialParameterInfo(), UTextureCollection* InValue = nullptr)
+		: ParameterInfo(InParameterInfo), ParameterValue(InValue)
+	{
+	}
+
+	bool IsOverride() const { return true; }
+
+	bool IsValid() const { return GetValue(*this) != nullptr; }
+
+	bool operator==(const FTextureCollectionParameterValue& Other) const
+	{
+		return
+			ParameterInfo == Other.ParameterInfo &&
+			ParameterValue == Other.ParameterValue &&
+			ExpressionGUID == Other.ExpressionGUID;
+	}
+	bool operator!=(const FTextureCollectionParameterValue& Other) const
+	{
+		return !((*this) == Other);
+	}
+
+	typedef const UTextureCollection* ValueType;
+	static ValueType GetValue(const FTextureCollectionParameterValue& Parameter) { return Parameter.ParameterValue; }
+
+	void GetValue(FMaterialParameterMetadata& OutResult) const
+	{
+		OutResult.Value = ParameterValue;
+#if WITH_EDITORONLY_DATA
+		OutResult.ExpressionGuid = ExpressionGUID;
+#endif
+	}
+};
+
 /** Editable runtime virtual texture parameter. */
 USTRUCT(BlueprintType)
 struct FRuntimeVirtualTextureParameterValue
@@ -632,6 +681,10 @@ public:
 	UPROPERTY(EditAnywhere, BlueprintReadOnly, Category=MaterialInstance, meta = (EditFixedOrder))
 	TArray<struct FTextureParameterValue> TextureParameterValues;
 
+	/** Texture Collection parameters. */
+	UPROPERTY(EditAnywhere, BlueprintReadOnly, Category=MaterialInstance, meta = (EditFixedOrder))
+	TArray<struct FTextureCollectionParameterValue> TextureCollectionParameterValues;
+
 	/** RuntimeVirtualTexture parameters. */
 	UPROPERTY(EditAnywhere, BlueprintReadOnly, Category = MaterialInstance, meta = (EditFixedOrder))
 	TArray<struct FRuntimeVirtualTextureParameterValue> RuntimeVirtualTextureParameterValues;
@@ -757,7 +810,8 @@ public:
 
 	virtual ENGINE_API void GetUsedTextures(TArray<UTexture*>& OutTextures, EMaterialQualityLevel::Type QualityLevel, bool bAllQualityLevels, ERHIFeatureLevel::Type FeatureLevel, bool bAllFeatureLevels) const override;
 	virtual ENGINE_API void GetUsedTexturesAndIndices(TArray<UTexture*>& OutTextures, TArray< TArray<int32> >& OutIndices, EMaterialQualityLevel::Type QualityLevel, ERHIFeatureLevel::Type FeatureLevel) const;
-	virtual ENGINE_API bool GetTextureParameterValue(const FHashedMaterialParameterInfo& ParameterInfo, class UTexture*& OutValue, bool bOveriddenOnly = false) const;
+	virtual ENGINE_API bool GetTextureParameterValue(const FHashedMaterialParameterInfo& ParameterInfo, class UTexture*& OutValue, bool bOveriddenOnly = false) const override;
+	virtual ENGINE_API bool GetTextureCollectionParameterValue(const FHashedMaterialParameterInfo& ParameterInfo, class UTextureCollection*& OutValue, bool bOveriddenOnly = false) const override;
 	virtual ENGINE_API void OverrideTexture(const UTexture* InTextureToOverride, UTexture* OverrideTexture, ERHIFeatureLevel::Type InFeatureLevel) override;
 	virtual ENGINE_API void OverrideNumericParameterDefault(EMaterialParameterType Type, const FHashedMaterialParameterInfo& ParameterInfo, const UE::Shader::FValue& Value, bool bOverride, ERHIFeatureLevel::Type FeatureLevel) override;
 	virtual ENGINE_API bool CheckMaterialUsage(const EMaterialUsage Usage) override;
@@ -1033,6 +1087,7 @@ protected:
 	void SetScalarParameterAtlasInternal(const FMaterialParameterInfo& ParameterInfo, FScalarParameterAtlasInstanceData AtlasData);
 #endif
 	void SetTextureParameterValueInternal(const FMaterialParameterInfo& ParameterInfo, class UTexture* Value);
+	void SetTextureCollectionParameterValueInternal(const FMaterialParameterInfo& ParameterInfo, class UTextureCollection* Value);
 	void SetRuntimeVirtualTextureParameterValueInternal(const FMaterialParameterInfo& ParameterInfo, class URuntimeVirtualTexture* Value);
 	void SetSparseVolumeTextureParameterValueInternal(const FMaterialParameterInfo& ParameterInfo, class USparseVolumeTexture* Value);
 	void SetFontParameterValueInternal(const FMaterialParameterInfo& ParameterInfo, class UFont* FontValue, int32 FontPage);
diff --git a/Engine/Source/Runtime/Engine/Classes/Materials/MaterialInstanceUpdateParameterSet.h b/Engine/Source/Runtime/Engine/Classes/Materials/MaterialInstanceUpdateParameterSet.h
--- a/Engine/Source/Runtime/Engine/Classes/Materials/MaterialInstanceUpdateParameterSet.h
+++ b/Engine/Source/Runtime/Engine/Classes/Materials/MaterialInstanceUpdateParameterSet.h
@@ -268,6 +268,11 @@ bool UpdateParameterSet(TArray<FTextureParameterValue>& Parameters, UMaterial* P
 	return MaterialInstance_Private::UpdateParameterSet_WithCachedData<FTextureParameterValue, ExpressionType>(EMaterialParameterType::Texture, Parameters, ParentMaterial);
 }
 template <typename ParameterType, typename ExpressionType>
+bool UpdateParameterSet(TArray<FTextureCollectionParameterValue>& Parameters, UMaterial* ParentMaterial)
+{
+	return MaterialInstance_Private::UpdateParameterSet_WithCachedData<FTextureCollectionParameterValue, ExpressionType>(EMaterialParameterType::TextureCollection, Parameters, ParentMaterial);
+}
+template <typename ParameterType, typename ExpressionType>
 bool UpdateParameterSet(TArray<FFontParameterValue>& Parameters, UMaterial* ParentMaterial)
 {
 	return MaterialInstance_Private::UpdateParameterSet_WithCachedData<FFontParameterValue, ExpressionType>(EMaterialParameterType::Font, Parameters, ParentMaterial);
diff --git a/Engine/Source/Runtime/Engine/Classes/Materials/MaterialInterface.h b/Engine/Source/Runtime/Engine/Classes/Materials/MaterialInterface.h
--- a/Engine/Source/Runtime/Engine/Classes/Materials/MaterialInterface.h
+++ b/Engine/Source/Runtime/Engine/Classes/Materials/MaterialInterface.h
@@ -667,6 +667,7 @@ public:
 	ENGINE_API void GetAllVectorParameterInfo(TArray<FMaterialParameterInfo>& OutParameterInfo, TArray<FGuid>& OutParameterIds) const;
 	ENGINE_API void GetAllDoubleVectorParameterInfo(TArray<FMaterialParameterInfo>& OutParameterInfo, TArray<FGuid>& OutParameterIds) const;
 	ENGINE_API void GetAllTextureParameterInfo(TArray<FMaterialParameterInfo>& OutParameterInfo, TArray<FGuid>& OutParameterIds) const;
+	ENGINE_API void GetAllTextureCollectionParameterInfo(TArray<FMaterialParameterInfo>& OutParameterInfo, TArray<FGuid>& OutParameterIds) const;
 	ENGINE_API void GetAllRuntimeVirtualTextureParameterInfo(TArray<FMaterialParameterInfo>& OutParameterInfo, TArray<FGuid>& OutParameterIds) const;
 	ENGINE_API void GetAllSparseVolumeTextureParameterInfo(TArray<FMaterialParameterInfo>& OutParameterInfo, TArray<FGuid>& OutParameterIds) const;
 	ENGINE_API void GetAllFontParameterInfo(TArray<FMaterialParameterInfo>& OutParameterInfo, TArray<FGuid>& OutParameterIds) const;
@@ -686,6 +687,7 @@ public:
 	ENGINE_API bool GetVectorParameterDefaultValue(const FHashedMaterialParameterInfo& ParameterInfo, FLinearColor& OutValue) const;
 	ENGINE_API bool GetDoubleVectorParameterDefaultValue(const FHashedMaterialParameterInfo& ParameterInfo, FVector4d& OutValue) const;
 	ENGINE_API bool GetTextureParameterDefaultValue(const FHashedMaterialParameterInfo& ParameterInfo, class UTexture*& OutValue) const;
+	ENGINE_API bool GetTextureCollectionParameterDefaultValue(const FHashedMaterialParameterInfo& ParameterInfo, class UTextureCollection*& OutValue) const;
 	ENGINE_API bool GetRuntimeVirtualTextureParameterDefaultValue(const FHashedMaterialParameterInfo& ParameterInfo, class URuntimeVirtualTexture*& OutValue) const;
 	ENGINE_API bool GetSparseVolumeTextureParameterDefaultValue(const FHashedMaterialParameterInfo& ParameterInfo, class USparseVolumeTexture*& OutValue) const;
 	ENGINE_API bool GetFontParameterDefaultValue(const FHashedMaterialParameterInfo& ParameterInfo, class UFont*& OutFontValue, int32& OutFontPage) const;
@@ -701,6 +703,8 @@ public:
 	/** Get textures referenced by expressions, including nested functions. */
 	ENGINE_API TArrayView<const TObjectPtr<UObject>> GetReferencedTextures() const;
 
+	ENGINE_API TConstArrayView<TObjectPtr<UTextureCollection>> GetReferencedTextureCollections() const;
+
 	virtual void SaveShaderStableKeysInner(const class ITargetPlatform* TP, const struct FStableShaderKeyAndValue& SaveKeyVal)
 		PURE_VIRTUAL(UMaterialInterface::SaveShaderStableKeysInner, );
 
@@ -907,6 +911,7 @@ public:
 	ENGINE_API bool GetDoubleVectorParameterChannelNames(const FHashedMaterialParameterInfo& ParameterInfo, FParameterChannelNames& OutValue) const;
 #endif // WITH_EDITOR
 	ENGINE_API virtual bool GetTextureParameterValue(const FHashedMaterialParameterInfo& ParameterInfo, class UTexture*& OutValue, bool bOveriddenOnly = false) const;
+	ENGINE_API virtual bool GetTextureCollectionParameterValue(const FHashedMaterialParameterInfo& ParameterInfo, class UTextureCollection*& OutValue, bool bOveriddenOnly = false) const;
 	ENGINE_API bool GetRuntimeVirtualTextureParameterValue(const FHashedMaterialParameterInfo& ParameterInfo, class URuntimeVirtualTexture*& OutValue, bool bOveriddenOnly = false) const;
 	ENGINE_API bool GetSparseVolumeTextureParameterValue(const FHashedMaterialParameterInfo& ParameterInfo, class USparseVolumeTexture*& OutValue, bool bOveriddenOnly = false) const;
 #if WITH_EDITOR
diff --git a/Engine/Source/Runtime/Engine/Private/Materials/HLSLMaterialTranslator.cpp b/Engine/Source/Runtime/Engine/Private/Materials/HLSLMaterialTranslator.cpp
--- a/Engine/Source/Runtime/Engine/Private/Materials/HLSLMaterialTranslator.cpp
+++ b/Engine/Source/Runtime/Engine/Private/Materials/HLSLMaterialTranslator.cpp
@@ -9,6 +9,7 @@
 #include "Engine/Engine.h"
 #include "MaterialDomain.h"
 #include "Engine/Texture.h"
+#include "Engine/TextureCollection.h"
 #include "Materials/MaterialAttributeDefinitionMap.h"
 #include "Field/FieldSystemTypes.h"
 #include "Materials/MaterialExpressionCustom.h"
@@ -21,7 +22,6 @@
 #include "Engine/RendererSettings.h"
 #include "DataDrivenShaderPlatformInfo.h"
 #include "Materials/Material.h"
-#include <functional>
 #include "Materials/MaterialExpressionAbsorptionMediumMaterialOutput.h"
 #include "Materials/MaterialExpressionCustomOutput.h"
 #include "Materials/MaterialExpressionFunctionInput.h"
@@ -64,6 +64,8 @@
 #include "ShaderPlatformCachedIniValue.h"
 #endif
 
+#include <functional>
+
 #if ENABLE_COOK_STATS && STATS
 
 static bool GEnableMaterialTranslationLogFile = false;
@@ -2154,6 +2156,7 @@ void FHLSLMaterialTranslator::DoTranslate()
 			TextureExpression->GetTextureParameterInfo(MaterialCompilationOutput.UniformExpressionSet.UniformTextureParameters[TypeIndex].AddDefaulted_GetRef());
 		}
 	}
+
 	MaterialCompilationOutput.UniformExpressionSet.UniformExternalTextureParameters.Empty(UniformExternalTextureExpressions.Num());
 	for (FMaterialUniformExpressionExternalTexture* TextureExpression : UniformExternalTextureExpressions)
 	{
@@ -2161,6 +2164,13 @@ void FHLSLMaterialTranslator::DoTranslate()
 		TextureExpression->GetExternalTextureParameterInfo(MaterialCompilationOutput.UniformExpressionSet.UniformExternalTextureParameters.AddDefaulted_GetRef());
 	}
 
+	MaterialCompilationOutput.UniformExpressionSet.UniformTextureCollectionParameters.Empty(UniformTextureCollectionExpressions.Num());
+	for (const FMaterialUniformExpressionTextureCollection* TextureExpression : UniformTextureCollectionExpressions)
+	{
+		CHECK_DDC_QUERY_FINISHED_ELSE_RETURN();
+		TextureExpression->GetTextureCollectionParameterInfo(MaterialCompilationOutput.UniformExpressionSet.UniformTextureCollectionParameters.AddDefaulted_GetRef());
+	}
+
 	for (uint32 SubstrateCompilationContextIndex = 0; SubstrateCompilationContextIndex < ESubstrateCompilationContext::SCC_MAX; ++SubstrateCompilationContextIndex)
 	{
 		CHECK_DDC_QUERY_FINISHED_ELSE_RETURN();
@@ -3289,6 +3299,7 @@ const TCHAR* FHLSLMaterialTranslator::DescribeType(EMaterialValueType Type) cons
 	case MCT_LWCVector2:			return TEXT("LWCVector2");
 	case MCT_LWCVector3:			return TEXT("LWCVector3");
 	case MCT_LWCVector4:			return TEXT("LWCVector4");
+	case MCT_TextureCollection:		return TEXT("TextureCollection");
 	default:						return TEXT("unknown");
 	};
 }
@@ -3303,11 +3314,11 @@ const TCHAR* FHLSLMaterialTranslator::HLSLTypeString(EMaterialValueType Type) co
 	case MCT_Float3:				return TEXT("MaterialFloat3");
 	case MCT_Float4:				return TEXT("MaterialFloat4");
 	case MCT_Float:					return TEXT("MaterialFloat");
-	case MCT_Texture2D:				return TEXT("texture2D");
-	case MCT_TextureCube:			return TEXT("textureCube");
-	case MCT_Texture2DArray:		return TEXT("texture2DArray");
-	case MCT_TextureCubeArray:		return TEXT("textureCubeArray");
-	case MCT_VolumeTexture:			return TEXT("volumeTexture");
+	case MCT_Texture2D:				return TEXT("Texture2D");
+	case MCT_TextureCube:			return TEXT("TextureCube");
+	case MCT_Texture2DArray:		return TEXT("Texture2DArray");
+	case MCT_TextureCubeArray:		return TEXT("TextureCubeArray");
+	case MCT_VolumeTexture:			return TEXT("Texture3D");
 	case MCT_StaticBool:			return TEXT("static bool");
 	case MCT_Bool:					return TEXT("bool");
 	case MCT_MaterialAttributes:	return TEXT("FMaterialAttributes");
@@ -3326,6 +3337,7 @@ const TCHAR* FHLSLMaterialTranslator::HLSLTypeString(EMaterialValueType Type) co
 	case MCT_LWCVector2:			return TEXT("FWSVector2");
 	case MCT_LWCVector3:			return TEXT("FWSVector3");
 	case MCT_LWCVector4:			return TEXT("FWSVector4");
+	case MCT_TextureCollection:		return TEXT("FResourceCollection");
 	default:						return TEXT("unknown");
 	};
 }
@@ -3363,6 +3375,7 @@ const TCHAR* FHLSLMaterialTranslator::HLSLTypeStringDeriv(EMaterialValueType Typ
 	case MCT_LWCVector2:			return (DerivativeStatus == EDerivativeStatus::Valid) ? TEXT("FWSVector2Deriv") : TEXT("FWSVector2");
 	case MCT_LWCVector3:			return (DerivativeStatus == EDerivativeStatus::Valid) ? TEXT("FWSVector3Deriv") : TEXT("FWSVector3");
 	case MCT_LWCVector4:			return (DerivativeStatus == EDerivativeStatus::Valid) ? TEXT("FWSVector4Deriv") : TEXT("FWSVector4");
+	case MCT_TextureCollection:		return TEXT("FResourceCollection");
 	default:						return TEXT("unknown");
 	};
 }
@@ -3474,6 +3487,8 @@ int32 FHLSLMaterialTranslator::AddCodeChunkInner(uint64 Hash, const TCHAR* Forma
 		return INDEX_NONE;
 	}
 
+	const bool bFromTextureCollection = (Type & MCT_TextureCollection) != 0;
+
 	int32 CodeIndex = INDEX_NONE;
 	if (Type == MCT_VoidStatement)
 	{
@@ -3488,7 +3503,7 @@ int32 FHLSLMaterialTranslator::AddCodeChunkInner(uint64 Hash, const TCHAR* Forma
 		new(*CurrentScopeChunks) FShaderCodeChunk(Hash, FormattedCode, FormattedCode, TEXT(""), Type, DerivativeStatus, true);
 	}
 	// Can only create temporaries for certain types
-	else if ((Type & (MCT_Float | MCT_LWCType | MCT_VTPageTableResult | MCT_UInt)) || Type == MCT_ShadingModel || Type == MCT_MaterialAttributes || Type == MCT_Substrate || Type == MCT_UInt)
+	else if ((Type & (MCT_Float | MCT_LWCType | MCT_VTPageTableResult | MCT_UInt)) || Type == MCT_ShadingModel || Type == MCT_MaterialAttributes || Type == MCT_Substrate || Type == MCT_UInt || bFromTextureCollection)
 	{
 		// Check for existing
 		for (int32 i = 0; i < CurrentScopeChunks->Num(); ++i)
@@ -3502,13 +3517,15 @@ int32 FHLSLMaterialTranslator::AddCodeChunkInner(uint64 Hash, const TCHAR* Forma
 
 		if (CodeIndex == INDEX_NONE)
 		{
+			const EMaterialValueType EffectiveType = EMaterialValueType(Type & ~MCT_TextureCollection);
+
 			CodeIndex = CurrentScopeChunks->Num();
 			// Allocate a local variable name
 			const FString SymbolName = CreateSymbolName(TEXT("Local"));
 			// Construct the definition string which stores the result in a temporary and adds a newline for readability
-			const FString LocalVariableDefinitionFinite = FString("	") + HLSLTypeString(Type) + TEXT(" ") + SymbolName + TEXT(" = ") + FormattedCode + TEXT(";") + HLSL_LINE_TERMINATOR;
+			const FString LocalVariableDefinitionFinite = FString("	") + HLSLTypeString(EffectiveType) + TEXT(" ") + SymbolName + TEXT(" = ") + FormattedCode + TEXT(";") + HLSL_LINE_TERMINATOR;
 			// Construct the definition string which stores the result in a temporary and adds a newline for readability
-			const FString LocalVariableDefinitionAnalytic = FString("	") + HLSLTypeString(Type) + TEXT(" ") + SymbolName + TEXT(" = ") + FormattedCode + TEXT(";") + HLSL_LINE_TERMINATOR;
+			const FString LocalVariableDefinitionAnalytic = FString("	") + HLSLTypeString(EffectiveType) + TEXT(" ") + SymbolName + TEXT(" = ") + FormattedCode + TEXT(";") + HLSL_LINE_TERMINATOR;
 			// Adding a code chunk that creates a local variable
 			new(*CurrentScopeChunks) FShaderCodeChunk(Hash, *LocalVariableDefinitionFinite, *LocalVariableDefinitionAnalytic, SymbolName, Type, DerivativeStatus, false);
 		}
@@ -3821,11 +3838,14 @@ int32 FHLSLMaterialTranslator::AccessUniformExpression(int32 Index)
 	}
 
 	FMaterialUniformExpressionTexture* TextureUniformExpression = CodeChunk.UniformExpression->GetTextureUniformExpression();
+	FMaterialUniformExpressionTextureCollection* TextureCollectionUniformExpression = CodeChunk.UniformExpression->GetTextureCollectionUniformExpression();
 	FMaterialUniformExpressionExternalTexture* ExternalTextureUniformExpression = CodeChunk.UniformExpression->GetExternalTextureUniformExpression();
 
 	// Any code chunk can have a texture uniform expression (eg FMaterialUniformExpressionFlipBookTextureParameter),
 	// But a texture code chunk must have a texture uniform expression
 	check(!(CodeChunk.Type & MCT_Texture) || TextureUniformExpression || ExternalTextureUniformExpression);
+	// Texture collection samples must have a corresponding uniform expression
+	check(!(CodeChunk.Type & MCT_TextureCollection) || TextureCollectionUniformExpression);
 	// External texture samples must have a corresponding uniform expression
 	check(!(CodeChunk.Type & MCT_TextureExternal) || ExternalTextureUniformExpression);
 	// Virtual texture samples must have a corresponding uniform expression
@@ -4002,6 +4022,11 @@ int32 FHLSLMaterialTranslator::AccessUniformExpression(int32 Index)
 		int32 TextureInputIndex = UniformTextureExpressions[(uint32)EMaterialTextureParameterType::SparseVolume].AddUnique(TextureUniformExpression);
 		FormattedCode.Appendf(TEXT("SparseVolumeTextureUnpackUniforms(Material.SVTPackedUniform[%d*2], Material.SVTPackedUniform[%d*2+1])"), TextureInputIndex, TextureInputIndex);
 	}
+	else if (CodeChunk.Type == MCT_TextureCollection)
+	{
+		int32 TextureCollectionInputIndex = UniformTextureCollectionExpressions.AddUnique(TextureCollectionUniformExpression);
+		FormattedCode.Appendf(TEXT("Material.TextureCollection_%u"), TextureCollectionInputIndex);
+	}
 	else
 	{
 		UE_LOG(LogMaterial, Fatal,TEXT("User input of unknown type: %s"),DescribeType(CodeChunk.Type));
@@ -4050,7 +4075,7 @@ bool FHLSLMaterialTranslator::GetTextureForExpression(int32 Index, int32& OutTex
 	check(Index >= 0 && Index < CurrentScopeChunks->Num());
 	const FShaderCodeChunk& Chunk = (*CurrentScopeChunks)[Index];
 	const EMaterialValueType TexInputType = Chunk.Type;
-	if (!(TexInputType & MCT_Texture))
+	if ((TexInputType & MCT_Texture) == 0)
 	{
 		return false;
 	}
@@ -4082,6 +4107,33 @@ bool FHLSLMaterialTranslator::GetTextureForExpression(int32 Index, int32& OutTex
 	return true;
 }
 
+bool FHLSLMaterialTranslator::GetTextureCollectionForExpression(int32 Index, int32& OutTextureCollectionIndex, TOptional<FName>& OutParameterName) const
+{
+	check(Index >= 0 && Index < CurrentScopeChunks->Num());
+	const FShaderCodeChunk& Chunk = (*CurrentScopeChunks)[Index];
+	const EMaterialValueType TexInputType = Chunk.Type;
+	if ((TexInputType & MCT_TextureCollection) == 0)
+	{
+		return false;
+	}
+
+	// If 'InputExpression' is connected, we use need to find the texture collection object that was passed in
+	// In this case, the texture collection assigned on this expression node is not used
+	FMaterialUniformExpression* UniformExpression = Chunk.UniformExpression;
+	checkf(UniformExpression, TEXT("TexInputType is %d, but missing FMaterialUniformExpression"), TexInputType);
+
+	if (FMaterialUniformExpressionTextureCollection* TextureCollectionUniform = UniformExpression->GetTextureCollectionUniformExpression())
+	{
+		OutTextureCollectionIndex = TextureCollectionUniform->GetTextureCollectionIndex();
+		if (FMaterialUniformExpressionTextureCollectionParameter* TextureCollectionParameterUniform = TextureCollectionUniform->GetTextureCollectionParameterUniformExpression())
+		{
+			OutParameterName = TextureCollectionParameterUniform->GetParameterName();
+		}
+	}
+
+	return true;
+}
+
 // GetArithmeticResultType
 EMaterialValueType FHLSLMaterialTranslator::GetArithmeticResultType(EMaterialValueType TypeA, EMaterialValueType TypeB)
 {
@@ -4858,6 +4910,12 @@ FString FHLSLMaterialTranslator::CastValue(const FString& Code, EMaterialValueTy
 		return Result;
 	}
 
+	// If the type came from a texture collection, we'll need to remove that flag to resolve the cast.
+	if (SourceType & MCT_TextureCollection && SourceType != MCT_TextureCollection)
+	{
+		return CastValue(Code, EMaterialValueType(SourceType & ~MCT_TextureCollection), DestType, Flags);
+	}
+
 	Errorf(TEXT("Cannot cast between non-numeric types %s to %s."), DescribeType(SourceType), DescribeType(DestType));
 	return FString();
 }
@@ -6826,6 +6884,9 @@ int32 FHLSLMaterialTranslator::TextureSample(
 		return INDEX_NONE;
 	}
 
+	const bool bFromCollection = (TextureType & MCT_TextureCollection) != 0;
+	TextureType = EMaterialValueType(TextureType & ~MCT_TextureCollection);
+
 	if (ShaderFrequency != SF_Pixel && MipValueMode == TMVM_MipBias)
 	{
 		Errorf(TEXT("MipBias is only supported in the pixel shader"));
@@ -6960,7 +7021,7 @@ int32 FHLSLMaterialTranslator::TextureSample(
 
 	// Won't be able to get the texture, if this is an external texture sample
 	const UTexture* Texture = nullptr;
-	if (TextureType != MCT_TextureExternal)
+	if (TextureType != MCT_TextureExternal && !bFromCollection)
 	{
 		FMaterialUniformExpression* Expression = (*CurrentScopeChunks)[TextureIndex].UniformExpression;
 		const FMaterialUniformExpressionTexture* TextureExpression = Expression ? Expression->GetTextureUniformExpression() : nullptr;
@@ -7107,25 +7168,39 @@ int32 FHLSLMaterialTranslator::TextureSample(
 		{
 			SamplerStateCode = FString::Printf(TEXT("%sSampler"), *TextureName);
 		}
-		else if (SamplerSource == SSM_Wrap_WorldGroupSettings)
-		{
-			// Use the shared sampler to save sampler slots
-			SamplerStateCode = FString::Printf(TEXT("GetMaterialSharedSampler(%sSampler,%s)"),
-				*TextureName, AutomaticViewMipBias ? TEXT("View.MaterialTextureBilinearWrapedSampler") : TEXT("Material.Wrap_WorldGroupSettings"));
-			RequiresManualViewMipBias = false;
-		}
-		else if (SamplerSource == SSM_Clamp_WorldGroupSettings)
-		{
-			// Use the shared sampler to save sampler slots
-			SamplerStateCode = FString::Printf(TEXT("GetMaterialSharedSampler(%sSampler,%s)"),
-				*TextureName, AutomaticViewMipBias ? TEXT("View.MaterialTextureBilinearClampedSampler") : TEXT("Material.Clamp_WorldGroupSettings"));
-			RequiresManualViewMipBias = false;
-		}
-		else if (SamplerSource == SSM_TerrainWeightmapGroupSettings)
+		else
 		{
-			SamplerStateCode = FString::Printf(TEXT("GetMaterialSharedSampler(%sSampler,%s)"),
-				*TextureName, TEXT("View.LandscapeWeightmapSampler"));
-			RequiresManualViewMipBias = false;
+			const TCHAR* SharedSamplerName = nullptr;
+
+			if (SamplerSource == SSM_Wrap_WorldGroupSettings)
+			{
+				// Use the shared sampler to save sampler slots
+				SharedSamplerName = AutomaticViewMipBias ? TEXT("View.MaterialTextureBilinearWrapedSampler") : TEXT("Material.Wrap_WorldGroupSettings");
+				RequiresManualViewMipBias = false;
+			}
+			else if (SamplerSource == SSM_Clamp_WorldGroupSettings)
+			{
+				// Use the shared sampler to save sampler slots
+				SharedSamplerName = AutomaticViewMipBias ? TEXT("View.MaterialTextureBilinearClampedSampler") : TEXT("Material.Clamp_WorldGroupSettings");
+				RequiresManualViewMipBias = false;
+			}
+			else if (SamplerSource == SSM_TerrainWeightmapGroupSettings)
+			{
+				SharedSamplerName = TEXT("View.LandscapeWeightmapSampler");
+				RequiresManualViewMipBias = false;
+			}
+
+			if (SharedSamplerName)
+			{
+				if (bFromCollection)
+				{
+					SamplerStateCode = SharedSamplerName;
+				}
+				else
+				{
+					SamplerStateCode = FString::Printf(TEXT("GetMaterialSharedSampler(%sSampler,%s)"), *TextureName, SharedSamplerName);
+				}
+			}
 		}
 	}
 
@@ -7401,6 +7476,62 @@ int32 FHLSLMaterialTranslator::TextureProperty(int32 TextureIndex, EMaterialExpo
 	return AddUniformExpression(Scope, new FMaterialUniformExpressionTextureProperty(TextureExpression, Property), ValueType, TEXT(""));
 }
 
+static const TCHAR* GetTextureObjectType(EMaterialValueType Type)
+{
+	switch (Type)
+	{
+	case MCT_Texture2D:				return TEXT("Texture2D");
+	case MCT_TextureCube:			return TEXT("TextureCube");
+	case MCT_Texture2DArray:		return TEXT("Texture2DArray");
+	case MCT_TextureCubeArray:		return TEXT("TextureCubeArray");
+	case MCT_VolumeTexture:			return TEXT("Texture3D");
+	default:						return TEXT("unknown");
+	};
+}
+
+int32 FHLSLMaterialTranslator::TextureFromCollection(int32 TextureCollectionCodeIndex, int32 IndexIntoCollectionCodeIndex, EMaterialValueType ResultTextureType)
+{
+	if (ErrorUnlessFeatureLevelSupported(ERHIFeatureLevel::SM6) == INDEX_NONE)
+	{
+		return INDEX_NONE;
+	}
+
+	if (TextureCollectionCodeIndex == INDEX_NONE)
+	{
+		return Errorf(TEXT("Missing TextureCollection parameter"));
+	}
+
+	if (IndexIntoCollectionCodeIndex == INDEX_NONE)
+	{
+		return Errorf(TEXT("Missing CollectionIndex parameter"));
+	}
+
+	const EMaterialValueType TextureType = GetParameterType(TextureCollectionCodeIndex);
+	if (TextureType != MCT_TextureCollection)
+	{
+		return Errorf(TEXT("TextureFromCollection only available for Texture Collections, not %s"), DescribeType(TextureType));
+	}
+
+	FMaterialUniformExpressionTextureCollection* TextureCollectionExpression = GetParameterUniformExpression(TextureCollectionCodeIndex)->GetTextureCollectionUniformExpression();
+	if (!TextureCollectionExpression)
+	{
+		return Errorf(TEXT("Expected a texture collection expression"));
+	}
+
+	if ((ResultTextureType & MCT_Texture) == 0)
+	{
+		return Errorf(TEXT("Expected a texture return type"));
+	}
+
+	return AddCodeChunk(
+		EMaterialValueType(ResultTextureType | MCT_TextureCollection),
+		TEXT("TextureFromCollection_%s(%s, %s)"),
+		GetTextureObjectType(ResultTextureType),
+		*GetParameterCode(TextureCollectionCodeIndex),
+		*GetParameterCode(IndexIntoCollectionCodeIndex)
+	);
+}
+
 int32 FHLSLMaterialTranslator::TextureDecalMipmapLevel(int32 TextureSizeInput)
 {
 	if (Material->GetMaterialDomain() != MD_DeferredDecal)
@@ -7969,6 +8100,71 @@ int32 FHLSLMaterialTranslator::TextureParameter(FName ParameterName, UTexture* I
 	return AddUniformExpression(Scope, new FMaterialUniformExpressionTextureParameter(ParameterInfo, TextureReferenceIndex, SamplerType, SamplerSource, bVirtual),ShaderType,TEXT(""));
 }
 
+int32 FHLSLMaterialTranslator::TextureCollection(UTextureCollection* InTextureCollection, int32& TextureCollectionReferenceIndex)
+{
+	TextureCollectionReferenceIndex = Material->GetReferencedTextureCollections().Find(InTextureCollection);
+	if (TextureCollectionReferenceIndex < 0)
+	{
+		return INDEX_NONE;
+	}
+
+	FAddUniformExpressionScope Scope(this);
+	return AddUniformExpression(Scope, new FMaterialUniformExpressionTextureCollection(TextureCollectionReferenceIndex), MCT_TextureCollection, TEXT(""));
+}
+
+int32 FHLSLMaterialTranslator::TextureCollectionParameter(FName ParameterName, UTextureCollection* InDefaultValue, int32& TextureCollectionReferenceIndex)
+{
+	UTextureCollection* DefaultValue = InDefaultValue;
+
+	// If we're compiling a function, give the function a chance to override the default parameter value
+	FMaterialParameterMetadata Meta;
+	if (GetParameterOverrideValueForCurrentFunction(EMaterialParameterType::TextureCollection, ParameterName, Meta))
+	{
+		DefaultValue = Meta.Value.TextureCollection;
+	}
+
+	TextureCollectionReferenceIndex = Material->GetReferencedTextureCollections().Find(DefaultValue);
+	if (TextureCollectionReferenceIndex == INDEX_NONE)
+	{
+		FString TextureCollectionName;
+		DefaultValue->GetName(TextureCollectionName);
+		Errorf(TEXT("Could not resolve referenced texture collection '%s'."), *TextureCollectionName);
+		return INDEX_NONE;
+	}
+
+	FMaterialParameterInfo ParameterInfo = GetParameterAssociationInfo();
+	ParameterInfo.Name = ParameterName;
+
+	FAddUniformExpressionScope Scope(this);
+	return AddUniformExpression(Scope, new FMaterialUniformExpressionTextureCollectionParameter(ParameterInfo, TextureCollectionReferenceIndex), MCT_TextureCollection, TEXT(""));
+}
+
+int32 FHLSLMaterialTranslator::TextureCollectionCount(int32 InTextureCollectionCodeIndex)
+{
+	if (ErrorUnlessFeatureLevelSupported(ERHIFeatureLevel::SM6) == INDEX_NONE)
+	{
+		return INDEX_NONE;
+	}
+
+	const EMaterialValueType TextureType = GetParameterType(InTextureCollectionCodeIndex);
+	if (TextureType != MCT_TextureCollection)
+	{
+		return Errorf(TEXT("TextureCollectionCount only available for Texture Collections, not %s"), DescribeType(TextureType));
+	}
+
+	FMaterialUniformExpressionTextureCollection* TextureCollectionExpression = GetParameterUniformExpression(InTextureCollectionCodeIndex)->GetTextureCollectionUniformExpression();
+	if (!TextureCollectionExpression)
+	{
+		return Errorf(TEXT("Expected a texture collection expression"));
+	}
+
+	return AddCodeChunk(
+		EMaterialValueType(MCT_UInt1),
+		TEXT("GetCountFromResourceCollection(%s)"),
+		*GetParameterCode(InTextureCollectionCodeIndex)
+	);
+}
+
 int32 FHLSLMaterialTranslator::VirtualTexture(URuntimeVirtualTexture* InTexture, int32 TextureLayerIndex, int32 PageTableLayerIndex, int32& TextureReferenceIndex, EMaterialSamplerType SamplerType) 
 {
 	if (!UseVirtualTexturing(Platform))
@@ -8181,6 +8377,11 @@ UObject* FHLSLMaterialTranslator::GetReferencedTexture(int32 Index)
 	return Material->GetReferencedTextures()[Index];
 }
 
+UTextureCollection* FHLSLMaterialTranslator::GetReferencedTextureCollection(int32 Index)
+{
+	return Material->GetReferencedTextureCollections()[Index];
+}
+
 int32 FHLSLMaterialTranslator::StaticBool(bool bValue)
 {
 	return AddInlinedCodeChunk(MCT_StaticBool,(bValue ? TEXT("true") : TEXT("false")));
diff --git a/Engine/Source/Runtime/Engine/Private/Materials/HLSLMaterialTranslator.h b/Engine/Source/Runtime/Engine/Private/Materials/HLSLMaterialTranslator.h
--- a/Engine/Source/Runtime/Engine/Private/Materials/HLSLMaterialTranslator.h
+++ b/Engine/Source/Runtime/Engine/Private/Materials/HLSLMaterialTranslator.h
@@ -301,6 +301,7 @@ protected:
 	// Uniform expressions used across all material properties
 	TArray<FShaderCodeChunk> UniformExpressions;
 	TArray<TRefCountPtr<FMaterialUniformExpressionTexture> > UniformTextureExpressions[NumMaterialTextureParameterTypes];
+	TArray<TRefCountPtr<FMaterialUniformExpressionTextureCollection>> UniformTextureCollectionExpressions;
 	TArray<TRefCountPtr<FMaterialUniformExpressionExternalTexture>> UniformExternalTextureExpressions;
 	TMap<UE::Shader::FValue, uint32> DefaultUniformValues;
 	uint32 UniformPreshaderOffset = 0u;
@@ -791,6 +792,7 @@ protected:
 	virtual FMaterialUniformExpression* GetParameterUniformExpression(int32 Index) const override;
 
 	virtual bool GetTextureForExpression(int32 Index, int32& OutTextureIndex, EMaterialSamplerType& OutSamplerType, TOptional<FName>& OutParameterName) const override;
+	virtual bool GetTextureCollectionForExpression(int32 Index, int32& OutTextureCollectionIndex, TOptional<FName>& OutParameterName) const override;
 
 	// GetArithmeticResultType
 	EMaterialValueType GetArithmeticResultType(EMaterialValueType TypeA, EMaterialValueType TypeB);
@@ -1020,6 +1022,7 @@ protected:
 	) override;
 
 	virtual int32 TextureProperty(int32 TextureIndex, EMaterialExposedTextureProperty Property) override;
+	virtual int32 TextureFromCollection(int32 TextureCollectionCodeIndex, int32 IndexIntoCollection, EMaterialValueType ResultTextureType) override;
 
 	virtual int32 TextureDecalMipmapLevel(int32 TextureSizeInput) override;
 	virtual int32 TextureDecalDerivative(bool bDDY) override;
@@ -1051,6 +1054,10 @@ protected:
 	virtual int32 Switch(int32 SwitchValueInput, int32 DefaultInput, TArray<int32>& CompiledInputs) override;
 	virtual int32 Texture(UTexture* InTexture, int32& TextureReferenceIndex, EMaterialSamplerType SamplerType, ESamplerSourceMode SamplerSource = SSM_FromTextureAsset, ETextureMipValueMode MipValueMode = TMVM_None) override;
 	virtual int32 TextureParameter(FName ParameterName, UTexture* DefaultValue, int32& TextureReferenceIndex, EMaterialSamplerType SamplerType, ESamplerSourceMode SamplerSource = SSM_FromTextureAsset) override;
+	
+	virtual int32 TextureCollection(UTextureCollection* TextureCollection, int32& TextureCollectionReferenceIndex) override;
+	virtual int32 TextureCollectionParameter(FName ParameterName, UTextureCollection* DefaultValue, int32& TextureCollectionReferenceIndex) override;
+	virtual int32 TextureCollectionCount(int32 InTextureCollectionCodeIndex) override;
 
 	virtual int32 VirtualTexture(URuntimeVirtualTexture* InTexture, int32 TextureLayerIndex, int32 PageTableLayerIndex, int32& TextureReferenceIndex, EMaterialSamplerType SamplerType) override;
 	virtual int32 VirtualTextureParameter(FName ParameterName, URuntimeVirtualTexture* DefaultValue, int32 TextureLayerIndex, int32 PageTableLayerIndex, int32& TextureReferenceIndex, EMaterialSamplerType SamplerType) override;
@@ -1075,7 +1082,8 @@ protected:
 	virtual int32 SparseVolumeTextureSamplePhysicalTileData(int32 SparseVolumeTextureIndex, int32 VoxelCoordIndex, int32 PhysicalTileDataIdxIndex) override;
 	virtual int32 SparseVolumeTextureSample(int32 SparseVolumeTextureIndex, int32 UVWIndex, int32 MipValue0Index, int32 MipValue1Index, int32 PhysicalTileDataIdxIndex, ETextureMipValueMode MipValueMode, ESamplerSourceMode SamplerSource) override;
 
-	virtual UObject* GetReferencedTexture(int32 Index);
+	virtual UObject* GetReferencedTexture(int32 Index) override;
+	virtual UTextureCollection* GetReferencedTextureCollection(int32 Index) override;
 
 	virtual int32 StaticBool(bool bValue) override;
 	virtual int32 StaticBoolParameter(FName ParameterName, bool bDefaultValue) override;
diff --git a/Engine/Source/Runtime/Engine/Private/Materials/MaterialCachedData.cpp b/Engine/Source/Runtime/Engine/Private/Materials/MaterialCachedData.cpp
--- a/Engine/Source/Runtime/Engine/Private/Materials/MaterialCachedData.cpp
+++ b/Engine/Source/Runtime/Engine/Private/Materials/MaterialCachedData.cpp
@@ -39,6 +39,7 @@
 #include "Curves/CurveLinearColor.h"
 #include "Curves/CurveLinearColorAtlas.h"
 #include "UObject/UE5MainStreamObjectVersion.h"
+#include "Engine/TextureCollection.h"
 
 #include UE_INLINE_GENERATED_CPP_BY_NAME(MaterialCachedData)
 
@@ -145,7 +146,7 @@ static bool TryAddParameter(FMaterialCachedExpressionData& CachedData,
 	return false;
 }
 
-bool FMaterialCachedExpressionData::AddParameter(const FMaterialParameterInfo& ParameterInfo, const FMaterialParameterMetadata& ParameterMeta, UObject*& OutReferencedTexture)
+bool FMaterialCachedExpressionData::AddParameter(const FMaterialParameterInfo& ParameterInfo, const FMaterialParameterMetadata& ParameterMeta, UObject*& OutReferencedTexture, UTextureCollection*& OutReferencedTextureCollection)
 {
 	check(EditorOnlyData);
 	int32 AssetIndex = INDEX_NONE;
@@ -194,6 +195,11 @@ bool FMaterialCachedExpressionData::AddParameter(const FMaterialParameterInfo& P
 			OutReferencedTexture = ParameterMeta.Value.Texture;
 			break;
 
+		case EMaterialParameterType::TextureCollection:
+			TextureCollectionValues.Insert(ParameterMeta.Value.TextureCollection, Index);
+			OutReferencedTextureCollection = ParameterMeta.Value.TextureCollection;
+			break;
+
 		case EMaterialParameterType::Font:
 			FontValues.Insert(ParameterMeta.Value.Font.Value, Index);
 			FontPageValues.Insert(ParameterMeta.Value.Font.Page, Index);
@@ -248,6 +254,10 @@ bool FMaterialCachedExpressionData::AddParameter(const FMaterialParameterInfo& P
 			bSameValue = TextureValues[Index] == ParameterMeta.Value.Texture;
 			break;
 
+		case EMaterialParameterType::TextureCollection:
+			bSameValue = TextureCollectionValues[Index] == ParameterMeta.Value.TextureCollection;
+			break;
+
 		case EMaterialParameterType::Font:
 			bSameValue = FontValues[Index] == ParameterMeta.Value.Font.Value && FontPageValues[Index] == ParameterMeta.Value.Font.Page;
 			break;
@@ -327,6 +337,7 @@ void FMaterialCachedExpressionData::UpdateForExpressions(const FMaterialCachedEx
 		}
 
 		UObject* ReferencedTexture = nullptr;
+		UTextureCollection* ReferencedTextureCollection = nullptr;
 
 		FMaterialParameterMetadata ParameterMeta;
 		if (Expression->GetParameterValue(ParameterMeta))
@@ -350,26 +361,34 @@ void FMaterialCachedExpressionData::UpdateForExpressions(const FMaterialCachedEx
 			const FMaterialParameterInfo ParameterInfo(ParameterName, Association, ParameterIndex);
 
 			// Try add the parameter. If this fails, the parameter is being added twice with different values. Report it as error.
-			if (!AddParameter(ParameterInfo, ParameterMeta, ReferencedTexture))
+			if (!AddParameter(ParameterInfo, ParameterMeta, ReferencedTexture, ReferencedTextureCollection))
 			{
 				DuplicateParameterErrors.AddUnique({ Expression, ParameterName });
 			}
 		}
 
-		// We first try to extract the referenced texture from the parameter value, that way we'll also get the proper texture in case value is overriden by a function instance
-		const bool bCanReferenceTexture = Expression->CanReferenceTexture();
-		if (!ReferencedTexture && bCanReferenceTexture)
+
+		if (ReferencedTexture)
+		{
+			ReferencedTextures.AddUnique(ReferencedTexture);
+		}
+		else if (ReferencedTextureCollection)
 		{
+			ReferencedTextureCollections.AddUnique(ReferencedTextureCollection);
+		}
+		else if (UTextureCollection* TextureCollection = Expression->GetReferencedTextureCollection())
+		{
+			ReferencedTextureCollections.AddUnique(TextureCollection);
+		}
+		else if (Expression->CanReferenceTexture())
+		{
+			// We first try to extract the referenced texture from the parameter value, that way we'll also get the proper texture in case value is overriden by a function instance
 			const UMaterialExpression::ReferencedTextureArray ExpressionReferencedTextures = Expression->GetReferencedTextures();
 			for (UObject* ExpressionReferencedTexture : ExpressionReferencedTextures)
 			{
 				ReferencedTextures.AddUnique(ExpressionReferencedTexture);
 			}
 		}
-		else if (ReferencedTexture)
-		{
-			ReferencedTextures.AddUnique(ReferencedTexture);
-		}
 
 		Expression->GetLandscapeLayerNames(EditorOnlyData->LandscapeLayerNames);
 
@@ -1125,6 +1144,9 @@ void FMaterialCachedExpressionData::GetParameterValueByIndex(EMaterialParameterT
 		}
 #endif // WITH_EDITORONLY_DATA
 		break;
+	case EMaterialParameterType::TextureCollection:
+		OutResult.Value = TextureCollectionValues[ParameterIndex].LoadSynchronous();
+		break;
 	case EMaterialParameterType::RuntimeVirtualTexture:
 		OutResult.Value = RuntimeVirtualTextureValues[ParameterIndex].LoadSynchronous();
 		break;
diff --git a/Engine/Source/Runtime/Engine/Private/Materials/MaterialExpressionHLSL.cpp b/Engine/Source/Runtime/Engine/Private/Materials/MaterialExpressionHLSL.cpp
--- a/Engine/Source/Runtime/Engine/Private/Materials/MaterialExpressionHLSL.cpp
+++ b/Engine/Source/Runtime/Engine/Private/Materials/MaterialExpressionHLSL.cpp
@@ -211,8 +211,11 @@
 #include "Materials/MaterialExpressionTangent.h"
 #include "Materials/MaterialExpressionTangentOutput.h"
 #include "Materials/MaterialExpressionTemporalSobol.h"
+#include "Materials/MaterialExpressionTextureCollection.h"
+#include "Materials/MaterialExpressionTextureCollectionParameter.h"
 #include "Materials/MaterialExpressionTextureCoordinate.h"
 #include "Materials/MaterialExpressionTextureObject.h"
+#include "Materials/MaterialExpressionTextureObjectFromCollection.h"
 #include "Materials/MaterialExpressionTextureObjectParameter.h"
 #include "Materials/MaterialExpressionTextureProperty.h"
 #include "Materials/MaterialExpressionTextureSample.h"
@@ -1492,6 +1495,62 @@ bool UMaterialExpressionTextureObjectParameter::GenerateHLSLExpression(FMaterial
 	return true;
 }
 
+bool UMaterialExpressionTextureCollection::GenerateHLSLExpression(FMaterialHLSLGenerator& Generator, UE::HLSLTree::FScope& Scope, int32 OutputIndex, UE::HLSLTree::FExpression const*& OutExpression) const
+{
+	using namespace UE::HLSLTree;
+	using namespace UE::Shader;
+
+	const FMaterialParameterMetadata ParameterMeta(TextureCollection);
+	OutExpression = Generator.GenerateMaterialParameter(FName(), ParameterMeta);
+	return true;
+}
+
+bool UMaterialExpressionTextureCollectionParameter::GenerateHLSLExpression(FMaterialHLSLGenerator& Generator, UE::HLSLTree::FScope& Scope, int32 OutputIndex, UE::HLSLTree::FExpression const*& OutExpression) const
+{
+	using namespace UE::HLSLTree;
+	using namespace UE::Shader;
+
+	FMaterialParameterMetadata ParameterMeta;
+	if (!GetParameterValue(ParameterMeta))
+	{
+		return Generator.Error(TEXT("Failed to get parameter value"));
+	}
+
+	OutExpression = Generator.GenerateMaterialParameter(ParameterName, ParameterMeta);
+	return true;
+}
+
+
+bool UMaterialExpressionTextureObjectFromCollection::GenerateHLSLExpression(FMaterialHLSLGenerator& Generator, UE::HLSLTree::FScope& Scope, int32 OutputIndex, UE::HLSLTree::FExpression const*& OutExpression) const
+{
+	using namespace UE::HLSLTree;
+	using namespace UE::Shader;
+
+	const FExpression* TextureCollectionExpression = nullptr;
+	if (TextureCollection.GetTracedInput().Expression)
+	{
+		TextureCollectionExpression = TextureCollection.AcquireHLSLExpression(Generator, Scope);
+	}
+	else if (TextureCollectionObject)
+	{
+		const FMaterialParameterMetadata ParameterMeta(TextureCollectionObject);
+		TextureCollectionExpression = Generator.GenerateMaterialParameter(FName(), ParameterMeta);
+	}
+	if (!TextureCollectionExpression)
+	{
+		return Generator.Error(TEXT("Failed to get Texture Collection"));
+	}
+
+	const FExpression* CollectionIndexExpression = CollectionIndex.AcquireHLSLExpressionOrConstant(Generator, Scope, ConstCollectionIndex);
+	if (!CollectionIndexExpression)
+	{
+		return Generator.Error(TEXT("Failed to get Texture Collection Index"));
+	}
+
+	OutExpression = Generator.GetTree().NewExpression<Material::FExpressionTextureObjectFromCollection>(TextureCollectionExpression, CollectionIndexExpression, TextureType);
+	return true;
+}
+
 bool UMaterialExpressionSkyLightEnvMapSample::GenerateHLSLExpression(FMaterialHLSLGenerator& Generator, UE::HLSLTree::FScope& Scope, int32 OutputIndex, UE::HLSLTree::FExpression const*& OutExpression) const
 {
 	using namespace UE::HLSLTree;
diff --git a/Engine/Source/Runtime/Engine/Private/Materials/MaterialExpressions.cpp b/Engine/Source/Runtime/Engine/Private/Materials/MaterialExpressions.cpp
--- a/Engine/Source/Runtime/Engine/Private/Materials/MaterialExpressions.cpp
+++ b/Engine/Source/Runtime/Engine/Private/Materials/MaterialExpressions.cpp
@@ -234,7 +234,10 @@
 #include "Materials/MaterialExpressionTextureSample.h"
 #include "Materials/MaterialExpressionParticleSubUV.h"
 #include "Materials/MaterialExpressionParticleSubUVProperties.h"
+#include "Materials/MaterialExpressionTextureCollection.h"
+#include "Materials/MaterialExpressionTextureCollectionParameter.h"
 #include "Materials/MaterialExpressionTextureSampleParameter.h"
+#include "Materials/MaterialExpressionTextureObjectFromCollection.h"
 #include "Materials/MaterialExpressionTextureObjectParameter.h"
 #include "Materials/MaterialExpressionTextureSampleParameter2D.h"
 #include "Materials/MaterialExpressionAntialiasedTextureMask.h"
@@ -2727,6 +2730,46 @@ FName UMaterialExpressionTextureSample::GetInputName(int32 InputIndex) const
 }
 #undef IF_INPUT_RETURN
 
+bool UMaterialExpressionTextureBase::VerifySamplerType(
+	const FString& TexturePathName,
+	EMaterialSamplerType CorrectSamplerType,
+	bool bSRGB,
+	EMaterialSamplerType SamplerType,
+	FString& OutErrorMessage)
+{
+	if (SamplerType != CorrectSamplerType)
+	{
+		UEnum* SamplerTypeEnum = UMaterialInterface::GetSamplerTypeEnum();
+		check(SamplerTypeEnum);
+
+		const FString SamplerTypeDisplayName = SamplerTypeEnum->GetDisplayNameTextByValue(SamplerType).ToString();
+		const FString TextureTypeDisplayName = SamplerTypeEnum->GetDisplayNameTextByValue(CorrectSamplerType).ToString();
+
+		OutErrorMessage = FString::Printf(TEXT("Sampler type is %s, should be %s for %s"),
+			*SamplerTypeDisplayName,
+			*TextureTypeDisplayName,
+			*TexturePathName);
+
+		return false;
+	}
+
+	if ((SamplerType == SAMPLERTYPE_Normal || SamplerType == SAMPLERTYPE_Masks) && bSRGB)
+	{
+		UEnum* SamplerTypeEnum = UMaterialInterface::GetSamplerTypeEnum();
+		check(SamplerTypeEnum);
+
+		const FString SamplerTypeDisplayName = SamplerTypeEnum->GetDisplayNameTextByValue(SamplerType).ToString();
+
+		OutErrorMessage = FString::Printf(TEXT("To use '%s' as sampler type, SRGB must be disabled for %s"),
+			*SamplerTypeDisplayName,
+			*TexturePathName);
+
+		return false;
+	}
+
+	return true;
+}
+
 bool UMaterialExpressionTextureBase::VerifySamplerType(
 	EShaderPlatform ShaderPlatform,
 	const ITargetPlatform* TargetPlatform,
@@ -2742,32 +2785,8 @@ bool UMaterialExpressionTextureBase::VerifySamplerType(
 		{
 			SamplerType = UMaterialExpressionTextureBase::GetSamplerTypeForTexture(Texture, !bIsVirtualTextured);
 		}
-		if ( SamplerType != CorrectSamplerType )
-		{
-			UEnum* SamplerTypeEnum = UMaterialInterface::GetSamplerTypeEnum();
-			check( SamplerTypeEnum );
-
-			FString SamplerTypeDisplayName = SamplerTypeEnum->GetDisplayNameTextByValue(SamplerType).ToString();
-			FString TextureTypeDisplayName = SamplerTypeEnum->GetDisplayNameTextByValue(CorrectSamplerType).ToString();
-
-			OutErrorMessage = FString::Printf(TEXT("Sampler type is %s, should be %s for %s"),
-				*SamplerTypeDisplayName,
-				*TextureTypeDisplayName,
-				*Texture->GetPathName() );
-			return false;
-		}
-		if((SamplerType == SAMPLERTYPE_Normal || SamplerType == SAMPLERTYPE_Masks) && Texture->SRGB)
-		{
-			UEnum* SamplerTypeEnum = UMaterialInterface::GetSamplerTypeEnum();
-			check( SamplerTypeEnum );
-
-			FString SamplerTypeDisplayName = SamplerTypeEnum->GetDisplayNameTextByValue(SamplerType).ToString();
-
-			OutErrorMessage = FString::Printf(TEXT("To use '%s' as sampler type, SRGB must be disabled for %s"),
-				*SamplerTypeDisplayName,
-				*Texture->GetPathName() );
-			return false;
-		}
+		
+		return VerifySamplerType(Texture->GetPathName(), CorrectSamplerType, Texture->SRGB, SamplerType, OutErrorMessage);
 	}
 	return true;
 }
@@ -2812,54 +2831,96 @@ int32 UMaterialExpressionTextureSample::Compile(class FMaterialCompiler* Compile
 			return INDEX_NONE;
 		}
 
-		UTexture* EffectiveTexture = Texture;
-		EMaterialSamplerType EffectiveSamplerType = SamplerType;
-		TOptional<FName> EffectiveParameterName;
-		if (InputExpression)
-		{
-			if (!Compiler->GetTextureForExpression(TextureCodeIndex, TextureReferenceIndex, EffectiveSamplerType, EffectiveParameterName))
-			{
-				return CompilerError(Compiler, TEXT("Tex input requires a texture value"));
-			}
-			if (TextureReferenceIndex != INDEX_NONE)
-			{
-				EffectiveTexture = Cast<UTexture>(Compiler->GetReferencedTexture(TextureReferenceIndex));
-			}
-		}
+		EMaterialValueType TextureType = Compiler->GetParameterType(TextureCodeIndex);
 
-		FString SamplerTypeError;
-		if (EffectiveTexture && VerifySamplerType(Compiler->GetShaderPlatform(), Compiler->GetTargetPlatform(), EffectiveTexture, EffectiveSamplerType, SamplerTypeError))
+		auto CheckForMissingUVWInput = [this, Compiler, TextureType](int32 ExpressionInput) -> TOptional<int32>
 		{
-			if (TextureCodeIndex != INDEX_NONE)
+			const EMaterialValueType TypesToCheck = EMaterialValueType(MCT_TextureCube | MCT_VolumeTexture | MCT_Texture2DArray | MCT_TextureCubeArray);
+			if (ExpressionInput != INDEX_NONE && (TextureType & TypesToCheck) != 0 && !Coordinates.GetTracedInput().Expression)
 			{
-				const EMaterialValueType TextureType = Compiler->GetParameterType(TextureCodeIndex);
-				if (TextureType == MCT_TextureCube && !Coordinates.GetTracedInput().Expression)
+				if (TextureType == MCT_TextureCube)
 				{
 					return CompilerError(Compiler, TEXT("UVW input required for cubemap sample"));
 				}
-				else if (TextureType == MCT_VolumeTexture && !Coordinates.GetTracedInput().Expression)
+				else if (TextureType == MCT_VolumeTexture)
 				{
 					return CompilerError(Compiler, TEXT("UVW input required for volume sample"));
 				}
-				else if (TextureType == MCT_Texture2DArray && !Coordinates.GetTracedInput().Expression)
+				else if (TextureType == MCT_Texture2DArray)
 				{
 					return CompilerError(Compiler, TEXT("UVW input required for texturearray sample"));
 				}
-				else if (TextureType == MCT_TextureCubeArray && !Coordinates.GetTracedInput().Expression)
+				else if (TextureType == MCT_TextureCubeArray)
 				{
 					return CompilerError(Compiler, TEXT("UVWX input required for texturecubearray sample"));
 				}
 			}
 
+			return TOptional<int32>();
+		};
+
+		auto GetCoordinateIndex = [this, Compiler](int32 ExpressionInput, EMaterialSamplerType EffectiveSamplerType, const TOptional<FName>& EffectiveParameterName)
+		{
 			int32 CoordinateIndex = Coordinates.GetTracedInput().Expression ? Coordinates.Compile(Compiler) : Compiler->TextureCoordinate(ConstCoordinate, false, false);
 
 			// If the sampler type is an external texture, we have might have a scale/bias to apply to the UV coordinates.
 			// Generate that code for the TextureReferenceIndex here so we compile it using the correct texture based on possible reroute textures above
 			if (EffectiveSamplerType == SAMPLERTYPE_External)
 			{
-				CoordinateIndex = CompileExternalTextureCoordinates(Compiler, CoordinateIndex, TextureReferenceIndex, EffectiveParameterName);
+				CoordinateIndex = CompileExternalTextureCoordinates(Compiler, CoordinateIndex, ExpressionInput, EffectiveParameterName);
+			}
+
+			return CoordinateIndex;
+		};
+
+		if (TextureType & MCT_TextureCollection)
+		{
+			// There's no UTexture object to get here
+
+			if (TOptional<int32> MissingError = CheckForMissingUVWInput(TextureCodeIndex))
+			{
+				return *MissingError;
+			}
+
+			const int32 CoordinateIndex = GetCoordinateIndex(TextureReferenceIndex, SamplerType, {});
+
+			return Compiler->TextureSample(
+				TextureCodeIndex,
+				CoordinateIndex,
+				SamplerType,
+				CompileMipValue0(Compiler),
+				CompileMipValue1(Compiler),
+				MipValueMode,
+				SamplerSource,
+				TextureReferenceIndex,
+				bDoAutomaticViewMipBias);
+		}
+
+		UTexture* EffectiveTexture = Texture;
+		EMaterialSamplerType EffectiveSamplerType = SamplerType;
+		TOptional<FName> EffectiveParameterName;
+		if (InputExpression)
+		{
+			if (!Compiler->GetTextureForExpression(TextureCodeIndex, TextureReferenceIndex, EffectiveSamplerType, EffectiveParameterName))
+			{
+				return CompilerError(Compiler, TEXT("Tex input requires a texture value"));
+			}
+			if (TextureReferenceIndex != INDEX_NONE)
+			{
+				EffectiveTexture = Cast<UTexture>(Compiler->GetReferencedTexture(TextureReferenceIndex));
+			}
+		}
+
+		FString SamplerTypeError;
+		if (EffectiveTexture && VerifySamplerType(Compiler->GetShaderPlatform(), Compiler->GetTargetPlatform(), EffectiveTexture, EffectiveSamplerType, SamplerTypeError))
+		{
+			if (TOptional<int32> MissingError = CheckForMissingUVWInput(TextureCodeIndex))
+			{
+				return *MissingError;
 			}
 
+			const int32 CoordinateIndex = GetCoordinateIndex(TextureReferenceIndex, EffectiveSamplerType, EffectiveParameterName);
+
 			return Compiler->TextureSample(
 				TextureCodeIndex,
 				CoordinateIndex,
@@ -10168,6 +10229,207 @@ void UMaterialExpressionBindlessSwitch::GetExpressionToolTip(TArray<FString>& Ou
 }
 #endif
 
+//
+// UMaterialExpressionTextureCollection
+//
+
+UMaterialExpressionTextureCollection::UMaterialExpressionTextureCollection(const FObjectInitializer& ObjectInitializer)
+	: Super(ObjectInitializer)
+{
+#if WITH_EDITORONLY_DATA
+	Outputs.Reset();
+	Outputs.Emplace(TEXT("TextureCollection"));
+	Outputs.Emplace(TEXT("TextureCount"));
+	
+	bShowOutputNameOnPin = true;
+#endif
+}
+
+#if WITH_EDITOR
+int32 UMaterialExpressionTextureCollection::Compile(FMaterialCompiler* Compiler, int32 OutputIndex)
+{
+	const int32 TextureCollectionCodeIndex = Compiler->TextureCollection(TextureCollection);
+
+	if (OutputIndex == 1)
+	{
+		return Compiler->TextureCollectionCount(TextureCollectionCodeIndex);
+	}
+
+	return TextureCollectionCodeIndex;
+}
+
+void UMaterialExpressionTextureCollection::GetCaption(TArray<FString>& OutCaptions) const
+{
+	OutCaptions.Emplace(TEXT("Texture Collection"));
+}
+
+uint32 UMaterialExpressionTextureCollection::GetOutputType(int32 OutputIndex)
+{
+	if (OutputIndex == 1)
+	{
+		return MCT_UInt1;
+	}
+
+	return MCT_TextureCollection;
+}
+#endif
+
+//
+// UMaterialExpressionTextureCollectionParameter
+//
+
+UMaterialExpressionTextureCollectionParameter::UMaterialExpressionTextureCollectionParameter(const FObjectInitializer& ObjectInitializer)
+	: Super(ObjectInitializer)
+{
+#if WITH_EDITORONLY_DATA
+	bIsParameterExpression = true;
+#endif
+}
+
+#if WITH_EDITOR
+int32 UMaterialExpressionTextureCollectionParameter::Compile(FMaterialCompiler* Compiler, int32 OutputIndex)
+{
+	const int32 TextureCollectionCodeIndex = Compiler->TextureCollectionParameter(ParameterName, TextureCollection);
+
+	if (OutputIndex == 1)
+	{
+		return Compiler->TextureCollectionCount(TextureCollectionCodeIndex);
+	}
+
+	return TextureCollectionCodeIndex;
+}
+
+void UMaterialExpressionTextureCollectionParameter::GetCaption(TArray<FString>& OutCaptions) const
+{
+	OutCaptions.Emplace(TEXT("Texture Collection Parameter"));
+}
+
+bool UMaterialExpressionTextureCollectionParameter::CanRenameNode() const
+{
+	return true;
+}
+
+FString UMaterialExpressionTextureCollectionParameter::GetEditableName() const
+{
+	return ParameterName.ToString();
+}
+
+void UMaterialExpressionTextureCollectionParameter::SetEditableName(const FString& NewName)
+{
+	ParameterName = *NewName;
+}
+
+bool UMaterialExpressionTextureCollectionParameter::HasAParameterName() const
+{
+	return true;
+}
+
+FName UMaterialExpressionTextureCollectionParameter::GetParameterName() const
+{
+	return ParameterName;
+}
+
+void UMaterialExpressionTextureCollectionParameter::SetParameterName(const FName& Name)
+{
+	ParameterName = Name;
+}
+
+void UMaterialExpressionTextureCollectionParameter::ValidateParameterName(const bool bAllowDuplicateName)
+{
+	ValidateParameterNameInternal(this, Material, bAllowDuplicateName);
+}
+
+bool UMaterialExpressionTextureCollectionParameter::GetParameterValue(FMaterialParameterMetadata& OutMeta) const
+{
+	OutMeta.Value = TextureCollection;
+	OutMeta.Description = Desc;
+	OutMeta.ExpressionGuid = ExpressionGUID;
+	OutMeta.Group = Group;
+	OutMeta.SortPriority = SortPriority;
+	OutMeta.AssetPath = GetAssetPathName();
+	return true;
+}
+
+bool UMaterialExpressionTextureCollectionParameter::SetParameterValue(const FName& Name, const FMaterialParameterMetadata& Meta, EMaterialExpressionSetParameterValueFlags Flags)
+{
+	if (Meta.Value.Type == EMaterialParameterType::TextureCollection)
+	{
+		if (SetParameterValue(Name, Meta.Value.TextureCollection, Flags))
+		{
+			if (EnumHasAnyFlags(Flags, EMaterialExpressionSetParameterValueFlags::AssignGroupAndSortPriority))
+			{
+				Group = Meta.Group;
+				SortPriority = Meta.SortPriority;
+			}
+			return true;
+		}
+	}
+	return false;
+}
+
+bool UMaterialExpressionTextureCollectionParameter::SetParameterValue(const FName& InParameterName, UTextureCollection* InValue, EMaterialExpressionSetParameterValueFlags Flags)
+{
+	if (InParameterName == ParameterName)
+	{
+		TextureCollection = InValue;
+		if (EnumHasAnyFlags(Flags, EMaterialExpressionSetParameterValueFlags::SendPostEditChangeProperty))
+		{
+			SendPostEditChangeProperty(this, GET_MEMBER_NAME_STRING_CHECKED(ThisClass, TextureCollection));
+		}
+		return true;
+	}
+	return false;
+}
+#endif
+
+FGuid& UMaterialExpressionTextureCollectionParameter::GetParameterExpressionId()
+{
+	return ExpressionGUID;
+}
+
+//
+// UMaterialExpressionTextureObjectFromCollection
+//
+
+UMaterialExpressionTextureObjectFromCollection::UMaterialExpressionTextureObjectFromCollection(const FObjectInitializer& ObjectInitializer)
+	: Super(ObjectInitializer)
+{
+}
+
+#if WITH_EDITOR
+int32 UMaterialExpressionTextureObjectFromCollection::Compile(FMaterialCompiler* Compiler, int32 OutputIndex)
+{
+	int32 TextureCollectionCodeIndex = TextureCollection.GetTracedInput().Expression ? TextureCollection.Compile(Compiler) : Compiler->TextureCollection(TextureCollectionObject);
+	int32 IndexIntoCollectionCodeIndex = CollectionIndex.GetTracedInput().Expression ? CollectionIndex.Compile(Compiler) : Compiler->Constant(ConstCollectionIndex);
+	int32 TextureFromCollectionCodeIndex = Compiler->TextureFromCollection(
+		TextureCollectionCodeIndex,
+		IndexIntoCollectionCodeIndex,
+		MaterialValueTypeFromTextureCollectionMemberType(TextureType)
+	);
+	return TextureFromCollectionCodeIndex;
+}
+
+uint32 UMaterialExpressionTextureObjectFromCollection::GetInputType(int32 InputIndex)
+{
+	if (InputIndex == 0)
+	{
+		return MCT_TextureCollection;
+	}
+
+	return MCT_UInt1;
+}
+
+uint32 UMaterialExpressionTextureObjectFromCollection::GetOutputType(int32 OutputIndex)
+{
+	return MaterialValueTypeFromTextureCollectionMemberType(TextureType);
+}
+
+void UMaterialExpressionTextureObjectFromCollection::GetCaption(TArray<FString>& OutCaptions) const
+{
+	OutCaptions.Emplace(TEXT("Texture Object From Collection"));
+}
+#endif
+
 //
 //	UMaterialExpressionRequiredSamplersSwitch
 //
@@ -14919,6 +15181,17 @@ bool UMaterialFunctionInterface::OverrideNamedTextureParameter(const FHashedMate
 	return false;
 }
 
+bool UMaterialFunctionInterface::OverrideNamedTextureCollectionParameter(const FHashedMaterialParameterInfo& ParameterInfo, class UTextureCollection*& OutValue)
+{
+	FMaterialParameterMetadata Meta;
+	if (GetParameterOverrideValue(EMaterialParameterType::TextureCollection, ParameterInfo.GetName(), Meta))
+	{
+		OutValue = Meta.Value.TextureCollection;
+		return true;
+	}
+	return false;
+}
+
 bool UMaterialFunctionInterface::OverrideNamedRuntimeVirtualTextureParameter(const FHashedMaterialParameterInfo& ParameterInfo, class URuntimeVirtualTexture*& OutValue)
 {
 	FMaterialParameterMetadata Meta;
@@ -16067,6 +16340,17 @@ void UMaterialFunctionInstance::UpdateParameterSet()
 						}
 					}
 				}
+				else if (const UMaterialExpressionTextureCollectionParameter* TextureCollectionParameter = Cast<const UMaterialExpressionTextureCollectionParameter>(FunctionExpression))
+				{
+					for (FTextureCollectionParameterValue& TextureCollectionParameterValue : TextureCollectionParameterValues)
+					{
+						if (TextureCollectionParameterValue.ExpressionGUID == TextureCollectionParameter->ExpressionGUID)
+						{
+							TextureCollectionParameterValue.ParameterInfo.Name = TextureCollectionParameter->ParameterName;
+							break;
+						}
+					}
+				}
 				else if (const UMaterialExpressionRuntimeVirtualTextureSampleParameter* RuntimeVirtualTextureParameter = Cast<const UMaterialExpressionRuntimeVirtualTextureSampleParameter>(FunctionExpression))
 				{
 					for (FRuntimeVirtualTextureParameterValue& RuntimeVirtualTextureParameterValue : RuntimeVirtualTextureParameterValues)
@@ -16134,6 +16418,7 @@ void UMaterialFunctionInstance::OverrideMaterialInstanceParameterValues(UMateria
 	Instance->VectorParameterValues = VectorParameterValues;
 	Instance->DoubleVectorParameterValues = DoubleVectorParameterValues;
 	Instance->TextureParameterValues = TextureParameterValues;
+	Instance->TextureCollectionParameterValues = TextureCollectionParameterValues;
 	Instance->RuntimeVirtualTextureParameterValues = RuntimeVirtualTextureParameterValues;
 	Instance->SparseVolumeTextureParameterValues = SparseVolumeTextureParameterValues;
 	Instance->FontParameterValues = FontParameterValues;
@@ -16223,6 +16508,14 @@ void UMaterialFunctionInstance::PostLoad()
 		}
 	}
 
+	for (const FTextureCollectionParameterValue& Param : TextureCollectionParameterValues)
+	{
+		if (UTextureCollection* TextureCollection = Param.ParameterValue)
+		{
+			TextureCollection->ConditionalPostLoad();
+		}
+	}
+
 	for (const FFontParameterValue& Param : FontParameterValues)
 	{
 		if (UFont* Font = Param.FontValue)
@@ -16342,6 +16635,7 @@ bool UMaterialFunctionInstance::GetParameterOverrideValue(EMaterialParameterType
 	case EMaterialParameterType::Vector: bResult = GameThread_GetParameterValue(VectorParameterValues, ParameterInfo, OutResult); break;
 	case EMaterialParameterType::DoubleVector: bResult = GameThread_GetParameterValue(DoubleVectorParameterValues, ParameterInfo, OutResult); break;
 	case EMaterialParameterType::Texture: bResult = GameThread_GetParameterValue(TextureParameterValues, ParameterInfo, OutResult); break;
+	case EMaterialParameterType::TextureCollection: bResult = GameThread_GetParameterValue(TextureCollectionParameterValues, ParameterInfo, OutResult); break;
 	case EMaterialParameterType::RuntimeVirtualTexture: bResult = GameThread_GetParameterValue(RuntimeVirtualTextureParameterValues, ParameterInfo, OutResult); break;
 	case EMaterialParameterType::SparseVolumeTexture: bResult = GameThread_GetParameterValue(SparseVolumeTextureParameterValues, ParameterInfo, OutResult); break;
 	case EMaterialParameterType::Font: bResult = GameThread_GetParameterValue(FontParameterValues, ParameterInfo, OutResult); break;
diff --git a/Engine/Source/Runtime/Engine/Private/Materials/MaterialHLSLEmitter.cpp b/Engine/Source/Runtime/Engine/Private/Materials/MaterialHLSLEmitter.cpp
--- a/Engine/Source/Runtime/Engine/Private/Materials/MaterialHLSLEmitter.cpp
+++ b/Engine/Source/Runtime/Engine/Private/Materials/MaterialHLSLEmitter.cpp
@@ -50,6 +50,7 @@ static const TCHAR* HLSLTypeString(EMaterialValueType Type)
 	case MCT_UInt3:					return TEXT("uint3");
 	case MCT_UInt4:					return TEXT("uint4");
 	case MCT_Substrate:				return TEXT("FSubstrateData");
+	case MCT_TextureCollection:		return TEXT("FResourceCollection");
 	default:						return TEXT("unknown");
 	};
 }
diff --git a/Engine/Source/Runtime/Engine/Private/Materials/MaterialHLSLTree.cpp b/Engine/Source/Runtime/Engine/Private/Materials/MaterialHLSLTree.cpp
--- a/Engine/Source/Runtime/Engine/Private/Materials/MaterialHLSLTree.cpp
+++ b/Engine/Source/Runtime/Engine/Private/Materials/MaterialHLSLTree.cpp
@@ -1771,6 +1771,19 @@ void FExpressionTextureSample::EmitValueShader(FEmitContext& Context, FEmitScope
 		bAutomaticViewMipBias);
 }
 
+bool FExpressionTextureObjectFromCollection::PrepareValue(FEmitContext& Context, FEmitScope& Scope, const FRequestedType& RequestedType, FPrepareValueResult& OutResult) const
+{
+	// TODO: christopher.waters - to implement
+	checkNoEntry();
+	return false;
+}
+
+void FExpressionTextureObjectFromCollection::EmitValueShader(FEmitContext& Context, FEmitScope& Scope, const FRequestedType& RequestedType, FEmitValueShaderResult& OutResult) const
+{
+	// TODO: christopher.waters - to implement
+	checkNoEntry();
+}
+
 FName FExpressionStaticTerrainLayerWeight::BuildWeightmapName(const TCHAR* Weightmap, int32 Index, bool bUseIndex) const
 {
 	FName Name;
diff --git a/Engine/Source/Runtime/Engine/Private/Materials/MaterialInstance.cpp b/Engine/Source/Runtime/Engine/Private/Materials/MaterialInstance.cpp
--- a/Engine/Source/Runtime/Engine/Private/Materials/MaterialInstance.cpp
+++ b/Engine/Source/Runtime/Engine/Private/Materials/MaterialInstance.cpp
@@ -8,6 +8,8 @@
 #include "Stats/StatsMisc.h"
 #include "EngineModule.h"
 #include "Engine/Font.h"
+#include "Engine/Texture.h"
+#include "Engine/TextureCollection.h"
 #include "Materials/Material.h"
 #include "UObject/Package.h"
 #include "Materials/MaterialExpressionStaticBoolParameter.h"
@@ -21,6 +23,7 @@
 #include "Materials/MaterialExpressionScalarParameter.h"
 #include "Materials/MaterialExpressionVectorParameter.h"
 #include "Materials/MaterialExpressionDoubleVectorParameter.h"
+#include "Materials/MaterialExpressionTextureCollectionParameter.h"
 #include "Materials/MaterialExpressionTextureSampleParameter.h"
 #include "Materials/MaterialExpressionFontSampleParameter.h"
 #include "Materials/MaterialExpressionMaterialAttributeLayers.h"
@@ -307,6 +310,7 @@ bool FMaterialInstanceResource::GetParameterValue(EMaterialParameterType Type, c
 		case EMaterialParameterType::Vector: bResult = RenderThread_GetParameterValue<FLinearColor>(ParameterInfo, OutValue); break;
 		case EMaterialParameterType::DoubleVector: bResult = RenderThread_GetParameterValue<FVector4d>(ParameterInfo, OutValue); break;
 		case EMaterialParameterType::Texture: bResult = RenderThread_GetParameterValue<const UTexture*>(ParameterInfo, OutValue); break;
+		case EMaterialParameterType::TextureCollection: bResult = RenderThread_GetParameterValue<const UTextureCollection*>(ParameterInfo, OutValue); break;
 		case EMaterialParameterType::RuntimeVirtualTexture: bResult = RenderThread_GetParameterValue<const URuntimeVirtualTexture*>(ParameterInfo, OutValue); break;
 		case EMaterialParameterType::SparseVolumeTexture: bResult = RenderThread_GetParameterValue<const USparseVolumeTexture*>(ParameterInfo, OutValue); break;
 		default: ensure(false); break; // other parameter types are not expected on the render thread
@@ -399,6 +403,7 @@ void FMaterialInstanceResource::InitMIParameters(FMaterialInstanceParameterSet&
 	ParameterSet.VectorParameters.Sort(SortMaterialInstanceParametersPredicate<FLinearColor>);
 	ParameterSet.DoubleVectorParameters.Sort(SortMaterialInstanceParametersPredicate<FVector4d>);
 	ParameterSet.TextureParameters.Sort(SortMaterialInstanceParametersPredicate<const UTexture*>);
+	ParameterSet.TextureCollectionParameters.Sort(SortMaterialInstanceParametersPredicate<const UTextureCollection*>);
 	ParameterSet.RuntimeVirtualTextureParameters.Sort(SortMaterialInstanceParametersPredicate<const URuntimeVirtualTexture*>);
 	ParameterSet.SparseVolumeTextureParameters.Sort(SortMaterialInstanceParametersPredicate<const USparseVolumeTexture*>);
 
@@ -407,6 +412,7 @@ void FMaterialInstanceResource::InitMIParameters(FMaterialInstanceParameterSet&
 	VectorParameterArray.Array = MoveTemp(ParameterSet.VectorParameters);
 	DoubleVectorParameterArray.Array = MoveTemp(ParameterSet.DoubleVectorParameters);
 	TextureParameterArray.Array = MoveTemp(ParameterSet.TextureParameters);
+	TextureCollectionParameterArray.Array = MoveTemp(ParameterSet.TextureCollectionParameters);
 	RuntimeVirtualTextureParameterArray.Array = MoveTemp(ParameterSet.RuntimeVirtualTextureParameters);
 	SparseVolumeTextureParameterArray.Array = MoveTemp(ParameterSet.SparseVolumeTextureParameters);
 
@@ -417,6 +423,7 @@ void FMaterialInstanceResource::InitMIParameters(FMaterialInstanceParameterSet&
 	VectorParameterArray.HashAddAllItems();
 	DoubleVectorParameterArray.HashAddAllItems();
 	TextureParameterArray.HashAddAllItems();
+	TextureCollectionParameterArray.HashAddAllItems();
 	RuntimeVirtualTextureParameterArray.HashAddAllItems();
 	SparseVolumeTextureParameterArray.HashAddAllItems();
 }
@@ -570,6 +577,7 @@ void UMaterialInstance::SwapLayerParameterIndices(int32 OriginalIndex, int32 New
 		SwapLayerParameterIndicesArray(VectorParameterValues, OriginalIndex, NewIndex);
 		SwapLayerParameterIndicesArray(DoubleVectorParameterValues, OriginalIndex, NewIndex);
 		SwapLayerParameterIndicesArray(TextureParameterValues, OriginalIndex, NewIndex);
+		SwapLayerParameterIndicesArray(TextureCollectionParameterValues, OriginalIndex, NewIndex);
 		SwapLayerParameterIndicesArray(RuntimeVirtualTextureParameterValues, OriginalIndex, NewIndex);
 		SwapLayerParameterIndicesArray(SparseVolumeTextureParameterValues, OriginalIndex, NewIndex);
 		SwapLayerParameterIndicesArray(FontParameterValues, OriginalIndex, NewIndex);
@@ -588,6 +596,7 @@ void UMaterialInstance::RemoveLayerParameterIndex(int32 Index)
 	RemoveLayerParameterIndicesArray(VectorParameterValues, Index);
 	RemoveLayerParameterIndicesArray(DoubleVectorParameterValues, Index);
 	RemoveLayerParameterIndicesArray(TextureParameterValues, Index);
+	RemoveLayerParameterIndicesArray(TextureCollectionParameterValues, Index);
 	RemoveLayerParameterIndicesArray(RuntimeVirtualTextureParameterValues, Index);
 	RemoveLayerParameterIndicesArray(SparseVolumeTextureParameterValues, Index);
 	RemoveLayerParameterIndicesArray(FontParameterValues, Index);
@@ -637,6 +646,9 @@ bool UMaterialInstance::UpdateParameters()
 			// Texture parameters
 			bDirty = UpdateParameterSet<FTextureParameterValue, UMaterialExpressionTextureSampleParameter>(TextureParameterValues, ParentMaterial) || bDirty;
 
+			// Texture Collection parameters
+			bDirty = UpdateParameterSet<FTextureCollectionParameterValue, UMaterialExpressionTextureCollectionParameter>(TextureCollectionParameterValues, ParentMaterial) || bDirty;
+
 			// Runtime Virtual Texture parameters
 			bDirty = UpdateParameterSet<FRuntimeVirtualTextureParameterValue, UMaterialExpressionRuntimeVirtualTextureSampleParameter>(RuntimeVirtualTextureParameterValues, ParentMaterial) || bDirty;
 
@@ -669,6 +681,7 @@ bool UMaterialInstance::UpdateParameters()
 					RemapLayerParameterIndicesArray(VectorParameterValues, RemapLayerIndices);
 					RemapLayerParameterIndicesArray(DoubleVectorParameterValues, RemapLayerIndices);
 					RemapLayerParameterIndicesArray(TextureParameterValues, RemapLayerIndices);
+					RemapLayerParameterIndicesArray(TextureCollectionParameterValues, RemapLayerIndices);
 					RemapLayerParameterIndicesArray(RuntimeVirtualTextureParameterValues, RemapLayerIndices);
 					RemapLayerParameterIndicesArray(SparseVolumeTextureParameterValues, RemapLayerIndices);
 					RemapLayerParameterIndicesArray(FontParameterValues, RemapLayerIndices);
@@ -681,7 +694,7 @@ bool UMaterialInstance::UpdateParameters()
 
 		if (bDirty)
 		{
-			FObjectCacheEventSink::NotifyReferencedTextureChanged_Concurrent(this);
+			FObjectCacheEventSink::NotifyMaterialChanged_Concurrent(this);
 		}
 	}
 #endif // WITH_EDITOR
@@ -775,6 +788,14 @@ void GameThread_InitMIParameters(const UMaterialInstance& Instance)
 		ParamRef.Value = FFontParameterValue::GetValue(Parameter);
 	}
 
+	ParameterSet.TextureCollectionParameters.Reserve(Instance.TextureCollectionParameterValues.Num());
+	for (const FTextureCollectionParameterValue& Parameter : Instance.TextureCollectionParameterValues)
+	{
+		auto& ParamRef = ParameterSet.TextureCollectionParameters.AddDefaulted_GetRef();
+		ParamRef.Info = Parameter.ParameterInfo;
+		ParamRef.Value = FTextureCollectionParameterValue::GetValue(Parameter);
+	}
+
 	// RuntimeVirtualTexture parameters
 	ParameterSet.RuntimeVirtualTextureParameters.Reserve(Instance.RuntimeVirtualTextureParameterValues.Num());
 	for (const FRuntimeVirtualTextureParameterValue& Parameter : Instance.RuntimeVirtualTextureParameterValues)
@@ -974,6 +995,7 @@ bool UMaterialInstance::GetParameterOverrideValue(EMaterialParameterType Type, c
 	case EMaterialParameterType::Vector: bResult = GameThread_GetParameterValue(VectorParameterValues, ParameterInfo, OutResult); break;
 	case EMaterialParameterType::DoubleVector: bResult = GameThread_GetParameterValue(DoubleVectorParameterValues, ParameterInfo, OutResult); break;
 	case EMaterialParameterType::Texture: bResult = GameThread_GetParameterValue(TextureParameterValues, ParameterInfo, OutResult); break;
+	case EMaterialParameterType::TextureCollection: bResult = GameThread_GetParameterValue(TextureCollectionParameterValues, ParameterInfo, OutResult); break;
 	case EMaterialParameterType::RuntimeVirtualTexture: bResult = GameThread_GetParameterValue(RuntimeVirtualTextureParameterValues, ParameterInfo, OutResult); break;
 	case EMaterialParameterType::SparseVolumeTexture: bResult = GameThread_GetParameterValue(SparseVolumeTextureParameterValues, ParameterInfo, OutResult); break;
 	case EMaterialParameterType::Font: bResult = GameThread_GetParameterValue(FontParameterValues, ParameterInfo, OutResult); break;
@@ -1784,7 +1806,7 @@ void UMaterialInstance::CopyMaterialInstanceParameters(UMaterialInterface* Sourc
 		InitResources();
 
 #if WITH_EDITOR
-		FObjectCacheEventSink::NotifyReferencedTextureChanged_Concurrent(this);
+		FObjectCacheEventSink::NotifyMaterialChanged_Concurrent(this);
 #endif
 	}
 }
@@ -1973,6 +1995,7 @@ void UMaterialInstance::GetAllParametersOfType(EMaterialParameterType Type, TMap
 		case EMaterialParameterType::Vector: GameThread_ApplyParameterOverrides(Instance->VectorParameterValues, LayerIndexRemap, bSetOverride, OverridenParameters, OutParameters); break;
 		case EMaterialParameterType::DoubleVector: GameThread_ApplyParameterOverrides(Instance->DoubleVectorParameterValues, LayerIndexRemap, bSetOverride, OverridenParameters, OutParameters); break;
 		case EMaterialParameterType::Texture: GameThread_ApplyParameterOverrides(Instance->TextureParameterValues, LayerIndexRemap, bSetOverride, OverridenParameters, OutParameters); break;
+		case EMaterialParameterType::TextureCollection: GameThread_ApplyParameterOverrides(Instance->TextureCollectionParameterValues, LayerIndexRemap, bSetOverride, OverridenParameters, OutParameters); break;
 		case EMaterialParameterType::RuntimeVirtualTexture: GameThread_ApplyParameterOverrides(Instance->RuntimeVirtualTextureParameterValues, LayerIndexRemap, bSetOverride, OverridenParameters, OutParameters); break;
 		case EMaterialParameterType::SparseVolumeTexture: GameThread_ApplyParameterOverrides(Instance->SparseVolumeTextureParameterValues, LayerIndexRemap, bSetOverride, OverridenParameters, OutParameters); break;
 		case EMaterialParameterType::Font: GameThread_ApplyParameterOverrides(Instance->FontParameterValues, LayerIndexRemap, bSetOverride, OverridenParameters, OutParameters); break;
@@ -3117,6 +3140,17 @@ void UMaterialInstance::PostLoad()
 		}
 	}
 
+	// We have to make sure the resources are created for all used texture collections.
+	for( int32 ValueIndex=0; ValueIndex<TextureCollectionParameterValues.Num(); ValueIndex++ )
+	{
+		// Make sure the texture is postloaded so the resource isn't null.
+		UTextureCollection* TextureCollection = TextureCollectionParameterValues[ValueIndex].ParameterValue;
+		if( TextureCollection)
+		{
+			TextureCollection->ConditionalPostLoad();
+		}
+	}
+
 	// do the same for runtime virtual textures
 	for (int32 ValueIndex = 0; ValueIndex < RuntimeVirtualTextureParameterValues.Num(); ValueIndex++)
 	{
@@ -3437,7 +3471,7 @@ bool UMaterialInstance::SetParentInternal(UMaterialInterface* NewParent, bool Re
 
 #if WITH_EDITOR
 			// Important to notify when the parent change for Material -> Material relationship update
-			FObjectCacheEventSink::NotifyReferencedTextureChanged_Concurrent(this);
+			FObjectCacheEventSink::NotifyMaterialChanged_Concurrent(this);
 #endif
 
 			if( Parent )
@@ -3560,6 +3594,7 @@ void UMaterialInstance::ReserveParameterValuesInternal(EMaterialParameterType Ty
 	case EMaterialParameterType::Vector: VectorParameterValues.Reserve(Capacity); break;
 	case EMaterialParameterType::DoubleVector: DoubleVectorParameterValues.Reserve(Capacity); break;
 	case EMaterialParameterType::Texture: TextureParameterValues.Reserve(Capacity); break;
+	case EMaterialParameterType::TextureCollection: TextureCollectionParameterValues.Reserve(Capacity); break;
 	case EMaterialParameterType::Font: FontParameterValues.Reserve(Capacity); break;
 	case EMaterialParameterType::RuntimeVirtualTexture: RuntimeVirtualTextureParameterValues.Reserve(Capacity); break;
 	case EMaterialParameterType::SparseVolumeTexture: SparseVolumeTextureParameterValues.Reserve(Capacity); break;
@@ -3588,6 +3623,7 @@ void UMaterialInstance::AddParameterValueInternal(const FMaterialParameterInfo&
 	case EMaterialParameterType::Vector: VectorParameterValues.Emplace(ParameterInfo, Value.AsLinearColor()); break;
 	case EMaterialParameterType::DoubleVector: DoubleVectorParameterValues.Emplace(ParameterInfo, Value.AsVector4d()); break;
 	case EMaterialParameterType::Texture: TextureParameterValues.Emplace(ParameterInfo, Value.Texture); break;
+	case EMaterialParameterType::TextureCollection: TextureCollectionParameterValues.Emplace(ParameterInfo, Value.TextureCollection); break;
 	case EMaterialParameterType::Font: FontParameterValues.Emplace(ParameterInfo, Value.Font.Value, Value.Font.Page); break;
 	case EMaterialParameterType::RuntimeVirtualTexture: RuntimeVirtualTextureParameterValues.Emplace(ParameterInfo, Value.RuntimeVirtualTexture); break;
 	case EMaterialParameterType::SparseVolumeTexture: SparseVolumeTextureParameterValues.Emplace(ParameterInfo, Value.SparseVolumeTexture); break;
@@ -3616,6 +3652,7 @@ void UMaterialInstance::SetParameterValueInternal(const FMaterialParameterInfo&
 	case EMaterialParameterType::Vector: SetVectorParameterValueInternal(ParameterInfo, Value.AsLinearColor()); break;
 	case EMaterialParameterType::DoubleVector: SetDoubleVectorParameterValueInternal(ParameterInfo, Value.AsVector4d()); break;
 	case EMaterialParameterType::Texture: SetTextureParameterValueInternal(ParameterInfo, Value.Texture); break;
+	case EMaterialParameterType::TextureCollection: SetTextureCollectionParameterValueInternal(ParameterInfo, Value.TextureCollection); break;
 	case EMaterialParameterType::Font: SetFontParameterValueInternal(ParameterInfo, Value.Font.Value, Value.Font.Page); break;
 	case EMaterialParameterType::RuntimeVirtualTexture: SetRuntimeVirtualTextureParameterValueInternal(ParameterInfo, Value.RuntimeVirtualTexture); break;
 	case EMaterialParameterType::SparseVolumeTexture: SetSparseVolumeTextureParameterValueInternal(ParameterInfo, Value.SparseVolumeTexture); break;
@@ -3800,7 +3837,46 @@ void UMaterialInstance::SetTextureParameterValueInternal(const FMaterialParamete
 			GameThread_UpdateMIParameter(this, *ParameterValue);
 
 #if WITH_EDITOR
-			FObjectCacheEventSink::NotifyReferencedTextureChanged_Concurrent(this);
+			FObjectCacheEventSink::NotifyMaterialChanged_Concurrent(this);
+#endif
+		}		
+	}
+}
+
+void UMaterialInstance::SetTextureCollectionParameterValueInternal(const FMaterialParameterInfo& ParameterInfo, UTextureCollection* Value)
+{
+	LLM_SCOPE(ELLMTag::MaterialInstance);
+
+	if (Value)
+	{
+		Value->ConditionalPostLoad();
+	}
+
+	FTextureCollectionParameterValue* ParameterValue = GameThread_FindParameterByName(TextureCollectionParameterValues, ParameterInfo);
+
+	bool bForceUpdate = false;
+	if(!ParameterValue)
+	{
+		// If there's no element for the named parameter in array yet, add one.
+		ParameterValue = new(TextureCollectionParameterValues) FTextureCollectionParameterValue;
+		ParameterValue->ParameterInfo = ParameterInfo;
+		ParameterValue->ExpressionGUID.Invalidate();
+		bForceUpdate = true;
+	}
+
+	// Don't enqueue an update if it isn't needed
+	if (bForceUpdate || ParameterValue->ParameterValue != Value)
+	{
+		// set as an ensure, because it is somehow possible to accidentally pass non-textures into here via blueprints...
+		if (Value && ensureMsgf(Value->IsA(UTextureCollection::StaticClass()), TEXT("Expecting a UTextureCollection! Value='%s' class='%s'"), *Value->GetName(), *Value->GetClass()->GetName()))
+		{
+			ParameterValue->ParameterValue = Value;
+			Value->AddToCluster(this, true);
+			// Update the material instance data in the rendering thread.
+			GameThread_UpdateMIParameter(this, *ParameterValue);
+
+#if WITH_EDITOR
+			FObjectCacheEventSink::NotifyMaterialChanged_Concurrent(this);
 #endif
 		}		
 	}
@@ -3839,7 +3915,7 @@ void UMaterialInstance::SetRuntimeVirtualTextureParameterValueInternal(const FMa
 			GameThread_UpdateMIParameter(this, *ParameterValue);
 
 #if WITH_EDITOR
-			FObjectCacheEventSink::NotifyReferencedTextureChanged_Concurrent(this);
+			FObjectCacheEventSink::NotifyMaterialChanged_Concurrent(this);
 #endif
 		}
 	}
@@ -3877,7 +3953,7 @@ void UMaterialInstance::SetSparseVolumeTextureParameterValueInternal(const FMate
 			GameThread_UpdateMIParameter(this, *ParameterValue);
 
 #if WITH_EDITOR
-			FObjectCacheEventSink::NotifyReferencedTextureChanged_Concurrent(this);
+			FObjectCacheEventSink::NotifyMaterialChanged_Concurrent(this);
 #endif
 		}
 	}
@@ -3937,6 +4013,7 @@ void UMaterialInstance::ClearParameterValuesInternal(EMaterialInstanceClearParam
 		ResetAllTextureParameterOverrides();
 #endif
 		TextureParameterValues.Empty();
+		TextureCollectionParameterValues.Empty();
 		RuntimeVirtualTextureParameterValues.Empty();
 		SparseVolumeTextureParameterValues.Empty();
 		FontParameterValues.Empty();
@@ -3967,7 +4044,7 @@ void UMaterialInstance::ClearParameterValuesInternal(EMaterialInstanceClearParam
 	}
 
 #if WITH_EDITOR
-	FObjectCacheEventSink::NotifyReferencedTextureChanged_Concurrent(this);
+	FObjectCacheEventSink::NotifyMaterialChanged_Concurrent(this);
 #endif
 
 	InitResources();
@@ -5129,7 +5206,7 @@ void UMaterialInstance::CopyMaterialUniformParametersInternal(UMaterialInterface
 	}
 
 #if WITH_EDITOR
-	FObjectCacheEventSink::NotifyReferencedTextureChanged_Concurrent(this);
+	FObjectCacheEventSink::NotifyMaterialChanged_Concurrent(this);
 #endif
 }
 
@@ -5192,7 +5269,7 @@ void UMaterialInstance::OverrideTextureParameterValue(const UTexture* InTextureT
 				);
 			}
 #if WITH_EDITOR
-			FObjectCacheEventSink::NotifyReferencedTextureChanged_Concurrent(this);
+			FObjectCacheEventSink::NotifyMaterialChanged_Concurrent(this);
 #endif
 		}
 
@@ -5251,3 +5328,9 @@ bool UMaterialInstance::GetTextureParameterValue(const FHashedMaterialParameterI
 
 	return bResult;
 }
+
+bool UMaterialInstance::GetTextureCollectionParameterValue(const FHashedMaterialParameterInfo& ParameterInfo, class UTextureCollection*& OutValue, bool bOveriddenOnly) const
+{
+	bool bResult = Super::GetTextureCollectionParameterValue(ParameterInfo, OutValue, bOveriddenOnly);
+	return bResult;
+}
\ No newline at end of file
diff --git a/Engine/Source/Runtime/Engine/Private/Materials/MaterialInstanceSupport.h b/Engine/Source/Runtime/Engine/Private/Materials/MaterialInstanceSupport.h
--- a/Engine/Source/Runtime/Engine/Private/Materials/MaterialInstanceSupport.h
+++ b/Engine/Source/Runtime/Engine/Private/Materials/MaterialInstanceSupport.h
@@ -239,6 +239,7 @@ public:
 		DoubleVectorParameterArray.Empty();
 		ScalarParameterArray.Empty();
 		TextureParameterArray.Empty();
+		TextureCollectionParameterArray.Empty();
 		RuntimeVirtualTextureParameterArray.Empty();
 		SparseVolumeTextureParameterArray.Empty();
 	}
@@ -309,6 +310,7 @@ private:
 	static bool IsValidParameterValue(const FLinearColor&) { return true; }
 	static bool IsValidParameterValue(const FVector4d&) { return true; }
 	static bool IsValidParameterValue(const UTexture* Value) { return Value != nullptr; }
+	static bool IsValidParameterValue(const UTextureCollection* Value) { return Value != nullptr; }
 	static bool IsValidParameterValue(const URuntimeVirtualTexture* Value) { return Value != nullptr; }
 	static bool IsValidParameterValue(const USparseVolumeTexture* Value) { return Value != nullptr; }
 
@@ -396,6 +398,8 @@ private:
 	THashedMaterialParameterMap<float> ScalarParameterArray;
 	/** Texture parameters for this material instance. */
 	THashedMaterialParameterMap<const UTexture*> TextureParameterArray;
+	/** TextureCollection parameters for this material instance. */
+	THashedMaterialParameterMap<const UTextureCollection*> TextureCollectionParameterArray;
 	/** Runtime Virtual Texture parameters for this material instance. */
 	THashedMaterialParameterMap<const URuntimeVirtualTexture*> RuntimeVirtualTextureParameterArray;
 	/** Sparse Volume Texture parameters for this material instance. */
@@ -409,6 +413,7 @@ template <> FORCEINLINE THashedMaterialParameterMap<float>& FMaterialInstanceRes
 template <> FORCEINLINE THashedMaterialParameterMap<FLinearColor>& FMaterialInstanceResource::GetValueArray() { return VectorParameterArray; }
 template <> FORCEINLINE THashedMaterialParameterMap<FVector4d>& FMaterialInstanceResource::GetValueArray() { return DoubleVectorParameterArray; }
 template <> FORCEINLINE THashedMaterialParameterMap<const UTexture*>& FMaterialInstanceResource::GetValueArray() { return TextureParameterArray; }
+template <> FORCEINLINE THashedMaterialParameterMap<const UTextureCollection*>& FMaterialInstanceResource::GetValueArray() { return TextureCollectionParameterArray; }
 template <> FORCEINLINE THashedMaterialParameterMap<const URuntimeVirtualTexture*>& FMaterialInstanceResource::GetValueArray() { return RuntimeVirtualTextureParameterArray; }
 template <> FORCEINLINE THashedMaterialParameterMap<const USparseVolumeTexture*>& FMaterialInstanceResource::GetValueArray() { return SparseVolumeTextureParameterArray; }
 template <> FORCEINLINE const THashedMaterialParameterMap<bool>& FMaterialInstanceResource::GetValueArray() const { return StaticSwitchParameterArray; }
@@ -416,6 +421,7 @@ template <> FORCEINLINE const THashedMaterialParameterMap<float>& FMaterialInsta
 template <> FORCEINLINE const THashedMaterialParameterMap<FLinearColor>& FMaterialInstanceResource::GetValueArray() const { return VectorParameterArray; }
 template <> FORCEINLINE const THashedMaterialParameterMap<FVector4d>& FMaterialInstanceResource::GetValueArray() const { return DoubleVectorParameterArray; }
 template <> FORCEINLINE const THashedMaterialParameterMap<const UTexture*>& FMaterialInstanceResource::GetValueArray() const { return TextureParameterArray; }
+template <> FORCEINLINE const THashedMaterialParameterMap<const UTextureCollection*>& FMaterialInstanceResource::GetValueArray() const { return TextureCollectionParameterArray; }
 template <> FORCEINLINE const THashedMaterialParameterMap<const URuntimeVirtualTexture*>& FMaterialInstanceResource::GetValueArray() const { return RuntimeVirtualTextureParameterArray; }
 template <> FORCEINLINE const THashedMaterialParameterMap<const USparseVolumeTexture*>& FMaterialInstanceResource::GetValueArray() const { return SparseVolumeTextureParameterArray; }
 
@@ -426,6 +432,7 @@ struct FMaterialInstanceParameterSet
 	TArray<THashedMaterialParameterMap<FLinearColor>::TNamedParameter>					VectorParameters;
 	TArray<THashedMaterialParameterMap<FVector4d>::TNamedParameter>						DoubleVectorParameters;
 	TArray<THashedMaterialParameterMap<const UTexture*>::TNamedParameter>				TextureParameters;
+	TArray<THashedMaterialParameterMap<const UTextureCollection*>::TNamedParameter>		TextureCollectionParameters;
 	TArray<THashedMaterialParameterMap<const URuntimeVirtualTexture*>::TNamedParameter>	RuntimeVirtualTextureParameters;
 	TArray<THashedMaterialParameterMap<const USparseVolumeTexture*>::TNamedParameter>	SparseVolumeTextureParameters;
 };
diff --git a/Engine/Source/Runtime/Engine/Private/Materials/MaterialInterface.cpp b/Engine/Source/Runtime/Engine/Private/Materials/MaterialInterface.cpp
--- a/Engine/Source/Runtime/Engine/Private/Materials/MaterialInterface.cpp
+++ b/Engine/Source/Runtime/Engine/Private/Materials/MaterialInterface.cpp
@@ -226,7 +226,7 @@ void UMaterialInterface::Serialize(FArchive& Ar)
 		Struct->SerializeTaggedProperties(Ar, (uint8*)CachedExpressionData.Get(), Struct, nullptr);
 
 #if WITH_EDITOR
-		FObjectCacheEventSink::NotifyReferencedTextureChanged_Concurrent(this);
+		FObjectCacheEventSink::NotifyMaterialChanged_Concurrent(this);
 #endif
 	}
     // we don't consider bLoadedCachedExpressionData here because in editor we call UpdateCachedExpressionData which can
@@ -370,6 +370,11 @@ TArrayView<const TObjectPtr<UObject>> UMaterialInterface::GetReferencedTextures(
 	return GetCachedExpressionData().ReferencedTextures;
 }
 
+TConstArrayView<TObjectPtr<UTextureCollection>> UMaterialInterface::GetReferencedTextureCollections() const
+{
+	return GetCachedExpressionData().ReferencedTextureCollections;
+}
+
 #if WITH_EDITOR
 void UMaterialInterface::GetReferencedTexturesAndOverrides(TSet<const UTexture*>& InOutTextures) const
 {
@@ -955,6 +960,17 @@ bool UMaterialInterface::GetTextureParameterValue(const FHashedMaterialParameter
 	return false;
 }
 
+bool UMaterialInterface::GetTextureCollectionParameterValue(const FHashedMaterialParameterInfo& ParameterInfo, UTextureCollection*& OutValue, bool bOveriddenOnly) const
+{
+	FMaterialParameterMetadata Result;
+	if (GetParameterValue(EMaterialParameterType::TextureCollection, ParameterInfo, Result, MakeParameterValueFlags(bOveriddenOnly)))
+	{
+		OutValue = Result.Value.TextureCollection;
+		return true;
+	}
+	return false;
+}
+
 bool UMaterialInterface::GetRuntimeVirtualTextureParameterValue(const FHashedMaterialParameterInfo& ParameterInfo, URuntimeVirtualTexture*& OutValue, bool bOveriddenOnly) const
 {
 	FMaterialParameterMetadata Result;
@@ -1052,6 +1068,17 @@ bool UMaterialInterface::GetTextureParameterDefaultValue(const FHashedMaterialPa
 	return false;
 }
 
+bool UMaterialInterface::GetTextureCollectionParameterDefaultValue(const FHashedMaterialParameterInfo& ParameterInfo, class UTextureCollection*& OutValue) const
+{
+	FMaterialParameterMetadata Result;
+	if (GetParameterDefaultValue(EMaterialParameterType::TextureCollection, ParameterInfo, Result))
+	{
+		OutValue = Result.Value.TextureCollection;
+		return true;
+	}
+	return false;
+}
+
 bool UMaterialInterface::GetRuntimeVirtualTextureParameterDefaultValue(const FHashedMaterialParameterInfo& ParameterInfo, class URuntimeVirtualTexture*& OutValue) const
 {
 	FMaterialParameterMetadata Result;
@@ -1150,6 +1177,11 @@ void UMaterialInterface::GetAllTextureParameterInfo(TArray<FMaterialParameterInf
 	GetAllParameterInfoOfType(EMaterialParameterType::Texture, OutParameterInfo, OutParameterIds);
 }
 
+void UMaterialInterface::GetAllTextureCollectionParameterInfo(TArray<FMaterialParameterInfo>& OutParameterInfo, TArray<FGuid>& OutParameterIds) const
+{
+	GetAllParameterInfoOfType(EMaterialParameterType::TextureCollection, OutParameterInfo, OutParameterIds);
+}
+
 void UMaterialInterface::GetAllRuntimeVirtualTextureParameterInfo(TArray<FMaterialParameterInfo>& OutParameterInfo, TArray<FGuid>& OutParameterIds) const
 {
 	GetAllParameterInfoOfType(EMaterialParameterType::RuntimeVirtualTexture, OutParameterInfo, OutParameterIds);
diff --git a/Engine/Source/Runtime/Engine/Private/Materials/MaterialRenderProxy.cpp b/Engine/Source/Runtime/Engine/Private/Materials/MaterialRenderProxy.cpp
--- a/Engine/Source/Runtime/Engine/Private/Materials/MaterialRenderProxy.cpp
+++ b/Engine/Source/Runtime/Engine/Private/Materials/MaterialRenderProxy.cpp
@@ -239,6 +239,17 @@ bool FMaterialRenderProxy::GetTextureValue(const FHashedMaterialParameterInfo& P
 	return false;
 }
 
+bool FMaterialRenderProxy::GetTextureCollectionValue(const FHashedMaterialParameterInfo& ParameterInfo, const UTextureCollection** OutValue, const FMaterialRenderContext& Context) const
+{
+	FMaterialParameterValue Value;
+	if (GetParameterValue(EMaterialParameterType::TextureCollection, ParameterInfo, Value, Context))
+	{
+		*OutValue = Value.TextureCollection;
+		return true;
+	}
+	return false;
+}
+
 static void OnVirtualTextureDestroyedCB(const FVirtualTextureProducerHandle& InHandle, void* Baton)
 {
 	FMaterialRenderProxy* MaterialProxy = static_cast<FMaterialRenderProxy*>(Baton);
diff --git a/Engine/Source/Runtime/Engine/Private/Materials/MaterialShared.cpp b/Engine/Source/Runtime/Engine/Private/Materials/MaterialShared.cpp
--- a/Engine/Source/Runtime/Engine/Private/Materials/MaterialShared.cpp
+++ b/Engine/Source/Runtime/Engine/Private/Materials/MaterialShared.cpp
@@ -98,6 +98,7 @@ IMPLEMENT_TYPE_LAYOUT(FMaterialUniformPreshaderHeader);
 IMPLEMENT_TYPE_LAYOUT(FMaterialUniformPreshaderField);
 IMPLEMENT_TYPE_LAYOUT(FMaterialNumericParameterInfo);
 IMPLEMENT_TYPE_LAYOUT(FMaterialTextureParameterInfo);
+IMPLEMENT_TYPE_LAYOUT(FMaterialTextureCollectionParameterInfo);
 IMPLEMENT_TYPE_LAYOUT(FMaterialExternalTextureParameterInfo);
 IMPLEMENT_TYPE_LAYOUT(FMaterialVirtualTextureStack);
 
@@ -1615,6 +1616,25 @@ TArrayView<const TObjectPtr<UObject>> FMaterialResource::GetReferencedTextures()
 	return UMaterial::GetDefaultMaterial(MD_Surface)->GetReferencedTextures();
 }
 
+TConstArrayView<TObjectPtr<UTextureCollection>> FMaterialResource::GetReferencedTextureCollections() const
+{
+	if (MaterialInstance)
+	{
+		TConstArrayView<TObjectPtr<UTextureCollection>> TextureCollections = MaterialInstance->GetReferencedTextureCollections();
+		if (TextureCollections.Num())
+		{
+			return TextureCollections;
+		}
+	}
+
+	if (Material)
+	{
+		return Material->GetReferencedTextureCollections();
+	}
+
+	return UMaterial::GetDefaultMaterial(MD_Surface)->GetReferencedTextureCollections();
+}
+
 void FMaterialResource::AddReferencedObjects(FReferenceCollector& Collector)
 {
 	FMaterial::AddReferencedObjects(Collector);
@@ -5314,6 +5334,7 @@ UE::Shader::FValue FMaterialParameterValue::AsShaderValue() const
 	case EMaterialParameterType::StaticSwitch: return Bool[0];
 	case EMaterialParameterType::StaticComponentMask: return UE::Shader::FValue(Bool[0], Bool[1], Bool[2], Bool[3]);
 	case EMaterialParameterType::Texture:
+	case EMaterialParameterType::TextureCollection:
 	case EMaterialParameterType::Font:
 	case EMaterialParameterType::RuntimeVirtualTexture:
 	case EMaterialParameterType::SparseVolumeTexture:
@@ -5355,6 +5376,7 @@ UE::Shader::FType GetShaderValueType(EMaterialParameterType Type)
 	case EMaterialParameterType::StaticSwitch: return UE::Shader::EValueType::Bool1;
 	case EMaterialParameterType::StaticComponentMask: return UE::Shader::EValueType::Bool4;
 	case EMaterialParameterType::Texture:
+	case EMaterialParameterType::TextureCollection:
 	case EMaterialParameterType::RuntimeVirtualTexture:
 	case EMaterialParameterType::Font:
 		return FMaterialTextureValue::GetTypeName();
diff --git a/Engine/Source/Runtime/Engine/Private/Materials/MaterialUniformExpressions.cpp b/Engine/Source/Runtime/Engine/Private/Materials/MaterialUniformExpressions.cpp
--- a/Engine/Source/Runtime/Engine/Private/Materials/MaterialUniformExpressions.cpp
+++ b/Engine/Source/Runtime/Engine/Private/Materials/MaterialUniformExpressions.cpp
@@ -6,6 +6,7 @@
 
 #include "Materials/MaterialUniformExpressions.h"
 #include "Engine/Texture.h"
+#include "Engine/TextureCollection.h"
 #include "SceneManagement.h"
 #include "Materials/MaterialInterface.h"
 #include "Materials/MaterialParameterCollection.h"
@@ -143,6 +144,14 @@ void FMaterialRenderContext::GetTextureParameterValue(const FHashedMaterialParam
 	}
 }
 
+void FMaterialRenderContext::GetTextureCollectionParameterValue(const FHashedMaterialParameterInfo& ParameterInfo, int32 TextureCollectionIndex, const UTextureCollection*& OutValue) const
+{
+	if (ParameterInfo.Name.IsNone() || !MaterialRenderProxy || !MaterialRenderProxy->GetTextureCollectionValue(ParameterInfo, &OutValue, *this))
+	{
+		OutValue = GetIndexedTextureCollection(Material, TextureCollectionIndex);
+	}
+}
+
 FMaterialUniformExpressionType::FMaterialUniformExpressionType(const TCHAR* InName)
 	: Name(InName)
 {
@@ -240,6 +249,7 @@ bool FUniformExpressionSet::IsEmpty() const
 	return UniformNumericParameters.Num() == 0
 		&& UniformPreshaders.Num() == 0
 		&& UniformExternalTextureParameters.Num() == 0
+		&& UniformTextureCollectionParameters.Num() == 0
 		&& VTStacks.Num() == 0
 		&& ParameterCollections.Num() == 0;
 }
@@ -257,6 +267,7 @@ bool FUniformExpressionSet::operator==(const FUniformExpressionSet& ReferenceSet
 	if (UniformNumericParameters.Num() != ReferenceSet.UniformNumericParameters.Num()
 		|| UniformPreshaders.Num() != ReferenceSet.UniformPreshaders.Num()
 		|| UniformExternalTextureParameters.Num() != ReferenceSet.UniformExternalTextureParameters.Num()
+		|| UniformTextureCollectionParameters.Num() != ReferenceSet.UniformTextureCollectionParameters.Num()
 		|| VTStacks.Num() != ReferenceSet.VTStacks.Num()
 		|| ParameterCollections.Num() != ReferenceSet.ParameterCollections.Num())
 	{
@@ -298,6 +309,15 @@ bool FUniformExpressionSet::operator==(const FUniformExpressionSet& ReferenceSet
 		}
 	}
 
+
+	for (int32 i = 0; i < UniformTextureCollectionParameters.Num(); i++)
+	{
+		if (UniformTextureCollectionParameters[i] != ReferenceSet.UniformTextureCollectionParameters[i])
+		{
+			return false;
+		}
+	}
+
 	for (int32 i = 0; i < VTStacks.Num(); i++)
 	{
 		if (VTStacks[i] != ReferenceSet.VTStacks[i])
@@ -324,7 +344,7 @@ bool FUniformExpressionSet::operator==(const FUniformExpressionSet& ReferenceSet
 
 FString FUniformExpressionSet::GetSummaryString() const
 {
-	return FString::Printf(TEXT("(%u preshaders, %u 2d tex, %u cube tex, %u 2darray tex, %u cubearray tex, %u 3d tex, %u virtual tex, %u sparse volume tex, %u external tex, %u VT stacks, %u collections)"),
+	return FString::Printf(TEXT("(%u preshaders, %u 2d tex, %u cube tex, %u 2darray tex, %u cubearray tex, %u 3d tex, %u virtual tex, %u sparse volume tex, %u external tex, %u tex collections, %u VT stacks, %u collections)"),
 		UniformPreshaders.Num(),
 		UniformTextureParameters[(uint32)EMaterialTextureParameterType::Standard2D].Num(),
 		UniformTextureParameters[(uint32)EMaterialTextureParameterType::Cube].Num(),
@@ -334,6 +354,7 @@ FString FUniformExpressionSet::GetSummaryString() const
 		UniformTextureParameters[(uint32)EMaterialTextureParameterType::Virtual].Num(),
 		UniformTextureParameters[(uint32)EMaterialTextureParameterType::SparseVolume].Num(),
 		UniformExternalTextureParameters.Num(),
+		UniformTextureCollectionParameters.Num(),
 		VTStacks.Num(),
 		ParameterCollections.Num()
 		);
@@ -397,6 +418,7 @@ FShaderParametersMetadata* FUniformExpressionSet::CreateBufferStruct()
 	static FString VolumeTextureNames[128];
 	static FString VolumeTextureSamplerNames[128];
 	static FString ExternalTextureNames[128];
+	static FString TextureCollectionNames[128];
 	static FString MediaTextureSamplerNames[128];
 	static FString VirtualTexturePageTableNames0[128];
 	static FString VirtualTexturePageTableNames1[128];
@@ -424,6 +446,7 @@ FShaderParametersMetadata* FUniformExpressionSet::CreateBufferStruct()
 			VolumeTextureNames[i] = FString::Printf(TEXT("VolumeTexture_%d"), i);
 			VolumeTextureSamplerNames[i] = FString::Printf(TEXT("VolumeTexture_%dSampler"), i);
 			ExternalTextureNames[i] = FString::Printf(TEXT("ExternalTexture_%d"), i);
+			TextureCollectionNames[i] = FString::Printf(TEXT("TextureCollection_%d"), i);
 			MediaTextureSamplerNames[i] = FString::Printf(TEXT("ExternalTexture_%dSampler"), i);
 			VirtualTexturePageTableNames0[i] = FString::Printf(TEXT("VirtualTexturePageTable0_%d"), i);
 			VirtualTexturePageTableNames1[i] = FString::Printf(TEXT("VirtualTexturePageTable1_%d"), i);
@@ -524,6 +547,13 @@ FShaderParametersMetadata* FUniformExpressionSet::CreateBufferStruct()
 		NextMemberOffset += SHADER_PARAMETER_POINTER_ALIGNMENT;
 	}
 
+	for (int32 i = 0; i < UniformTextureCollectionParameters.Num(); ++i)
+	{
+		check((NextMemberOffset % SHADER_PARAMETER_POINTER_ALIGNMENT) == 0);
+		new(Members) FShaderParametersMetadata::FMember(*TextureCollectionNames[i], TEXT("FResourceCollection"), __LINE__, NextMemberOffset, UBMT_RESOURCE_COLLECTION, EShaderPrecisionModifier::Float, 1, 1, 0, NULL);
+		NextMemberOffset += SHADER_PARAMETER_POINTER_ALIGNMENT;
+	}
+
 	for (int32 i = 0; i < VTStacks.Num(); ++i)
 	{
 		const FMaterialVirtualTextureStack& Stack = VTStacks[i];
@@ -617,6 +647,19 @@ int32 FUniformExpressionSet::FindOrAddTextureParameter(EMaterialTextureParameter
 	return UniformTextureParameters[(int32)Type].Add(Info);
 }
 
+int32 FUniformExpressionSet::FindOrAddTextureCollectionParameter(const FMaterialTextureCollectionParameterInfo& Info)
+{
+	for (int32 Index = 0; Index < UniformTextureCollectionParameters.Num(); Index++)
+	{
+		if (UniformTextureCollectionParameters[Index] == Info)
+		{
+			return Index;
+		}
+	}
+
+	return UniformTextureCollectionParameters.Add(Info);
+}
+
 int32 FUniformExpressionSet::FindOrAddExternalTextureParameter(const FMaterialExternalTextureParameterInfo& Info)
 {
 	for (int32 i = 0; i < UniformExternalTextureParameters.Num(); ++i)
@@ -998,6 +1041,7 @@ void FUniformExpressionSet::FillUniformBuffer(const FMaterialRenderContext& Mate
 				+ UniformTextureParameters[(uint32)EMaterialTextureParameterType::Volume].Num() * 2
 				+ UniformTextureParameters[(uint32)EMaterialTextureParameterType::SparseVolume].Num() * 4
 				+ UniformExternalTextureParameters.Num() * 2
+				+ UniformTextureCollectionParameters.Num()
 				+ UniformTextureParameters[(uint32)EMaterialTextureParameterType::Virtual].Num() * 2
 				+ NumPageTableTextures
 				+ NumPageTableIndirectionTextures
@@ -1333,6 +1377,28 @@ void FUniformExpressionSet::FillUniformBuffer(const FMaterialRenderContext& Mate
 			}
 		}
 
+		for (const FMaterialTextureCollectionParameterInfo& ParameterInfo : UniformTextureCollectionParameters)
+		{
+			void** ResourceTableBufferPtr = reinterpret_cast<void**>(BufferCursor);
+			BufferCursor = reinterpret_cast<uint8*>(BufferCursor) + SHADER_PARAMETER_POINTER_ALIGNMENT;
+			check(BufferCursor <= TempBuffer + TempBufferSize);
+
+			const UTextureCollection* TextureCollection = nullptr;
+			ParameterInfo.GetTextureCollection(MaterialRenderContext, TextureCollection);
+
+			if (TextureCollection != nullptr)
+			{
+				FTextureCollectionResource* Resource = TextureCollection->GetResource();
+				check(Resource);
+				*ResourceTableBufferPtr = Resource->GetRHI();
+			}
+			else
+			{
+				// TODO: christopher.waters - implement default texture collections
+				unimplemented();
+			}
+		}
+
 		// Cache virtual texture page table uniform expressions.
 		for (int32 VTStackIndex = 0; VTStackIndex < VTStacks.Num(); ++VTStackIndex)
 		{
@@ -1571,6 +1637,71 @@ bool FMaterialUniformExpressionTexture::IsIdentical(const FMaterialUniformExpres
 		bVirtualTexture == OtherTextureExpression->bVirtualTexture;
 }
 
+// FMaterialUniformExpressionTextureCollection
+
+FMaterialUniformExpressionTextureCollection::FMaterialUniformExpressionTextureCollection(int32 InTextureCollectionIndex)
+	: TextureCollectionIndex(InTextureCollectionIndex)
+{
+}
+
+bool FMaterialUniformExpressionTextureCollection::IsIdentical(const FMaterialUniformExpression* OtherExpression) const
+{
+	if (GetType() != OtherExpression->GetType())
+	{
+		return false;
+	}
+
+	const FMaterialUniformExpressionTextureCollection* OtherTextureCollectionExpression = static_cast<const FMaterialUniformExpressionTextureCollection*>(OtherExpression);
+	return TextureCollectionIndex == OtherTextureCollectionExpression->TextureCollectionIndex;
+}
+
+FMaterialUniformExpressionTextureCollectionParameter* FMaterialUniformExpressionTextureCollection::GetTextureCollectionParameterUniformExpression()
+{
+	return nullptr;
+}
+
+void FMaterialUniformExpressionTextureCollection::GetTextureCollectionParameterInfo(FMaterialTextureCollectionParameterInfo& OutParameter) const
+{
+	OutParameter.TextureCollectionIndex = this->TextureCollectionIndex;
+}
+
+// FMaterialUniformExpressionTextureCollectionParameter
+
+FMaterialUniformExpressionTextureCollectionParameter::FMaterialUniformExpressionTextureCollectionParameter(const FMaterialParameterInfo& InParameterInfo, int32 InTextureCollectionIndex)
+	: Super(InTextureCollectionIndex)
+	, ParameterInfo(InParameterInfo)
+{
+}
+
+bool FMaterialUniformExpressionTextureCollectionParameter::IsIdentical(const FMaterialUniformExpression* OtherExpression) const
+{
+	if (GetType() != OtherExpression->GetType())
+	{
+		return false;
+	}
+
+	const FMaterialUniformExpressionTextureCollectionParameter* OtherParameter = static_cast<const FMaterialUniformExpressionTextureCollectionParameter*>(OtherExpression);
+	return ParameterInfo == OtherParameter->ParameterInfo && Super::IsIdentical(OtherParameter);
+}
+
+bool FMaterialUniformExpressionTextureCollectionParameter::IsConstant() const
+{
+	return false;
+}
+
+FMaterialUniformExpressionTextureCollectionParameter* FMaterialUniformExpressionTextureCollectionParameter::GetTextureCollectionParameterUniformExpression()
+{
+	return this;
+}
+
+void FMaterialUniformExpressionTextureCollectionParameter::GetTextureCollectionParameterInfo(FMaterialTextureCollectionParameterInfo& OutParameter) const
+{
+	Super::GetTextureCollectionParameterInfo(OutParameter);
+	OutParameter.ParameterInfo = ParameterInfo;
+}
+
+// FMaterialUniformExpressionExternalTextureBase
+
 FMaterialUniformExpressionExternalTextureBase::FMaterialUniformExpressionExternalTextureBase(int32 InSourceTextureIndex)
 	: SourceTextureIndex(InSourceTextureIndex)
 {}
@@ -1636,6 +1767,11 @@ void FMaterialTextureParameterInfo::GetGameThreadTextureValue(const UMaterialInt
 	}
 }
 
+void FMaterialTextureCollectionParameterInfo::GetTextureCollection(const FMaterialRenderContext& Context, const UTextureCollection*& OutTextureCollection) const
+{
+	Context.GetTextureCollectionParameterValue(ParameterInfo, TextureCollectionIndex, OutTextureCollection);
+}
+
 void FMaterialTextureParameterInfo::GetGameThreadTextureValue(const UMaterialInterface* MaterialInterface, const FMaterial& Material, USparseVolumeTexture*& OutValue) const
 {
 	if (ParameterInfo.Name.IsNone() || !MaterialInterface->GetSparseVolumeTextureParameterValue(ParameterInfo, OutValue))
@@ -1773,6 +1909,8 @@ class FMaterialUniformExpressionRuntimeVirtualTextureParameter_DEPRECATED : publ
 };
 
 IMPLEMENT_MATERIALUNIFORMEXPRESSION_TYPE(FMaterialUniformExpressionTexture);
+IMPLEMENT_MATERIALUNIFORMEXPRESSION_TYPE(FMaterialUniformExpressionTextureCollection);
+IMPLEMENT_MATERIALUNIFORMEXPRESSION_TYPE(FMaterialUniformExpressionTextureCollectionParameter);
 IMPLEMENT_MATERIALUNIFORMEXPRESSION_TYPE(FMaterialUniformExpressionConstant);
 IMPLEMENT_MATERIALUNIFORMEXPRESSION_TYPE(FMaterialUniformExpressionGenericConstant);
 IMPLEMENT_MATERIALUNIFORMEXPRESSION_TYPE(FMaterialUniformExpressionNumericParameter);
diff --git a/Engine/Source/Runtime/Engine/Private/Materials/MaterialUniformExpressions.h b/Engine/Source/Runtime/Engine/Private/Materials/MaterialUniformExpressions.h
--- a/Engine/Source/Runtime/Engine/Private/Materials/MaterialUniformExpressions.h
+++ b/Engine/Source/Runtime/Engine/Private/Materials/MaterialUniformExpressions.h
@@ -61,6 +61,7 @@ public:
 	virtual FMaterialUniformExpressionType* GetType() const = 0;
 	virtual class FMaterialUniformExpressionTexture* GetTextureUniformExpression() { return nullptr; }
 	virtual class FMaterialUniformExpressionExternalTexture* GetExternalTextureUniformExpression() { return nullptr; }
+	virtual class FMaterialUniformExpressionTextureCollection* GetTextureCollectionUniformExpression() { return nullptr; }
 	virtual bool IsConstant() const { return false; }
 	virtual bool IsIdentical(const FMaterialUniformExpression* OtherExpression) const { return false; }
 
@@ -130,6 +131,63 @@ protected:
 	bool bVirtualTexture;
 };
 
+class FMaterialUniformExpressionTextureCollection : public FMaterialUniformExpression
+{
+	DECLARE_MATERIALUNIFORMEXPRESSION_TYPE(FMaterialUniformExpressionTextureCollection);
+
+public:
+	FMaterialUniformExpressionTextureCollection(int32 InTextureCollectionIndex);
+
+	//~ Begin FMaterialUniformExpression Interface.
+	virtual class FMaterialUniformExpressionTextureCollection* GetTextureCollectionUniformExpression() final { return this; }
+	virtual bool IsIdentical(const FMaterialUniformExpression* OtherExpression) const override;
+	//~ End FMaterialUniformExpression Interface.
+
+	virtual class FMaterialUniformExpressionTextureCollectionParameter* GetTextureCollectionParameterUniformExpression();
+	virtual void GetTextureCollectionParameterInfo(FMaterialTextureCollectionParameterInfo& OutParameter) const;
+
+	int32 GetTextureCollectionIndex() const
+	{
+		return TextureCollectionIndex;
+	}
+
+protected:
+	int32 TextureCollectionIndex;
+};
+
+class FMaterialUniformExpressionTextureCollectionParameter : public FMaterialUniformExpressionTextureCollection
+{
+	using Super = FMaterialUniformExpressionTextureCollection;
+	DECLARE_MATERIALUNIFORMEXPRESSION_TYPE(FMaterialUniformExpressionTextureCollectionParameter);
+
+public:
+	FMaterialUniformExpressionTextureCollectionParameter(const FMaterialParameterInfo& InParameterInfo, int32 InTextureCollectionIndex);
+
+	//~ Begin FMaterialUniformExpression Interface.
+	virtual bool IsIdentical(const FMaterialUniformExpression* OtherExpression) const override;
+	virtual bool IsConstant() const override;
+	//~ End FMaterialUniformExpression Interface.
+
+	//~ Begin FMaterialUniformExpressionTextureCollection Interface.
+	virtual FMaterialUniformExpressionTextureCollectionParameter* GetTextureCollectionParameterUniformExpression() final;
+	virtual void GetTextureCollectionParameterInfo(FMaterialTextureCollectionParameterInfo& OutParameter) const override;
+	//~ End FMaterialUniformExpressionTextureCollection Interface.
+
+	FName GetParameterName() const
+	{
+		return ParameterInfo.GetName();
+	}
+
+	const FHashedMaterialParameterInfo& GetParameterInfo() const
+	{
+		return ParameterInfo;
+	}
+
+private:
+	FHashedMaterialParameterInfo ParameterInfo;
+	int32 ParameterIndex;
+};
+
 class FMaterialUniformExpressionExternalTextureBase : public FMaterialUniformExpression
 {
 	DECLARE_MATERIALUNIFORMEXPRESSION_TYPE(FMaterialUniformExpressionExternalTextureBase);
@@ -393,6 +451,28 @@ static TextureType* GetIndexedTexture(const FMaterial& Material, int32 TextureIn
 	return Cast<TextureType>(IndexedTexture);
 }
 
+inline UTextureCollection* GetIndexedTextureCollection(const FMaterial& Material, int32 TextureCollectionIndex)
+{
+	UTextureCollection* IndexedTextureCollection = nullptr;
+	const TArrayView<const TObjectPtr<UTextureCollection>> ReferencedTextureCollections = Material.GetReferencedTextureCollections();
+	if (ReferencedTextureCollections.IsValidIndex(TextureCollectionIndex))
+	{
+		IndexedTextureCollection = ReferencedTextureCollections[TextureCollectionIndex];
+	}
+
+	if (IndexedTextureCollection == nullptr)
+	{
+		static bool bWarnedOnce = false;
+		if (!bWarnedOnce)
+		{
+			UE_LOG(LogMaterial, Warning, TEXT("%s: Requesting an invalid TextureIndex! (%u / %u)"), *Material.GetFriendlyName(), TextureCollectionIndex, ReferencedTextureCollections.Num());
+			bWarnedOnce = true;
+		}
+	}
+
+	return IndexedTextureCollection;
+}
+
 /**
  * A texture parameter expression.
  */
diff --git a/Engine/Source/Runtime/Engine/Private/ObjectCacheContext.cpp b/Engine/Source/Runtime/Engine/Private/ObjectCacheContext.cpp
--- a/Engine/Source/Runtime/Engine/Private/ObjectCacheContext.cpp
+++ b/Engine/Source/Runtime/Engine/Private/ObjectCacheContext.cpp
@@ -8,6 +8,7 @@
 #include "Materials/MaterialInterface.h"
 #include "Engine/StaticMesh.h"
 #include "Engine/Texture.h"
+#include "Engine/TextureCollection.h"
 #include "Components/StaticMeshComponent.h"
 #include "Components/SkinnedMeshComponent.h"
 #include "Containers/Queue.h"
@@ -79,6 +80,7 @@ enum class EObjectReverseLookupMask : int32
 	MaterialToPrimitiveLookupCache = 4,
 	StaticMeshToComponentLookupCache = 8,
 	MaterialToMaterialsLookupCache = 16,
+	TextureCollectionToMaterialLookupCache = 32,
 	All = -1
 };
 
@@ -367,11 +369,14 @@ private:
 
 FObjectReverseLookupCache<UMaterialInterface, UMaterialInstance>   GMaterialToMaterialInstanceLookupCache; // Parent -> Children
 FObjectReverseLookupCache<UTexture, UMaterialInterface>            GTextureToMaterialLookupCache;
+FObjectReverseLookupCache<UTextureCollection, UMaterialInterface>  GTextureCollectionToMaterialLookupCache;
 FObjectReverseLookupCache<UTexture, UTexture>                      GTextureToTextureLookupCache; // CompositeTexture -> Texture
 FObjectReverseLookupCache<UStaticMesh, IStaticMeshComponent>       GStaticMeshToComponentLookupCache;
 FObjectReverseLookupCache<UMaterialInterface, IPrimitiveComponent> GMaterialToPrimitiveLookupCache;
 
 void GetReferencedTextures(UMaterialInterface* MaterialInterface, TSet<UTexture*>& OutReferencedTextures);
+void GetReferencedTextureCollections(UMaterialInterface* MaterialInterface, TSet<UTextureCollection*>& OutReferencedTextureCollections);
+
 void Validate()
 {
 	// Scan and compare UStaticMesh -> UStaticMeshComponent
@@ -400,6 +405,18 @@ void Validate()
 		ErrorCount += TempLookup.Compare(GTextureToMaterialLookupCache);
 	}
 
+	// Scan and compare UTextureCollection -> UMaterialInterface
+	{
+		FObjectReverseLookupCache<UTextureCollection, UMaterialInterface> TempLookup;
+		for (TObjectIterator<UMaterialInterface> It(RF_ClassDefaultObject, true /*bIncludeDerivedClasses*/, GetObjectCacheInternalFlagsExclusion()); It; ++It)
+		{
+			TSet<UTextureCollection*> ReferencedTextureCollections;
+			GetReferencedTextureCollections(*It, ReferencedTextureCollections);
+			TempLookup.Update(*It, ReferencedTextureCollections.Array());
+		}
+		ErrorCount += TempLookup.Compare(GTextureCollectionToMaterialLookupCache);
+	}
+
 	// Scan and compare UTexture -> UTexture
 	{
 		FObjectReverseLookupCache<UTexture, UTexture> TempLookup;
@@ -480,6 +497,32 @@ namespace ObjectCacheContextImpl {
 		}
 	}
 
+	void GetReferencedTextureCollections(UMaterialInterface* MaterialInterface, TSet<UTextureCollection*>& OutReferencedTextureCollections)
+	{
+		TSet<UTextureCollection*> TextureCollections;
+		for (UTextureCollection* TextureCollection : MaterialInterface->GetReferencedTextureCollections())
+		{
+			if (TextureCollection)
+			{
+				OutReferencedTextureCollections.FindOrAdd(TextureCollection);
+			}
+		}
+
+		UMaterialInstance* MaterialInstance = Cast<UMaterialInstance>(MaterialInterface);
+		while (MaterialInstance)
+		{
+			for (const FTextureCollectionParameterValue& TextureCollectionParam : MaterialInstance->TextureCollectionParameterValues)
+			{
+				if (TextureCollectionParam.ParameterValue)
+				{
+					OutReferencedTextureCollections.FindOrAdd(TextureCollectionParam.ParameterValue);
+				}
+			}
+
+			MaterialInstance = Cast<UMaterialInstance>(MaterialInstance->Parent);
+		}
+	}
+
 	// This is the old method of finding dependent materials by scanning all MaterialInstance currently reachable
 	// Adapted from FMaterialUpdateContext::~FMaterialUpdateContext() / MaterialShared.cpp
 	void GetMaterialsAffectedByMaterials_Iteration(TArrayView<UMaterialInterface*> InMaterials, TSet<UMaterialInterface*>& OutAffectedMaterials)
@@ -697,6 +740,43 @@ TObjectCacheIterator<UTexture> FObjectCacheContext::GetUsedTextures(UMaterialInt
 	return TObjectCacheIterator<UTexture>(Textures->Array());
 }
 
+TObjectCacheIterator<UTextureCollection> FObjectCacheContext::GetUsedTextureCollections(UMaterialInterface* MaterialInterface)
+{
+	using namespace ObjectCacheContextImpl;
+#if WITH_EDITOR
+	if (ShouldUseLookupTable(EObjectReverseLookupMask::TextureCollectionToMaterialLookupCache))
+	{
+		return GTextureCollectionToMaterialLookupCache.GetTo(MaterialInterface);
+	}
+#endif
+
+	TSet<UTextureCollection*>* TextureCollections = MaterialUsedTextureCollections.Find(MaterialInterface);
+	if (TextureCollections == nullptr)
+	{
+		TextureCollections = &MaterialUsedTextureCollections.Add(MaterialInterface);
+		ObjectCacheContextImpl::GetReferencedTextureCollections(MaterialInterface, *TextureCollections);
+	}
+
+#if WITH_EDITOR
+	if (IsReverseLookupComparisonActive(EObjectReverseLookupMask::TextureCollectionToMaterialLookupCache))
+	{
+		TSet<UTextureCollection*> LookupResult;
+		for (UTextureCollection* TextureCollection : GTextureCollectionToMaterialLookupCache.GetTo(MaterialInterface))
+		{
+			LookupResult.FindOrAdd(TextureCollection);
+			checkf(TextureCollections->Contains(TextureCollection), TEXT("Permanent Object Cache has an additional texture collection %s on material %s"), *TextureCollection->GetFullName(), *MaterialInterface->GetFullName());
+		}
+
+		for (UTextureCollection* TextureCollection : *TextureCollections)
+		{
+			checkf(LookupResult.Contains(TextureCollection), TEXT("Permanent Object Cache is missing a texture collection %s on material %s"), *TextureCollection->GetFullName(), *MaterialInterface->GetFullName());
+		}
+	}
+#endif
+
+	return TObjectCacheIterator<UTextureCollection>(TextureCollections->Array());
+}
+
 TObjectCacheIterator<USkinnedMeshComponent> FObjectCacheContext::GetSkinnedMeshComponents()
 {
 	using namespace ObjectCacheContextImpl;
@@ -864,6 +944,57 @@ TObjectCacheIterator<UMaterialInterface> FObjectCacheContext::GetMaterialsAffect
 	return TObjectCacheIterator<UMaterialInterface>(ComputeResult.Array());
 }
 
+TObjectCacheIterator<UMaterialInterface> FObjectCacheContext::GetMaterialsAffectedByTextureCollection(UTextureCollection* InTextureCollection)
+{
+	using namespace ObjectCacheContextImpl;
+#if WITH_EDITOR
+	if (ShouldUseLookupTable(EObjectReverseLookupMask::TextureCollectionToMaterialLookupCache))
+	{
+		return GTextureCollectionToMaterialLookupCache.GetFrom(InTextureCollection);
+	}
+#endif
+
+	if (!TextureCollectionToMaterials.IsSet())
+	{
+		TRACE_CPUPROFILER_EVENT_SCOPE(ComputeMaterialsAffectedByTexture);
+
+		TMap<TObjectKey<UTextureCollection>, TSet<UMaterialInterface*>> TempMap;
+		TempMap.Reserve(8192);
+		for (TObjectIterator<UMaterialInterface> It(RF_ClassDefaultObject, true /*bIncludeDerivedClasses*/, GetObjectCacheInternalFlagsExclusion()); It; ++It)
+		{
+			UMaterialInterface* MaterialInterface = *It;
+			for (UTextureCollection* TextureCollection : GetUsedTextureCollections(MaterialInterface))
+			{
+				TempMap.FindOrAdd(TextureCollection).FindOrAdd(MaterialInterface);
+			}
+		}
+		TextureCollectionToMaterials = MoveTemp(TempMap);
+	}
+
+	static TSet<UMaterialInterface*> EmptySet;
+	TSet<UMaterialInterface*>* Set = TextureCollectionToMaterials.GetValue().Find(InTextureCollection);
+	const TSet<UMaterialInterface*>& ComputeResult = Set ? *Set : EmptySet;
+
+#if WITH_EDITOR
+	if (IsReverseLookupComparisonActive(EObjectReverseLookupMask::TextureCollectionToMaterialLookupCache))
+	{
+		TSet<UMaterialInterface*> LookupResult;
+		for (UMaterialInterface* Material : GTextureCollectionToMaterialLookupCache.GetFrom(InTextureCollection))
+		{
+			LookupResult.FindOrAdd(Material);
+			checkf(ComputeResult.Contains(Material), TEXT("Permanent Object Cache has an additional material %s for texture %s"), *Material->GetFullName(), *InTextureCollection->GetFullName());
+		}
+
+		for (UMaterialInterface* Material : ComputeResult)
+		{
+			checkf(LookupResult.Contains(Material), TEXT("Permanent Object Cache is missing a material %s for texture %s"), *Material->GetFullName(), *InTextureCollection->GetFullName());
+		}
+	}
+#endif
+
+	return TObjectCacheIterator<UMaterialInterface>(ComputeResult.Array());
+}
+
 #if WITH_EDITOR
 
 TObjectCacheIterator<UTexture> FObjectCacheContext::GetTexturesAffectedByTexture(UTexture* InTexture)
@@ -1057,7 +1188,7 @@ struct FObjectCacheEventSinkPrivate
 
 	static void NotifyUsedMaterialsChanged_Concurrent(const IPrimitiveComponent* PrimitiveComponent, const TArray<UMaterialInterface*>& UsedMaterials);
 	static void NotifyRenderStateChanged_Concurrent(const IPrimitiveComponent*);
-	static void NotifyReferencedTextureChanged_Concurrent(UMaterialInterface*);
+	static void NotifyMaterialChanged_Concurrent(UMaterialInterface*);
 	static void NotifyStaticMeshChanged_Concurrent(IStaticMeshComponent*);
 	static void NotifyMaterialDestroyed_Concurrent(UMaterialInterface*);
 	static void NotifyCompositeTextureChanged_Concurrent(UTexture*);
@@ -1067,7 +1198,7 @@ struct FObjectCacheEventSinkPrivate
 		EMaterialDestroyed,
 		EMaterialsChanged,
 		ERenderStateChanged,
-		EReferencedTextureChanged,
+		EMaterialChanged,
 		EStaticMeshChanged,
 		ECompositeTextureChanged,
 	};
@@ -1080,6 +1211,7 @@ struct FObjectCacheEventSinkPrivate
 		IStaticMeshComponent* StaticMeshComponent;
 		TArray<UMaterialInterface*> UsedMaterials;
 		UTexture* Texture;
+		UTextureCollection* TextureCollection;
 	};
 
 	static std::atomic<bool> bShouldQueueSinkEvents;
@@ -1202,7 +1334,7 @@ void FObjectCacheEventSink::NotifyRenderStateChanged_Concurrent(IPrimitiveCompon
 }
 
 
-void FObjectCacheEventSink::NotifyReferencedTextureChanged_Concurrent(UMaterialInterface* MaterialInterface)
+void FObjectCacheEventSink::NotifyMaterialChanged_Concurrent(UMaterialInterface* MaterialInterface)
 {
 	using namespace ObjectCacheContextImpl;
 
@@ -1210,11 +1342,11 @@ void FObjectCacheEventSink::NotifyReferencedTextureChanged_Concurrent(UMaterialI
 	{
 		if (FObjectCacheEventSinkPrivate::bShouldQueueSinkEvents)
 		{
-			FObjectCacheEventSinkPrivate::AddNotifyEvent(FObjectCacheEventSinkPrivate::EReferencedTextureChanged, MaterialInterface, nullptr, nullptr, nullptr, nullptr);
+			FObjectCacheEventSinkPrivate::AddNotifyEvent(FObjectCacheEventSinkPrivate::EMaterialChanged, MaterialInterface, nullptr, nullptr, nullptr, nullptr);
 		}
 		else
 		{
-			FObjectCacheEventSinkPrivate::NotifyReferencedTextureChanged_Concurrent(MaterialInterface);
+			FObjectCacheEventSinkPrivate::NotifyMaterialChanged_Concurrent(MaterialInterface);
 		}
 	}
 }
@@ -1297,8 +1429,8 @@ void FObjectCacheEventSinkPrivate::ProcessQueuedNotifyEvents()
 		case ERenderStateChanged:
 			NotifyRenderStateChanged_Concurrent(Event->PrimitiveComponent);
 			break;
-		case EReferencedTextureChanged:
-			NotifyReferencedTextureChanged_Concurrent(Event->MaterialInterface);
+		case EMaterialChanged:
+			NotifyMaterialChanged_Concurrent(Event->MaterialInterface);
 			break;
 		case EStaticMeshChanged:
 			NotifyStaticMeshChanged_Concurrent(Event->StaticMeshComponent);
@@ -1351,6 +1483,9 @@ void FObjectCacheEventSinkPrivate::NotifyMaterialDestroyed_Concurrent(UMaterialI
 	// Remove any Material to Texture that might be present in the cache for this Material
 	GTextureToMaterialLookupCache.Update(MaterialInterface, {});
 
+	// Remove any Material to Texture that might be present in the cache for this Material
+	GTextureCollectionToMaterialLookupCache.Update(MaterialInterface, {});
+
 	// Cleanup up any mapping that this material could have with other materials
 	if (UMaterialInstance* MaterialInstance = Cast<UMaterialInstance>(MaterialInterface))
 	{
@@ -1375,7 +1510,7 @@ void FObjectCacheEventSinkPrivate::NotifyRenderStateChanged_Concurrent(const IPr
 	GMaterialToPrimitiveLookupCache.Update(const_cast<IPrimitiveComponent*>(PrimitiveComponent), {});
 }
 
-void FObjectCacheEventSinkPrivate::NotifyReferencedTextureChanged_Concurrent(UMaterialInterface* MaterialInterface)
+void FObjectCacheEventSinkPrivate::NotifyMaterialChanged_Concurrent(UMaterialInterface* MaterialInterface)
 {
 	using namespace ObjectCacheContextImpl;
 
@@ -1385,9 +1520,17 @@ void FObjectCacheEventSinkPrivate::NotifyReferencedTextureChanged_Concurrent(UMa
 	}
 	else
 	{
-		TSet<UTexture*> Textures;
-		ObjectCacheContextImpl::GetReferencedTextures(MaterialInterface, Textures);
-		GTextureToMaterialLookupCache.Update(MaterialInterface, Textures.Array());
+		{
+			TSet<UTexture*> Textures;
+			ObjectCacheContextImpl::GetReferencedTextures(MaterialInterface, Textures);
+			GTextureToMaterialLookupCache.Update(MaterialInterface, Textures.Array());
+		}
+
+		{
+			TSet<UTextureCollection*> TextureCollections;
+			ObjectCacheContextImpl::GetReferencedTextureCollections(MaterialInterface, TextureCollections);
+			GTextureCollectionToMaterialLookupCache.Update(MaterialInterface, TextureCollections.Array());
+		}
 
 		if (UMaterialInstance* MaterialInstance = Cast<UMaterialInstance>(MaterialInterface))
 		{
diff --git a/Engine/Source/Runtime/Engine/Private/TextureCollection.cpp b/Engine/Source/Runtime/Engine/Private/TextureCollection.cpp
new file mode 100644
--- /dev/null
+++ b/Engine/Source/Runtime/Engine/Private/TextureCollection.cpp
@@ -0,0 +1,198 @@
+// Copyright Epic Games, Inc. All Rights Reserved.
+
+#include "Engine/TextureCollection.h"
+#include "Engine/Engine.h"
+#include "Engine/Texture.h"
+#include "RenderingThread.h"
+#include "TextureResource.h"
+
+#if WITH_EDITOR
+#include "Materials/Material.h"
+#include "Materials/MaterialInterface.h"
+#include "MaterialShared.h"
+#include "ObjectCacheContext.h"
+#endif
+
+#define LOCTEXT_NAMESPACE "UTextureCollection"
+
+FTextureCollectionResource::FTextureCollectionResource(UTextureCollection* InParent)
+{
+	TConstArrayView<TObjectPtr<UTexture>> ParentTextures = InParent->Textures;
+	InputTextureResources.Reserve(ParentTextures.Num());
+
+	for (TObjectPtr<UTexture> Texture : ParentTextures)
+	{
+		// TODO: figure out what we do with null textures
+		UTexture* ActualTexture = Texture != nullptr ? (UTexture*)Texture : (UTexture*)GEngine->DefaultTexture;
+		InputTextureResources.Emplace(ActualTexture->GetResource());
+	}
+}
+
+void FTextureCollectionResource::InitRHI(FRHICommandListBase& RHICmdList)
+{
+	InputTextures.Empty(InputTextureResources.Num());
+
+	TArray<FRHIResourceCollectionMember> CollectionMembers;
+	CollectionMembers.Reserve(InputTextureResources.Num());
+
+	for (FTextureResource* TextureResource : InputTextureResources)
+	{
+		FRHITextureReference* TextureReference = TextureResource ? TextureResource->GetTextureReference() : nullptr;
+		InputTextures.Emplace(TextureReference);
+		CollectionMembers.Emplace(TextureReference);
+	}
+
+	ResourceCollectionRHI = RHICmdList.CreateResourceCollection(CollectionMembers);
+}
+
+void FTextureCollectionResource::ReleaseRHI()
+{
+}
+
+#if WITH_EDITOR
+void UTextureCollection::NotifyMaterials()
+{
+	FMaterialUpdateContext UpdateContext;
+	FObjectCacheContextScope ObjectCache;
+
+	// Notify any material that uses this texture
+	TSet<UMaterial*> BaseMaterialsThatUseThisTexture;
+	for (UMaterialInterface* MaterialInterface : ObjectCache.GetContext().GetMaterialsAffectedByTextureCollection(this))
+	{
+		UpdateContext.AddMaterialInterface(MaterialInterface);
+		// This is a bit tricky. We want to make sure all materials using this texture are
+		// updated. Materials are always updated. Material instances may also have to be
+		// updated and if they have static permutations their children must be updated
+		// whether they use the texture or not! The safe thing to do is to add the instance's
+		// base material to the update context causing all materials in the tree to update.
+		BaseMaterialsThatUseThisTexture.Add(MaterialInterface->GetMaterial());
+	}
+
+	// Go ahead and update any base materials that need to be.
+	for (TSet<UMaterial*>::TConstIterator It(BaseMaterialsThatUseThisTexture); It; ++It)
+	{
+		(*It)->PostEditChange();
+	}
+}
+
+void UTextureCollection::PostEditChangeProperty(FPropertyChangedEvent& PropertyChangedEvent)
+{
+	UpdateResource();
+
+	NotifyMaterials();
+}
+#endif // WITH_EDITOR
+
+void UTextureCollection::PostLoad()
+{
+	Super::PostLoad();
+
+	if (!IsTemplate())
+	{
+		UpdateResource();
+	}
+}
+
+void UTextureCollection::SetResource(FTextureCollectionResource* InResource)
+{
+	check (!IsInActualRenderingThread() && !IsInRHIThread());
+
+	PrivateResource = InResource;
+	ENQUEUE_RENDER_COMMAND(SetResourceRenderThread)([this, InResource](FRHICommandListImmediate& RHICmdList)
+	{
+		PrivateResourceRenderThread = InResource;
+	});
+}
+
+FTextureCollectionResource* UTextureCollection::GetResource() const
+{
+	if (IsInParallelGameThread() || IsInGameThread() || IsInSlateThread() || IsInAsyncLoadingThread())
+	{
+		return PrivateResource;
+	}
+
+	if (IsInParallelRenderingThread() || IsInRHIThread())
+	{
+		return PrivateResourceRenderThread;
+	}
+
+	ensureMsgf(false, TEXT("Attempted to access a texture resource from an unkown thread."));
+	return nullptr;
+}
+
+FTextureCollectionResource* UTextureCollection::GetResource()
+{
+	if (IsInParallelGameThread() || IsInGameThread() || IsInSlateThread() || IsInAsyncLoadingThread())
+	{
+		return PrivateResource;
+	}
+
+	if (IsInParallelRenderingThread() || IsInRHIThread())
+	{
+		return PrivateResourceRenderThread;
+	}
+
+	ensureMsgf(false, TEXT("Attempted to access a texture resource from an unkown thread."));
+	return nullptr;
+}
+
+FTextureCollectionResource* UTextureCollection::CreateResource()
+{
+	return new FTextureCollectionResource(this);
+}
+
+void UTextureCollection::ReleaseResource()
+{
+	if (PrivateResource)
+	{
+		check(!IsInActualRenderingThread() && !IsInRHIThread());
+
+		// Free the resource.
+		ENQUEUE_RENDER_COMMAND(DeleteResource)([this, ToDelete = PrivateResource](FRHICommandListImmediate& RHICmdList)
+		{
+			PrivateResourceRenderThread = nullptr;
+			ToDelete->ReleaseResource();
+			delete ToDelete;
+		});
+		PrivateResource = nullptr;
+	}
+}
+
+void UTextureCollection::UpdateResource()
+{
+	// Release the existing texture resource.
+	ReleaseResource();
+
+	// Dedicated servers have no texture internals
+	if (!FApp::CanEverRender() || HasAnyFlags(RF_ClassDefaultObject))
+	{
+		return;
+	}
+
+	// TODO: christopher.waters - We should add a way to query the RHI to see if bindless is enabled.
+	ERHIBindlessConfiguration BindlessResourcesConfig = RHIGetRuntimeBindlessResourcesConfiguration(GMaxRHIShaderPlatform);
+	if (BindlessResourcesConfig == ERHIBindlessConfiguration::Disabled || GRHIGlobals.BindlessSupport == ERHIBindlessSupport::Unsupported)
+	{
+		return;
+	}
+
+	if (FTextureCollectionResource* NewResource = CreateResource())
+	{
+		LLM_SCOPE(ELLMTag::Textures);
+
+		PrivateResource = NewResource;
+		NewResource->SetOwnerName(FName(GetPathName()));
+
+		ENQUEUE_RENDER_COMMAND(SetTextureCollectionResource)([this, NewResource](FRHICommandListImmediate& RHICmdList)
+		{
+			PrivateResourceRenderThread = NewResource;
+			NewResource->InitResource(RHICmdList);
+		});
+	}
+	else
+	{
+		SetResource(nullptr);
+	}
+}
+
+#undef LOCTEXT_NAMESPACE
diff --git a/Engine/Source/Runtime/Engine/Public/MaterialCachedData.h b/Engine/Source/Runtime/Engine/Public/MaterialCachedData.h
--- a/Engine/Source/Runtime/Engine/Public/MaterialCachedData.h
+++ b/Engine/Source/Runtime/Engine/Public/MaterialCachedData.h
@@ -13,6 +13,7 @@
 #include "MaterialCachedData.generated.h"
 
 class UTexture;
+class UTextureCollection;
 class UCurveLinearColor;
 class UCurveLinearColorAtlas;
 class UFont;
@@ -197,8 +198,13 @@ struct FMaterialCachedExpressionData
 	void Validate(const UMaterialInterface& Material);
 
 	/** Adds a parameter. If this returns false, a parameter with identical name has already been added but it was set to a different value. */
-	bool AddParameter(const FMaterialParameterInfo& ParameterInfo, const FMaterialParameterMetadata& ParameterMeta, UObject*& OutReferencedTexture);
-	
+	bool AddParameter(const FMaterialParameterInfo& ParameterInfo, const FMaterialParameterMetadata& ParameterMeta, UObject*& OutReferencedTexture, UTextureCollection*& OutReferencedTextureCollection);
+
+	bool AddParameter(const FMaterialParameterInfo& ParameterInfo, const FMaterialParameterMetadata& ParameterMeta, UObject*& OutReferencedTexture)
+	{
+		UTextureCollection* UnusedReferencedTextureCollection = nullptr;
+		return AddParameter(ParameterInfo, ParameterMeta, OutReferencedTexture, UnusedReferencedTextureCollection);
+	}
 #endif // WITH_EDITOR
 
 #if WITH_EDITORONLY_DATA
@@ -293,6 +299,9 @@ struct FMaterialCachedExpressionData
 	UPROPERTY()
 	TArray<TSoftObjectPtr<UTexture>> TextureValues;
 
+	UPROPERTY()
+	TArray<TSoftObjectPtr<UTextureCollection>> TextureCollectionValues;
+
 	UPROPERTY()
 	TArray<TSoftObjectPtr<UFont>> FontValues;
 
@@ -309,6 +318,9 @@ struct FMaterialCachedExpressionData
 	UPROPERTY()
 	TArray<TObjectPtr<UObject>> ReferencedTextures;
 
+	UPROPERTY()
+	TArray<TObjectPtr<UTextureCollection>> ReferencedTextureCollections;
+
 	/** Array of all functions this material depends on. */
 	UPROPERTY()
 	TArray<FMaterialFunctionInfo> FunctionInfos;
diff --git a/Engine/Source/Runtime/Engine/Public/MaterialCompiler.h b/Engine/Source/Runtime/Engine/Public/MaterialCompiler.h
--- a/Engine/Source/Runtime/Engine/Public/MaterialCompiler.h
+++ b/Engine/Source/Runtime/Engine/Public/MaterialCompiler.h
@@ -29,6 +29,7 @@ class Error;
 class UMaterialParameterCollection;
 class URuntimeVirtualTexture;
 class UTexture;
+class UTextureCollection;
 struct FMaterialParameterInfo;
 class USparseVolumeTexture;
 
@@ -165,6 +166,7 @@ public:
 	virtual FMaterialUniformExpression* GetParameterUniformExpression(int32 Index) const = 0;
 
 	virtual bool GetTextureForExpression(int32 Index, int32& OutTextureIndex, EMaterialSamplerType& OutSamplerType, TOptional<FName>& OutParameterName) const = 0;
+	virtual bool GetTextureCollectionForExpression(int32 Index, int32& OutTextureCollectionIndex, TOptional<FName>& OutParameterName) const = 0;
 
 	virtual bool IsMaterialPropertyUsed(EMaterialProperty Property, int32 CodeChunkIdx) const = 0;
 	
@@ -285,6 +287,11 @@ public:
 	virtual int32 Texture(UTexture* Texture,int32& TextureReferenceIndex, EMaterialSamplerType SamplerType, ESamplerSourceMode SamplerSource=SSM_FromTextureAsset,ETextureMipValueMode MipValueMode=TMVM_None) = 0;
 	virtual int32 TextureParameter(FName ParameterName,UTexture* DefaultTexture,int32& TextureReferenceIndex, EMaterialSamplerType SamplerType, ESamplerSourceMode SamplerSource=SSM_FromTextureAsset) = 0;
 
+	virtual int32 TextureCollection(UTextureCollection* TextureCollection, int32& TextureCollectionReferenceIndex) = 0;
+	virtual int32 TextureCollectionParameter(FName ParameterName, UTextureCollection* DefaultValue, int32& TextureCollectionReferenceIndex) = 0;
+	virtual int32 TextureCollectionCount(int32 InTextureCollectionCodeIndex) = 0;
+	virtual int32 TextureFromCollection(int32 TextureCollectionCodeIndex, int32 IndexIntoCollectionCodeIndex, EMaterialValueType ResultTextureType) = 0;
+
 	virtual int32 VirtualTexture(URuntimeVirtualTexture* InTexture, int32 TextureLayerIndex, int32 PageTableLayerIndex, int32& TextureReferenceIndex, EMaterialSamplerType SamplerType) = 0;
 	virtual int32 VirtualTextureParameter(FName ParameterName, URuntimeVirtualTexture* DefaultValue, int32 TextureLayerIndex, int32 PageTableLayerIndex, int32& TextureReferenceIndex, EMaterialSamplerType SamplerType) = 0;
 	virtual int32 VirtualTextureUniform(int32 TextureIndex, int32 VectorIndex, UE::Shader::EValueType Type) = 0;
@@ -359,7 +366,15 @@ public:
 
 	virtual int32 SparseVolumeTextureSample(int32 SparseVolumeTextureIndex, int32 UVWIndex, int32 MipValue0Index, int32 MipValue1Index, int32 PhysicalTileDataIdxIndex, ETextureMipValueMode MipValueMode, ESamplerSourceMode SamplerSource) = 0;
 
-	virtual UObject* GetReferencedTexture(int32 Index) { return nullptr; }
+	virtual UObject* GetReferencedTexture(int32 Index)
+	{
+		return nullptr;
+	}
+
+	virtual UTextureCollection* GetReferencedTextureCollection(int32 Index)
+	{
+		return nullptr;
+	}
 
 	int32 Texture(UTexture* InTexture, EMaterialSamplerType SamplerType, ESamplerSourceMode SamplerSource=SSM_FromTextureAsset)
 	{
@@ -403,6 +418,18 @@ public:
 		return ExternalTextureParameter(ParameterName, DefaultTexture, TextureReferenceIndex);
 	}
 
+	int32 TextureCollection(UTextureCollection* InTextureCollection)
+	{
+		int32 TextureCollectionReferenceIndex = INDEX_NONE;
+		return TextureCollection(InTextureCollection, TextureCollectionReferenceIndex);
+	}
+
+	int32 TextureCollectionParameter(FName ParameterName, UTextureCollection* DefaultValue)
+	{
+		int32 TextureCollectionReferenceIndex = INDEX_NONE;
+		return TextureCollectionParameter(ParameterName, DefaultValue, TextureCollectionReferenceIndex);
+	}
+
 	virtual	int32 PixelDepth()=0;
 	virtual int32 SceneDepth(int32 Offset, int32 ViewportUV, bool bUseOffset) = 0;
 	virtual int32 SceneColor(int32 Offset, int32 ViewportUV, bool bUseOffset) = 0;
@@ -769,6 +796,7 @@ public:
 	virtual EMaterialValueType GetParameterType(int32 Index) const { return Compiler->GetParameterType(Index); }
 	virtual FMaterialUniformExpression* GetParameterUniformExpression(int32 Index) const { return Compiler->GetParameterUniformExpression(Index); }
 	virtual bool GetTextureForExpression(int32 Index, int32& OutTextureIndex, EMaterialSamplerType& OutSamplerType, TOptional<FName>& OutParameterName) const override { return Compiler->GetTextureForExpression(Index, OutTextureIndex, OutSamplerType, OutParameterName); }
+	virtual bool GetTextureCollectionForExpression(int32 Index, int32& OutTextureCollectionIndex, TOptional<FName>& OutParameterName) const override { return Compiler->GetTextureCollectionForExpression(Index, OutTextureCollectionIndex, OutParameterName); }
 	virtual void SetMaterialProperty(EMaterialProperty InProperty, EShaderFrequency OverrideShaderFrequency, bool bUsePreviousFrameTime) override { Compiler->SetMaterialProperty(InProperty, OverrideShaderFrequency, bUsePreviousFrameTime); }
 	virtual void PushMaterialAttribute(const FGuid& InAttributeID) override { Compiler->PushMaterialAttribute(InAttributeID); }
 	virtual FGuid PopMaterialAttribute() override { return Compiler->PopMaterialAttribute(); }
@@ -894,6 +922,22 @@ public:
 	virtual int32 Texture(UTexture* InTexture, int32& TextureReferenceIndex, EMaterialSamplerType SamplerType, ESamplerSourceMode SamplerSource = SSM_FromTextureAsset, ETextureMipValueMode MipValueMode = TMVM_None) override { return Compiler->Texture(InTexture, TextureReferenceIndex, SamplerType, SamplerSource, MipValueMode); }
 	virtual int32 TextureParameter(FName ParameterName, UTexture* DefaultValue, int32& TextureReferenceIndex, EMaterialSamplerType SamplerType, ESamplerSourceMode SamplerSource = SSM_FromTextureAsset) override { return Compiler->TextureParameter(ParameterName, DefaultValue, TextureReferenceIndex, SamplerType, SamplerSource); }
 
+	virtual int32 TextureCollection(UTextureCollection* TextureCollection, int32& TextureCollectionReferenceIndex) override
+	{
+		return Compiler->TextureCollection(TextureCollection, TextureCollectionReferenceIndex);
+	}
+	virtual int32 TextureCollectionParameter(FName ParameterName, UTextureCollection* DefaultValue, int32& TextureCollectionReferenceIndex) override
+	{
+		return Compiler->TextureCollectionParameter(ParameterName, DefaultValue, TextureCollectionReferenceIndex);
+	}
+	virtual int32 TextureCollectionCount(int32 InTextureCollectionCodeIndex) override
+	{
+		return Compiler->TextureCollectionCount(InTextureCollectionCodeIndex);
+	}
+	virtual int32 TextureFromCollection(int32 TextureCollectionCodeIndex, int32 IndexIntoCollectionCodeIndex, EMaterialValueType ResultTextureType) override
+	{
+		return Compiler->TextureFromCollection(TextureCollectionCodeIndex, IndexIntoCollectionCodeIndex, ResultTextureType);
+	}
 
 	virtual int32 VirtualTexture(URuntimeVirtualTexture* InTexture, int32 TextureLayerIndex, int32 PageTableLayerIndex, int32& TextureReferenceIndex, EMaterialSamplerType SamplerType) override
 	{
@@ -924,7 +968,15 @@ public:
 	virtual int32 SparseVolumeTextureSamplePhysicalTileData(int32 SparseVolumeTextureIndex, int32 VoxelCoordIndex, int32 PhysicalTileDataIdxIndex) override													{ return Compiler->SparseVolumeTextureSamplePhysicalTileData(SparseVolumeTextureIndex, VoxelCoordIndex, PhysicalTileDataIdxIndex); }
 	virtual int32 SparseVolumeTextureSample(int32 SparseVolumeTextureIndex, int32 UVWIndex, int32 MipValue0Index, int32 MipValue1Index, int32 PhysicalTileDataIdxIndex, ETextureMipValueMode MipValueMode, ESamplerSourceMode SamplerSource) { return Compiler->SparseVolumeTextureSample(SparseVolumeTextureIndex, UVWIndex, MipValue0Index, MipValue1Index, PhysicalTileDataIdxIndex, MipValueMode, SamplerSource); }
 
-	virtual UObject* GetReferencedTexture(int32 Index) override { return Compiler->GetReferencedTexture(Index); }
+	virtual UObject* GetReferencedTexture(int32 Index) override
+	{
+		return Compiler->GetReferencedTexture(Index);
+	}
+
+	virtual UTextureCollection* GetReferencedTextureCollection(int32 Index) override
+	{
+		return Compiler->GetReferencedTextureCollection(Index);
+	}
 
 	virtual	int32 PixelDepth() override { return Compiler->PixelDepth(); }
 	virtual int32 SceneDepth(int32 Offset, int32 ViewportUV, bool bUseOffset) override { return Compiler->SceneDepth(Offset, ViewportUV, bUseOffset); }
diff --git a/Engine/Source/Runtime/Engine/Public/MaterialHLSLTree.h b/Engine/Source/Runtime/Engine/Public/MaterialHLSLTree.h
--- a/Engine/Source/Runtime/Engine/Public/MaterialHLSLTree.h
+++ b/Engine/Source/Runtime/Engine/Public/MaterialHLSLTree.h
@@ -15,6 +15,7 @@
 
 class UTexture;
 enum class EMaterialParameterType : uint8;
+enum class ETextureCollectionMemberType : uint8;
 struct FMaterialCachedExpressionData;
 struct FMaterialLayersFunctions;
 
@@ -38,6 +39,7 @@ inline void AppendHash(FHasher& Hasher, const FMaterialParameterValue& Value)
 	case EMaterialParameterType::Vector: AppendHash(Hasher, Value.Float); break;
 	case EMaterialParameterType::DoubleVector: AppendHash(Hasher, Value.Double); break;
 	case EMaterialParameterType::Texture: AppendHash(Hasher, Value.Texture); break;
+	case EMaterialParameterType::TextureCollection: AppendHash(Hasher, Value.TextureCollection); break;
 	case EMaterialParameterType::Font: AppendHash(Hasher, Value.Font); break;
 	case EMaterialParameterType::RuntimeVirtualTexture: AppendHash(Hasher, Value.RuntimeVirtualTexture); break;
 	case EMaterialParameterType::SparseVolumeTexture: AppendHash(Hasher, Value.SparseVolumeTexture); break;
@@ -632,6 +634,29 @@ public:
 	virtual void EmitValueShader(FEmitContext& Context, FEmitScope& Scope, const FRequestedType& RequestedType, FEmitValueShaderResult& OutResult) const override;
 };
 
+class FExpressionTextureObjectFromCollection : public FExpression
+{
+public:
+	FExpressionTextureObjectFromCollection(
+		const FExpression* InTextureCollectionExpression,
+		const FExpression* InCollectionIndexExpression,
+		ETextureCollectionMemberType InTextureType)
+		: TextureCollectionExpression(InTextureCollectionExpression)
+		, CollectionIndexExpression(InCollectionIndexExpression)
+		, TextureType(InTextureType)
+	{
+	}
+
+	const FExpression* const TextureCollectionExpression;
+	const FExpression* const CollectionIndexExpression;
+	const ETextureCollectionMemberType TextureType;
+
+	// FExpression
+	virtual bool PrepareValue(FEmitContext& Context, FEmitScope& Scope, const FRequestedType& RequestedType, FPrepareValueResult& OutResult) const override;
+	virtual void EmitValueShader(FEmitContext& Context, FEmitScope& Scope, const FRequestedType& RequestedType, FEmitValueShaderResult& OutResult) const override;
+	//~FExpression
+};
+
 class FExpressionAntiAliasedTextureMask : public FExpression
 {
 public:
diff --git a/Engine/Source/Runtime/Engine/Public/MaterialShared.h b/Engine/Source/Runtime/Engine/Public/MaterialShared.h
--- a/Engine/Source/Runtime/Engine/Public/MaterialShared.h
+++ b/Engine/Source/Runtime/Engine/Public/MaterialShared.h
@@ -376,6 +376,7 @@ struct FMaterialRenderContext
 	ENGINE_API void GetTextureParameterValue(const FHashedMaterialParameterInfo& ParameterInfo, int32 TextureIndex, const UTexture*& OutValue) const;
 	ENGINE_API void GetTextureParameterValue(const FHashedMaterialParameterInfo& ParameterInfo, int32 TextureIndex, const URuntimeVirtualTexture*& OutValue) const;
 	ENGINE_API void GetTextureParameterValue(const FHashedMaterialParameterInfo& ParameterInfo, int32 TextureIndex, const USparseVolumeTexture*& OutValue) const;
+	ENGINE_API void GetTextureCollectionParameterValue(const FHashedMaterialParameterInfo& ParameterInfo, int32 TextureCollectionIndex, const UTextureCollection*& OutValue) const;
 	ENGINE_API FGuid GetExternalTextureGuid(const FGuid& ExternalTextureGuid, const FName& ParameterName, int32 SourceTextureIndex) const;
 };
 
@@ -533,6 +534,24 @@ public:
 	LAYOUT_FIELD_INITIALIZED(uint8, VirtualTextureLayerIndex, 0u);
 };
 
+class FMaterialTextureCollectionParameterInfo
+{
+	DECLARE_TYPE_LAYOUT(FMaterialTextureCollectionParameterInfo, NonVirtual);
+public:
+	friend inline bool operator==(const FMaterialTextureCollectionParameterInfo& Lhs, const FMaterialTextureCollectionParameterInfo& Rhs)
+	{
+		return Lhs.ParameterInfo == Rhs.ParameterInfo && Lhs.TextureCollectionIndex == Rhs.TextureCollectionIndex;
+	}
+	friend inline bool operator!=(const FMaterialTextureCollectionParameterInfo& Lhs, const FMaterialTextureCollectionParameterInfo& Rhs)
+	{
+		return !operator==(Lhs, Rhs);
+	}
+
+	void GetTextureCollection(const FMaterialRenderContext& Context, const UTextureCollection*& OutTextureCollection) const;
+
+	LAYOUT_FIELD(FHashedMaterialParameterInfo, ParameterInfo);
+	LAYOUT_FIELD(int32, TextureCollectionIndex);
+};
 
 class FMaterialExternalTextureParameterInfo
 {
@@ -641,6 +660,7 @@ public:
 	ENGINE_API void GetTextureValue(int32 Index, const FMaterialRenderContext& Context, const FMaterial& Material, const USparseVolumeTexture*& OutValue) const;
 
 	int32 FindOrAddTextureParameter(EMaterialTextureParameterType Type, const FMaterialTextureParameterInfo& Info);
+	int32 FindOrAddTextureCollectionParameter(const FMaterialTextureCollectionParameterInfo& Info);
 	int32 FindOrAddExternalTextureParameter(const FMaterialExternalTextureParameterInfo& Info);
 	int32 FindOrAddNumericParameter(EMaterialParameterType Type, const FMaterialParameterInfo& ParameterInfo, uint32 DefaultValueOffset);
 	uint32 AddDefaultParameterValue(const UE::Shader::FValue& Value);
@@ -681,6 +701,7 @@ protected:
 	LAYOUT_FIELD(TMemoryImageArray<FMaterialNumericParameterInfo>, UniformNumericParameters);
 	LAYOUT_ARRAY(TMemoryImageArray<FMaterialTextureParameterInfo>, UniformTextureParameters, NumMaterialTextureParameterTypes);
 	LAYOUT_FIELD(TMemoryImageArray<FMaterialExternalTextureParameterInfo>, UniformExternalTextureParameters);
+	LAYOUT_FIELD(TMemoryImageArray<FMaterialTextureCollectionParameterInfo>, UniformTextureCollectionParameters);
 	LAYOUT_FIELD(uint32, UniformPreshaderBufferSize);
 
 	LAYOUT_FIELD(UE::Shader::FPreshaderData, UniformPreshaderData);
@@ -2361,6 +2382,7 @@ public:
 	ENGINE_API virtual void AddReferencedObjects(FReferenceCollector& Collector);
 
 	virtual TArrayView<const TObjectPtr<UObject>> GetReferencedTextures() const = 0;
+	virtual TConstArrayView<TObjectPtr<UTextureCollection>> GetReferencedTextureCollections() const = 0;
 
 	/**
 	 * Finds the shader matching the template type and the passed in vertex factory, asserts if not found.
@@ -2914,6 +2936,7 @@ public:
 	ENGINE_API virtual void LegacySerialize(FArchive& Ar) override;
 
 	ENGINE_API virtual TArrayView<const TObjectPtr<UObject>> GetReferencedTextures() const override;
+	ENGINE_API virtual TConstArrayView<TObjectPtr<UTextureCollection>> GetReferencedTextureCollections() const override;
 
 	ENGINE_API virtual void AddReferencedObjects(FReferenceCollector& Collector) override;
 
diff --git a/Engine/Source/Runtime/Engine/Public/MaterialTypes.h b/Engine/Source/Runtime/Engine/Public/MaterialTypes.h
--- a/Engine/Source/Runtime/Engine/Public/MaterialTypes.h
+++ b/Engine/Source/Runtime/Engine/Public/MaterialTypes.h
@@ -12,6 +12,7 @@
 
 class FShaderKeyGenerator;
 class UTexture;
+class UTextureCollection;
 class UCurveLinearColor;
 class UCurveLinearColorAtlas;
 class UFont;
@@ -186,6 +187,7 @@ enum class EMaterialParameterType : uint8
 	Vector,
 	DoubleVector,
 	Texture,
+	TextureCollection,
 	Font,
 	RuntimeVirtualTexture,
 	SparseVolumeTexture,
@@ -322,6 +324,7 @@ struct FMaterialParameterValue
 	FMaterialParameterValue(const FVector4d& InValue) : Type(EMaterialParameterType::DoubleVector) { Double[0] = InValue.X; Double[1] = InValue.Y; Double[2] = InValue.Z; Double[3] = InValue.W; }
 	FMaterialParameterValue(UTexture* InValue) : Type(EMaterialParameterType::Texture) { Texture = InValue; }
 	FMaterialParameterValue(const TObjectPtr<UTexture>& InValue) : Type(EMaterialParameterType::Texture) { Texture = InValue; }
+	FMaterialParameterValue(const TObjectPtr<UTextureCollection>& InValue) : Type(EMaterialParameterType::TextureCollection) { TextureCollection = InValue; }
 	FMaterialParameterValue(URuntimeVirtualTexture* InValue) : Type(EMaterialParameterType::RuntimeVirtualTexture) { RuntimeVirtualTexture = InValue; }
 	FMaterialParameterValue(const TObjectPtr<URuntimeVirtualTexture>& InValue) : Type(EMaterialParameterType::RuntimeVirtualTexture) { RuntimeVirtualTexture = InValue; }
 	FMaterialParameterValue(USparseVolumeTexture* InValue) : Type(EMaterialParameterType::SparseVolumeTexture) { SparseVolumeTexture = InValue; }
@@ -331,6 +334,8 @@ struct FMaterialParameterValue
 	// Would be possible to store an additional const-flag member, and provide runtime checks to ensure constness is not violated...maybe worth doing in the future
 	FMaterialParameterValue(const UTexture* InValue) : Type(EMaterialParameterType::Texture) { Texture = const_cast<UTexture*>(InValue); }
 	FMaterialParameterValue(const TObjectPtr<const UTexture>& InValue) : Type(EMaterialParameterType::Texture) { Texture = const_cast<UTexture*>(InValue.Get()); }
+	FMaterialParameterValue(const UTextureCollection* InValue) : Type(EMaterialParameterType::TextureCollection) { TextureCollection = const_cast<UTextureCollection*>(InValue); }
+	FMaterialParameterValue(const TObjectPtr<const UTextureCollection>& InValue) : Type(EMaterialParameterType::TextureCollection) { TextureCollection = const_cast<UTextureCollection*>(InValue.Get()); }
 	FMaterialParameterValue(const URuntimeVirtualTexture* InValue) : Type(EMaterialParameterType::RuntimeVirtualTexture) { RuntimeVirtualTexture = const_cast<URuntimeVirtualTexture*>(InValue); }
 	FMaterialParameterValue(const TObjectPtr<const URuntimeVirtualTexture>& InValue) : Type(EMaterialParameterType::RuntimeVirtualTexture) { RuntimeVirtualTexture = const_cast<URuntimeVirtualTexture*>(InValue.Get()); }
 	FMaterialParameterValue(const USparseVolumeTexture* InValue) : Type(EMaterialParameterType::SparseVolumeTexture) { SparseVolumeTexture = const_cast<USparseVolumeTexture*>(InValue); }
@@ -357,6 +362,7 @@ struct FMaterialParameterValue
 		float Float[4];
 		bool Bool[4];
 		UTexture* Texture;
+		UTextureCollection* TextureCollection;
 		URuntimeVirtualTexture* RuntimeVirtualTexture;
 		USparseVolumeTexture* SparseVolumeTexture;
 		struct
@@ -389,6 +395,7 @@ struct FMaterialParameterValue
 			Lhs.Double[2] == Rhs.Double[2] &&
 			Lhs.Double[3] == Rhs.Double[3];
 		case EMaterialParameterType::Texture: return Lhs.Texture == Rhs.Texture;
+		case EMaterialParameterType::TextureCollection: return Lhs.TextureCollection == Rhs.TextureCollection;
 		case EMaterialParameterType::Font: return Lhs.Font.Value == Rhs.Font.Value && Lhs.Font.Page == Rhs.Font.Page;
 		case EMaterialParameterType::RuntimeVirtualTexture: return Lhs.RuntimeVirtualTexture == Rhs.RuntimeVirtualTexture;
 		case EMaterialParameterType::SparseVolumeTexture: return Lhs.SparseVolumeTexture == Rhs.SparseVolumeTexture;
diff --git a/Engine/Source/Runtime/Engine/Public/MaterialValueType.h b/Engine/Source/Runtime/Engine/Public/MaterialValueType.h
--- a/Engine/Source/Runtime/Engine/Public/MaterialValueType.h
+++ b/Engine/Source/Runtime/Engine/Public/MaterialValueType.h
@@ -9,23 +9,22 @@
 /**
  * The types which can be used by materials.
  */
-enum EMaterialValueType
+enum EMaterialValueType : int
 {
 	/** 
 	 * A scalar float type.  
 	 * Note that MCT_Float1 will not auto promote to any other float types, 
 	 * So use MCT_Float instead for scalar expression return types.
 	 */
-	MCT_Float1		= 1,
-	MCT_Float2		= 2,
-	MCT_Float3		= 4,
-	MCT_Float4		= 8,
+	MCT_Float1                = 1 << 0,
+	MCT_Float2                = 1 << 1,
+	MCT_Float3                = 1 << 2,
+	MCT_Float4                = 1 << 3,
 
 	/** 
 	 * Any size float type by definition, but this is treated as a scalar which can auto convert (by replication) to any other size float vector.
 	 * Use this as the type for any scalar expressions.
 	 */
-	MCT_Float                 = 8|4|2|1,
 	MCT_Texture2D	          = 1 << 4,
 	MCT_TextureCube	          = 1 << 5,
 	MCT_Texture2DArray		  = 1 << 6,
@@ -37,21 +36,17 @@ enum EMaterialValueType
 	MCT_TextureExternal       = 1 << 12,
 	MCT_TextureVirtual        = 1 << 13,
 	MCT_SparseVolumeTexture   = 1 << 14,
-	/** MCT_SparseVolumeTexture is intentionally not (yet) included here because it differs a lot from the other texture types and may not be supported/appropriate for all MCT_Texture use cases. */
-	MCT_Texture               = MCT_Texture2D | MCT_TextureCube | MCT_Texture2DArray | MCT_TextureCubeArray | MCT_VolumeTexture | MCT_TextureExternal | MCT_TextureVirtual,
 
 	/** Used internally when sampling from virtual textures */
 	MCT_VTPageTableResult     = 1 << 15,
 	
 	MCT_ShadingModel		  = 1 << 16,
-
 	MCT_Substrate			  = 1 << 17,
 
 	MCT_LWCScalar			  = 1 << 18,
 	MCT_LWCVector2			  = 1 << 19,
 	MCT_LWCVector3			  = 1 << 20,
 	MCT_LWCVector4			  = 1 << 21,
-	MCT_LWCType				  = MCT_LWCScalar | MCT_LWCVector2 | MCT_LWCVector3 | MCT_LWCVector4,
 
 	MCT_Execution             = 1 << 22,
 
@@ -59,14 +54,21 @@ enum EMaterialValueType
 	MCT_VoidStatement         = 1 << 23,
 
 	/** Non-static bool, only used in new HLSL translator */
-	MCT_Bool				 = 1 << 24,
+	MCT_Bool                  = 1 << 24,
 
 	/** Unsigned int types */
-	MCT_UInt1				 = 1 << 25,
-	MCT_UInt2				 = 1 << 26,
-	MCT_UInt3				 = 1 << 27,
-	MCT_UInt4				 = 1 << 28,
-	MCT_UInt				 = MCT_UInt1 | MCT_UInt2 | MCT_UInt3 | MCT_UInt4,
+	MCT_UInt1                 = 1 << 25,
+	MCT_UInt2                 = 1 << 26,
+	MCT_UInt3                 = 1 << 27,
+	MCT_UInt4                 = 1 << 28,
+
+	MCT_TextureCollection     = 1 << 29,
+
+	/** MCT_SparseVolumeTexture is intentionally not (yet) included here because it differs a lot from the other texture types and may not be supported/appropriate for all MCT_Texture use cases. */
+	MCT_Texture = MCT_Texture2D | MCT_TextureCube | MCT_Texture2DArray | MCT_TextureCubeArray | MCT_VolumeTexture | MCT_TextureExternal | MCT_TextureVirtual,
 
+	MCT_Float   = MCT_Float1 | MCT_Float2 | MCT_Float3 | MCT_Float4,
+	MCT_UInt    = MCT_UInt1 | MCT_UInt2 | MCT_UInt3 | MCT_UInt4,
+	MCT_LWCType = MCT_LWCScalar | MCT_LWCVector2 | MCT_LWCVector3 | MCT_LWCVector4,
 	MCT_Numeric = MCT_Float | MCT_LWCType | MCT_Bool,
 };
diff --git a/Engine/Source/Runtime/Engine/Public/Materials/MaterialRenderProxy.h b/Engine/Source/Runtime/Engine/Public/Materials/MaterialRenderProxy.h
--- a/Engine/Source/Runtime/Engine/Public/Materials/MaterialRenderProxy.h
+++ b/Engine/Source/Runtime/Engine/Public/Materials/MaterialRenderProxy.h
@@ -27,6 +27,7 @@ class USubsurfaceProfile;
 class USpecularProfile;
 class UNeuralProfile;
 class UTexture;
+class UTextureCollection;
 
 struct FMaterialParameterValue;
 struct FMaterialRenderContext;
@@ -170,6 +171,7 @@ public:
 	ENGINE_API bool GetTextureValue(const FHashedMaterialParameterInfo& ParameterInfo, const UTexture** OutValue, const FMaterialRenderContext& Context) const;
 	ENGINE_API bool GetTextureValue(const FHashedMaterialParameterInfo& ParameterInfo, const URuntimeVirtualTexture** OutValue, const FMaterialRenderContext& Context) const;
 	ENGINE_API bool GetTextureValue(const FHashedMaterialParameterInfo& ParameterInfo, const USparseVolumeTexture** OutValue, const FMaterialRenderContext& Context) const;
+	ENGINE_API bool GetTextureCollectionValue(const FHashedMaterialParameterInfo& ParameterInfo, const UTextureCollection** OutValue, const FMaterialRenderContext& Context) const;
 	virtual bool GetParameterValue(EMaterialParameterType Type, const FHashedMaterialParameterInfo& ParameterInfo, FMaterialParameterValue& OutValue, const FMaterialRenderContext& Context) const = 0;
 
 	bool IsDeleted() const
diff --git a/Engine/Source/Runtime/Engine/Public/ObjectCacheContext.h b/Engine/Source/Runtime/Engine/Public/ObjectCacheContext.h
--- a/Engine/Source/Runtime/Engine/Public/ObjectCacheContext.h
+++ b/Engine/Source/Runtime/Engine/Public/ObjectCacheContext.h
@@ -15,6 +15,7 @@ class USkinnedMeshComponent;
 class UMaterialInterface;
 class UStaticMesh;
 class UTexture;
+class UTextureCollection;
 class USkeletalMesh;
 class USkinnedAsset;
 
@@ -75,9 +76,11 @@ public:
 	ENGINE_API TObjectCacheIterator<USkinnedMeshComponent> GetSkinnedMeshComponents(USkinnedAsset* InSkinnedAsset);
 	ENGINE_API TObjectCacheIterator<IStaticMeshComponent>  GetStaticMeshComponents(UStaticMesh* InStaticMesh);
 	ENGINE_API TObjectCacheIterator<UMaterialInterface>    GetMaterialsAffectedByTexture(UTexture* InTexture);
+	ENGINE_API TObjectCacheIterator<UMaterialInterface>    GetMaterialsAffectedByTextureCollection(UTextureCollection* InTextureCollection);
 	ENGINE_API TObjectCacheIterator<IPrimitiveComponent>   GetPrimitivesAffectedByMaterial(UMaterialInterface* InMaterial);
 	ENGINE_API TObjectCacheIterator<IPrimitiveComponent>   GetPrimitivesAffectedByMaterials(TArrayView<UMaterialInterface*> InMaterials);
 	ENGINE_API TObjectCacheIterator<UTexture>              GetUsedTextures(UMaterialInterface* InMaterial);
+	ENGINE_API TObjectCacheIterator<UTextureCollection>    GetUsedTextureCollections(UMaterialInterface* InMaterial);
 	ENGINE_API TObjectCacheIterator<UMaterialInterface>    GetUsedMaterials(IPrimitiveComponent* InComponent);	
 	ENGINE_API TObjectCacheIterator<UMaterialInterface>    GetMaterialsAffectedByMaterials(TArrayView<UMaterialInterface*> InMaterials);
 #if WITH_EDITOR
@@ -89,7 +92,9 @@ private:
 	FObjectCacheContext() = default;
 	TMap<IPrimitiveComponent*, TSet<UMaterialInterface*>> PrimitiveComponentToMaterial;
 	TMap<TObjectKey<UMaterialInterface>, TSet<UTexture*>> MaterialUsedTextures;
+	TMap<TObjectKey<UMaterialInterface>, TSet<UTextureCollection*>> MaterialUsedTextureCollections;
 	TOptional<TMap<TObjectKey<UTexture>, TSet<UMaterialInterface*>>> TextureToMaterials;
+	TOptional<TMap<TObjectKey<UTextureCollection>, TSet<UMaterialInterface*>>> TextureCollectionToMaterials;
 #if WITH_EDITOR
 	TOptional<TMap<TObjectKey<UTexture>, TSet<UTexture*>>> TextureToTextures;
 #endif
diff --git a/Engine/Source/Runtime/Engine/Public/ObjectCacheEventSink.h b/Engine/Source/Runtime/Engine/Public/ObjectCacheEventSink.h
--- a/Engine/Source/Runtime/Engine/Public/ObjectCacheEventSink.h
+++ b/Engine/Source/Runtime/Engine/Public/ObjectCacheEventSink.h
@@ -26,7 +26,11 @@ struct  FObjectCacheEventSink
 	static ENGINE_API void NotifyUsedMaterialsChanged_Concurrent(const IPrimitiveComponent* PrimitiveComponent, const TArray<UMaterialInterface*>& UsedMaterials);
 	static ENGINE_API void NotifyRenderStateChanged_Concurrent(UActorComponent*);
 	static ENGINE_API void NotifyRenderStateChanged_Concurrent(IPrimitiveComponent*);
-	static ENGINE_API void NotifyReferencedTextureChanged_Concurrent(UMaterialInterface*);
+	static ENGINE_API void NotifyMaterialChanged_Concurrent(UMaterialInterface*);
+	static inline void NotifyReferencedTextureChanged_Concurrent(UMaterialInterface* InMaterial)
+	{
+		NotifyMaterialChanged_Concurrent(InMaterial);
+	}
 	static ENGINE_API void NotifyStaticMeshChanged_Concurrent(IStaticMeshComponent*);
 	static ENGINE_API void NotifyMaterialDestroyed_Concurrent(UMaterialInterface*);
 	static ENGINE_API void NotifyCompositeTextureChanged_Concurrent(UTexture*);
diff --git a/Engine/Source/Runtime/Engine/Public/TextureResource.h b/Engine/Source/Runtime/Engine/Public/TextureResource.h
--- a/Engine/Source/Runtime/Engine/Public/TextureResource.h
+++ b/Engine/Source/Runtime/Engine/Public/TextureResource.h
@@ -154,6 +154,11 @@ public:
 		TextureReferenceRHI = TextureReference;
 	}
 
+	FRHITextureReference* GetTextureReference() const
+	{
+		return TextureReferenceRHI;
+	}
+
 #if STATS
 	/* The Stat_ FName corresponding to each TEXTUREGROUP */
 	static FName TextureGroupStatFNames[TEXTUREGROUP_MAX];
