-Make sure to reset/initialize the Data Layers before BeginPlay since replication can already be happening at that time.

diff --git a/Engine/Source/Runtime/Engine/Private/WorldPartition/DataLayer/WorldDataLayers.cpp b/Engine/Source/Runtime/Engine/Private/WorldPartition/DataLayer/WorldDataLayers.cpp
--- a/Engine/Source/Runtime/Engine/Private/WorldPartition/DataLayer/WorldDataLayers.cpp
+++ b/Engine/Source/Runtime/Engine/Private/WorldPartition/DataLayer/WorldDataLayers.cpp
@@ -84,9 +84,18 @@ void AWorldDataLayers::GetLifetimeReplicatedProps(TArray< FLifetimeProperty >& O
 	DOREPLIFETIME_WITH_PARAMS_FAST(AWorldDataLayers, RepEffectiveActiveDataLayerNames, Params);
 }
 
-void AWorldDataLayers::BeginPlay()
+void AWorldDataLayers::RewindForReplay()
+{
+	Super::RewindForReplay();
+
+	// Same as PostRegisterAllComponents when rewinding we want to reset our state to the initial state and rely on the Replay/Replication.
+	ResetDataLayerRuntimeStates();
+	InitializeDataLayerRuntimeStates();
+}
+
+void AWorldDataLayers::PostRegisterAllComponents() 
 {
-	Super::BeginPlay();
+	Super::PostRegisterAllComponents();
 
 	// When running a Replay we want to reset our state to the initial state and rely on the Replay/Replication.
 	// Unfortunately this can't be tested in the PostLoad as the World doesn't have a demo driver yet.
@@ -97,15 +106,6 @@ void AWorldDataLayers::BeginPlay()
 	}
 }
 
-void AWorldDataLayers::RewindForReplay()
-{
-	Super::RewindForReplay();
-
-	// Same as BeginPlay when rewinding we want to reset our state to the initial state and rely on the Replay/Replication.
-	ResetDataLayerRuntimeStates();
-	InitializeDataLayerRuntimeStates();
-}
-
 #define AWORLDDATALAYERS_UPDATE_REPLICATED_DATALAYERS(ReplicatedArray, SourceArray) \
 	MARK_PROPERTY_DIRTY_FROM_NAME(AWorldDataLayers, ReplicatedArray, this); \
 	ReplicatedArray = SourceArray;
diff --git a/Engine/Source/Runtime/Engine/Public/WorldPartition/DataLayer/WorldDataLayers.h b/Engine/Source/Runtime/Engine/Public/WorldPartition/DataLayer/WorldDataLayers.h
--- a/Engine/Source/Runtime/Engine/Public/WorldPartition/DataLayer/WorldDataLayers.h
+++ b/Engine/Source/Runtime/Engine/Public/WorldPartition/DataLayer/WorldDataLayers.h
@@ -63,7 +63,7 @@ public:
 
 	//~ Begin AActor Interface
 	ENGINE_API virtual void RewindForReplay() override;
-	ENGINE_API virtual void BeginPlay() override;
+	ENGINE_API virtual void PostRegisterAllComponents() override;
 #if WITH_EDITOR
 	ENGINE_API virtual void PreEditUndo() override;
 	ENGINE_API virtual void PostEditUndo() override;
