[World Partition] Fix crash when calling UWorldPartitionBlueprintLibrary::GetActorDescsForActors for actors in level instances.

diff --git a/Engine/Source/Runtime/Engine/Private/WorldPartitionBlueprintLibrary.cpp b/Engine/Source/Runtime/Engine/Private/WorldPartitionBlueprintLibrary.cpp
--- a/Engine/Source/Runtime/Engine/Private/WorldPartitionBlueprintLibrary.cpp
+++ b/Engine/Source/Runtime/Engine/Private/WorldPartitionBlueprintLibrary.cpp
@@ -275,7 +275,10 @@ bool UWorldPartitionBlueprintLibrary::GetActorDescsForActors(const TArray<AActor
 	{
 		for (AActor* Actor : InActors)
 		{
-			OutActorDescs.Emplace(*WorldPartition->GetActorDescInstance(Actor->GetActorGuid()));
+			if (FWorldPartitionActorDescInstance* ActorDescInstance = WorldPartition->GetActorDescInstance(Actor->GetActorGuid()))
+			{
+				OutActorDescs.Emplace(*ActorDescInstance);
+			}
 		}
 	}
 #endif
