PR #11443: Set VirtualTextureRenderPassType to BlueprintReadWrite

diff --git a/Engine/Source/Runtime/Engine/Classes/Components/PrimitiveComponent.h b/Engine/Source/Runtime/Engine/Classes/Components/PrimitiveComponent.h
--- a/Engine/Source/Runtime/Engine/Classes/Components/PrimitiveComponent.h
+++ b/Engine/Source/Runtime/Engine/Classes/Components/PrimitiveComponent.h
@@ -825,7 +825,7 @@ public:
 	int8 VirtualTextureMinCoverage = 0;
 
 	/** Controls if this component draws in the main pass as well as in the virtual texture. */
-	UPROPERTY(EditAnywhere, BlueprintReadOnly, Category = VirtualTexture, meta = (DisplayName = "Draw in Main Pass"))
+	UPROPERTY(EditAnywhere, BlueprintReadWrite, Category = VirtualTexture, meta = (DisplayName = "Draw in Main Pass"))
 	ERuntimeVirtualTextureMainPassType VirtualTextureRenderPassType = ERuntimeVirtualTextureMainPassType::Exclusive;
 
 	/** Get the array of runtime virtual textures into which we render the mesh for this actor. */
