[NO NON-PII COMMIT MESSAGE]

diff --git a/Engine/Source/Runtime/Engine/Public/WorldPartition/DataLayer/WorldDataLayers.h b/Engine/Source/Runtime/Engine/Public/WorldPartition/DataLayer/WorldDataLayers.h
--- a/Engine/Source/Runtime/Engine/Public/WorldPartition/DataLayer/WorldDataLayers.h
+++ b/Engine/Source/Runtime/Engine/Public/WorldPartition/DataLayer/WorldDataLayers.h
@@ -58,6 +58,7 @@ class AWorldDataLayers : public AInfo, public IDataLayerInstanceProvider
 public:
 	//~ Begin UObject Interface
 	ENGINE_API virtual void PostLoad() override;
+	ENGINE_API virtual void PostRegisterAllComponents() override;
 	virtual void Serialize(FArchive& Ar) override;
 	//~ End UObject Interface
 
