[Backout] - CL32898933 - CIS Compile Error
Original CL Desc
-----------------------------------------------------------------
-Make sure to reset/initialize the Data Layers before BeginPlay since replication can already be happening at that time.

diff --git a/Engine/Source/Runtime/Engine/Private/WorldPartition/DataLayer/WorldDataLayers.cpp b/Engine/Source/Runtime/Engine/Private/WorldPartition/DataLayer/WorldDataLayers.cpp
--- a/Engine/Source/Runtime/Engine/Private/WorldPartition/DataLayer/WorldDataLayers.cpp
+++ b/Engine/Source/Runtime/Engine/Private/WorldPartition/DataLayer/WorldDataLayers.cpp
@@ -84,18 +84,9 @@ void AWorldDataLayers::GetLifetimeReplicatedProps(TArray< FLifetimeProperty >& O
 	DOREPLIFETIME_WITH_PARAMS_FAST(AWorldDataLayers, RepEffectiveActiveDataLayerNames, Params);
 }
 
-void AWorldDataLayers::RewindForReplay()
-{
-	Super::RewindForReplay();
-
-	// Same as PostRegisterAllComponents when rewinding we want to reset our state to the initial state and rely on the Replay/Replication.
-	ResetDataLayerRuntimeStates();
-	InitializeDataLayerRuntimeStates();
-}
-
-void AWorldDataLayers::PostRegisterAllComponents() 
+void AWorldDataLayers::BeginPlay()
 {
-	Super::PostRegisterAllComponents();
+	Super::BeginPlay();
 
 	// When running a Replay we want to reset our state to the initial state and rely on the Replay/Replication.
 	// Unfortunately this can't be tested in the PostLoad as the World doesn't have a demo driver yet.
@@ -106,6 +97,15 @@ void AWorldDataLayers::PostRegisterAllComponents()
 	}
 }
 
+void AWorldDataLayers::RewindForReplay()
+{
+	Super::RewindForReplay();
+
+	// Same as BeginPlay when rewinding we want to reset our state to the initial state and rely on the Replay/Replication.
+	ResetDataLayerRuntimeStates();
+	InitializeDataLayerRuntimeStates();
+}
+
 #define AWORLDDATALAYERS_UPDATE_REPLICATED_DATALAYERS(ReplicatedArray, SourceArray) \
 	MARK_PROPERTY_DIRTY_FROM_NAME(AWorldDataLayers, ReplicatedArray, this); \
 	ReplicatedArray = SourceArray;
