Added TComponentInterfaceIterator so that we can iterate over the various components interfaces instances in cases where ObjectCacheContext doesn't maintain a global cache
This is necessary for the few usages of this we have in runtime

diff --git a/Engine/Source/Runtime/Engine/Classes/Components/ComponentInterfaceIterator.h b/Engine/Source/Runtime/Engine/Classes/Components/ComponentInterfaceIterator.h
new file mode 100644
--- /dev/null
+++ b/Engine/Source/Runtime/Engine/Classes/Components/ComponentInterfaceIterator.h
@@ -0,0 +1,112 @@
+// Copyright Epic Games, Inc. All Rights Reserved.
+#pragma once
+
+#include "Components/ComponentInterfaces.h"
+#include "UObject/UObjectHash.h"
+#include "UObject/UObjectIterator.h"
+
+template< class T > 
+class TComponentInterfaceIterator
+{
+public:
+
+	enum EEndTagType
+	{
+		EndTag
+	};
+
+	explicit TComponentInterfaceIterator(EObjectFlags AdditionalExclusionFlags = RF_ClassDefaultObject, bool bIncludeDerivedClasses = true, EInternalObjectFlags InInternalExclusionFlags = EInternalObjectFlags::None)
+		: Index(-1)
+	{
+		TArray<UObject*>	ObjectArray;
+
+		for (FComponentInterfaceImplementation& Implementation : T::Implementers)
+		{
+			ObjectArray.Reset();
+			GetObjectsOfClass(Implementation.Class, ObjectArray, bIncludeDerivedClasses, AdditionalExclusionFlags, GetObjectIteratorDefaultInternalExclusionFlags(InInternalExclusionFlags));
+			Interfaces.Reserve(ObjectArray.Num() + Interfaces.Num());
+			for (UObject* Object : ObjectArray)
+			{
+				T* Interface = reinterpret_cast<T*>(Implementation.Resolver(Object));
+				check(Interface);
+				Interfaces.Add(Interface);
+			}
+		}
+
+		Index = 0;
+	}
+
+	TComponentInterfaceIterator(EEndTagType, const TComponentInterfaceIterator<T>& InIter )
+		: Index(InIter.Interfaces.Num())
+	{
+	}
+
+	FORCEINLINE void operator++(int)
+	{		
+		Index++;
+	}	
+
+	FORCEINLINE void operator++()
+	{		
+		Index++;
+	}
+
+	/** Conversion to "bool" returning true if the iterator is valid. */
+	FORCEINLINE explicit operator bool() const
+	{ 
+		return Interfaces.IsValidIndex(Index); 
+	}
+
+	/** Conversion to "bool" returning true if the iterator is valid. */
+	FORCEINLINE bool operator !() const 
+	{
+		return !(bool)*this;
+	}
+
+	FORCEINLINE T* operator*() const
+	{
+		if (Interfaces.IsValidIndex(Index))
+		{
+			return  Interfaces[Index];
+		}
+		else
+		{
+			return nullptr;
+		}
+	}
+
+	FORCEINLINE T* operator->() const
+	{
+		if (Interfaces.IsValidIndex(Index))
+		{
+			return  Interfaces[Index];
+		}
+		else
+		{
+			return nullptr;
+		}
+	}
+
+	FORCEINLINE bool operator==(const TComponentInterfaceIterator& Rhs) const { return Index == Rhs.Index; }
+	FORCEINLINE bool operator!=(const TComponentInterfaceIterator& Rhs) const { return Index != Rhs.Index; }
+
+protected:
+	/** Resolved results for all the Results from the GetObjectsOfClass queries */
+	TArray<T*>	Interfaces;
+	
+	/** index of the current element in the object array */
+	int32 Index;
+};
+
+template <typename T>
+struct TComponentInterfaceRange
+{
+	TComponentInterfaceRange()
+	{
+	}
+
+	friend TComponentInterfaceIterator<T> begin(const TComponentInterfaceRange& Range) { return Range.Begin; }
+	friend TComponentInterfaceIterator<T> end  (const TComponentInterfaceRange& Range) { return TComponentInterfaceIterator<T>(TComponentInterfaceIterator<T>::EndTag, Range.Begin); }
+
+	TComponentInterfaceIterator<T> Begin;
+};
diff --git a/Engine/Source/Runtime/Engine/Classes/Components/ComponentInterfaces.h b/Engine/Source/Runtime/Engine/Classes/Components/ComponentInterfaces.h
--- a/Engine/Source/Runtime/Engine/Classes/Components/ComponentInterfaces.h
+++ b/Engine/Source/Runtime/Engine/Classes/Components/ComponentInterfaces.h
@@ -4,6 +4,7 @@
 #include "Containers/Array.h"
 #include "PrimitiveComponentId.h"
 #include "Templates/RefCounting.h"
+#include "Templates/Function.h"
 
 class UWorld;
 class FSceneInterface;
@@ -48,6 +49,11 @@ struct FPrimitiveStats
 	TArray<FPrimitiveLODStats> LODStats;
 };
 
+struct FComponentInterfaceImplementation
+{
+	UClass*	Class;
+	TFunction<void*(UObject*)> Resolver;
+};
 
 class IPrimitiveComponent
 {
@@ -92,6 +98,15 @@ public:
 	virtual HHitProxy* CreateMeshHitProxy(int32 SectionIndex, int32 MaterialIndex) = 0;
 #endif
 	virtual HHitProxy* CreatePrimitiveHitProxies(TArray<TRefCountPtr<HHitProxy> >& OutHitProxies) = 0;
+	
+	ENGINE_API static void AddImplementer(const FComponentInterfaceImplementation& Implementer);
+
+protected:
+
+	template<class T> 
+	friend class TComponentInterfaceIterator;
+
+	ENGINE_API static TArray<FComponentInterfaceImplementation> Implementers;
 };
 
 class IStaticMeshComponent
@@ -109,6 +124,16 @@ public:
 		// use the non-const version and return it as a const object to avoid duplicating the code in implementers
 		return (const_cast<IStaticMeshComponent*>(this))->GetPrimitiveComponentInterface();
 	}
+	
+	ENGINE_API static void AddImplementer(const FComponentInterfaceImplementation& Implementer);
+
+protected:
+
+	template<class T> 
+	friend class TComponentInterfaceIterator;
+
+	ENGINE_API static TArray<FComponentInterfaceImplementation> Implementers;
+	
 };
 
 
diff --git a/Engine/Source/Runtime/Engine/Private/Components/ComponentInterfaces.cpp b/Engine/Source/Runtime/Engine/Private/Components/ComponentInterfaces.cpp
new file mode 100644
--- /dev/null
+++ b/Engine/Source/Runtime/Engine/Private/Components/ComponentInterfaces.cpp
@@ -0,0 +1,21 @@
+// Copyright Epic Games, Inc. All Rights Reserved.
+
+#include "Components/ComponentInterfaces.h"
+#include "Components/ComponentInterfaceIterator.h"
+#include "HAL/IConsoleManager.h"
+#include "UnrealEngine.h"
+
+TArray<FComponentInterfaceImplementation> IPrimitiveComponent::Implementers;
+
+void IPrimitiveComponent::AddImplementer(const FComponentInterfaceImplementation& Implementer)
+{	
+	Implementers.Add(Implementer);
+
+}
+
+TArray<FComponentInterfaceImplementation> IStaticMeshComponent::Implementers;
+
+void IStaticMeshComponent::AddImplementer(const FComponentInterfaceImplementation& Implementer)
+{
+	Implementers.Add(Implementer);
+}
\ No newline at end of file
diff --git a/Engine/Source/Runtime/Engine/Private/ObjectCacheContext.cpp b/Engine/Source/Runtime/Engine/Private/ObjectCacheContext.cpp
--- a/Engine/Source/Runtime/Engine/Private/ObjectCacheContext.cpp
+++ b/Engine/Source/Runtime/Engine/Private/ObjectCacheContext.cpp
@@ -15,6 +15,7 @@
 #include "Logging/LogMacros.h"
 #include "HAL/LowLevelMemStats.h"
 #include "Misc/ScopeRWLock.h"
+#include "Components/ComponentInterfaceIterator.h"
 
 #define LOCTEXT_NAMESPACE "ObjectCache"
 
@@ -383,11 +384,11 @@ void Validate()
 	int32 ErrorCount = 0;
 	{
 		FObjectReverseLookupCache<UStaticMesh, IStaticMeshComponent> TempLookup;
-		for (TObjectIterator<UStaticMeshComponent> It(RF_ClassDefaultObject, true /*bIncludeDerivedClasses*/, GetObjectCacheInternalFlagsExclusion()); It; ++It)
+		for (TComponentInterfaceIterator<IStaticMeshComponent> It(RF_ClassDefaultObject, true /*bIncludeDerivedClasses*/, GetObjectCacheInternalFlagsExclusion()); It; ++It)
 		{
 			if (It->GetStaticMesh())
 			{
-				TempLookup.Update((*It)->GetStaticMeshComponentInterface(), { It->GetStaticMesh() } );
+				TempLookup.Update(*It, { It->GetStaticMesh() } );
 			}
 		}
 		ErrorCount += TempLookup.Compare(GStaticMeshToComponentLookupCache);
@@ -434,14 +435,14 @@ void Validate()
 	// Scan and compare UMaterialInterface -> UPrimitiveComponent
 	{
 		FObjectReverseLookupCache<UMaterialInterface, IPrimitiveComponent> TempLookup;
-		for (TObjectIterator<UPrimitiveComponent> It(RF_ClassDefaultObject, true /*bIncludeDerivedClasses*/, GetObjectCacheInternalFlagsExclusion()); It; ++It)
+		for (TComponentInterfaceIterator<IPrimitiveComponent> It(RF_ClassDefaultObject, true /*bIncludeDerivedClasses*/, GetObjectCacheInternalFlagsExclusion()); It; ++It)
 		{
-			UPrimitiveComponent* Component = *It;
-			if (Component->IsRenderStateCreated() && Component->SceneProxy && !Component->IsRenderStateDirty())
+			IPrimitiveComponent* Component = *It;
+			if (Component->IsRenderStateCreated() && Component->GetSceneProxy() && !Component->IsRenderStateDirty())
 			{
 				TArray<UMaterialInterface*> UsedMaterials;
 				Component->GetUsedMaterials(UsedMaterials);
-				TempLookup.Update(Component->GetPrimitiveComponentInterface(), UsedMaterials);
+				TempLookup.Update(Component, UsedMaterials);
 			}
 		}
 		ErrorCount += TempLookup.Compare(GMaterialToPrimitiveLookupCache);
@@ -674,9 +675,9 @@ TObjectCacheIterator<IPrimitiveComponent> FObjectCacheContext::GetPrimitiveCompo
 
 		TArray<IPrimitiveComponent*> Array;
 		Array.Reserve(4096);
-		for (TObjectIterator<UPrimitiveComponent> It(RF_ClassDefaultObject, true /*bIncludeDerivedClasses*/, GetObjectCacheInternalFlagsExclusion()); It; ++It)
+		for (TComponentInterfaceIterator<IPrimitiveComponent> It(RF_ClassDefaultObject, true /*bIncludeDerivedClasses*/, GetObjectCacheInternalFlagsExclusion()); It; ++It)
 		{
-			Array.Add((*It)->GetPrimitiveComponentInterface());
+			Array.Add(*It);
 		}
 		PrimitiveComponents = MoveTemp(Array);
 	}
@@ -694,9 +695,9 @@ TObjectCacheIterator<IStaticMeshComponent> FObjectCacheContext::GetStaticMeshCom
 
 		TArray<IStaticMeshComponent*> Array;
 		Array.Reserve(4096);
-		for (TObjectIterator<UStaticMeshComponent> It(RF_ClassDefaultObject, true /*bIncludeDerivedClasses*/, GetObjectCacheInternalFlagsExclusion()); It; ++It)
+		for (TComponentInterfaceIterator<IStaticMeshComponent> It(RF_ClassDefaultObject, true /*bIncludeDerivedClasses*/, GetObjectCacheInternalFlagsExclusion()); It; ++It)
 		{
-			Array.Add((*It)->GetStaticMeshComponentInterface());
+			Array.Add(*It);
 		}
 		StaticMeshComponents = MoveTemp(Array);
 	}
diff --git a/Engine/Source/Runtime/Engine/Private/UnrealEngine.cpp b/Engine/Source/Runtime/Engine/Private/UnrealEngine.cpp
--- a/Engine/Source/Runtime/Engine/Private/UnrealEngine.cpp
+++ b/Engine/Source/Runtime/Engine/Private/UnrealEngine.cpp
@@ -308,6 +308,9 @@ void FEngineModule::StartupModule()
 	USkinnedMeshComponent::BindWorldDelegates();
 #endif
 
+	IPrimitiveComponent::AddImplementer({UPrimitiveComponent::StaticClass(),  [](UObject* Obj){return Cast<UPrimitiveComponent>(Obj)->GetPrimitiveComponentInterface();}});
+	IStaticMeshComponent::AddImplementer({UStaticMeshComponent::StaticClass(), [](UObject* Obj){return Cast<UStaticMeshComponent>(Obj)->GetStaticMeshComponentInterface();}});
+
 #if OBJECT_TRACE_ENABLED
 	FObjectTrace::Init();
 #endif
