Fixed initialization error in LocalLightComponent.cpp

diff --git a/Engine/Source/Runtime/Engine/Private/Components/LocalLightComponent.cpp b/Engine/Source/Runtime/Engine/Private/Components/LocalLightComponent.cpp
--- a/Engine/Source/Runtime/Engine/Private/Components/LocalLightComponent.cpp
+++ b/Engine/Source/Runtime/Engine/Private/Components/LocalLightComponent.cpp
@@ -15,6 +15,7 @@ ULocalLightComponent::ULocalLightComponent(const FObjectInitializer& ObjectIniti
 	: Super(ObjectInitializer)
 {
 	Intensity = 5000;
+	IntensityUnits = ELightUnits::Unitless;
 	Radius_DEPRECATED = 1024.0f;
 	AttenuationRadius = 1000;
 	InverseExposureBlend = 0;
