Level Instance: Prevent label editing of sub selection
- Also move Brush specific code to ABrush::IsActorLabelEditable()

#rnx

diff --git a/Engine/Source/Runtime/Engine/Classes/Engine/Brush.h b/Engine/Source/Runtime/Engine/Classes/Engine/Brush.h
--- a/Engine/Source/Runtime/Engine/Classes/Engine/Brush.h
+++ b/Engine/Source/Runtime/Engine/Classes/Engine/Brush.h
@@ -183,6 +183,7 @@ public:
 	ENGINE_API virtual bool SupportsExternalPackaging() const override;
 	ENGINE_API virtual bool CanDeleteSelectedActor(FText& OutReason) const override;
 	ENGINE_API virtual bool CanReplaceSelectedActor(FText& OutReason) const override;
+	ENGINE_API virtual bool IsActorLabelEditable() const override;
 public:
 
 	ENGINE_API virtual void InitPosRotScale();
diff --git a/Engine/Source/Runtime/Engine/Private/ActorEditor.cpp b/Engine/Source/Runtime/Engine/Private/ActorEditor.cpp
--- a/Engine/Source/Runtime/Engine/Private/ActorEditor.cpp
+++ b/Engine/Source/Runtime/Engine/Private/ActorEditor.cpp
@@ -1271,7 +1271,12 @@ void AActor::SetActorLabel(const FString& NewActorLabelDirty, bool bMarkDirty)
 
 bool AActor::IsActorLabelEditable() const
 {
-	return bActorLabelEditable && !FActorEditorUtils::IsABuilderBrush(this);
+	if (IsInLevelInstance() && !IsInEditLevelInstance())
+	{
+		return false;
+	}
+
+	return bActorLabelEditable;
 }
 
 void AActor::ClearActorLabel()
diff --git a/Engine/Source/Runtime/Engine/Private/Brush.cpp b/Engine/Source/Runtime/Engine/Private/Brush.cpp
--- a/Engine/Source/Runtime/Engine/Private/Brush.cpp
+++ b/Engine/Source/Runtime/Engine/Private/Brush.cpp
@@ -285,6 +285,16 @@ bool ABrush::CanReplaceSelectedActor(FText& OutReason) const
 	return BrushUtils::CanDeleteOrReplaceCommon(this, OutReason);
 }
 
+bool ABrush::IsActorLabelEditable() const
+{
+	if (!Super::IsActorLabelEditable())
+	{
+		return false;
+	}
+
+	return !FActorEditorUtils::IsABuilderBrush(this);
+}
+
 void ABrush::PostLoad()
 {
 	Super::PostLoad();
