[shaders] enable printing of warnings by default in shader dev mode, and add a filtering option such that this can either be fully disabled or apply only to global shader jobs as desired (the default state is to print all warnings). as a bonus cleaned up the code around error/warning processing and related cvars.

diff --git a/Engine/Plugins/FX/Niagara/Source/Niagara/Private/NiagaraSystem.cpp b/Engine/Plugins/FX/Niagara/Source/Niagara/Private/NiagaraSystem.cpp
--- a/Engine/Plugins/FX/Niagara/Source/Niagara/Private/NiagaraSystem.cpp
+++ b/Engine/Plugins/FX/Niagara/Source/Niagara/Private/NiagaraSystem.cpp
@@ -48,6 +48,7 @@
 #include "UObject/LinkerLoad.h"
 #include "UObject/Package.h"
 #include "PipelineStateCache.h"
+#include "ShaderDiagnostics.h"
 #include "NiagaraDataChannel.h"
 #include "UObject/UObjectIterator.h"
 
@@ -3378,8 +3379,7 @@ bool UNiagaraSystem::RequestCompile(bool bForce, FNiagaraSystemUpdateContext* Op
 		ForceGraphToRecompileOnNextCheck();
 
 		// if we're forcing a recompile in development mode then flush the shader file cache to catch any datainterface files that may have been edited
-		static IConsoleVariable* CVarShaderDevMode = IConsoleManager::Get().FindConsoleVariable(TEXT("r.ShaderDevelopmentMode"));
-		if (CVarShaderDevMode && CVarShaderDevMode->GetInt() != 0)
+		if (IsShaderDevelopmentModeEnabled())
 		{
 			FlushShaderFileCache();
 		}
diff --git a/Engine/Source/Runtime/D3D12RHI/Private/D3D12Adapter.cpp b/Engine/Source/Runtime/D3D12RHI/Private/D3D12Adapter.cpp
--- a/Engine/Source/Runtime/D3D12RHI/Private/D3D12Adapter.cpp
+++ b/Engine/Source/Runtime/D3D12RHI/Private/D3D12Adapter.cpp
@@ -12,6 +12,7 @@ D3D12Adapter.cpp:D3D12 Adapter implementation.
 #include "Misc/CommandLine.h"
 #include "Misc/EngineVersion.h"
 #include "Misc/OutputDeviceRedirector.h"
+#include "ShaderDiagnostics.h"
 #include "DataDrivenShaderPlatformInfo.h"
 #include "Windows/AllowWindowsPlatformTypes.h"
 #include "GenericPlatform/GenericPlatformCrashContext.h"
@@ -490,10 +491,9 @@ void FD3D12Adapter::CreateRootDevice(bool bWithDebug)
 #if !PLATFORM_CPU_ARM_FAMILY && (PLATFORM_WINDOWS)
 	if (IsRHIDeviceAMD() && FD3D12DynamicRHI::GetD3DRHI()->GetAmdAgsContext())
 	{
-		auto* CVarShaderDevelopmentMode = IConsoleManager::Get().FindTConsoleVariableDataInt(TEXT("r.ShaderDevelopmentMode"));
 		auto* CVarDisableEngineAndAppRegistration = IConsoleManager::Get().FindTConsoleVariableDataInt(TEXT("r.DisableEngineAndAppRegistration"));
 
-		const bool bDisableEngineRegistration = (CVarShaderDevelopmentMode && CVarShaderDevelopmentMode->GetValueOnAnyThread() != 0) ||
+		const bool bDisableEngineRegistration = IsShaderDevelopmentModeEnabled() ||
 			(CVarDisableEngineAndAppRegistration && CVarDisableEngineAndAppRegistration->GetValueOnAnyThread() != 0);
 		const bool bDisableAppRegistration = bDisableEngineRegistration || !FApp::HasProjectName();
 
diff --git a/Engine/Source/Runtime/Engine/Private/Materials/MaterialShader.cpp b/Engine/Source/Runtime/Engine/Private/Materials/MaterialShader.cpp
--- a/Engine/Source/Runtime/Engine/Private/Materials/MaterialShader.cpp
+++ b/Engine/Source/Runtime/Engine/Private/Materials/MaterialShader.cpp
@@ -2250,14 +2250,9 @@ int32 FMaterialShaderMap::SubmitCompileJobs(uint32 CompilingShaderMapId,
 	FString DebugExtensionStr(TEXT(""));
 	FString DebugDescriptionStr(TEXT(""));
 
-	// DebugExtension and Description make the jobs unnecessarily different. Use them if the job cache is disabled or shader dev mode is on
-	static IConsoleVariable* CVarCacheEnabled = IConsoleManager::Get().FindConsoleVariable(TEXT("r.ShaderCompiler.JobCache"));
-	static IConsoleVariable* CVarShaderDevMode = IConsoleManager::Get().FindConsoleVariable(TEXT("r.ShaderDevelopmentMode"));
-	if ((CVarCacheEnabled && CVarCacheEnabled->GetInt() == 0) || (CVarShaderDevMode && CVarShaderDevMode->GetInt() != 0))
-	{
-		DebugExtensionStr = FString::Printf(TEXT("_%08x%08x"), ShaderMapId.BaseMaterialId.A, ShaderMapId.BaseMaterialId.B);
-		DebugDescriptionStr = GetDebugDescription();
-	}
+	DebugExtensionStr = FString::Printf(TEXT("_%08x%08x"), ShaderMapId.BaseMaterialId.A, ShaderMapId.BaseMaterialId.B);
+	DebugDescriptionStr = GetDebugDescription();
+
 	const TCHAR* DebugExtension = *DebugExtensionStr;
 	const TCHAR* DebugDescription = *DebugDescriptionStr;
 
diff --git a/Engine/Source/Runtime/Engine/Private/Materials/MeshMaterialShader.cpp b/Engine/Source/Runtime/Engine/Private/Materials/MeshMaterialShader.cpp
--- a/Engine/Source/Runtime/Engine/Private/Materials/MeshMaterialShader.cpp
+++ b/Engine/Source/Runtime/Engine/Private/Materials/MeshMaterialShader.cpp
@@ -6,6 +6,7 @@
 
 #include "MeshMaterialShader.h"
 #include "ShaderCompiler.h"
+#include "ShaderDiagnostics.h"
 
 #if WITH_EDITOR
 
@@ -30,8 +31,7 @@ static void PrepareMeshMaterialShaderCompileJob(EShaderPlatform Platform,
 	NewJob->bIsDefaultMaterial = Material->IsDefaultMaterial();
 	NewJob->bIsGlobalShader = false;
 
-	static IConsoleVariable* CVarShaderDevMode = IConsoleManager::Get().FindConsoleVariable(TEXT("r.ShaderDevelopmentMode"));
-	if (CVarShaderDevMode && CVarShaderDevMode->GetInt() != 0)
+	if (IsShaderDevelopmentModeEnabled())
 	{
 		NewJob->bErrorsAreLikelyToBeCode = true;
 	}
diff --git a/Engine/Source/Runtime/Engine/Private/ShaderCompiler/ShaderCompiler.cpp b/Engine/Source/Runtime/Engine/Private/ShaderCompiler/ShaderCompiler.cpp
--- a/Engine/Source/Runtime/Engine/Private/ShaderCompiler/ShaderCompiler.cpp
+++ b/Engine/Source/Runtime/Engine/Private/ShaderCompiler/ShaderCompiler.cpp
@@ -51,6 +51,7 @@
 #include "Serialization/MemoryReader.h"
 #include "Serialization/NameAsStringProxyArchive.h"
 #include "ShaderCodeLibrary.h"
+#include "ShaderDiagnostics.h"
 #include "ShaderPlatformCachedIniValue.h"
 #include "StaticBoundShaderState.h"
 #include "StereoRenderUtils.h"
@@ -2134,14 +2135,6 @@ static FAutoConsoleVariableRef CVarLogShaderCompilerStats(
 	TEXT("When set to 1, Log detailed shader compiler stats.")
 );
 
-
-static int32 GShowShaderWarnings = 0;
-static FAutoConsoleVariableRef CVarShowShaderWarnings(
-	TEXT("r.ShowShaderCompilerWarnings"),
-	GShowShaderWarnings,
-	TEXT("When set to 1, will display all warnings.")
-	);
-
 static int32 GForceAllCoresForShaderCompiling = 0;
 static FAutoConsoleVariableRef CVarForceAllCoresForShaderCompiling(
 	TEXT("r.ForceAllCoresForShaderCompiling"),
@@ -2803,23 +2796,13 @@ bool FShaderCompileUtilities::DoWriteTasks(const TArray<FShaderCommonCompileJobP
 	return DoWriteTasksInner(QueuedJobs, InTransferFile, BuildDistributionController, bUseRelativePaths, bCompressTaskFile);
 }
 
-struct FShaderErrorInfo
-{
-	TArray<FShaderCommonCompileJob*> ErrorJobs;
-	TArray<FString> UniqueErrors;
-	TArray<FString> UniqueErrorPrefixes;
-	TArray<FString> UniqueWarnings;
-	TArray<EShaderPlatform> ErrorPlatforms;
-	FString TargetShaderPlatformString;
-};
-
-static void BuildErrorStringAndReport(const FShaderErrorInfo& ErrorInfo, FString& ErrorString)
+static void BuildErrorStringAndReport(const FShaderDiagnosticInfo& DiagInfo, FString& ErrorString)
 {
 	bool bReportedDebugInfo = false;
 
-	for (int32 ErrorIndex = 0; ErrorIndex < ErrorInfo.UniqueErrors.Num(); ErrorIndex++)
+	for (int32 ErrorIndex = 0; ErrorIndex < DiagInfo.UniqueErrors.Num(); ErrorIndex++)
 	{
-		FString UniqueErrorString = ErrorInfo.UniqueErrorPrefixes[ErrorIndex] + ErrorInfo.UniqueErrors[ErrorIndex] + TEXT("\n");
+		FString UniqueErrorString = DiagInfo.UniqueErrorPrefixes[ErrorIndex] + DiagInfo.UniqueErrors[ErrorIndex] + TEXT("\n");
 
 		if (FPlatformMisc::IsDebuggerPresent())
 		{
@@ -2865,22 +2848,6 @@ static bool ReadSingleJob(FShaderCompileJob* CurrentJob, FArchive& WorkerOutputF
 	return false;
 }
 
-static FString GetSingleJobCompilationDump(const FShaderCompileJob* SingleJob)
-{
-	if (!SingleJob)
-	{
-		return TEXT("Internal error, not a Job!");
-	}
-	FString String = SingleJob->Input.GenerateShaderName();
-	if (SingleJob->Key.VFType)
-	{
-		String += FString::Printf(TEXT(" VF '%s'"), SingleJob->Key.VFType->GetName());
-	}
-	String += FString::Printf(TEXT(" Type '%s'"), SingleJob->Key.ShaderType->GetName());
-	String += FString::Printf(TEXT(" '%s' Entry '%s' Permutation %i "), *SingleJob->Input.VirtualSourceFilePath, *SingleJob->Input.EntryPointName, SingleJob->Key.PermutationId);
-	return String;
-}
-
 static const TCHAR* GetCompileJobSuccessText(FShaderCompileJob* SingleJob)
 {
 	if (SingleJob)
@@ -3232,15 +3199,6 @@ static bool CheckSingleJob(const FShaderCompileJob& SingleJob, TArray<FString>&
 		checkf(SingleJob.Output.ShaderCode.GetShaderCodeSize() > 0, TEXT("Abnormal shader code size for a succeded job: %d bytes"), SingleJob.Output.ShaderCode.GetShaderCodeSize());
 	}
 
-	if (GShowShaderWarnings || !SingleJob.bSucceeded)
-	{
-		for (int32 ErrorIndex = 0; ErrorIndex < SingleJob.Output.Errors.Num(); ErrorIndex++)
-		{
-			const FShaderCompilerError& InError = SingleJob.Output.Errors[ErrorIndex];
-			OutErrors.AddUnique(InError.GetErrorStringWithLineMarker());
-		}
-	}
-
 	bool bSucceeded = SingleJob.bSucceeded;
 
 	if (SingleJob.Key.ShaderType)
@@ -3270,123 +3228,6 @@ static bool CheckSingleJob(const FShaderCompileJob& SingleJob, TArray<FString>&
 };
 #endif // WITH_EDITOR
 
-static int32 AddAndProcessErrorsForFailedJobFiltered(FShaderCompileJob& CurrentJob, FShaderErrorInfo& OutShaderErrorInfo, const TCHAR* FilterMessage)
-{
-	int32 NumAddedErrors = 0;
-
-	bool bReportedDebugInfo = false;
-
-	for (int32 ErrorIndex = 0; ErrorIndex < CurrentJob.Output.Errors.Num(); ErrorIndex++)
-	{
-		FShaderCompilerError& CurrentError = CurrentJob.Output.Errors[ErrorIndex];
-		FString CurrentErrorString = CurrentError.GetErrorString();
-
-		// Include warnings if LogShaders is unsuppressed, otherwise only include filtered messages
-		if (UE_LOG_ACTIVE(LogShaders, Log) || FilterMessage == nullptr || CurrentError.StrippedErrorMessage.Contains(FilterMessage))
-		{
-			// Extract source location from error message if the shader backend doesn't provide it separated from the stripped message
-			CurrentError.ExtractSourceLocation();
-
-			// Remap filenames
-			if (CurrentError.ErrorVirtualFilePath == TEXT("/Engine/Generated/Material.ush"))
-			{
-				// MaterialTemplate.usf is dynamically included as Material.usf
-				// Currently the material translator does not add new lines when filling out MaterialTemplate.usf,
-				// So we don't need the actual filled out version to find the line of a code bug.
-				CurrentError.ErrorVirtualFilePath = TEXT("/Engine/Private/MaterialTemplate.ush");
-			}
-			else if (CurrentError.ErrorVirtualFilePath.Contains(TEXT("memory")))
-			{
-				check(CurrentJob.Key.ShaderType);
-
-				// Files passed to the shader compiler through memory will be named memory
-				// Only the shader's main file is passed through memory without a filename
-				CurrentError.ErrorVirtualFilePath = FString(CurrentJob.Key.ShaderType->GetShaderFilename());
-			}
-			else if (CurrentError.ErrorVirtualFilePath == TEXT("/Engine/Generated/VertexFactory.ush"))
-			{
-				// VertexFactory.usf is dynamically included from whichever vertex factory the shader was compiled with.
-				check(CurrentJob.Key.VFType);
-				CurrentError.ErrorVirtualFilePath = FString(CurrentJob.Key.VFType->GetShaderFilename());
-			}
-			else if (CurrentError.ErrorVirtualFilePath == TEXT("") && CurrentJob.Key.ShaderType)
-			{
-				// Some shader compiler errors won't have a file and line number, so we just assume the error happened in file containing the entrypoint function.
-				CurrentError.ErrorVirtualFilePath = FString(CurrentJob.Key.ShaderType->GetShaderFilename());
-			}
-
-			if (OutShaderErrorInfo.UniqueErrors.Find(CurrentErrorString) == INDEX_NONE)
-			{
-				// build up additional info in a "prefix" string; only do this once for each unique error
-				FString UniqueErrorPrefix;
-
-				// If we dumped the shader info, add it before the first error string
-				if (!GIsBuildMachine && !bReportedDebugInfo && CurrentJob.Input.DumpDebugInfoPath.Len() > 0)
-				{
-					UniqueErrorPrefix += FString::Printf(TEXT("Shader debug info dumped to: \"%s\"\n"), *CurrentJob.Input.DumpDebugInfoPath);
-					bReportedDebugInfo = true;
-				}
-
-				if (CurrentJob.Key.ShaderType)
-				{
-					// Construct a path that will enable VS.NET to find the shader file, relative to the solution
-					const FString SolutionPath = FPaths::RootDir();
-					FString ShaderFilePath = IFileManager::Get().ConvertToAbsolutePathForExternalAppForRead(*CurrentError.GetShaderSourceFilePath());
-					UniqueErrorPrefix += FString::Printf(TEXT("%s(%s): Shader %s, Permutation %d, VF %s:\n\t"),
-						*ShaderFilePath,
-						*CurrentError.ErrorLineString,
-						CurrentJob.Key.ShaderType->GetName(),
-						CurrentJob.Key.PermutationId,
-						CurrentJob.Key.VFType ? CurrentJob.Key.VFType->GetName() : TEXT("None"));
-				}
-				else
-				{
-					UniqueErrorPrefix += FString::Printf(TEXT("%s(0): "),
-						*CurrentJob.Input.VirtualSourceFilePath);
-				}
-
-				OutShaderErrorInfo.UniqueErrors.Add(CurrentErrorString);
-				OutShaderErrorInfo.UniqueErrorPrefixes.Add(UniqueErrorPrefix);
-				OutShaderErrorInfo.ErrorJobs.AddUnique(&CurrentJob);
-			}
-			++NumAddedErrors;
-		}
-	}
-
-	return NumAddedErrors;
-}
-
-static void AddAndProcessErrorsForFailedJob(FShaderCompileJob& CurrentJob, FShaderErrorInfo& OutShaderErrorInfo)
-{
-	OutShaderErrorInfo.ErrorPlatforms.AddUnique((EShaderPlatform)CurrentJob.Input.Target.Platform);
-
-	if (CurrentJob.Output.Errors.Num() == 0)
-	{
-		// Job hard crashed
-		FShaderCompilerError Error(*(FString("Internal Error!\n\t") + GetSingleJobCompilationDump(&CurrentJob)));
-		CurrentJob.Output.Errors.Add(Error);
-	}
-
-	// If we filter all error messages because they are interpreted as warnings, we have to assume all error messages are in fact errors and not warnings.
-	// In that case, add jobs again without a filter; e.g. when the stripped message starts with "Internal exception".
-	if (AddAndProcessErrorsForFailedJobFiltered(CurrentJob, OutShaderErrorInfo, TEXT("error")) == 0)
-	{
-		AddAndProcessErrorsForFailedJobFiltered(CurrentJob, OutShaderErrorInfo, nullptr);
-	}
-}
-
-static void AddWarningsForJob(const FShaderCompileJob& CurrentJob, FShaderErrorInfo& OutShaderErrorInfo)
-{
-	if (GShowShaderWarnings && CurrentJob.bSucceeded)
-	{
-		for (int32 ErrorIndex = 0; ErrorIndex < CurrentJob.Output.Errors.Num(); ErrorIndex++)
-		{
-			// If the job succeeded the Errors array will contain warnings.
-			OutShaderErrorInfo.UniqueWarnings.AddUnique(CurrentJob.Output.Errors[ErrorIndex].GetErrorString());
-		}
-	}
-}
-
 /** Information tracked for each shader compile worker process instance. */
 struct FShaderCompileWorkerInfo
 {
@@ -3685,7 +3526,7 @@ void FShaderCompileThreadRunnable::PushCompletedJobsToManager()
 					if (const FShaderCompileJob* SingleJob = Job.GetSingleShaderJob())
 					{
 						const TCHAR* JobName = Manager->bLogJobCompletionTimes ? *SingleJob->Input.DebugGroupName : SingleJob->Key.ShaderType->GetName();
-						JobNames += FString::Printf(TEXT("%s [WorkerTime=%.3fs]"), JobName, SingleJob->Output.CompileTime);
+						JobNames += FString::Printf(TEXT("%s:%d(%s) [WorkerTime=%.3fs]"), JobName, SingleJob->Key.PermutationId, *SingleJob->Input.ShaderFormat.ToString(), SingleJob->Output.CompileTime);
 					}
 					else
 					{
@@ -6700,8 +6541,7 @@ void FShaderCompilingManager::ProcessCompiledShaderMaps(
 		if (CompilingShaderMap)
 		{
 			TArray<TRefCountPtr<FMaterial>>& MaterialDependencies = CompilingShaderMap->CompilingMaterialDependencies;
-
-			TArray<FString> Errors;
+			FShaderDiagnosticInfo ErrorInfo(FinishedJobs);
 
 			bool bSuccess = true;
 			for (int32 JobIndex = 0; JobIndex < FinishedJobs.Num(); JobIndex++)
@@ -6710,14 +6550,14 @@ void FShaderCompilingManager::ProcessCompiledShaderMaps(
 
 				if (FShaderCompileJob* SingleJob = CurrentJob.GetSingleShaderJob())
 				{
-					const bool bCheckSucceeded = CheckSingleJob(*SingleJob, Errors);
+					const bool bCheckSucceeded = CheckSingleJob(*SingleJob, ErrorInfo.UniqueErrors);
 					bSuccess = bCheckSucceeded && bSuccess;
 				}
 				else if (FShaderPipelineCompileJob* PipelineJob = CurrentJob.GetShaderPipelineJob())
 				{
 					for (int32 Index = 0; Index < PipelineJob->StageJobs.Num(); ++Index)
 					{
-						const bool bCheckSucceeded = CheckSingleJob(*PipelineJob->StageJobs[Index], Errors);
+						const bool bCheckSucceeded = CheckSingleJob(*PipelineJob->StageJobs[Index], ErrorInfo.UniqueErrors);
 						bSuccess = PipelineJob->StageJobs[Index]->bSucceeded && bCheckSucceeded && bSuccess;
 					}
 				}
@@ -6808,14 +6648,14 @@ void FShaderCompilingManager::ProcessCompiledShaderMaps(
 						++NumIncompleteMaterials;
 					}
 
-					if (GShowShaderWarnings && Errors.Num() > 0)
+					if (ErrorInfo.UniqueWarnings.Num() > 0)
 					{
 						UE_LOG(LogShaderCompilers, Warning, TEXT("Warnings while compiling Material %s for platform %s:"),
 							*Material->GetDebugName(),
 							*LegacyShaderPlatformToShaderFormat(CompilingShaderMap->GetShaderPlatform()).ToString());
-						for (int32 ErrorIndex = 0; ErrorIndex < Errors.Num(); ErrorIndex++)
+						for (const FString& UniqueWarning : ErrorInfo.UniqueWarnings)
 						{
-							UE_LOG(LogShaders, Warning, TEXT("  %s"), *Errors[ErrorIndex]);
+							UE_LOG(LogShaders, Warning, TEXT("  %s"), *UniqueWarning);
 						}
 					}
 				}
@@ -6823,7 +6663,7 @@ void FShaderCompilingManager::ProcessCompiledShaderMaps(
 				{
 					bReleaseCompilingId = true;
 					// Propagate error messages
-					Material->CompileErrors = Errors;
+					Material->CompileErrors = ErrorInfo.UniqueErrors;
 
 					MaterialsToUpdate.Add(Material, nullptr);
 
@@ -6832,9 +6672,9 @@ void FShaderCompilingManager::ProcessCompiledShaderMaps(
 						FString ErrorString;
 
 						// Log the errors unsuppressed before the fatal error, so it's always obvious from the log what the compile error was
-						for (int32 ErrorIndex = 0; ErrorIndex < Errors.Num(); ErrorIndex++)
+						for (const FString& UniqueError : ErrorInfo.UniqueErrors)
 						{
-							ErrorString += FString::Printf(TEXT("  %s\n"), *Errors[ErrorIndex]);
+							ErrorString += FString::Printf(TEXT("  %s\n"), *UniqueError);
 						}
 
 						ErrorString += FString::Printf(TEXT("Failed to compile default material %s!"), *Material->GetBaseMaterialPathName());
@@ -6855,9 +6695,9 @@ void FShaderCompilingManager::ProcessCompiledShaderMaps(
 					ErrorString += FString::Printf(TEXT("Failed to compile Material %s for platform %s, Default Material will be used in game.\n"),
 						*Material->GetDebugName(), *LegacyShaderPlatformToShaderFormat(CompilingShaderMap->GetShaderPlatform()).ToString());
 
-					for (int32 ErrorIndex = 0; ErrorIndex < Errors.Num(); ErrorIndex++)
+					for (const FString& UniqueError : ErrorInfo.UniqueErrors)
 					{
-						FString ErrorMessage = Errors[ErrorIndex];
+						FString ErrorMessage = UniqueError;
 						// Work around build machine string matching heuristics that will cause a cook to fail
 						ErrorMessage.ReplaceInline(TEXT("error "), TEXT("err0r "), ESearchCase::CaseSensitive);
 						ErrorString += FString::Printf(TEXT("  %s\n"), *ErrorMessage);
@@ -7036,63 +6876,6 @@ FShaderCompileMemoryUsage FShaderCompilingManager::GetExternalMemoryUsage()
 	return TotalMemoryUsage;
 }
 
-static bool GatherUniqueErrors(const TArray<FShaderCommonCompileJobPtr>& CompleteJobs, FShaderErrorInfo& OutShaderErrorInfo)
-{
-	// Gather unique errors
-	for (int32 JobIndex = 0; JobIndex < CompleteJobs.Num(); JobIndex++)
-	{
-		FShaderCommonCompileJob& CurrentJob = *CompleteJobs[JobIndex];
-		if (!CurrentJob.bSucceeded)
-		{
-			FShaderCompileJob* SingleJob = CurrentJob.GetSingleShaderJob();
-			if (SingleJob)
-			{
-				AddAndProcessErrorsForFailedJob(*SingleJob, OutShaderErrorInfo);
-			}
-			else
-			{
-				FShaderPipelineCompileJob* PipelineJob = CurrentJob.GetShaderPipelineJob();
-				check(PipelineJob);
-				for (TRefCountPtr<FShaderCompileJob>& CommonJob : PipelineJob->StageJobs)
-				{
-					AddAndProcessErrorsForFailedJob(*CommonJob, OutShaderErrorInfo);
-				}
-			}
-		}
-		else if (GShowShaderWarnings)
-		{
-			const FShaderCompileJob* SingleJob = CurrentJob.GetSingleShaderJob();
-			if (SingleJob)
-			{
-				AddWarningsForJob(*SingleJob, OutShaderErrorInfo);
-			}
-			else
-			{
-				const FShaderPipelineCompileJob* PipelineJob = CurrentJob.GetShaderPipelineJob();
-				check(PipelineJob);
-				for (const TRefCountPtr<FShaderCompileJob>& CommonJob : PipelineJob->StageJobs)
-				{
-					AddWarningsForJob(*CommonJob, OutShaderErrorInfo);
-				}
-			}
-		}
-	}
-
-	for (int32 PlatformIndex = 0; PlatformIndex < OutShaderErrorInfo.ErrorPlatforms.Num(); PlatformIndex++)
-	{
-		if (OutShaderErrorInfo.TargetShaderPlatformString.IsEmpty())
-		{
-			OutShaderErrorInfo.TargetShaderPlatformString = FDataDrivenShaderPlatformInfo::GetName(OutShaderErrorInfo.ErrorPlatforms[PlatformIndex]).ToString();
-		}
-		else
-		{
-			OutShaderErrorInfo.TargetShaderPlatformString += FString(TEXT(", ")) + FDataDrivenShaderPlatformInfo::GetName(OutShaderErrorInfo.ErrorPlatforms[PlatformIndex]).ToString();
-		}
-	}
-
-	return OutShaderErrorInfo.UniqueErrors.Num() > 0;
-}
-
 bool FShaderCompilingManager::HandlePotentialRetryOnError(TMap<int32, FShaderMapFinalizeResults>& CompletedShaderMaps)
 {
 	if (FApp::IsUnattended())
@@ -7130,21 +6913,20 @@ bool FShaderCompilingManager::HandlePotentialRetryOnError(TMap<int32, FShaderMap
 				|| It.Key() == GlobalShaderMapId)
 			{
 				TArray<FShaderCommonCompileJobPtr>& CompleteJobs = Results.FinishedJobs;
-				FShaderErrorInfo ShaderErrorInfo;
-				GatherUniqueErrors(CompleteJobs, ShaderErrorInfo);
+				FShaderDiagnosticInfo ShaderDiagInfo(CompleteJobs);
 
 				const TCHAR* MaterialName = ShaderMap ? ShaderMap->GetFriendlyName() : TEXT("global shaders");
-				FString ErrorString = FString::Printf(TEXT("%i Shader compiler errors compiling %s for platform %s:"), ShaderErrorInfo.UniqueErrors.Num(), MaterialName, *ShaderErrorInfo.TargetShaderPlatformString);
+				FString ErrorString = FString::Printf(TEXT("%i Shader compiler errors compiling %s for platform %s:"), ShaderDiagInfo.UniqueErrors.Num(), MaterialName, *ShaderDiagInfo.TargetShaderPlatformString);
 				UE_LOG(LogShaderCompilers, Warning, TEXT("%s"), *ErrorString);
 				ErrorString += TEXT("\n");
 
 				bool bAnyErrorLikelyToBeCodeError = false;
-				for (const FShaderCommonCompileJob* Job : ShaderErrorInfo.ErrorJobs)
+				for (const FShaderCommonCompileJob* Job : ShaderDiagInfo.ErrorJobs)
 				{
 					bAnyErrorLikelyToBeCodeError |= Job->bErrorsAreLikelyToBeCode;
 				}
 
-				BuildErrorStringAndReport(ShaderErrorInfo, ErrorString);
+				BuildErrorStringAndReport(ShaderDiagInfo, ErrorString);
 
 				if (UE_LOG_ACTIVE(LogShaders, Log) && (bAnyErrorLikelyToBeCodeError || bPromptToRetryFailedShaderCompiles || bSpecialEngineMaterial))
 				{
@@ -10598,27 +10380,22 @@ void ProcessCompiledGlobalShaders(const TArray<FShaderCommonCompileJobPtr>& Comp
 
 	UE_LOG(LogShaders, Verbose, TEXT("Compiled %u global shaders"), CompilationResults.Num());
 
-	FShaderErrorInfo ShaderErrorInfo;
-	GatherUniqueErrors(CompilationResults, ShaderErrorInfo);
+	FShaderDiagnosticInfo ShaderDiagInfo(CompilationResults);
 
 	// Report unique errors for global shaders.
-	for (int32 ErrorIndex = 0; ErrorIndex < ShaderErrorInfo.UniqueErrors.Num(); ++ErrorIndex)
+	for (int32 ErrorIndex = 0; ErrorIndex < ShaderDiagInfo.UniqueErrors.Num(); ++ErrorIndex)
 	{
-		FString ErrorString = ShaderErrorInfo.UniqueErrorPrefixes[ErrorIndex] + ShaderErrorInfo.UniqueErrors[ErrorIndex];
+		FString ErrorString = ShaderDiagInfo.UniqueErrorPrefixes[ErrorIndex] + ShaderDiagInfo.UniqueErrors[ErrorIndex];
 		UE_LOGFMT_NSLOC(LogShaders, Error, "Shaders", "GlobalShaderCompileError", "{ErrorMessage}", 
 			("ErrorMessage", ErrorString));
 	}
 
-	if (GShowShaderWarnings)
+	for (const FString& WarningString : ShaderDiagInfo.UniqueWarnings)
 	{
-		for (const FString& WarningString : ShaderErrorInfo.UniqueWarnings)
-		{
-			UE_LOGFMT_NSLOC(LogShaders, Warning, "Shaders", "GlobalShaderCompileWarning", "{WarningMessage}",
-				("WarningMessage", WarningString));
-		}
+		UE_LOGFMT_NSLOC(LogShaders, Warning, "Shaders", "GlobalShaderCompileWarning", "{WarningMessage}", ("WarningMessage", WarningString));
 	}
 
-	const int32 UniqueErrorCount = ShaderErrorInfo.UniqueErrors.Num();
+	const int32 UniqueErrorCount = ShaderDiagInfo.UniqueErrors.Num();
 	if (UniqueErrorCount)
 	{
 		const TCHAR* RetryMsg = TEXT(" Enable 'r.ShaderDevelopmentMode' in ConsoleVariables.ini for retries.");
@@ -10626,7 +10403,7 @@ void ProcessCompiledGlobalShaders(const TArray<FShaderCommonCompileJobPtr>& Comp
 		{
 			UE_LOGFMT_NSLOC(LogShaders, Fatal, "Shaders", "GlobalShadersCompilationFailed", "{NumErrors} Shader compiler errors compiling GlobalShaders for platform {Platform}. {RetryMsg}",
 				("NumErrors", UniqueErrorCount),
-				("Platform", ShaderErrorInfo.TargetShaderPlatformString),
+				("Platform", ShaderDiagInfo.TargetShaderPlatformString),
 				("RetryMsg", IsRunningCommandlet() ? TEXT("") : RetryMsg)
 			);
 		}
@@ -10634,7 +10411,7 @@ void ProcessCompiledGlobalShaders(const TArray<FShaderCommonCompileJobPtr>& Comp
 		{
 			UE_LOGFMT_NSLOC(LogShaders, Error, "Shaders", "GlobalShadersCompilationFailed", "{NumErrors} Shader compiler errors compiling GlobalShaders for platform {Platform}. {RetryMsg}",
 				("NumErrors", UniqueErrorCount),
-				("Platform", ShaderErrorInfo.TargetShaderPlatformString),
+				("Platform", ShaderDiagInfo.TargetShaderPlatformString),
 				("RetryMsg", IsRunningCommandlet() ? TEXT("") : RetryMsg)
 			);
 		}
diff --git a/Engine/Source/Runtime/RenderCore/Internal/ShaderDiagnostics.h b/Engine/Source/Runtime/RenderCore/Internal/ShaderDiagnostics.h
new file mode 100644
--- /dev/null
+++ b/Engine/Source/Runtime/RenderCore/Internal/ShaderDiagnostics.h
@@ -0,0 +1,27 @@
+// Copyright Epic Games, Inc. All Rights Reserved.
+
+#pragma once
+
+#include "Containers/Array.h"
+#include "RHIShaderPlatform.h"
+#include "ShaderCompilerJobTypes.h"
+
+struct FShaderDiagnosticInfo
+{
+	TArray<FShaderCommonCompileJob*> ErrorJobs;
+	TArray<FString> UniqueErrors;
+	TArray<FString> UniqueErrorPrefixes;
+	TArray<FString> UniqueWarnings;
+	TArray<EShaderPlatform> ErrorPlatforms;
+	FString TargetShaderPlatformString;
+
+	RENDERCORE_API FShaderDiagnosticInfo(const TArray<FShaderCommonCompileJobPtr>& Jobs);
+
+private:
+	void AddAndProcessErrorsForJob(FShaderCommonCompileJob& Job);
+	void AddWarningsForJob(const FShaderCommonCompileJob& Job);
+};
+
+RENDERCORE_API FString GetSingleJobCompilationDump(const FShaderCompileJob* SingleJob);
+RENDERCORE_API bool IsShaderDevelopmentModeEnabled();
+
diff --git a/Engine/Source/Runtime/RenderCore/Private/ShaderCore.cpp b/Engine/Source/Runtime/RenderCore/Private/ShaderCore.cpp
--- a/Engine/Source/Runtime/RenderCore/Private/ShaderCore.cpp
+++ b/Engine/Source/Runtime/RenderCore/Private/ShaderCore.cpp
@@ -35,6 +35,7 @@
 #include "ShaderCompilerCore.h"
 #include "ShaderCompilerDefinitions.h"
 #include "ShaderCompilerJobTypes.h"
+#include "ShaderDiagnostics.h"
 #include "Stats/StatsMisc.h"
 #include "String/Find.h"
 #include "Tasks/Task.h"
@@ -50,11 +51,6 @@
 #include "Windows/HideWindowsPlatformTypes.h"
 #endif
 
-static TAutoConsoleVariable<int32> CVarShaderDevelopmentMode(
-	TEXT("r.ShaderDevelopmentMode"),
-	0,
-	TEXT("0: Default, 1: Enable various shader development utilities, such as the ability to retry on failed shader compile, and extra logging as shaders are compiled."),
-	ECVF_Default);
 
 static TAutoConsoleVariable<bool> CVarDumpDebugInfoForCacheHits(
 	TEXT("r.ShaderCompiler.DumpDebugInfoForCacheHits"),
@@ -68,7 +64,7 @@ void UpdateShaderDevelopmentMode()
 	// r.ShaderDevelopmentMode==1 results in all LogShaders log messages being displayed.
 	// if r.ShaderDevelopmentMode isn't set, we leave the category alone (it defaults to Error, but we can be overriding it to something higher)
 	bool bLogShadersUnsuppressed = UE_LOG_ACTIVE(LogShaders, Log);
-	bool bDesiredLogShadersUnsuppressed = CVarShaderDevelopmentMode.GetValueOnGameThread() == 1;
+	bool bDesiredLogShadersUnsuppressed = IsShaderDevelopmentModeEnabled();
 
 	if (bLogShadersUnsuppressed != bDesiredLogShadersUnsuppressed)
 	{
diff --git a/Engine/Source/Runtime/RenderCore/Private/ShaderDiagnostics.cpp b/Engine/Source/Runtime/RenderCore/Private/ShaderDiagnostics.cpp
new file mode 100644
--- /dev/null
+++ b/Engine/Source/Runtime/RenderCore/Private/ShaderDiagnostics.cpp
@@ -0,0 +1,208 @@
+#include "ShaderDiagnostics.h"
+
+#include "DataDrivenShaderPlatformInfo.h"
+#include "ShaderCompilerJobTypes.h"
+
+static TAutoConsoleVariable<int32> CVarShaderDevelopmentMode(
+	TEXT("r.ShaderDevelopmentMode"),
+	0,
+	TEXT("0: Default, 1: Enable various shader development utilities, such as the ability to retry on failed shader compile, and extra logging as shaders are compiled."),
+	ECVF_Default);
+
+static TAutoConsoleVariable<int32> CVarShowShaderWarnings(
+	TEXT("r.ShowShaderCompilerWarnings"),
+	0,
+	TEXT("When set to 1, will display all warnings. Note that this flag is ignored if r.ShaderDevelopmentMode=1 (in dev mode warnings are shown by default).")
+	);
+
+static TAutoConsoleVariable<int32> CVarShaderWarningsFilter(
+	TEXT("r.ShaderCompilerWarningsFilter"),
+	2,
+	TEXT("Additional filtering for shader warnings; 2=show all, 1=show global shader warnings only, 0=show no shader warnings.")
+);
+
+
+int32 AddAndProcessErrorsForFailedJobFiltered(FShaderCompileJob& CurrentJob, FShaderDiagnosticInfo& OutShaderDiagInfo, const TCHAR* FilterMessage)
+{
+	int32 NumAddedErrors = 0;
+
+	bool bReportedDebugInfo = false;
+
+	for (int32 ErrorIndex = 0; ErrorIndex < CurrentJob.Output.Errors.Num(); ErrorIndex++)
+	{
+		FShaderCompilerError& CurrentError = CurrentJob.Output.Errors[ErrorIndex];
+		FString CurrentErrorString = CurrentError.GetErrorString();
+
+		// Include warnings if LogShaders is unsuppressed, otherwise only include filtered messages
+		if (UE_LOG_ACTIVE(LogShaders, Log) || FilterMessage == nullptr || CurrentError.StrippedErrorMessage.Contains(FilterMessage))
+		{
+			// Extract source location from error message if the shader backend doesn't provide it separated from the stripped message
+			CurrentError.ExtractSourceLocation();
+
+			// Remap filenames
+			if (CurrentError.ErrorVirtualFilePath == TEXT("/Engine/Generated/Material.ush"))
+			{
+				// MaterialTemplate.usf is dynamically included as Material.usf
+				// Currently the material translator does not add new lines when filling out MaterialTemplate.usf,
+				// So we don't need the actual filled out version to find the line of a code bug.
+				CurrentError.ErrorVirtualFilePath = TEXT("/Engine/Private/MaterialTemplate.ush");
+			}
+			else if (CurrentError.ErrorVirtualFilePath.Contains(TEXT("memory")))
+			{
+				check(CurrentJob.Key.ShaderType);
+
+				// Files passed to the shader compiler through memory will be named memory
+				// Only the shader's main file is passed through memory without a filename
+				CurrentError.ErrorVirtualFilePath = FString(CurrentJob.Key.ShaderType->GetShaderFilename());
+			}
+			else if (CurrentError.ErrorVirtualFilePath == TEXT("/Engine/Generated/VertexFactory.ush"))
+			{
+				// VertexFactory.usf is dynamically included from whichever vertex factory the shader was compiled with.
+				check(CurrentJob.Key.VFType);
+				CurrentError.ErrorVirtualFilePath = FString(CurrentJob.Key.VFType->GetShaderFilename());
+			}
+			else if (CurrentError.ErrorVirtualFilePath == TEXT("") && CurrentJob.Key.ShaderType)
+			{
+				// Some shader compiler errors won't have a file and line number, so we just assume the error happened in file containing the entrypoint function.
+				CurrentError.ErrorVirtualFilePath = FString(CurrentJob.Key.ShaderType->GetShaderFilename());
+			}
+
+			if (OutShaderDiagInfo.UniqueErrors.Find(CurrentErrorString) == INDEX_NONE)
+			{
+				// build up additional info in a "prefix" string; only do this once for each unique error
+				FString UniqueErrorPrefix;
+
+				// If we dumped the shader info, add it before the first error string
+				if (!GIsBuildMachine && !bReportedDebugInfo && CurrentJob.Input.DumpDebugInfoPath.Len() > 0)
+				{
+					UniqueErrorPrefix += FString::Printf(TEXT("Shader debug info dumped to: \"%s\"\n"), *CurrentJob.Input.DumpDebugInfoPath);
+					bReportedDebugInfo = true;
+				}
+
+				if (CurrentJob.Key.ShaderType)
+				{
+					// Construct a path that will enable VS.NET to find the shader file, relative to the solution
+					const FString SolutionPath = FPaths::RootDir();
+					FString ShaderFilePath = IFileManager::Get().ConvertToAbsolutePathForExternalAppForRead(*CurrentError.GetShaderSourceFilePath());
+					UniqueErrorPrefix += FString::Printf(TEXT("%s(%s): Shader %s, Permutation %d, VF %s:\n\t"),
+						*ShaderFilePath,
+						*CurrentError.ErrorLineString,
+						CurrentJob.Key.ShaderType->GetName(),
+						CurrentJob.Key.PermutationId,
+						CurrentJob.Key.VFType ? CurrentJob.Key.VFType->GetName() : TEXT("None"));
+				}
+				else
+				{
+					UniqueErrorPrefix += FString::Printf(TEXT("%s(0): "),
+						*CurrentJob.Input.VirtualSourceFilePath);
+				}
+
+				OutShaderDiagInfo.UniqueErrors.Add(CurrentErrorString);
+				OutShaderDiagInfo.UniqueErrorPrefixes.Add(UniqueErrorPrefix);
+				OutShaderDiagInfo.ErrorJobs.AddUnique(&CurrentJob);
+			}
+			++NumAddedErrors;
+		}
+	}
+
+	return NumAddedErrors;
+}
+
+FString GetSingleJobCompilationDump(const FShaderCompileJob* SingleJob)
+{
+	if (!SingleJob)
+	{
+		return TEXT("Internal error, not a Job!");
+	}
+	FString String = SingleJob->Input.GenerateShaderName();
+	if (SingleJob->Key.VFType)
+	{
+		String += FString::Printf(TEXT(" VF '%s'"), SingleJob->Key.VFType->GetName());
+	}
+	String += FString::Printf(TEXT(" Type '%s'"), SingleJob->Key.ShaderType->GetName());
+	String += FString::Printf(TEXT(" '%s' Entry '%s' Permutation %i "), *SingleJob->Input.VirtualSourceFilePath, *SingleJob->Input.EntryPointName, SingleJob->Key.PermutationId);
+	return String;
+}
+
+bool IsShaderDevelopmentModeEnabled()
+{
+	return (CVarShaderDevelopmentMode.GetValueOnAnyThread() != 0);
+}
+
+bool ShouldShowWarnings()
+{
+	// show warnings if explicitly requested via the r.ShowShaderCompilerWarnings cvar, or if shader dev mode is enabled.
+	// in either case additional optional filtering happens at the job level via r.ShaderCompilerWarningsFilter (this can
+	// be used to show just only global shader warnings, or disable warning prints entirely for shader dev mode)
+	return (CVarShowShaderWarnings.GetValueOnAnyThread() != 0) || IsShaderDevelopmentModeEnabled();
+}
+
+FShaderDiagnosticInfo::FShaderDiagnosticInfo(const TArray<FShaderCommonCompileJobPtr>& Jobs)
+{
+	// Gather unique errors
+	for (const FShaderCommonCompileJobPtr& Job : Jobs)
+	{
+		if (!Job->bSucceeded)
+		{
+			AddAndProcessErrorsForJob(*Job);
+		}
+		else if (ShouldShowWarnings())
+		{
+			AddWarningsForJob(*Job);
+		}
+	}
+
+	for (EShaderPlatform ErrorPlatform : ErrorPlatforms)
+	{
+		if (TargetShaderPlatformString.IsEmpty())
+		{
+			TargetShaderPlatformString = FDataDrivenShaderPlatformInfo::GetName(ErrorPlatform).ToString();
+		}
+		else
+		{
+			TargetShaderPlatformString += FString(TEXT(", ")) + FDataDrivenShaderPlatformInfo::GetName(ErrorPlatform).ToString();
+		}
+	}
+}
+
+void FShaderDiagnosticInfo::AddAndProcessErrorsForJob(FShaderCommonCompileJob& Job)
+{
+	Job.ForEachSingleShaderJob([this](FShaderCompileJob& Job)
+		{
+			ErrorPlatforms.AddUnique((EShaderPlatform)Job.Input.Target.Platform);
+
+			if (Job.Output.Errors.Num() == 0)
+			{
+				// Job hard crashed
+				UniqueErrors.Add(FString::Printf(TEXT("Internal Error!\n\t%s"), *GetSingleJobCompilationDump(&Job)));
+			}
+
+			// If we filter all error messages because they are interpreted as warnings, we have to assume all error messages are in fact errors and not warnings.
+			// In that case, add jobs again without a filter; e.g. when the stripped message starts with "Internal exception".
+			if (AddAndProcessErrorsForFailedJobFiltered(Job, *this, TEXT("error")) == 0)
+			{
+				AddAndProcessErrorsForFailedJobFiltered(Job, *this, nullptr);
+			}
+
+		});
+}
+
+void FShaderDiagnosticInfo::AddWarningsForJob(const FShaderCommonCompileJob& Job)
+{
+	Job.ForEachSingleShaderJob([this](const FShaderCompileJob& Job)
+		{
+			bool bIsGlobalShader = Job.Key.ShaderType->GetTypeForDynamicCast() == FShaderType::EShaderTypeForDynamicCast::Global;
+			int32 FilterValue = CVarShaderWarningsFilter.GetValueOnAnyThread();
+			bool bFilter = (FilterValue == 0) || ((FilterValue == 1) && !bIsGlobalShader);
+			// Append the "errors" to the UniqueWarnings array if the job succeeded; there's nothing distinguishing errors
+			// from warnings in the compile job so any errors that exist on a successful job are in fact warnings. 
+			// Note that if the failed any warnings will already be interspersed with the errors in the UniqueErrors array.
+			if (Job.bSucceeded && !bFilter)
+			{
+				for (const FShaderCompilerError& Warning : Job.Output.Errors)
+				{
+					UniqueWarnings.AddUnique(Warning.GetErrorString());
+				}
+			}
+		});
+}
\ No newline at end of file
diff --git a/Engine/Source/Runtime/RenderCore/Public/ShaderCompilerJobTypes.h b/Engine/Source/Runtime/RenderCore/Public/ShaderCompilerJobTypes.h
--- a/Engine/Source/Runtime/RenderCore/Public/ShaderCompilerJobTypes.h
+++ b/Engine/Source/Runtime/RenderCore/Public/ShaderCompilerJobTypes.h
@@ -166,6 +166,7 @@ public:
 
 	/** Calls the specified predicate for each single compile job, i.e. FShaderCompileJob and each stage of FShaderPipelineCompileJob. */
 	void ForEachSingleShaderJob(const TFunction<void(const FShaderCompileJob& SingleJob)>& Predicate) const;
+	void ForEachSingleShaderJob(const TFunction<void(FShaderCompileJob& SingleJob)>& Predicate);
 
 	/** This returns a unique id for a shader compiler job */
 	RENDERCORE_API static uint32 GetNextJobId();
@@ -350,11 +351,11 @@ inline void FShaderCommonCompileJob::Destroy() const
 	}
 }
 
-inline void FShaderCommonCompileJob::ForEachSingleShaderJob(const TFunction<void(const FShaderCompileJob&)>& Predicate) const
+inline void FShaderCommonCompileJob::ForEachSingleShaderJob(const TFunction<void(const FShaderCompileJob&)>& Function) const
 {
 	if (const FShaderCompileJob* SingleJob = GetSingleShaderJob())
 	{
-		Predicate(*SingleJob);
+		Function(*SingleJob);
 	}
 	else if (const FShaderPipelineCompileJob* PipelineJob = GetShaderPipelineJob())
 	{
@@ -362,7 +363,25 @@ inline void FShaderCommonCompileJob::ForEachSingleShaderJob(const TFunction<void
 		{
 			if (const FShaderCompileJob* SingleStageJob = StageJob->GetSingleShaderJob())
 			{
-				Predicate(*SingleStageJob);
+				Function(*SingleStageJob);
+			}
+		}
+	}
+}
+
+inline void FShaderCommonCompileJob::ForEachSingleShaderJob(const TFunction<void(FShaderCompileJob&)>& Function)
+{
+	if (FShaderCompileJob* SingleJob = GetSingleShaderJob())
+	{
+		Function(*SingleJob);
+	}
+	else if (FShaderPipelineCompileJob* PipelineJob = GetShaderPipelineJob())
+	{
+		for (TRefCountPtr<FShaderCompileJob>& StageJob : PipelineJob->StageJobs)
+		{
+			if (FShaderCompileJob* SingleStageJob = StageJob->GetSingleShaderJob())
+			{
+				Function(*SingleStageJob);
 			}
 		}
 	}
diff --git a/Engine/Source/Runtime/VulkanRHI/Private/VulkanRHI.cpp b/Engine/Source/Runtime/VulkanRHI/Private/VulkanRHI.cpp
--- a/Engine/Source/Runtime/VulkanRHI/Private/VulkanRHI.cpp
+++ b/Engine/Source/Runtime/VulkanRHI/Private/VulkanRHI.cpp
@@ -23,6 +23,7 @@
 #include "VulkanResourceCollection.h"
 #include "RHIValidation.h"
 #include "RHIUtilities.h"
+#include "ShaderDiagnostics.h"
 #include "IHeadMountedDisplayModule.h"
 #include "VulkanRenderpass.h"
 #include "VulkanTransientResourceAllocator.h"
@@ -656,10 +657,9 @@ void FVulkanDynamicRHI::Shutdown()
 void FVulkanDynamicRHI::CreateInstance()
 {
 	// Engine registration can be disabled via console var. Also disable automatically if ShaderDevelopmentMode is on.
-	auto* CVarShaderDevelopmentMode = IConsoleManager::Get().FindTConsoleVariableDataInt(TEXT("r.ShaderDevelopmentMode"));
 	auto* CVarDisableEngineAndAppRegistration = IConsoleManager::Get().FindTConsoleVariableDataInt(TEXT("r.DisableEngineAndAppRegistration"));
 	bool bDisableEngineRegistration = (CVarDisableEngineAndAppRegistration && CVarDisableEngineAndAppRegistration->GetValueOnAnyThread() != 0) ||
-		(CVarShaderDevelopmentMode && CVarShaderDevelopmentMode->GetValueOnAnyThread() != 0);
+		IsShaderDevelopmentModeEnabled();
 
 	// Use the API version stored in the profile
 	ApiVersion = GetVulkanApiVersionForFeatureLevel(GMaxRHIFeatureLevel, false);
diff --git a/Engine/Source/Runtime/Windows/D3D11RHI/Private/Windows/WindowsD3D11Device.cpp b/Engine/Source/Runtime/Windows/D3D11RHI/Private/Windows/WindowsD3D11Device.cpp
--- a/Engine/Source/Runtime/Windows/D3D11RHI/Private/Windows/WindowsD3D11Device.cpp
+++ b/Engine/Source/Runtime/Windows/D3D11RHI/Private/Windows/WindowsD3D11Device.cpp
@@ -26,6 +26,7 @@
 #include "GenericPlatform/GenericPlatformCrashContext.h"
 #include "RHIValidation.h"
 #include "RHIUtilities.h"
+#include "ShaderDiagnostics.h"
 #include "HDRHelper.h"
 #include "GlobalShader.h"
 
@@ -1597,9 +1598,8 @@ void FD3D11DynamicRHI::InitD3DDevice()
 			};
 
 			// Engine registration can be disabled via console var. Also disable automatically if ShaderDevelopmentMode is on.
-			auto* CVarShaderDevelopmentMode = IConsoleManager::Get().FindTConsoleVariableDataInt(TEXT("r.ShaderDevelopmentMode"));
 			auto* CVarDisableEngineAndAppRegistration = IConsoleManager::Get().FindTConsoleVariableDataInt(TEXT("r.DisableEngineAndAppRegistration"));
-			const bool bDisableEngineRegistration = (CVarShaderDevelopmentMode && CVarShaderDevelopmentMode->GetValueOnAnyThread() != 0) || 
+			const bool bDisableEngineRegistration = IsShaderDevelopmentModeEnabled() || 
 				(CVarDisableEngineAndAppRegistration && CVarDisableEngineAndAppRegistration->GetValueOnAnyThread() != 0);
 			const bool bDisableAppRegistration = bDisableEngineRegistration || !FApp::HasProjectName();
 
