Fix pose watch crash when toggling pose watches

diff --git a/Engine/Source/Runtime/Engine/Private/Animation/AnimInstanceProxy.cpp b/Engine/Source/Runtime/Engine/Private/Animation/AnimInstanceProxy.cpp
--- a/Engine/Source/Runtime/Engine/Private/Animation/AnimInstanceProxy.cpp
+++ b/Engine/Source/Runtime/Engine/Private/Animation/AnimInstanceProxy.cpp
@@ -3367,25 +3367,6 @@ void FAnimInstanceProxy::RegisterWatchedPose(const FCompactPose& Pose, const FBl
 						PoseWatch.Object = GetAnimInstanceObject();
 						PoseWatch.PoseWatch->SetIsNodeEnabled(true);
 
-						/*
-						for (FCompactPoseBoneIndex BoneIndex : Pose.ForEachBoneIndex())
-						{
-							FMeshPoseBoneIndex MeshBoneIndex = Pose.GetBoneContainer().MakeMeshPoseIndex(BoneIndex);
-
-							int32 ParentIndex = Pose.GetBoneContainer().GetParentBoneIndex(MeshBoneIndex.GetInt());
-
-							if (ParentIndex == INDEX_NONE)
-							{
-								WorldTransforms[MeshBoneIndex.GetInt()] = Pose[BoneIndex] * MeshComponent->GetComponentTransform();
-							}
-							else
-							{
-								WorldTransforms[MeshBoneIndex.GetInt()] = Pose[BoneIndex] * WorldTransforms[ParentIndex];
-							}
-							BoneColors[MeshBoneIndex.GetInt()] = BoneColor;
-						}
-						*/
-
 						TArray<FTransform> BoneTransforms;
 						BoneTransforms.AddUninitialized(Pose.GetBoneContainer().GetNumBones());
 
@@ -3420,35 +3401,10 @@ void FAnimInstanceProxy::RegisterWatchedPose(const FCSPose<FCompactPose>& Pose,
 
 void FAnimInstanceProxy::RegisterWatchedPose(const FCSPose<FCompactPose>& Pose, const FBlendedCurve& InCurve, int32 LinkID)
 {
-	if (bIsBeingDebugged)
-	{
-		FAnimBlueprintDebugData* DebugData = GetAnimBlueprintDebugData();
-		if (DebugData)
-		{
-			if (USkeletalMeshComponent* SkelMeshComponent = GetSkelMeshComponent())
-			{
-				for (FAnimNodePoseWatch& PoseWatch : DebugData->AnimNodePoseWatch)
-				{
-					if (PoseWatch.PoseWatch && PoseWatch.NodeID == LinkID)
-					{
-						FCompactPose TempPose;
-						FCSPose<FCompactPose>::ConvertComponentPosesToLocalPoses(Pose, TempPose);
-						PoseWatch.Object = GetAnimInstanceObject();
-						PoseWatch.PoseWatch->SetIsNodeEnabled(true);
+	FCompactPose LocalPose;
+	FCSPose<FCompactPose>::ConvertComponentPosesToLocalPoses(Pose, LocalPose);
 
-						const TArray<FTransform, FAnimStackAllocator>& BoneTransforms = TempPose.GetBones();
-						const TArray<FBoneIndexType>& TmpRequiredBones = TempPose.GetBoneContainer().GetBoneIndicesArray();
-						PoseWatch.SetPose(TmpRequiredBones, BoneTransforms);
-						PoseWatch.SetCurves(InCurve);
-						PoseWatch.SetWorldTransform(SkelMeshComponent->GetComponentTransform());
-
-						TRACE_ANIM_POSE_WATCH(*this, PoseWatch.PoseWatchPoseElement, PoseWatch.NodeID, PoseWatch.GetBoneTransforms(), PoseWatch.GetCurves(), PoseWatch.GetRequiredBones(), PoseWatch.GetWorldTransform(), PoseWatch.PoseWatchPoseElement->GetIsVisible());
-						break;
-					}
-				}
-			}
-		}
-	}
+	RegisterWatchedPose(LocalPose, InCurve, LinkID);
 }
 #endif
 
diff --git a/Engine/Source/Runtime/Engine/Private/SkeletalDebugRendering.cpp b/Engine/Source/Runtime/Engine/Private/SkeletalDebugRendering.cpp
--- a/Engine/Source/Runtime/Engine/Private/SkeletalDebugRendering.cpp
+++ b/Engine/Source/Runtime/Engine/Private/SkeletalDebugRendering.cpp
@@ -7,6 +7,10 @@
 #include "Engine/PoseWatchRenderData.h"
 #include "ReferenceSkeleton.h"
 
+static TAutoConsoleVariable<bool> CVarDisablePoseWatchRendering(
+	TEXT("a.DisablePoseWatchRendering"),
+	false,
+	TEXT("Disable all active pose watches from being drawn."));
 
 namespace SkeletalDebugRendering
 {
@@ -180,6 +184,11 @@ void DrawBonesFromPoseWatch(
 	const bool bUseWorldTransform
 )
 {
+	if (CVarDisablePoseWatchRendering.GetValueOnAnyThread())
+	{
+		return;
+	}
+
 	const TArray<FTransform>& InBoneTransforms = PoseWatch.GetBoneTransforms();
 	const TArray<FBoneIndexType>& InRequiredBones = PoseWatch.GetRequiredBones();
 	if (InRequiredBones.Num() == 0 || InBoneTransforms.Num() < InRequiredBones.Num())
@@ -194,14 +203,14 @@ void DrawBonesFromPoseWatch(
 	const TArray<int32>& ParentIndices = PoseWatch.GetParentIndices();
 
 	TArray<FTransform> UseWorldTransforms;
-	UseWorldTransforms.AddDefaulted(ParentIndices.Num());
+	UseWorldTransforms.AddDefaulted(InBoneTransforms.Num());
 	
 	TArray<FBoneIndexType> UseRequiredBones;
 	UseRequiredBones.Reserve(InRequiredBones.Num());
 
 	for (const FBoneIndexType& BoneIndex : InRequiredBones)
 	{
-		if (ParentIndices.IsValidIndex(BoneIndex))
+		if (ParentIndices.IsValidIndex(BoneIndex) && InBoneTransforms.IsValidIndex(BoneIndex))
 		{
 			const int32 ParentIndex = ParentIndices[BoneIndex];
 
