Fixing Linux compilation errors when trying to log management of CVars

#rnx

diff --git a/Engine/Source/Runtime/Engine/Private/Tests/AutomationCommon.cpp b/Engine/Source/Runtime/Engine/Private/Tests/AutomationCommon.cpp
--- a/Engine/Source/Runtime/Engine/Private/Tests/AutomationCommon.cpp
+++ b/Engine/Source/Runtime/Engine/Private/Tests/AutomationCommon.cpp
@@ -302,7 +302,10 @@ void FTestConsoleVariable::Set(const FString& Value)
 		}
 
 		ConsoleVariable->AsVariable()->SetWithCurrentPriority(*Value);
+
+#if WITH_AUTOMATION_TESTS
 		UE_LOG(LogEngineAutomationTests, Verbose, TEXT("CVar '%s' was set with the value '%s'."), *ConsoleVariableName, *Value);
+#endif // WITH_AUTOMATION_TESTS
 	}
 }
 
@@ -325,7 +328,10 @@ void FTestConsoleVariable::Restore()
 		if (ensure(ConsoleVariable))
 		{
 			ConsoleVariable->AsVariable()->SetWithCurrentPriority(*OriginalValue);
+
+#if WITH_AUTOMATION_TESTS
 			UE_LOG(LogEngineAutomationTests, Verbose, TEXT("CVar '%s' was restored with the original value '%s'."), *ConsoleVariableName, *OriginalValue);
+#endif // WITH_AUTOMATION_TESTS
 		}
 
 		bModified = false;
