Added minimal init function for External Package DataLayer Instances to allow for DataLayer Instance iteration in the case where a World is loaded without a full Init of its systems (like in a validation cmdlet)

diff --git a/Engine/Source/Runtime/Engine/Private/WorldPartition/DataLayer/WorldDataLayers.cpp b/Engine/Source/Runtime/Engine/Private/WorldPartition/DataLayer/WorldDataLayers.cpp
--- a/Engine/Source/Runtime/Engine/Private/WorldPartition/DataLayer/WorldDataLayers.cpp
+++ b/Engine/Source/Runtime/Engine/Private/WorldPartition/DataLayer/WorldDataLayers.cpp
@@ -646,6 +646,10 @@ void AWorldDataLayers::OnLoadedActorRemovedFromLevel()
 
 	if (IsUsingExternalPackageDataLayerInstances())
 	{
+		// Validation will sometimes load the world without invoking the full initialization of its systems so
+		// perform a minimal initialization on ExternalPackage DataLayerInstances so that we can iterate over them
+		InitializeExternalPackageDataLayerInstances();
+
 		// Iterate on Data Layer Instances excluding Transient Data Layer Instances
 		for (UDataLayerInstance* DataLayerInstance : GetDataLayerInstances())
 		{
@@ -1199,13 +1203,23 @@ TSet<TObjectPtr<UDataLayerInstance>>& AWorldDataLayers::GetDataLayerInstances()
 	return DataLayerInstances;
 }
 
+void AWorldDataLayers::InitializeExternalPackageDataLayerInstances()
+{
+	check(IsUsingExternalPackageDataLayerInstances());
+	ExternalPackageDataLayerInstances.Append(LoadedExternalPackageDataLayerInstances);
+	LoadedExternalPackageDataLayerInstances.Reset();
+}
+
 void AWorldDataLayers::OnDataLayerManagerInitialized()
 {
 #if WITH_EDITOR
 	ConditionalPostLoad();
 	// At this point, LoadedExternalPackageDataLayerInstances are fully loaded, transfer them to the DataLayerInstances list.
-	ExternalPackageDataLayerInstances.Append(LoadedExternalPackageDataLayerInstances);
-	LoadedExternalPackageDataLayerInstances.Reset();
+	if (IsUsingExternalPackageDataLayerInstances())
+	{
+		InitializeExternalPackageDataLayerInstances();
+	}
+	
 	if (IsRunningCookCommandlet())
 	{
 		// Embed external DataLayerInstances when cooking
diff --git a/Engine/Source/Runtime/Engine/Public/WorldPartition/DataLayer/WorldDataLayers.h b/Engine/Source/Runtime/Engine/Public/WorldPartition/DataLayer/WorldDataLayers.h
--- a/Engine/Source/Runtime/Engine/Public/WorldPartition/DataLayer/WorldDataLayers.h
+++ b/Engine/Source/Runtime/Engine/Public/WorldPartition/DataLayer/WorldDataLayers.h
@@ -117,6 +117,7 @@ public:
 	ENGINE_API bool IsUsingExternalPackageDataLayerInstances() const { return bUseExternalPackageDataLayerInstances; }
 	ENGINE_API bool SetUseExternalPackageDataLayerInstances(bool bInNewValue, bool bInInteractiveMode = false);
 
+
 	ENGINE_API bool CanReferenceDataLayerAsset(const UDataLayerAsset* InDataLayerAsset, FText* OutFailureReason) const;
 #endif
 
@@ -219,6 +220,9 @@ private:
 	bool RemoveExternalDataLayerInstance(UExternalDataLayerInstance* ExternalDataLayerInstance);
 	ENGINE_API UExternalDataLayerInstance* GetExternalDataLayerInstance(const UExternalDataLayerAsset* InExternalDataLayerAsset);
 
+	// ExternalPackage Data Layer Instances
+	ENGINE_API void InitializeExternalPackageDataLayerInstances();
+
 	ENGINE_API void OnDataLayerManagerInitialized();
 	ENGINE_API void OnDataLayerManagerDeinitialized();
 	ENGINE_API void ResolveEffectiveRuntimeState(const UDataLayerInstance* InDataLayer, bool bInNotifyChange = true);
