Correctly mark InitializeExternalPackageDataLayerInstances editor only

diff --git a/Engine/Source/Runtime/Engine/Private/WorldPartition/DataLayer/WorldDataLayers.cpp b/Engine/Source/Runtime/Engine/Private/WorldPartition/DataLayer/WorldDataLayers.cpp
--- a/Engine/Source/Runtime/Engine/Private/WorldPartition/DataLayer/WorldDataLayers.cpp
+++ b/Engine/Source/Runtime/Engine/Private/WorldPartition/DataLayer/WorldDataLayers.cpp
@@ -1203,12 +1203,14 @@ TSet<TObjectPtr<UDataLayerInstance>>& AWorldDataLayers::GetDataLayerInstances()
 	return DataLayerInstances;
 }
 
+#if WITH_EDITOR
 void AWorldDataLayers::InitializeExternalPackageDataLayerInstances()
 {
 	check(IsUsingExternalPackageDataLayerInstances());
 	ExternalPackageDataLayerInstances.Append(LoadedExternalPackageDataLayerInstances);
 	LoadedExternalPackageDataLayerInstances.Reset();
 }
+#endif
 
 void AWorldDataLayers::OnDataLayerManagerInitialized()
 {
diff --git a/Engine/Source/Runtime/Engine/Public/WorldPartition/DataLayer/WorldDataLayers.h b/Engine/Source/Runtime/Engine/Public/WorldPartition/DataLayer/WorldDataLayers.h
--- a/Engine/Source/Runtime/Engine/Public/WorldPartition/DataLayer/WorldDataLayers.h
+++ b/Engine/Source/Runtime/Engine/Public/WorldPartition/DataLayer/WorldDataLayers.h
@@ -220,8 +220,10 @@ private:
 	bool RemoveExternalDataLayerInstance(UExternalDataLayerInstance* ExternalDataLayerInstance);
 	ENGINE_API UExternalDataLayerInstance* GetExternalDataLayerInstance(const UExternalDataLayerAsset* InExternalDataLayerAsset);
 
+#if WITH_EDITOR
 	// ExternalPackage Data Layer Instances
 	ENGINE_API void InitializeExternalPackageDataLayerInstances();
+#endif
 
 	ENGINE_API void OnDataLayerManagerInitialized();
 	ENGINE_API void OnDataLayerManagerDeinitialized();
