- Fix for undo / redo when changing Cascade mesh type data

#rnx
#tests Followed steps in Jira

diff --git a/Engine/Source/Runtime/Engine/Private/Particles/ParticleModules.cpp b/Engine/Source/Runtime/Engine/Private/Particles/ParticleModules.cpp
--- a/Engine/Source/Runtime/Engine/Private/Particles/ParticleModules.cpp
+++ b/Engine/Source/Runtime/Engine/Private/Particles/ParticleModules.cpp
@@ -3529,6 +3529,7 @@ void UParticleModuleTypeDataMesh::OnMeshChanged()
 
 void UParticleModuleTypeDataMesh::PreEditChange(FProperty* PropertyThatWillChange)
 {
+	Super::PreEditChange(PropertyThatWillChange);
 	if ( (PropertyThatWillChange != nullptr) && (PropertyThatWillChange->GetFName() == FName(TEXT("Mesh"))) )
 	{
 		if ( GIsEditor && (Mesh != nullptr) )
