PR #11474: Made a warning in UActorComponent::RegisterComponentWithWorld actionable by including the name of the owning actor class

diff --git a/Engine/Source/Runtime/Engine/Private/Components/ActorComponent.cpp b/Engine/Source/Runtime/Engine/Private/Components/ActorComponent.cpp
--- a/Engine/Source/Runtime/Engine/Private/Components/ActorComponent.cpp
+++ b/Engine/Source/Runtime/Engine/Private/Components/ActorComponent.cpp
@@ -1395,7 +1395,7 @@ void UActorComponent::RegisterComponentWithWorld(UWorld* InWorld, FRegisterCompo
 
 	if (MyOwner && MyOwner->GetClass()->HasAnyClassFlags(CLASS_NewerVersionExists))
 	{
-		UE_LOG(LogActorComponent, Log, TEXT("RegisterComponentWithWorld: Owner belongs to a DEADCLASS"));
+		UE_LOG(LogActorComponent, Log, TEXT("RegisterComponentWithWorld: Owner belongs to a DEADCLASS %s"), *GetPathNameSafe(MyOwner->GetClass()));
 		return;
 	}
 
