Allow modulate blend mode materials to be renderred in translucent after DOF pass on mobile.

diff --git a/Engine/Source/Runtime/Engine/Private/Materials/HLSLMaterialTranslator.cpp b/Engine/Source/Runtime/Engine/Private/Materials/HLSLMaterialTranslator.cpp
--- a/Engine/Source/Runtime/Engine/Private/Materials/HLSLMaterialTranslator.cpp
+++ b/Engine/Source/Runtime/Engine/Private/Materials/HLSLMaterialTranslator.cpp
@@ -1757,9 +1757,9 @@ void FHLSLMaterialTranslator::DoTranslate()
 
 	if (Domain == MD_Surface)
 	{
-		if (IsModulateBlendMode(BlendMode) && Material->IsTranslucencyAfterDOFEnabled() && !RHISupportsDualSourceBlending(Platform))
+		if (IsModulateBlendMode(BlendMode) && Material->IsTranslucencyAfterDOFEnabled() && !(RHISupportsDualSourceBlending(Platform) || IsMobilePlatform(Platform)))
 		{
-			Errorf(TEXT("Translucency after DOF with BLEND_Modulate is only allowed on platforms that support dual-blending. Consider using BLEND_Translucent with black emissive"));
+			Errorf(TEXT("Translucency after DOF with BLEND_Modulate is only allowed on mobile platforms or desktop platforms that support dual-blending. Consider using BLEND_Translucent with black emissive"));
 		}
 	}
 
