[NetPhysics] Smooth kinematic Predictive Interpolation replication

diff --git a/Engine/Source/Runtime/Engine/Private/PhysicsEngine/PhysicsReplication.cpp b/Engine/Source/Runtime/Engine/Private/PhysicsEngine/PhysicsReplication.cpp
--- a/Engine/Source/Runtime/Engine/Private/PhysicsEngine/PhysicsReplication.cpp
+++ b/Engine/Source/Runtime/Engine/Private/PhysicsEngine/PhysicsReplication.cpp
@@ -209,9 +209,15 @@ namespace PhysicsReplicationCVars
 		bool bSleepConnectedBodies = true;
 		static FAutoConsoleVariableRef CVarSleepConnectedBodies(TEXT("np2.PredictiveInterpolation.SleepConnectedBodies"), bSleepConnectedBodies, TEXT("When true, sleep state will be applied to any dynamic physics object connected to the replicated object."));
 
+		bool bKinematicPrediction = true;
+		static FAutoConsoleVariableRef CVarKinematicPrediction(TEXT("np2.PredictiveInterpolation.KinematicPrediction"), bKinematicPrediction, TEXT("When true, predictive interpolation will perform predictive movement instead of interpolation for kinematic objects."));
+	
+		bool bKinematicHardSnap = false;
+		static FAutoConsoleVariableRef CVarKinematicHardSnap(TEXT("np2.PredictiveInterpolation.KinematicHardSnap"), bKinematicHardSnap, TEXT("When true, predictive interpolation will perform a hard snap for objects that are kinematic."));
+	
 		bool bDisableSoftSnap = false;
 		static FAutoConsoleVariableRef CVarDisableSoftSnap(TEXT("np2.PredictiveInterpolation.DisableSoftSnap"), bDisableSoftSnap, TEXT("When true, predictive interpolation will not use softsnap to correct the replication with when velocity fails. Hardsnap will still eventually kick in if replication can't reach the target."));
-	
+
 		bool bAlwaysHardSnap = false;
 		static FAutoConsoleVariableRef CVarAlwaysHardSnap(TEXT("np2.PredictiveInterpolation.AlwaysHardSnap"), bAlwaysHardSnap, TEXT("When true, predictive interpolation replication mode will always hard snap. Used as a backup measure"));
 
@@ -227,6 +233,9 @@ namespace PhysicsReplicationCVars
 		bool bDrawDebugVectors = false;
 		static FAutoConsoleVariableRef CVarDrawDebugVectors(TEXT("np2.PredictiveInterpolation.DrawDebugVectors"), bDrawDebugVectors, TEXT("Draw replication vectors, target velocity, replicated velocity, velocity change between replication calls etc."));
 		
+		float DrawDebugZOffset = 50.0f;
+		static FAutoConsoleVariableRef CVarDrawDebugZOffset(TEXT("np2.PredictiveInterpolation.DrawDebugZOffset"), DrawDebugZOffset, TEXT("Offset in Z axis for draw debug calls"));
+
 		float SleepSecondsClearTarget = 15.0f;
 		static FAutoConsoleVariableRef CVarSleepSecondsClearTarget(TEXT("np2.PredictiveInterpolation.SleepSecondsClearTarget"), SleepSecondsClearTarget, TEXT("Wait for the object to sleep for this many seconds before clearing the replication target, to ensure nothing wakes up the object just after it goes to sleep on the client."));
 		
@@ -1000,7 +1009,7 @@ void FPhysicsReplicationAsync::UpdateAsyncTarget(const FPhysicsRepAsyncInputData
 #if !(UE_BUILD_SHIPPING || UE_BUILD_TEST)
 			if (PhysicsReplicationCVars::PredictiveInterpolationCVars::bDrawDebugTargets)
 			{
-				const FVector Offset = FVector(0.0f, 0.0f, 50.0f);
+				const FVector Offset = FVector(0.0f, 0.0f, PhysicsReplicationCVars::PredictiveInterpolationCVars::DrawDebugZOffset);
 				Chaos::FDebugDrawQueue::GetInstance().DrawDebugBox(Input.TargetState.Position + Offset, FVector(15.0f, 15.0f, 15.0f), Input.TargetState.Quaternion, FColor::MakeRandomSeededColor(Input.ServerFrame), false, CharacterMovementCVars::NetCorrectionLifetime, 0, 1.0f);
 			}
 #endif
@@ -1472,10 +1481,10 @@ bool FPhysicsReplicationAsync::PredictiveInterpolation(Chaos::FPBDRigidParticleH
 #if !(UE_BUILD_SHIPPING || UE_BUILD_TEST)
 	if (PhysicsReplicationCVars::PredictiveInterpolationCVars::bDrawDebugTargets)
 	{
-		const FVector Offset = FVector(0.0f, 0.0f, 50.0f);
-		const FVector StartPos = Target.TargetState.Position + Offset;
+		const FVector Offset = FVector(0.0f, 0.0f, PhysicsReplicationCVars::PredictiveInterpolationCVars::DrawDebugZOffset);
+		const FVector Pos = Target.TargetState.Position + Offset;
 		const int32 SizeMultiplier = FMath::Clamp(Target.TickCount, -4, 30);
-		Chaos::FDebugDrawQueue::GetInstance().DrawDebugBox(StartPos, FVector(5.0f + SizeMultiplier * 0.75f, 5.0f + SizeMultiplier * 0.75f, 5.0f + SizeMultiplier * 0.75f), Target.TargetState.Quaternion, FColor::MakeRandomSeededColor(Target.ServerFrame), false, CharacterMovementCVars::NetCorrectionLifetime, 0, 1.0f);
+		Chaos::FDebugDrawQueue::GetInstance().DrawDebugBox(Pos, FVector(5.0f + SizeMultiplier * 0.75f, 5.0f + SizeMultiplier * 0.75f, 5.0f + SizeMultiplier * 0.75f), Target.TargetState.Quaternion, FColor::MakeRandomSeededColor(Target.ServerFrame), false, CharacterMovementCVars::NetCorrectionLifetime, 0, 1.0f);
 	}
 #endif
 
@@ -1503,9 +1512,8 @@ bool FPhysicsReplicationAsync::PredictiveInterpolation(Chaos::FPBDRigidParticleH
 
 		// --- Should replication stop? ---
 		const bool bClearTarget =
-			(!bCanSimulate
-				|| (bOkToClear && bShouldSleep && Target.AccumulatedSleepSeconds >= PhysicsReplicationCVars::PredictiveInterpolationCVars::SleepSecondsClearTarget) // Don't clear the target due to sleeping until the object both should sleep and is sleeping for n seconds
-				|| (bOkToClear && !bReplicatingPhysics))
+			((bOkToClear && bShouldSleep && Target.AccumulatedSleepSeconds >= PhysicsReplicationCVars::PredictiveInterpolationCVars::SleepSecondsClearTarget) // Don't clear the target due to sleeping until the object both should sleep and is sleeping for n seconds
+			|| (bOkToClear && !bReplicatingPhysics))
 			&& !PhysicsReplicationCVars::PredictiveInterpolationCVars::bDontClearTarget;
 
 		// --- Target Prediction ---
@@ -1547,12 +1555,6 @@ bool FPhysicsReplicationAsync::PredictiveInterpolation(Chaos::FPBDRigidParticleH
 		}
 	}
 	
-	// Wake up if sleeping
-	if (bIsSleeping)
-	{
-		RigidsSolver->GetEvolution()->SetParticleObjectState(Handle, Chaos::EObjectStateType::Dynamic);
-	}
-
 	// Update the AverageReceiveInterval if Target.ReceiveInterval has a valid value to update from
 	Target.AverageReceiveInterval = Target.ReceiveInterval == 0 ? Target.AverageReceiveInterval : FMath::Lerp(Target.AverageReceiveInterval, Target.ReceiveInterval, FMath::Clamp((1.0f / (Target.ReceiveInterval * PhysicsReplicationCVars::PredictiveInterpolationCVars::AverageReceiveIntervalSmoothing)), 0.0f, 1.0f));
 
@@ -1606,9 +1608,9 @@ bool FPhysicsReplicationAsync::PredictiveInterpolation(Chaos::FPBDRigidParticleH
 		bSoftSnap = false;
 	}
 
-	const bool bHardSnap = !bCanSimulate ||
-		Target.AccumulatedErrorSeconds > PhysicsReplicationCVars::PredictiveInterpolationCVars::ErrorAccumulationSeconds ||
-		PhysicsReplicationCVars::PredictiveInterpolationCVars::bAlwaysHardSnap;
+	const bool bHardSnap = (!bCanSimulate && PhysicsReplicationCVars::PredictiveInterpolationCVars::bKinematicHardSnap)
+		|| Target.AccumulatedErrorSeconds > PhysicsReplicationCVars::PredictiveInterpolationCVars::ErrorAccumulationSeconds
+		|| PhysicsReplicationCVars::PredictiveInterpolationCVars::bAlwaysHardSnap;
 
 	if (bHardSnap)
 	{
@@ -1635,8 +1637,67 @@ bool FPhysicsReplicationAsync::PredictiveInterpolation(Chaos::FPBDRigidParticleH
 		// End replication and go to sleep if that's requested
 		return EndReplicationHelper(Target, true);
 	}
+	else if (Handle->IsKinematic()) // Smooth Kinematic Replication
+	{
+		const bool bKinematicPrediction = PhysicsReplicationCVars::PredictiveInterpolationCVars::bKinematicPrediction;
+		const float InterpolationTicks = FMath::CeilToInt(Target.AverageReceiveInterval) - (RigidsSolver->GetCurrentFrame() - Target.ReceiveFrame);
+
+		if (bKinematicPrediction || InterpolationTicks > 0)
+		{
+			/* Calculate the Lerp value for a smooth interpolation
+			* ------------------------------------------------------------------------------
+			* bKinematicPrediction is True :: Interpolate towards the target that gets forward predicted each tick
+			*	1 / 4 = 0.25 = 25% interpolation each time (if AverageReceiveInterval is 4)
+			* ------------------------------------------------------------------------------
+			* bKinematicPrediction is False :: Interpolate from current position to the static source for the current target, we need to cover the same amount of distance but from a decaying distance
+			*	| ---> | ------------------ |
+			*	0%    25%				   100%		(1 / 4 = 0.25)
+			*		   | ---> | ----------- |
+			*		   0%	 33%		   100%		(1 / 3 = 0.33)
+			*				  | ---> | ---- |
+			*				  0%    50%    100%		(1 / 2 = 0.5)
+			*						 | ---> |
+			*						 0%    100%		(1 / 1 = 1.0)
+			* ------------------------------------------------------------------------------
+			*/
+			const float Lerp = 1.f / (bKinematicPrediction ? Target.AverageReceiveInterval : InterpolationTicks);
+			
+			// Interpolate position and rotation from current position towards target position based on either predicted target or source target
+			const FVector KinTargetPos = FMath::Lerp(CurrentState.Position,
+				(bKinematicPrediction ? Target.TargetState.Position : Target.PrevPosTarget),
+				Lerp);
+			const FQuat KinTargetRot = FQuat::Slerp(CurrentState.Quaternion,
+				(bKinematicPrediction ? Target.TargetState.Quaternion : Target.PrevRotTarget),
+				Lerp);
+
+			// Apply kinematic target
+			const Chaos::FKinematicTarget KinTarget = Chaos::FKinematicTarget::MakePositionTarget(KinTargetPos, KinTargetRot); // Uses EKinematicTargetMode::Position
+			RigidsSolver->GetEvolution()->SetParticleKinematicTarget(Handle, KinTarget);
+
+#if !(UE_BUILD_SHIPPING || UE_BUILD_TEST)
+			if (PhysicsReplicationCVars::PredictiveInterpolationCVars::bDrawDebugTargets)
+			{
+				const FVector Offset = FVector(0.0f, 0.0f, PhysicsReplicationCVars::PredictiveInterpolationCVars::DrawDebugZOffset);
+				const FVector Pos = KinTargetPos + Offset;
+				const int32 SizeMultiplier = FMath::Clamp(Target.TickCount, -4, 30);
+				Chaos::FDebugDrawQueue::GetInstance().DrawDebugSphere(Pos, 3.0f + SizeMultiplier * 0.75f, 8, FColor::MakeRandomSeededColor(Target.ServerFrame), false, CharacterMovementCVars::NetCorrectionLifetime, 0, 1.0f);
+			}
+#endif
+		}
+		else
+		{
+			// End replication and allow to clear target
+			return EndReplicationHelper(Target, true);
+		}
+	}
 	else // Velocity-based Replication
 	{
+		// Wake up if sleeping
+		if (bIsSleeping)
+		{
+			RigidsSolver->GetEvolution()->SetParticleObjectState(Handle, Chaos::EObjectStateType::Dynamic);
+		}
+
 		// Calculate interpolation time based on current average receive rate
 		const float AverageReceiveIntervalSeconds = Target.AverageReceiveInterval * DeltaSeconds;
 		const float InterpolationTime = AverageReceiveIntervalSeconds * SettingsCurrent.PredictiveInterpolationSettings.GetPosInterpolationTimeMultiplier();
@@ -1694,7 +1755,7 @@ bool FPhysicsReplicationAsync::PredictiveInterpolation(Chaos::FPBDRigidParticleH
 #if !(UE_BUILD_SHIPPING || UE_BUILD_TEST)
 			if (PhysicsReplicationCVars::PredictiveInterpolationCVars::bDrawDebugVectors)
 			{
-				const FVector Offset = FVector(0.0f, 0.0f, 50.0f);
+				const FVector Offset = FVector(0.0f, 0.0f, PhysicsReplicationCVars::PredictiveInterpolationCVars::DrawDebugZOffset);
 				const FVector OffsetAdd = FVector(0.0f, 0.0f, 10.0f);
 				const FVector StartPos = TargetPos + Offset;
 				FVector Direction = TargetLinVel;
