Removed unused variables.

diff --git a/Engine/Plugins/Animation/BlendSpaceMotionAnalysis/Source/BlendSpaceMotionAnalysis/Private/RootMotionAnalysis.cpp b/Engine/Plugins/Animation/BlendSpaceMotionAnalysis/Source/BlendSpaceMotionAnalysis/Private/RootMotionAnalysis.cpp
--- a/Engine/Plugins/Animation/BlendSpaceMotionAnalysis/Source/BlendSpaceMotionAnalysis/Private/RootMotionAnalysis.cpp
+++ b/Engine/Plugins/Animation/BlendSpaceMotionAnalysis/Source/BlendSpaceMotionAnalysis/Private/RootMotionAnalysis.cpp
@@ -238,7 +238,6 @@ static bool CalculateRootMotion(
 		return false;
 	}
 
-	FVector Velocity;
 	switch (AnalysisProperties->FunctionAxis)
 	{
 	case EAnalysisRootMotionAxis::Speed:
diff --git a/Engine/Plugins/Animation/LiveLink/Source/LiveLinkMovieScene/Private/MovieScene/MovieSceneLiveLinkSectionTemplate.cpp b/Engine/Plugins/Animation/LiveLink/Source/LiveLinkMovieScene/Private/MovieScene/MovieSceneLiveLinkSectionTemplate.cpp
--- a/Engine/Plugins/Animation/LiveLink/Source/LiveLinkMovieScene/Private/MovieScene/MovieSceneLiveLinkSectionTemplate.cpp
+++ b/Engine/Plugins/Animation/LiveLink/Source/LiveLinkMovieScene/Private/MovieScene/MovieSceneLiveLinkSectionTemplate.cpp
@@ -171,7 +171,6 @@ bool FMovieSceneLiveLinkSectionTemplate::GetLiveLinkFrameArray(const FFrameTime&
 
 	bool bSendInterpolated = bAlwaysSendInterpolated || !TimeCodeFrameTime.IsSet() || LowerBound == UpperBound || bMustDoInterpolation;
 	FLiveLinkWorldTime WorldTime = FLiveLinkWorldTime(); //this calls FPlatform::Seconds()
-	FVector Vector;
 
 	if (!bSendInterpolated)
 	{
diff --git a/Engine/Plugins/Editor/ProxyLODPlugin/Source/ProxyLOD/Private/ProxyLODMeshUtilities.cpp b/Engine/Plugins/Editor/ProxyLODPlugin/Source/ProxyLOD/Private/ProxyLODMeshUtilities.cpp
--- a/Engine/Plugins/Editor/ProxyLODPlugin/Source/ProxyLOD/Private/ProxyLODMeshUtilities.cpp
+++ b/Engine/Plugins/Editor/ProxyLODPlugin/Source/ProxyLOD/Private/ProxyLODMeshUtilities.cpp
@@ -1132,7 +1132,6 @@ void ProxyLOD::ComputeFaceAveragedVertexNormals(FAOSMesh& InOutMesh)
 
 			if (AdjFaces.Num() != 0)
 			{
-				FVector3f Pos[3];
 				for (auto FaceId: AdjFaces)
 				{
 					checkSlow(FaceId > -1);
@@ -1605,7 +1604,6 @@ double ProxyLOD::GetWorldSpaceArea(const FVertexDataMesh& InMesh)
 
 	double Mesh3DArea = 0;
 
-	FVector3f Pos[3];
 	for (int32 IdxTri = 0; IdxTri < NumTriangles; ++IdxTri)
 	{
 		const uint32 Offset = 3 * IdxTri;
diff --git a/Engine/Plugins/Editor/SequencerAnimTools/Source/SequencerAnimTools/Private/MotionTrailTool.cpp b/Engine/Plugins/Editor/SequencerAnimTools/Source/SequencerAnimTools/Private/MotionTrailTool.cpp
--- a/Engine/Plugins/Editor/SequencerAnimTools/Source/SequencerAnimTools/Private/MotionTrailTool.cpp
+++ b/Engine/Plugins/Editor/SequencerAnimTools/Source/SequencerAnimTools/Private/MotionTrailTool.cpp
@@ -289,7 +289,6 @@ void UMotionTrailTool::GizmoTransformChanged(UTransformProxy* Proxy, FTransform
 	for (const TUniquePtr<UE::SequencerAnimTools::FTrailHierarchy>& TrailHierarchy : TrailHierarchies)
 	{
 		FVector WidgetLocation = Transform.GetLocation();
-		FVector Location;
 		bool bIsHandled = TrailHierarchy->ApplyDelta(LocationDiff, Diff.GetRotation().Rotator(), WidgetLocation);
 		if (bIsHandled)
 		{
@@ -328,7 +327,6 @@ void UMotionTrailTool::UpdateGizmoVisibility()
 
 		for (const TUniquePtr<UE::SequencerAnimTools::FTrailHierarchy>& TrailHierarchy : TrailHierarchies)
 		{
-			FVector Location;
 			bool bIsHandled = TrailHierarchy->IsAnythingSelected();
 			if (bIsHandled)
 			{
@@ -449,7 +447,6 @@ bool UMotionTrailTool::SomeKeysAreSelected() const
 
 	for (const TUniquePtr<UE::SequencerAnimTools::FTrailHierarchy>& TrailHierarchy : TrailHierarchies)
 	{
-		FVector Location;
 		bool bIsHandled = TrailHierarchy->IsAnythingSelected();
 		if (bIsHandled)
 		{
diff --git a/Engine/Plugins/Editor/SequencerAnimTools/Source/SequencerAnimTools/Private/TrailHierarchy.cpp b/Engine/Plugins/Editor/SequencerAnimTools/Source/SequencerAnimTools/Private/TrailHierarchy.cpp
--- a/Engine/Plugins/Editor/SequencerAnimTools/Source/SequencerAnimTools/Private/TrailHierarchy.cpp
+++ b/Engine/Plugins/Editor/SequencerAnimTools/Source/SequencerAnimTools/Private/TrailHierarchy.cpp
@@ -315,7 +315,6 @@ bool FTrailHierarchy::IsAnythingSelected() const
 {
 	for (const TPair<FGuid, TUniquePtr<FTrail>>& GuidTrailPair : GetAllTrails())
 	{
-		FVector Location;
 		bool bHandled = GuidTrailPair.Value->IsAnythingSelected();
 		if (bHandled)
 		{
diff --git a/Engine/Plugins/Editor/UVEditor/Source/UVEditorTools/Private/Operators/UVEditorUVTransformOp.cpp b/Engine/Plugins/Editor/UVEditor/Source/UVEditorTools/Private/Operators/UVEditorUVTransformOp.cpp
--- a/Engine/Plugins/Editor/UVEditor/Source/UVEditorTools/Private/Operators/UVEditorUVTransformOp.cpp
+++ b/Engine/Plugins/Editor/UVEditor/Source/UVEditorTools/Private/Operators/UVEditorUVTransformOp.cpp
@@ -69,7 +69,6 @@ void FUVEditorUVTransformBaseOp::RebuildBoundingBoxes()
 
 	ParallelFor(NumComponents, [&](int32 k)
 		{
-			FVector3d VertexPos[3];
 			PerComponentBoundingBoxes[k] = FAxisAlignedBox2d::Empty();
 			const TArray<int>& Vertices = (*UVComponents)[k].Indices;
 			for (int32 Vid : Vertices)
diff --git a/Engine/Source/Developer/BSPUtils/Private/BSPUtils.cpp b/Engine/Source/Developer/BSPUtils/Private/BSPUtils.cpp
--- a/Engine/Source/Developer/BSPUtils/Private/BSPUtils.cpp
+++ b/Engine/Source/Developer/BSPUtils/Private/BSPUtils.cpp
@@ -1988,7 +1988,6 @@ void FBSPUtils::polyUpdateBrush
 	check(Brush);
 
 	FVector ActorLocation;
-	FVector ActorPrePivot;
 	FVector ActorScale;
 	FRotator ActorRotation;
 
diff --git a/Engine/Source/Developer/MeshUtilities/Private/MeshUtilities.cpp b/Engine/Source/Developer/MeshUtilities/Private/MeshUtilities.cpp
--- a/Engine/Source/Developer/MeshUtilities/Private/MeshUtilities.cpp
+++ b/Engine/Source/Developer/MeshUtilities/Private/MeshUtilities.cpp
@@ -4583,8 +4583,7 @@ bool FMeshUtilities::BuildSkeletalMesh_Legacy(FSkeletalMeshLODModel& LODModel
 
 				FMatrix	TextureToLocal = ParameterToTexture.Inverse() * ParameterToLocal;
 				FVector3f	TangentX = FVector4f(TextureToLocal.TransformVector(FVector(1, 0, 0)).GetSafeNormal()),
-					TangentY = FVector4f(TextureToLocal.TransformVector(FVector(0, 1, 0)).GetSafeNormal()),
-					TangentZ;
+					TangentY = FVector4f(TextureToLocal.TransformVector(FVector(0, 1, 0)).GetSafeNormal());
 
 				TangentX = TangentX - TriangleNormal * (TangentX | TriangleNormal);
 				TangentY = TangentY - TriangleNormal * (TangentY | TriangleNormal);
diff --git a/Engine/Source/Editor/ComponentVisualizers/Private/SplineComponentVisualizer.cpp b/Engine/Source/Editor/ComponentVisualizers/Private/SplineComponentVisualizer.cpp
--- a/Engine/Source/Editor/ComponentVisualizers/Private/SplineComponentVisualizer.cpp
+++ b/Engine/Source/Editor/ComponentVisualizers/Private/SplineComponentVisualizer.cpp
@@ -2059,8 +2059,6 @@ void FSplineComponentVisualizer::SnapKeyToTransform(const ESplineComponentSnapMo
 	if (InSnapMode == ESplineComponentSnapMode::AlignToTangent || InSnapMode == ESplineComponentSnapMode::AlignPerpendicularToTangent)
 	{
 		// Copy tangents
-		FVector AlignTangent;
-		FQuat AlignRot;
 		const FVector WorldUpVector = InWorldUpVector.GetSafeNormal();
 		const FVector WorldForwardVector = InWorldForwardVector.GetSafeNormal();
 
diff --git a/Engine/Source/Editor/LandscapeEditor/Private/LandscapeEdModeSplineTools.cpp b/Engine/Source/Editor/LandscapeEditor/Private/LandscapeEdModeSplineTools.cpp
--- a/Engine/Source/Editor/LandscapeEditor/Private/LandscapeEdModeSplineTools.cpp
+++ b/Engine/Source/Editor/LandscapeEditor/Private/LandscapeEdModeSplineTools.cpp
@@ -781,8 +781,8 @@ public:
 		}
 
 		float t, t0, t1, tseg;
-		FVector Location, Location0, Location1;
-		FVector Tangent, Tangent0, Tangent1;
+		FVector Location0, Location1;
+		FVector Tangent0, Tangent1;
 		Segment->FindNearest(LocalLocation, t0, Location0, Tangent0);
 		NewSegment->FindNearest(LocalLocation, t1, Location1, Tangent1);
 
diff --git a/Engine/Source/Editor/UnrealEd/Private/UnrealWidget.cpp b/Engine/Source/Editor/UnrealEd/Private/UnrealWidget.cpp
--- a/Engine/Source/Editor/UnrealEd/Private/UnrealWidget.cpp
+++ b/Engine/Source/Editor/UnrealEd/Private/UnrealWidget.cpp
@@ -648,7 +648,6 @@ void FWidget::AbsoluteConvertMouseToAxis_TranslateRotateZ(const FMatrix& InputCo
                                                           FAbsoluteMovementParams& InOutParams, FVector& OutDrag,
                                                           FRotator& OutRotation)
 {
-	FVector LineToUse;
 	switch (CurrentAxis)
 	{
 	case EAxisList::X: {
diff --git a/Engine/Source/Runtime/AIModule/Private/EnvironmentQuery/Generators/EnvQueryGenerator_Donut.cpp b/Engine/Source/Runtime/AIModule/Private/EnvironmentQuery/Generators/EnvQueryGenerator_Donut.cpp
--- a/Engine/Source/Runtime/AIModule/Private/EnvironmentQuery/Generators/EnvQueryGenerator_Donut.cpp
+++ b/Engine/Source/Runtime/AIModule/Private/EnvironmentQuery/Generators/EnvQueryGenerator_Donut.cpp
@@ -27,7 +27,6 @@ FVector::FReal UEnvQueryGenerator_Donut::GetArcBisectorAngle(FEnvQueryInstance&
 {
 	FRotator::FReal BisectAngle = 0.;
 
-	FVector Direction;
 	if (bDefineArc)
 	{
 		if (ArcDirection.DirMode == EEnvDirection::TwoPoints)
diff --git a/Engine/Source/Runtime/Engine/Private/AI/Navigation/AvoidanceManager.cpp b/Engine/Source/Runtime/Engine/Private/AI/Navigation/AvoidanceManager.cpp
--- a/Engine/Source/Runtime/Engine/Private/AI/Navigation/AvoidanceManager.cpp
+++ b/Engine/Source/Runtime/Engine/Private/AI/Navigation/AvoidanceManager.cpp
@@ -447,8 +447,6 @@ FVector UAvoidanceManager::GetAvoidanceVelocity_Internal(const FNavAvoidanceData
 			FVector PointAWorld = inAvoidanceData.Center;
 			FVector PointBRelative = OtherObject.Center - PointAWorld;
 			FVector TowardB, SidewaysFromB;
-			FVector VelAdjustment;
-			FVector VelAfterAdjustment;
 			float RadiusB = OtherObject.Radius + inAvoidanceData.Radius;
 
 			PointBRelative.Z = 0.0f;
diff --git a/Engine/Source/Runtime/Engine/Private/LevelInstance/LevelInstanceSubsystem.cpp b/Engine/Source/Runtime/Engine/Private/LevelInstance/LevelInstanceSubsystem.cpp
--- a/Engine/Source/Runtime/Engine/Private/LevelInstance/LevelInstanceSubsystem.cpp
+++ b/Engine/Source/Runtime/Engine/Private/LevelInstance/LevelInstanceSubsystem.cpp
@@ -1042,7 +1042,6 @@ bool ULevelInstanceSubsystem::GetLevelInstanceBoundsFromPackage(const FTransform
 	FBox LevelBounds;
 	if (ULevel::GetLevelBoundsFromPackage(LevelPackage, LevelBounds))
 	{
-		FVector LevelBoundsLocation;
 		FVector BoundsLocation;
 		FVector BoundsExtent;
 		LevelBounds.GetCenterAndExtents(BoundsLocation, BoundsExtent);
diff --git a/Engine/Source/Runtime/Engine/Private/Particles/ParticleModules_Location.cpp b/Engine/Source/Runtime/Engine/Private/Particles/ParticleModules_Location.cpp
--- a/Engine/Source/Runtime/Engine/Private/Particles/ParticleModules_Location.cpp
+++ b/Engine/Source/Runtime/Engine/Private/Particles/ParticleModules_Location.cpp
@@ -2865,8 +2865,6 @@ bool UParticleModuleLocationSkelVertSurface::GetParticleLocation(FParticleEmitte
 			}
 			else
 			{
-				FVector VertPos;
-
 				OutPosition = (Verts[0] + Verts[1] + Verts[2]) / 3.0f;
 			}
 
diff --git a/Engine/Source/Runtime/Engine/Private/Particles/ParticleSystemRender.cpp b/Engine/Source/Runtime/Engine/Private/Particles/ParticleSystemRender.cpp
--- a/Engine/Source/Runtime/Engine/Private/Particles/ParticleSystemRender.cpp
+++ b/Engine/Source/Runtime/Engine/Private/Particles/ParticleSystemRender.cpp
@@ -2554,8 +2554,6 @@ void FDynamicBeam2EmitterData::RenderLines(const FParticleSystemSceneProxy* Prox
 		FVector	LastPosition, LastDrawPosition, LastTangent;
 		// The current position
 		FVector	CurrPosition, CurrDrawPosition;
-		// The target
-		FVector	TargetPosition, TargetDrawPosition;
 		// The next target
 		FVector	NextTargetPosition, NextTargetDrawPosition, TargetTangent;
 		// The interperted draw position
@@ -2566,8 +2564,6 @@ void FDynamicBeam2EmitterData::RenderLines(const FParticleSystemSceneProxy* Prox
 
 		FVector Location;
 		FVector EndPoint;
-		FVector Offset;
-		FVector LastOffset;
 		float	fStrength;
 		float	fTargetStrength;
 
@@ -2666,7 +2662,7 @@ void FDynamicBeam2EmitterData::RenderLines(const FParticleSystemSceneProxy* Prox
 
 			bool	bLocked	= BEAM2_TYPEDATA_LOCKED(BeamPayloadData->Lock_Max_NumNoisePoints);
 
-			FVector	UseNoisePoint, CheckNoisePoint;
+			FVector	CheckNoisePoint;
 			FVector	NoiseDir;
 
 			// Reset the texture coordinate
@@ -3553,8 +3549,6 @@ int32 FDynamicBeam2EmitterData::FillData_Noise(FAsyncBufferFillData& Me) const
 	FVector	LastPosition, LastDrawPosition, LastTangent;
 	// The current position
 	FVector	CurrPosition, CurrDrawPosition;
-	// The target
-	FVector	TargetPosition, TargetDrawPosition;
 	// The next target
 	FVector	NextTargetPosition, NextTargetDrawPosition, TargetTangent;
 	// The interperted draw position
@@ -3687,7 +3681,7 @@ int32 FDynamicBeam2EmitterData::FillData_Noise(FAsyncBufferFillData& Me) const
 
 			bool	bLocked	= BEAM2_TYPEDATA_LOCKED(BeamPayloadData->Lock_Max_NumNoisePoints);
 
-			FVector	UseNoisePoint, CheckNoisePoint;
+			FVector	CheckNoisePoint;
 			FVector	NoiseDir;
 
 			for (int32 SheetIndex = 0; SheetIndex < Source.Sheets; SheetIndex++)
@@ -4099,7 +4093,7 @@ int32 FDynamicBeam2EmitterData::FillData_Noise(FAsyncBufferFillData& Me) const
 
 			bool	bLocked	= BEAM2_TYPEDATA_LOCKED(BeamPayloadData->Lock_Max_NumNoisePoints);
 
-			FVector	UseNoisePoint, CheckNoisePoint;
+			FVector	CheckNoisePoint;
 			FVector	NoiseDir;
 
 			for (int32 SheetIndex = 0; SheetIndex < Source.Sheets; SheetIndex++)
@@ -4516,8 +4510,6 @@ int32 FDynamicBeam2EmitterData::FillData_InterpolatedNoise(FAsyncBufferFillData&
 	FVector	LastPosition, LastDrawPosition, LastTangent;
 	// The current position
 	FVector	CurrPosition, CurrDrawPosition;
-	// The target
-	FVector	TargetPosition, TargetDrawPosition;
 	// The next target
 	FVector	NextTargetPosition, NextTargetDrawPosition, TargetTangent;
 	// The interperted draw position
@@ -4664,7 +4656,7 @@ int32 FDynamicBeam2EmitterData::FillData_InterpolatedNoise(FAsyncBufferFillData&
 
 		bool	bLocked	= BEAM2_TYPEDATA_LOCKED(BeamPayloadData->Lock_Max_NumNoisePoints);
 
-		FVector	UseNoisePoint, CheckNoisePoint;
+		FVector	CheckNoisePoint;
 		FVector	NoiseDir;
 
 		for (int32 SheetIndex = 0; SheetIndex < Source.Sheets; SheetIndex++)
diff --git a/Engine/Source/Runtime/Engine/Private/Polygon.cpp b/Engine/Source/Runtime/Engine/Private/Polygon.cpp
--- a/Engine/Source/Runtime/Engine/Private/Polygon.cpp
+++ b/Engine/Source/Runtime/Engine/Private/Polygon.cpp
@@ -461,7 +461,6 @@ void FPoly::Transform
 	const FVector3f&		PostAdd
 )
 {
-	FVector 	Temp;
 	int32 		i;
 
 	Base += PostAdd;
diff --git a/Engine/Source/Runtime/GeometryCore/Private/DynamicMesh/MeshTangents.cpp b/Engine/Source/Runtime/GeometryCore/Private/DynamicMesh/MeshTangents.cpp
--- a/Engine/Source/Runtime/GeometryCore/Private/DynamicMesh/MeshTangents.cpp
+++ b/Engine/Source/Runtime/GeometryCore/Private/DynamicMesh/MeshTangents.cpp
@@ -262,7 +262,6 @@ void TMeshTangents<RealType>::ComputeTriangleTangents(const FDynamicMeshUVOverla
 		Mesh->GetTriVertices(TriangleID, TriVertices[0], TriVertices[1], TriVertices[2]);
 		FVector2f TriUVs[3];
 		UVOverlay->GetTriElements(TriangleID, TriUVs[0], TriUVs[1], TriUVs[2]);
-		FVector3f TriNormals[3];
 		FVector3d TriNormal = Mesh->GetTriNormal(TriangleID);
 
 		FVector3d Tangent, Bitangent;
diff --git a/Engine/Source/Runtime/MovieSceneTracks/Private/Sections/MovieScene3DTransformSection.cpp b/Engine/Source/Runtime/MovieSceneTracks/Private/Sections/MovieScene3DTransformSection.cpp
--- a/Engine/Source/Runtime/MovieSceneTracks/Private/Sections/MovieScene3DTransformSection.cpp
+++ b/Engine/Source/Runtime/MovieSceneTracks/Private/Sections/MovieScene3DTransformSection.cpp
@@ -840,10 +840,6 @@ EMovieSceneChannelProxyType UMovieScene3DTransformSection::CacheChannelProxy()
 
 TSharedPtr<FStructOnScope> UMovieScene3DTransformSection::GetKeyStruct(TArrayView<const FKeyHandle> KeyHandles)
 {
-	FVector  StartingLocation;
-	FRotator StartingRotation;
-	FVector  StartingScale;
-
 	TArrayView<FMovieSceneDoubleChannel* const> DoubleChannels = ChannelProxy->GetChannels<FMovieSceneDoubleChannel>();
 
 	TOptional<TTuple<FKeyHandle, FFrameNumber>> LocationKeys[3] = {
