Light Function Atlas - renamed a property internal name to be more explicit and diferenciante from other.

diff --git a/Engine/Source/Runtime/Engine/Classes/Materials/Material.h b/Engine/Source/Runtime/Engine/Classes/Materials/Material.h
--- a/Engine/Source/Runtime/Engine/Classes/Materials/Material.h
+++ b/Engine/Source/Runtime/Engine/Classes/Materials/Material.h
@@ -822,9 +822,10 @@ public:
 	 * Indicates that this material is safe to use with the light function atlas:
 	 *  - texture coordinates are manipulated in a way that works.
 	 *  - world position and depth are not used in the graph or do work with the atlas.
+	 * Forces all deferred lights to go the fast batched lighting code path (no screen shadow mask).
 	 */
-	UPROPERTY(EditAnywhere, Category = LightFunctionMaterial)
-	uint8 bCompatibleWithLightFunctionAtlas : 1;
+	UPROPERTY(EditAnywhere, Category = LightFunctionMaterial, meta = (DisplayName = "Compatible With Light Function Atlas"))
+	uint8 bForceCompatibleWithLightFunctionAtlas : 1;
 
 	/** 
 	 * Whether to automatically set usage flags based on what the material is applied to in the editor.
diff --git a/Engine/Source/Runtime/Engine/Private/Materials/HLSLMaterialTranslator.cpp b/Engine/Source/Runtime/Engine/Private/Materials/HLSLMaterialTranslator.cpp
--- a/Engine/Source/Runtime/Engine/Private/Materials/HLSLMaterialTranslator.cpp
+++ b/Engine/Source/Runtime/Engine/Private/Materials/HLSLMaterialTranslator.cpp
@@ -1721,7 +1721,7 @@ void FHLSLMaterialTranslator::DoTranslate()
 	 *		- Wrapped textures look different however if texcoords are scaled (looks correct for 1, 2, 3 but not for 1.5 or 1.1 for instance as UV space [0,1] will no longer align with the atlas tile edges)
 	 * But, an artist can specify and override the fact that a material is compatible with the light function atlas.
 	 */
-	MaterialCompilationOutput.bIsLightFunctionAtlasCompatible = (!bUsesVertexPosition && !bUsesSceneDepth && !MaterialCompilationOutput.bNeedsSceneTextures && !bPotentiallyManipulateTexCoords) || Material->IsCompatibleWithLightFunctionAtlas();
+	MaterialCompilationOutput.bIsLightFunctionAtlasCompatible = (!bUsesVertexPosition && !bUsesSceneDepth && !MaterialCompilationOutput.bNeedsSceneTextures && !bPotentiallyManipulateTexCoords) || Material->GetForceCompatibleWithLightFunctionAtlas();
 
 	EMaterialDecalResponse MDR = (EMaterialDecalResponse)Material->GetMaterialDecalResponse();
 	if (MDR == MDR_Color || MDR == MDR_ColorNormal || MDR == MDR_ColorRoughness || MDR == MDR_ColorNormalRoughness)
diff --git a/Engine/Source/Runtime/Engine/Private/Materials/Material.cpp b/Engine/Source/Runtime/Engine/Private/Materials/Material.cpp
--- a/Engine/Source/Runtime/Engine/Private/Materials/Material.cpp
+++ b/Engine/Source/Runtime/Engine/Private/Materials/Material.cpp
@@ -1021,7 +1021,7 @@ UMaterial::UMaterial(const FObjectInitializer& ObjectInitializer)
 	bUsedWithHeterogeneousVolumes = false;
 
 	// We default to false=incompatible, and rely onto the automatic detection executed from the HLSLTranslator.
-	bCompatibleWithLightFunctionAtlas = false;
+	bForceCompatibleWithLightFunctionAtlas = false;
 
 	bUseEmissiveForDynamicAreaLighting = false;
 	RefractionDepthBias = 0.0f;
diff --git a/Engine/Source/Runtime/Engine/Private/Materials/MaterialShared.cpp b/Engine/Source/Runtime/Engine/Private/Materials/MaterialShared.cpp
--- a/Engine/Source/Runtime/Engine/Private/Materials/MaterialShared.cpp
+++ b/Engine/Source/Runtime/Engine/Private/Materials/MaterialShared.cpp
@@ -1898,9 +1898,9 @@ bool FMaterialResource::IsFullyRough() const
 	return Material->bFullyRough;
 }
 
-bool FMaterialResource::IsCompatibleWithLightFunctionAtlas() const
+bool FMaterialResource::GetForceCompatibleWithLightFunctionAtlas() const
 {
-	return Material->bCompatibleWithLightFunctionAtlas;
+	return Material->bForceCompatibleWithLightFunctionAtlas;
 }
 
 bool FMaterialResource::UseNormalCurvatureToRoughness() const
diff --git a/Engine/Source/Runtime/Engine/Public/MaterialShared.h b/Engine/Source/Runtime/Engine/Public/MaterialShared.h
--- a/Engine/Source/Runtime/Engine/Public/MaterialShared.h
+++ b/Engine/Source/Runtime/Engine/Public/MaterialShared.h
@@ -2072,7 +2072,7 @@ public:
 	virtual bool IsUsedWithVirtualHeightfieldMesh() const { return false; }
 	virtual bool IsUsedWithNeuralNetworks() const { return false; }
 	virtual bool IsFullyRough() const { return false; }
-	virtual bool IsCompatibleWithLightFunctionAtlas() const { return false; }
+	virtual bool GetForceCompatibleWithLightFunctionAtlas() const { return false; }
 	virtual bool UseNormalCurvatureToRoughness() const { return false; }
 	virtual enum EMaterialFloatPrecisionMode GetMaterialFloatPrecisionMode() const { return EMaterialFloatPrecisionMode::MFPM_Default; };
 	virtual bool IsUsingAlphaToCoverage() const { return false; }
@@ -2836,7 +2836,7 @@ public:
 	ENGINE_API virtual bool IsUsedWithVolumetricCloud() const override;
 	ENGINE_API virtual bool IsUsedWithHeterogeneousVolumes() const override;
 	ENGINE_API virtual bool IsFullyRough() const override;
-	ENGINE_API virtual bool IsCompatibleWithLightFunctionAtlas() const override;
+	ENGINE_API virtual bool GetForceCompatibleWithLightFunctionAtlas() const override;
 	ENGINE_API virtual bool UseNormalCurvatureToRoughness() const override;
 	ENGINE_API virtual enum EMaterialFloatPrecisionMode GetMaterialFloatPrecisionMode() const override;
 	ENGINE_API virtual bool IsUsingAlphaToCoverage() const override;
