Fixed double processing of cloth buffers when skin cache is enabled.

diff --git a/Engine/Source/Runtime/Engine/Private/SkeletalRenderGPUSkin.cpp b/Engine/Source/Runtime/Engine/Private/SkeletalRenderGPUSkin.cpp
--- a/Engine/Source/Runtime/Engine/Private/SkeletalRenderGPUSkin.cpp
+++ b/Engine/Source/Runtime/Engine/Private/SkeletalRenderGPUSkin.cpp
@@ -762,14 +762,7 @@ void FSkeletalMeshObjectGPUSkin::ProcessUpdatedDynamicData(
 
 			if (ClothSimulationData)
 			{
-				ClothShaderData->ClothBlendWeight = DynamicData->ClothBlendWeight;
-				ClothShaderData->WorldScale = (FVector3f)WorldScale;
-				ClothShaderData->UpdateClothSimulationData(RHICmdList, ClothSimulationData->Positions, ClothSimulationData->Normals, RevisionNumber, OwnerName);
-
 				ClothToLocal = FMatrix44f(ClothSimulationData->ComponentRelativeTransform.ToMatrixWithScale());
-
-				// Transform from cloth space to local space. Cloth space is relative to cloth root bone, local space is component space.
-				ClothShaderData->GetClothToLocalForWriting() = ClothToLocal;
 			}
 		}
 
@@ -799,6 +792,17 @@ void FSkeletalMeshObjectGPUSkin::ProcessUpdatedDynamicData(
 			bAllowAddToSkinCache = bSectionUsingSkinCache;
 		}
 
+		// Process the cloth simulation once we know for sure we aren't using the skin cache for this section.
+		if (!bSectionUsingSkinCache && ClothSimulationData)
+		{
+			ClothShaderData->ClothBlendWeight = DynamicData->ClothBlendWeight;
+			ClothShaderData->WorldScale = (FVector3f)WorldScale;
+			ClothShaderData->UpdateClothSimulationData(RHICmdList, ClothSimulationData->Positions, ClothSimulationData->Normals, RevisionNumber, OwnerName);
+
+			// Transform from cloth space to local space. Cloth space is relative to cloth root bone, local space is component space.
+			ClothShaderData->GetClothToLocalForWriting() = ClothToLocal;
+		}
+
 		bAllSectionsUsingSkinCache &= bSectionUsingSkinCache;
 
 		if (Mode == EGPUSkinCacheEntryMode::Raster && DynamicData->GPUSkinTechnique != ESkeletalMeshGPUSkinTechnique::MeshDeformer && !bSectionUsingSkinCache)
