Avoid issues where use can see default decal material without knowing why.
Problem is that when we select non-decal materials on a decal component, the renderer will drop back to default material which can be confusing for the user.
Three changes have been made:
* Filter the material selection list to decal only materials. Note that this requires us to iterate up the parent AssetData chain for material instances which is potentially slow. But testing on our largest asset registry, this was still very usable.
* Pop a toast when a decal component has a non-decal material. This tells the user about the issue, and gives hyperlink to edit the material. This is to surface the legacy case or the case where material's are changed after selection.
* Show a map check warning when a decal component has a non-decal material. This is another way to surface the legacy case.

diff --git a/Engine/Source/Editor/DetailCustomizations/Private/DecalComponentDetails.cpp b/Engine/Source/Editor/DetailCustomizations/Private/DecalComponentDetails.cpp
new file mode 100644
--- /dev/null
+++ b/Engine/Source/Editor/DetailCustomizations/Private/DecalComponentDetails.cpp
@@ -0,0 +1,78 @@
+// Copyright Epic Games, Inc. All Rights Reserved.
+
+#include "DecalComponentDetails.h"
+
+#include "AssetRegistry/AssetRegistryModule.h"
+#include "DetailCategoryBuilder.h"
+#include "DetailLayoutBuilder.h"
+#include "DetailWidgetRow.h"
+#include "Materials/MaterialInterface.h"
+#include "Modules/ModuleManager.h"
+#include "PropertyCustomizationHelpers.h"
+
+#define LOCTEXT_NAMESPACE "DecalComponentDetails"
+
+static TAutoConsoleVariable<bool> CVarDecalFilterMaterialList(
+	TEXT("r.Decal.FilterMaterialList"),
+	true,
+	TEXT("Enable filtering of material list in Decal Component details panel.")
+);
+
+TSharedRef<IDetailCustomization> FDecalComponentDetails::MakeInstance()
+{
+	return MakeShareable( new FDecalComponentDetails);
+}
+
+void FDecalComponentDetails::CustomizeDetails(IDetailLayoutBuilder& DetailBuilder)
+{
+	TArray<TWeakObjectPtr<UObject>> ObjectsBeingCustomized;
+	DetailBuilder.GetObjectsBeingCustomized(ObjectsBeingCustomized);
+	if (ObjectsBeingCustomized.Num() > 1)
+	{
+		return;
+	}
+
+	// Filter only show decal materials.
+	TSharedRef<IPropertyHandle> PropertyHandle = DetailBuilder.GetProperty(TEXT("DecalMaterial"));
+	DetailBuilder.EditDefaultProperty(PropertyHandle)->CustomWidget()
+	.NameContent()
+	[
+		PropertyHandle->CreatePropertyNameWidget()
+	]
+	.ValueContent()
+	[
+		SNew(SObjectPropertyEntryBox)
+		.PropertyHandle(PropertyHandle)
+		.AllowedClass(UMaterialInterface::StaticClass())
+		.ThumbnailPool(DetailBuilder.GetThumbnailPool())
+		.OnShouldFilterAsset(FOnShouldFilterAsset::CreateStatic(&FDecalComponentDetails::ShouldFilterDecalMaterialAsset))
+	];
+}
+
+bool FDecalComponentDetails::ShouldFilterDecalMaterialAsset(FAssetData const& AssetData)
+{
+	if (!CVarDecalFilterMaterialList.GetValueOnAnyThread())
+	{
+		return false;
+	}
+
+	FAssetRegistryModule& AssetRegistryModule = FModuleManager::Get().LoadModuleChecked<FAssetRegistryModule>(TEXT("AssetRegistry"));
+	return !IsDecalMaterialAssetRecursive(AssetData, AssetRegistryModule.Get());
+}
+
+bool FDecalComponentDetails::IsDecalMaterialAssetRecursive(FAssetData const& AssetData, IAssetRegistry const& AssetRegistry)
+{
+	FString DataStr;
+	if (AssetData.GetTagValue("Parent", DataStr))
+	{
+		FAssetData ParentAsset = AssetRegistry.GetAssetByObjectPath(FSoftObjectPath(*DataStr));
+		return IsDecalMaterialAssetRecursive(ParentAsset, AssetRegistry);
+	}
+	else if (AssetData.GetTagValue("MaterialDomain", DataStr))
+	{
+		return DataStr.Equals(TEXT("MD_DeferredDecal"));
+	}
+	return false;
+}
+
+#undef LOCTEXT_NAMESPACE
diff --git a/Engine/Source/Editor/DetailCustomizations/Private/DecalComponentDetails.h b/Engine/Source/Editor/DetailCustomizations/Private/DecalComponentDetails.h
new file mode 100644
--- /dev/null
+++ b/Engine/Source/Editor/DetailCustomizations/Private/DecalComponentDetails.h
@@ -0,0 +1,24 @@
+// Copyright Epic Games, Inc. All Rights Reserved.
+
+#pragma once
+
+#include "IDetailCustomization.h"
+
+struct FAssetData;
+class IAssetRegistry;
+
+class FDecalComponentDetails : public IDetailCustomization
+{
+public:
+	/** Makes a new instance of this detail layout class for a specific detail view requesting it */
+	static TSharedRef<IDetailCustomization> MakeInstance();
+
+	/** IDetailCustomization interface */
+	virtual void CustomizeDetails(IDetailLayoutBuilder& DetailBuilder) override;
+
+private:
+	/** Returns false if asset should appear in filtered list. */
+	static bool ShouldFilterDecalMaterialAsset(FAssetData const& AssetData);
+	/** Returns true if this asset is a decal material, or child of a decal material. */
+	static bool IsDecalMaterialAssetRecursive(FAssetData const& AssetData, IAssetRegistry const& AssetRegistry);
+};
diff --git a/Engine/Source/Editor/DetailCustomizations/Private/DetailCustomizations.cpp b/Engine/Source/Editor/DetailCustomizations/Private/DetailCustomizations.cpp
--- a/Engine/Source/Editor/DetailCustomizations/Private/DetailCustomizations.cpp
+++ b/Engine/Source/Editor/DetailCustomizations/Private/DetailCustomizations.cpp
@@ -49,6 +49,7 @@
 #include "DataTableCustomization.h"
 #include "DateTimeStructCustomization.h"
 #include "DebugCameraControllerSettingsCustomization.h"
+#include "DecalComponentDetails.h"
 #include "Delegates/Delegate.h"
 #include "DeviceProfileDetails.h"
 #include "DialogueStructsCustomizations.h"
@@ -476,6 +477,8 @@ void FDetailCustomizationsModule::RegisterObjectCustomizations()
 	RegisterCustomClassLayout("BoundsCopyComponent", FOnGetDetailCustomizationInstance::CreateStatic(&FBoundsCopyComponentDetailsCustomization::MakeInstance));
 
 	RegisterCustomClassLayout("SoundNodeDistanceCrossFade", FOnGetDetailCustomizationInstance::CreateStatic(&FCrossFadeCustomization::MakeInstance));
+	
+	RegisterCustomClassLayout("DecalComponent", FOnGetDetailCustomizationInstance::CreateStatic(&FDecalComponentDetails::MakeInstance));
 }
 
 #define LOCTEXT_NAMESPACE "DetailsSections"
diff --git a/Engine/Source/Runtime/Engine/Classes/Components/DecalComponent.h b/Engine/Source/Runtime/Engine/Classes/Components/DecalComponent.h
--- a/Engine/Source/Runtime/Engine/Classes/Components/DecalComponent.h
+++ b/Engine/Source/Runtime/Engine/Classes/Components/DecalComponent.h
@@ -178,12 +178,16 @@ public:
 
 	
 	//~ Begin UActorComponent Interface
+	ENGINE_API virtual void OnRegister() override;
 	ENGINE_API virtual void BeginPlay() override;
 	ENGINE_API virtual void CreateRenderState_Concurrent(FRegisterComponentContext* Context) override;
 	ENGINE_API virtual void DestroyRenderState_Concurrent() override;
 	ENGINE_API virtual void SendRenderTransform_Concurrent() override;
 	ENGINE_API virtual const UObject* AdditionalStatObject() const override;
 	ENGINE_API virtual void PrecachePSOs() override;
+#if WITH_EDITOR
+	ENGINE_API virtual void CheckForErrors() override;
+#endif
 	//~ End UActorComponent Interface
 	
 	//~ Begin UObject Interface. 
diff --git a/Engine/Source/Runtime/Engine/Private/Components/DecalComponent.cpp b/Engine/Source/Runtime/Engine/Private/Components/DecalComponent.cpp
--- a/Engine/Source/Runtime/Engine/Private/Components/DecalComponent.cpp
+++ b/Engine/Source/Runtime/Engine/Private/Components/DecalComponent.cpp
@@ -16,8 +16,20 @@
 #include "Materials/MaterialInstanceDynamic.h"
 #include "MarkActorRenderStateDirtyTask.h"
 
+#if WITH_EDITOR
+#include "Editor.h"
+#include "Framework/Notifications/NotificationManager.h"
+#include "Logging/MessageLog.h"
+#include "Misc/MapErrors.h"
+#include "Misc/UObjectToken.h"
+#include "Subsystems/AssetEditorSubsystem.h"
+#include "Widgets/Notifications/SNotificationList.h"
+#endif
+
 #include UE_INLINE_GENERATED_CPP_BY_NAME(DecalComponent)
 
+#define LOCTEXT_NAMESPACE "DecalComponent"
+
 static TAutoConsoleVariable<float> CVarDecalFadeDurationScale(
 	TEXT("r.Decal.FadeDurationScale"),
 	1.0f,
@@ -396,6 +408,36 @@ FBoxSphereBounds UDecalComponent::CalcBounds(const FTransform& LocalToWorld) con
 	return FBoxSphereBounds(FVector(0, 0, 0), DecalSize, DecalSize.Size()).TransformBy(LocalToWorld);
 }
 
+void UDecalComponent::OnRegister()
+{
+	Super::OnRegister();
+
+#if WITH_EDITOR
+	if (DecalMaterial && DecalMaterial->GetMaterial()->MaterialDomain != MD_DeferredDecal && GEditor)
+	{
+		static TWeakPtr<class SNotificationItem> NotificationHandle;
+		if (!NotificationHandle.IsValid())
+		{
+			FNotificationInfo Info(LOCTEXT("DecalMaterial_Notify", "Decal Material must use Deferred Decal Material Domain."));
+			Info.bFireAndForget = true;
+			Info.ExpireDuration = 8.0f;
+			Info.SubText = FText::Format(
+				LOCTEXT("DecalMaterial_NotifySubtext", "Decal materials must use the Deferred Decal Material Domain.\nEither select a valid material for {0} or open the current material and select the Deferred Decal Material Domain."), 
+				FText::FromString(GetOwner()->GetActorNameOrLabel()));
+			Info.HyperlinkText = FText::Format(
+				LOCTEXT("DecalMaterial_Hyperlink", "Open {0}"), 
+				FText::FromString(DecalMaterial->GetName()));
+			Info.Hyperlink = FSimpleDelegate::CreateWeakLambda(this, [this]
+				{
+					GEditor->GetEditorSubsystem<UAssetEditorSubsystem>()->OpenEditorForAsset(DecalMaterial);
+				});
+
+			NotificationHandle = FSlateNotificationManager::Get().AddNotification(Info);
+		}
+	}
+#endif
+}
+
 void UDecalComponent::BeginPlay()
 {
 	Super::BeginPlay();
@@ -436,4 +478,24 @@ void UDecalComponent::DestroyRenderState_Concurrent()
 	GetWorld()->Scene->RemoveDecal(this);
 }
 
+#if WITH_EDITOR
+
+void UDecalComponent::CheckForErrors()
+{
+	Super::CheckForErrors();
+
+	if (DecalMaterial && DecalMaterial->GetMaterial()->MaterialDomain != MD_DeferredDecal)
+	{
+		FFormatNamedArguments Arguments;
+		Arguments.Add(TEXT("ComponentName"), FText::FromString(GetName()));
+		Arguments.Add(TEXT("OwnerName"), FText::FromString(GetNameSafe(GetOwner())));
+
+		FMessageLog("MapCheck").Warning()
+			->AddToken(FUObjectToken::Create(this))
+			->AddToken(FTextToken::Create(FText::Format(LOCTEXT("DecalMaterial_MapCheck", "{ComponentName}::{OwnerName} has a DecalMaterial that doesn't use the Deferred Decal Material Domain."), Arguments)));
+	}
+}
+
+#endif // WITH_EDITOR
 
+#undef LOCTEXT_NAMESPACE
