Check r.RayTracing.Enable during initialization when not using r.RayTracing.EnableOnDemand.

diff --git a/Engine/Source/Runtime/Engine/Classes/Engine/RendererSettings.h b/Engine/Source/Runtime/Engine/Classes/Engine/RendererSettings.h
--- a/Engine/Source/Runtime/Engine/Classes/Engine/RendererSettings.h
+++ b/Engine/Source/Runtime/Engine/Classes/Engine/RendererSettings.h
@@ -541,7 +541,7 @@ class URendererSettings : public UDeveloperSettings
 	 */
 	UPROPERTY(config, EditAnywhere, Category = HardwareRayTracing, meta = (
 		ConsoleVariable = "r.RayTracing", DisplayName = "Support Hardware Ray Tracing",
-		ToolTip = "Support Hardware Ray Tracing features.  Requires 'Support Compute Skincache' before project is allowed to set this.",
+		ToolTip = "Support Hardware Ray Tracing features.  Requires 'Support Compute Skin Cache' before project is allowed to set this.",
 		ConfigRestartRequired = true))
 		uint32 bEnableRayTracing : 1;
 
diff --git a/Engine/Source/Runtime/Engine/Private/Rendering/RayTracingGeometryManager.cpp b/Engine/Source/Runtime/Engine/Private/Rendering/RayTracingGeometryManager.cpp
--- a/Engine/Source/Runtime/Engine/Private/Rendering/RayTracingGeometryManager.cpp
+++ b/Engine/Source/Runtime/Engine/Private/Rendering/RayTracingGeometryManager.cpp
@@ -18,7 +18,8 @@ static bool bHasRayTracingEnableChanged = false;
 static TAutoConsoleVariable<int32> CVarRayTracingEnable(
 	TEXT("r.RayTracing.Enable"),
 	1,
-	TEXT("Runtime toggle for switching raytracing on/off (experimental)."),
+	TEXT("Whether ray tracing is enabled at runtime.\n")
+	TEXT("If r.RayTracing.EnableOnDemand is enabled, ray tracing can be toggled on/off at runtime. Otherwise this is only checked during initialization."),
 	FConsoleVariableDelegate::CreateLambda([](IConsoleVariable* InVariable)
 		{
 			FGlobalComponentRecreateRenderStateContext Context;
diff --git a/Engine/Source/Runtime/RenderCore/Private/RenderUtils.cpp b/Engine/Source/Runtime/RenderCore/Private/RenderUtils.cpp
--- a/Engine/Source/Runtime/RenderCore/Private/RenderUtils.cpp
+++ b/Engine/Source/Runtime/RenderCore/Private/RenderUtils.cpp
@@ -901,7 +901,13 @@ RENDERCORE_API void RenderUtilsInit()
 			}
 			else
 			{
-				DesiredRayTracingMode = ERayTracingMode::Enabled;
+				// if not using EnableOnDemand, check r.Raytracing.Enable during initialization (changing the cvar at runtime will have no effect)
+				static const auto RayTracingEnableCVar = IConsoleManager::Get().FindTConsoleVariableDataInt(TEXT("r.Raytracing.Enable"));
+
+				if (RayTracingEnableCVar && RayTracingEnableCVar->GetValueOnAnyThread() != 0)
+				{
+					DesiredRayTracingMode = ERayTracingMode::Enabled;
+				}
 			}
 		}
 
