Fixed client-only/server-only data layers not changing state once activated.

diff --git a/Engine/Source/Runtime/Engine/Private/WorldPartition/DataLayer/WorldDataLayers.cpp b/Engine/Source/Runtime/Engine/Private/WorldPartition/DataLayer/WorldDataLayers.cpp
--- a/Engine/Source/Runtime/Engine/Private/WorldPartition/DataLayer/WorldDataLayers.cpp
+++ b/Engine/Source/Runtime/Engine/Private/WorldPartition/DataLayer/WorldDataLayers.cpp
@@ -248,6 +248,7 @@ void AWorldDataLayers::SetDataLayerRuntimeState(const UDataLayerInstance* InData
 		LoadedDataLayerNames.Remove(InDataLayerInstance->GetDataLayerFName());
 		ActiveDataLayerNames.Remove(InDataLayerInstance->GetDataLayerFName());
 		LocalLoadedDataLayerNames.Remove(InDataLayerInstance->GetDataLayerFName());
+		LocalActiveDataLayerNames.Remove(InDataLayerInstance->GetDataLayerFName());
 
 		if (InState == EDataLayerRuntimeState::Loaded)
 		{
