Fixed potentially uninitialized FVectors in Engine.

diff --git a/Engine/Source/Runtime/Engine/Private/Actor.cpp b/Engine/Source/Runtime/Engine/Private/Actor.cpp
--- a/Engine/Source/Runtime/Engine/Private/Actor.cpp
+++ b/Engine/Source/Runtime/Engine/Private/Actor.cpp
@@ -1001,8 +1001,8 @@ void AActor::PostLoadSubobjects(FObjectInstancingGraph* OuterInstanceGraph)
 	USceneComponent* OldRootParent = (OldRoot ? OldRoot->GetAttachParent() : nullptr);
 	bool bHadRoot = !!OldRoot;
 	FRotator OldRotation;
-	FVector OldTranslation;
-	FVector OldScale;
+	FVector OldTranslation = FVector::ZeroVector;
+	FVector OldScale = FVector::ZeroVector;
 	if (bHadRoot)
 	{
 		OldRotation = OldRoot->GetRelativeRotation();
diff --git a/Engine/Source/Runtime/Engine/Private/Particles/ParticleBeamModules.cpp b/Engine/Source/Runtime/Engine/Private/Particles/ParticleBeamModules.cpp
--- a/Engine/Source/Runtime/Engine/Private/Particles/ParticleBeamModules.cpp
+++ b/Engine/Source/Runtime/Engine/Private/Particles/ParticleBeamModules.cpp
@@ -1910,94 +1910,88 @@ bool UParticleModuleBeamSource::ResolveSourceData(FParticleBeam2EmitterInstance*
 						BeamInst->ResolveSource();
 					}
 
-					if (BeamInst->SourceEmitter)
+					if (BeamInst->SourceEmitter && BeamInst->SourceEmitter->ParticleIndices)
 					{
-						FVector CalcSourcePosition;
-
 						int32 SourceIndex = LastSelectedParticleIndex;
 
-						if (BeamInst->SourceEmitter && BeamInst->SourceEmitter->ParticleIndices)
+						if (SourceIndex != -1)
 						{
-							if (SourceIndex != -1)
+							FBaseParticle* SourceParticle = BeamInst->SourceEmitter->GetParticleDirect(SourceIndex);
+							if (SourceParticle == NULL || SourceParticle->RelativeTime>1.0f)
 							{
-								FBaseParticle* SourceParticle = BeamInst->SourceEmitter->GetParticleDirect(SourceIndex);
-								if (SourceParticle == NULL || SourceParticle->RelativeTime>1.0f)
-								{
-									// If the previous particle is not found, force the trail to pick a new one
-									SourceIndex = -1;
-								}
+								// If the previous particle is not found, force the trail to pick a new one
+								SourceIndex = -1;
 							}
+						}
 
-							if (SourceIndex == -1)
+						if (SourceIndex == -1)
+						{
+							int32 Index = 0;
+							// TODO: add selection method and random selection
+							/*
+							switch (SelectionMethod)
 							{
-								int32 Index = 0;
-								// TODO: add selection method and random selection
-								/*
-								switch (SelectionMethod)
-								{
-								case EPSSM_Random:
-								{
-								Index = FMath::TruncToInt(FMath::FRand() * BeamInst->SourceEmitter->ActiveParticles);
-								}
-								break;
-
-								case EPSSM_Sequential:*/
-								{
-									if (++LastSelectedParticleIndex >= BeamInst->SourceEmitter->ActiveParticles)
-									{
-										LastSelectedParticleIndex = 0;
-									}
-									Index = LastSelectedParticleIndex;
-								}
-								/*
-								break;
-								}
-								*/
+							case EPSSM_Random:
+							{
+							Index = FMath::TruncToInt(FMath::FRand() * BeamInst->SourceEmitter->ActiveParticles);
+							}
+							break;
 
-								if (SourceIndex == BeamInst->SourceEmitter->ParticleIndices[Index])
+							case EPSSM_Sequential:*/
+							{
+								if (++LastSelectedParticleIndex >= BeamInst->SourceEmitter->ActiveParticles)
 								{
-									Index = -1;
+									LastSelectedParticleIndex = 0;
 								}
+								Index = LastSelectedParticleIndex;
+							}
+							/*
+							break;
+							}
+							*/
 
-								SourceIndex = (Index != -1) ? BeamInst->SourceEmitter->ParticleIndices[Index] : -1;
+							if (SourceIndex == BeamInst->SourceEmitter->ParticleIndices[Index])
+							{
+								Index = -1;
 							}
 
-							bool bEncounteredNaNError = false;
+							SourceIndex = (Index != -1) ? BeamInst->SourceEmitter->ParticleIndices[Index] : -1;
+						}
 
-							// Grab the particle
+						bool bEncounteredNaNError = false;
 
-							UParticleSystemComponent* Comp = BeamInst->SourceEmitter->Component;
-							FBaseParticle* SourceParticle = (SourceIndex >= 0) ? BeamInst->SourceEmitter->GetParticleDirect(SourceIndex) : nullptr;
-							if (SourceParticle != nullptr)
-							{
-								const FVector WorldOrigin = BeamInst->SourceEmitter->SimulationToWorld.GetOrigin();
+						// Grab the particle
 
-								if (!ensureMsgf(!SourceParticle->Location.ContainsNaN(), TEXT("NaN in SourceParticle Location. Template: %s, Component: %s"), Comp ? *GetNameSafe(Comp->Template) : TEXT("UNKNOWN"), *GetPathNameSafe(Comp)) ||
-									!ensureMsgf(!SourceParticle->OldLocation.ContainsNaN(), TEXT("NaN in SourceParticle OldLocation. Template: %s, Component: %s"), Comp ? *GetNameSafe(Comp->Template) : TEXT("UNKNOWN"), *GetPathNameSafe(Comp)) ||
-									!ensureMsgf(!WorldOrigin.ContainsNaN(), TEXT("NaN in WorldOrigin. Template: %s, Component: %s"), Comp ? *GetNameSafe(Comp->Template) : TEXT("UNKNOWN"), *GetPathNameSafe(Comp))
-									)
-								{
-									// Contains NaN!
-									bEncounteredNaNError = true;
-								}
-								else
-								{
-									CalcSourcePosition = SourceParticle->Location + WorldOrigin;
-								}
+						UParticleSystemComponent* Comp = BeamInst->SourceEmitter->Component;
+						FBaseParticle* SourceParticle = (SourceIndex >= 0) ? BeamInst->SourceEmitter->GetParticleDirect(SourceIndex) : nullptr;
+						if (SourceParticle != nullptr)
+						{
+							const FVector WorldOrigin = BeamInst->SourceEmitter->SimulationToWorld.GetOrigin();
+
+							if (!ensureMsgf(!SourceParticle->Location.ContainsNaN(), TEXT("NaN in SourceParticle Location. Template: %s, Component: %s"), Comp ? *GetNameSafe(Comp->Template) : TEXT("UNKNOWN"), *GetPathNameSafe(Comp)) ||
+								!ensureMsgf(!SourceParticle->OldLocation.ContainsNaN(), TEXT("NaN in SourceParticle OldLocation. Template: %s, Component: %s"), Comp ? *GetNameSafe(Comp->Template) : TEXT("UNKNOWN"), *GetPathNameSafe(Comp)) ||
+								!ensureMsgf(!WorldOrigin.ContainsNaN(), TEXT("NaN in WorldOrigin. Template: %s, Component: %s"), Comp ? *GetNameSafe(Comp->Template) : TEXT("UNKNOWN"), *GetPathNameSafe(Comp))
+								)
+							{
+								// Contains NaN!
+								bEncounteredNaNError = true;
 							}
 							else
 							{
-								// Fall back to the emitter location??
-								CalcSourcePosition = Comp->GetComponentLocation();
-								//@todo. How to handle this... can potentially cause a jump from the emitter to the
-								// particle...
-								SourceIndex = -1;//No valid particle source;
+								BeamData->SourcePoint = SourceParticle->Location + WorldOrigin;
+								bSetSource = true;
 							}
 						}
+						else
+						{
+							// Fall back to the emitter location??
+							BeamData->SourcePoint = Comp->GetComponentLocation();
+							bSetSource = true;
 
-						//@todo. fill this in correctly...
-						BeamData->SourcePoint = CalcSourcePosition;
-						bSetSource = true;
+							//@todo. How to handle this... can potentially cause a jump from the emitter to the
+							// particle...
+							SourceIndex = -1;//No valid particle source;
+						}
 					}
 				}
 			}
@@ -2378,7 +2372,7 @@ bool UParticleModuleBeamTarget::ResolveTargetData(FParticleBeam2EmitterInstance*
 
 					if (BeamInst->TargetEmitter)
 					{
-						FVector LocalTargetPosition;
+						FVector LocalTargetPosition = FVector::ZeroVector;
 
 						int32 TargetIndex = LastSelectedParticleIndex;
 
diff --git a/Engine/Source/Runtime/Engine/Private/Particles/ParticleModules_Location.cpp b/Engine/Source/Runtime/Engine/Private/Particles/ParticleModules_Location.cpp
--- a/Engine/Source/Runtime/Engine/Private/Particles/ParticleModules_Location.cpp
+++ b/Engine/Source/Runtime/Engine/Private/Particles/ParticleModules_Location.cpp
@@ -987,13 +987,14 @@ void UParticleModuleLocationPrimitiveCylinder::SpawnEx(FParticleEmitterInstance*
 
 
 	// Determine the unit direction
-	FVector vUnitDir, vUnitDirTemp;
+	FVector vUnitDir = FVector::ZeroVector;
 
 	bool bFoundValidValue = false;
 	int32 NumberOfAttempts = 0;
 	float RadiusSquared = fStartRadius * fStartRadius;
 	while (!bFoundValidValue)
 	{
+		FVector vUnitDirTemp;
 		DetermineUnitDirection(Owner, vUnitDirTemp, InRandomStream);
 		vUnitDir[RadialIndex0]	= vUnitDirTemp[RadialIndex0];
 		vUnitDir[RadialIndex1]	= vUnitDirTemp[RadialIndex1];
@@ -2048,7 +2049,7 @@ bool UParticleModuleLocationBoneSocket::GetSocketInfoForSourceIndex(FModuleLocat
 bool UParticleModuleLocationBoneSocket::GetBoneInfoForSourceIndex(FModuleLocationBoneSocketInstancePayload* InstancePayload, USkeletalMeshComponent* SourceComponent, int32 SourceIndex, FMatrix& OutBoneMatrix, FVector& OutOffset)const
 {
 	int32 BoneIndex = INDEX_NONE;
-	FVector Offset;
+	FVector Offset = FVector::ZeroVector;
 	if (SourceType == BONESOCKETSOURCE_Sockets)
 	{
 		USkeletalMeshSocket* Socket = nullptr;
diff --git a/Engine/Source/Runtime/Engine/Private/Particles/ParticleSystemRender.cpp b/Engine/Source/Runtime/Engine/Private/Particles/ParticleSystemRender.cpp
--- a/Engine/Source/Runtime/Engine/Private/Particles/ParticleSystemRender.cpp
+++ b/Engine/Source/Runtime/Engine/Private/Particles/ParticleSystemRender.cpp
@@ -3176,8 +3176,8 @@ int32 FDynamicBeam2EmitterData::FillVertexData_NoNoise(FAsyncBufferFillData& Me)
 
 			FVector EndPoint	= Particle->Location;
 			FVector Location	= (FVector)BeamPayloadData->SourcePoint;
-			FVector Right, Up;
-			FVector WorkingUp;
+			FVector Right;
+			FVector Up			= FVector::ZeroVector;
 
 			Right = Location - EndPoint;
 			Right.Normalize();
@@ -3214,6 +3214,7 @@ int32 FDynamicBeam2EmitterData::FillVertexData_NoNoise(FAsyncBufferFillData& Me)
 			// For the direct case, this isn't a big deal, as it will not require much work per sheet.
 			for (int32 SheetIndex = 0; SheetIndex < Source.Sheets; SheetIndex++)
 			{
+				FVector WorkingUp;
 				if (SheetIndex)
 				{
 					float	Angle		= ((float)UE_PI / (float)Source.Sheets) * SheetIndex;
@@ -3367,8 +3368,7 @@ int32 FDynamicBeam2EmitterData::FillVertexData_NoNoise(FAsyncBufferFillData& Me)
 			FVector Location;
 			FVector EndPoint;
 			FVector Right;
-			FVector Up;
-			FVector WorkingUp;
+			FVector Up = FVector::ZeroVector;
 			float	fU;
 
 			float	Tex_U2 = 0.0f;
@@ -3393,6 +3393,7 @@ int32 FDynamicBeam2EmitterData::FillVertexData_NoNoise(FAsyncBufferFillData& Me)
 					}
 				}
 
+				FVector WorkingUp;
 				if (SheetIndex)
 				{
 					Angle		= ((float)UE_PI / (float)Source.Sheets) * SheetIndex;
diff --git a/Engine/Source/Runtime/Engine/Private/Particles/ParticleTrail2EmitterInstance.cpp b/Engine/Source/Runtime/Engine/Private/Particles/ParticleTrail2EmitterInstance.cpp
--- a/Engine/Source/Runtime/Engine/Private/Particles/ParticleTrail2EmitterInstance.cpp
+++ b/Engine/Source/Runtime/Engine/Private/Particles/ParticleTrail2EmitterInstance.cpp
@@ -1753,70 +1753,63 @@ bool FParticleRibbonEmitterInstance::Spawn_Source(float DeltaTime)
 
 				if ((NextSpawnedParticle != NULL) && (NextNextSpawnedParticle != NULL))
 				{
-					FVector NewTangent;
-					if (NextNextSpawnedParticle != NULL)
+					// Calculate the new tangent from the previous and next position...
+					// LastSourcePosition will be that position of the first particle that will be spawned this frame
+					FVector PositionDelta = (CurrentSourcePosition[TrailIdx] - PositionOffsetThisTick - NextNextSpawnedParticle->Location);
+					float TimeDelta = ElapsedTime - NextNextSpawnedTrailData->SpawnTime;
+
+					FVector NewTangent = FVector::ZeroVector;
+					if (TimeDelta > UE_SMALL_NUMBER)
+					{
+						NewTangent = PositionDelta / TimeDelta;
+					}
+
+					// Calculate new tangents for all the interpolated particles between NextNext and Next
+					if (NextSpawnedTrailData->SpawnedTessellationPoints > 0)
 					{
-						// Calculate the new tangent from the previous and next position...
-						// LastSourcePosition will be that position of the first particle that will be spawned this frame
-						FVector PositionDelta = (CurrentSourcePosition[TrailIdx] - PositionOffsetThisTick - NextNextSpawnedParticle->Location);
-						float TimeDelta = ElapsedTime - NextNextSpawnedTrailData->SpawnTime;
+						FBaseParticle* CurrentParticle = NULL;
+						FRibbonTypeDataPayload* CurrentTrailData = NULL;
 						
-						if (TimeDelta > UE_SMALL_NUMBER)
-						{
-							NewTangent = PositionDelta / TimeDelta;
-						}
-						else
 						{
-							NewTangent = FVector::ZeroVector;
+							int32 Prev = TRAIL_EMITTER_GET_PREV(NextNextSpawnedTrailData->Flags);
+							check(Prev != TRAIL_EMITTER_NULL_PREV);
+							DECLARE_PARTICLE_PTR(PrevParticle, ParticleData + ParticleStride * Prev);
+							CurrentParticle = PrevParticle;
+							CurrentTrailData = ((FRibbonTypeDataPayload*)((uint8*)CurrentParticle + TypeDataOffset));
 						}
-
-						// Calculate new tangents for all the interpolated particles between NextNext and Next
-						if (NextSpawnedTrailData->SpawnedTessellationPoints > 0)
+						
+						// Fix up the next ones...
+						float Diff = NextSpawnedTrailData->SpawnTime - NextNextSpawnedTrailData->SpawnTime;
+						FVector CurrUp = FVector(0.0f, 0.0f, 1.0f);
+						float InvCount = 1.0f / NextSpawnedTrailData->SpawnedTessellationPoints;
+						// Spawn the given number of particles, interpolating between the current and last position/tangent
+						//@todo. Recalculate the number of interpolated spawn particles???
+						for (int32 SpawnIdx = 0; SpawnIdx < NextSpawnedTrailData->SpawnedTessellationPoints; SpawnIdx++)
 						{
-							FBaseParticle* CurrentParticle = NULL;
-							FRibbonTypeDataPayload* CurrentTrailData = NULL;
-							
+							float TimeStep = InvCount * (SpawnIdx + 1);
+							FVector CurrPosition = FMath::CubicInterp<FVector>(
+								NextNextSpawnedParticle->Location, (FVector)NextNextSpawnedTrailData->Tangent,
+								NextSpawnedParticle->Location, NewTangent * Diff, 
+								TimeStep);
+							FVector CurrTangent = FMath::CubicInterpDerivative<FVector>(
+								NextNextSpawnedParticle->Location, (FVector)NextNextSpawnedTrailData->Tangent,
+								NextSpawnedParticle->Location, NewTangent * Diff,
+								TimeStep);
+
+							// Trail specific...
+							CurrentParticle->OldLocation = CurrentParticle->Location;
+							CurrentParticle->Location = CurrPosition;
+							CurrentTrailData->Tangent = FVector3f(CurrTangent * InvCount);
+
+							// Get the next particle in the trail (previous)
+							if ((SpawnIdx + 1) < NextSpawnedTrailData->SpawnedTessellationPoints)
 							{
-								int32 Prev = TRAIL_EMITTER_GET_PREV(NextNextSpawnedTrailData->Flags);
+								int32 Prev = TRAIL_EMITTER_GET_PREV(CurrentTrailData->Flags);
 								check(Prev != TRAIL_EMITTER_NULL_PREV);
-								DECLARE_PARTICLE_PTR(PrevParticle, ParticleData + ParticleStride * Prev);
-								CurrentParticle = PrevParticle;
+								DECLARE_PARTICLE_PTR(PrevParticleInTrail, ParticleData + ParticleStride * Prev);
+								CurrentParticle = PrevParticleInTrail;
 								CurrentTrailData = ((FRibbonTypeDataPayload*)((uint8*)CurrentParticle + TypeDataOffset));
 							}
-							
-							// Fix up the next ones...
-							float Diff = NextSpawnedTrailData->SpawnTime - NextNextSpawnedTrailData->SpawnTime;
-							FVector CurrUp = FVector(0.0f, 0.0f, 1.0f);
-							float InvCount = 1.0f / NextSpawnedTrailData->SpawnedTessellationPoints;
-							// Spawn the given number of particles, interpolating between the current and last position/tangent
-							//@todo. Recalculate the number of interpolated spawn particles???
-							for (int32 SpawnIdx = 0; SpawnIdx < NextSpawnedTrailData->SpawnedTessellationPoints; SpawnIdx++)
-							{
-								float TimeStep = InvCount * (SpawnIdx + 1);
-								FVector CurrPosition = FMath::CubicInterp<FVector>(
-									NextNextSpawnedParticle->Location, (FVector)NextNextSpawnedTrailData->Tangent,
-									NextSpawnedParticle->Location, NewTangent * Diff, 
-									TimeStep);
-								FVector CurrTangent = FMath::CubicInterpDerivative<FVector>(
-									NextNextSpawnedParticle->Location, (FVector)NextNextSpawnedTrailData->Tangent,
-									NextSpawnedParticle->Location, NewTangent * Diff,
-									TimeStep);
-
-								// Trail specific...
-								CurrentParticle->OldLocation = CurrentParticle->Location;
-								CurrentParticle->Location = CurrPosition;
-								CurrentTrailData->Tangent = FVector3f(CurrTangent * InvCount);
-
-								// Get the next particle in the trail (previous)
-								if ((SpawnIdx + 1) < NextSpawnedTrailData->SpawnedTessellationPoints)
-								{
-									int32 Prev = TRAIL_EMITTER_GET_PREV(CurrentTrailData->Flags);
-									check(Prev != TRAIL_EMITTER_NULL_PREV);
-									DECLARE_PARTICLE_PTR(PrevParticleInTrail, ParticleData + ParticleStride * Prev);
-									CurrentParticle = PrevParticleInTrail;
-									CurrentTrailData = ((FRibbonTypeDataPayload*)((uint8*)CurrentParticle + TypeDataOffset));
-								}
-							}
 						}
 					}
 
