Fix for incorrect WorldPartitionRuntimeCell SoftObjectPath in Levels contained in ContentBundles/ExternalDataLayers

diff --git a/Engine/Source/Runtime/Engine/Private/WorldPartition/Cook/WorldPartitionCook.cpp b/Engine/Source/Runtime/Engine/Private/WorldPartition/Cook/WorldPartitionCook.cpp
--- a/Engine/Source/Runtime/Engine/Private/WorldPartition/Cook/WorldPartitionCook.cpp
+++ b/Engine/Source/Runtime/Engine/Private/WorldPartition/Cook/WorldPartitionCook.cpp
@@ -152,7 +152,7 @@ bool UWorldPartition::PopulateGeneratedPackageForCook(IWorldPartitionCookPackage
 		FStringBuilderBase CookPackageObjectPath;
 		CookPackageObjectPath += ExternalDataLayerManager->GetExternalStreamingObjectPackagePath(ExternalDataLayerAsset);
 		CookPackageObjectPath += TEXT(".");
-		CookPackageObjectPath += StreamingObject->GetName();
+		CookPackageObjectPath += URuntimeHashExternalStreamingObjectBase::GetCookedExternalStreamingObjectName();
 		CookPackageObjectPath += TEXT(".");
 		CookPackageObjectPath += Cast<UObject>(InCookPackageObject)->GetName();
 		return FSoftObjectPath(*CookPackageObjectPath);
diff --git a/Engine/Source/Runtime/Engine/Private/WorldPartition/WorldPartitionLevelHelper.cpp b/Engine/Source/Runtime/Engine/Private/WorldPartition/WorldPartitionLevelHelper.cpp
--- a/Engine/Source/Runtime/Engine/Private/WorldPartition/WorldPartitionLevelHelper.cpp
+++ b/Engine/Source/Runtime/Engine/Private/WorldPartition/WorldPartitionLevelHelper.cpp
@@ -378,7 +378,7 @@ bool FWorldPartitionLevelHelper::RemapLevelCellPathInContentBundle(ULevel* Level
 {
 	FString CellPath = ContentBundleEditor->GetExternalStreamingObjectPackagePath();
 	CellPath += TEXT(".");
-	CellPath += ContentBundleEditor->GetStreamingObject()->GetName();
+	CellPath += URuntimeHashExternalStreamingObjectBase::GetCookedExternalStreamingObjectName();
 	CellPath += TEXT(".");
 	CellPath += Cell->GetName();
 	FSetWorldPartitionRuntimeCell SetWorldPartitionRuntimeCell(Level, FSoftObjectPath(CellPath));
diff --git a/Engine/Source/Runtime/Engine/Private/WorldPartition/WorldPartitionRuntimeHash.cpp b/Engine/Source/Runtime/Engine/Private/WorldPartition/WorldPartitionRuntimeHash.cpp
--- a/Engine/Source/Runtime/Engine/Private/WorldPartition/WorldPartitionRuntimeHash.cpp
+++ b/Engine/Source/Runtime/Engine/Private/WorldPartition/WorldPartitionRuntimeHash.cpp
@@ -153,7 +153,7 @@ bool URuntimeHashExternalStreamingObjectBase::OnPopulateGeneratedPackageForCook(
 	if (PrepareForCook(InCookContext))
 	{
 		// We provide a new name for the URuntimeHashExternalStreamingObjectBase in the package so that we have a stable name (for cook determinism)
-		return Rename(TEXT("RuntimeHashExternalStreamingObjectBase"), InGeneratedPackage, REN_DontCreateRedirectors);
+		return Rename(URuntimeHashExternalStreamingObjectBase::GetCookedExternalStreamingObjectName(), InGeneratedPackage, REN_DontCreateRedirectors);
 	}
 	return false;
 }
diff --git a/Engine/Source/Runtime/Engine/Public/WorldPartition/WorldPartitionRuntimeHash.h b/Engine/Source/Runtime/Engine/Public/WorldPartition/WorldPartitionRuntimeHash.h
--- a/Engine/Source/Runtime/Engine/Public/WorldPartition/WorldPartitionRuntimeHash.h
+++ b/Engine/Source/Runtime/Engine/Public/WorldPartition/WorldPartitionRuntimeHash.h
@@ -92,6 +92,8 @@ public:
 	ENGINE_API virtual bool OnPopulateGeneratedPackageForCook(const IWorldPartitionCookPackageContext& InCookContext, UPackage* InPackage, TArray<UPackage*>& OutModifiedPackages) override;
 	//~End IWorldPartitionCookPackageObject interface
 
+	ENGINE_API const static TCHAR* GetCookedExternalStreamingObjectName() { return TEXT("RuntimeHashExternalStreamingObjectBase"); }
+
 protected:
 	virtual void DumpStateLog(FHierarchicalLogArchive& Ar);
 	UWorldPartitionRuntimeCell* GetCellForCookPackage(const FString& InCookPackageName) const;
