Add GetMemberNameCheck for SkeletalMeshAsset property.

diff --git a/Engine/Source/Runtime/Engine/Classes/Components/SkeletalMeshComponent.h b/Engine/Source/Runtime/Engine/Classes/Components/SkeletalMeshComponent.h
--- a/Engine/Source/Runtime/Engine/Classes/Components/SkeletalMeshComponent.h
+++ b/Engine/Source/Runtime/Engine/Classes/Components/SkeletalMeshComponent.h
@@ -572,6 +572,11 @@ public:
 	// helper function to get the member name and verify it exists, without making it public
 	static ENGINE_API const FName& GetAnimationModePropertyNameChecked();
 
+	// helper function to get the member name and verify it exists, without making it public
+#if WITH_EDITORONLY_DATA
+	static ENGINE_API FName GetSkeletalMeshAssetPropertyNameChecked();
+#endif // WITH_EDITORONLY_DATA
+
 private:
 	/** Teleport type to use on the next update */
 	ETeleportType PendingTeleportType;
diff --git a/Engine/Source/Runtime/Engine/Private/Components/SkeletalMeshComponent.cpp b/Engine/Source/Runtime/Engine/Private/Components/SkeletalMeshComponent.cpp
--- a/Engine/Source/Runtime/Engine/Private/Components/SkeletalMeshComponent.cpp
+++ b/Engine/Source/Runtime/Engine/Private/Components/SkeletalMeshComponent.cpp
@@ -4771,6 +4771,14 @@ const FName& USkeletalMeshComponent::GetAnimationModePropertyNameChecked()
 	return Name;
 }
 
+#if WITH_EDITORONLY_DATA
+FName USkeletalMeshComponent::GetSkeletalMeshAssetPropertyNameChecked()
+{
+PRAGMA_DISABLE_DEPRECATION_WARNINGS
+	return GET_MEMBER_NAME_CHECKED(USkeletalMeshComponent, SkeletalMeshAsset);
+PRAGMA_ENABLE_DEPRECATION_WARNINGS
+}
+#endif // WITH_EDITORONLY_DATA
 
 TArray<FTransform> USkeletalMeshComponent::GetBoneSpaceTransforms() 
 {
