Don't include TransformVector in LWC usage stats

diff --git a/Engine/Source/Runtime/Engine/Private/Materials/HLSLMaterialTranslator.cpp b/Engine/Source/Runtime/Engine/Private/Materials/HLSLMaterialTranslator.cpp
--- a/Engine/Source/Runtime/Engine/Private/Materials/HLSLMaterialTranslator.cpp
+++ b/Engine/Source/Runtime/Engine/Private/Materials/HLSLMaterialTranslator.cpp
@@ -10586,7 +10586,7 @@ int32 FHLSLMaterialTranslator::TransformBase(EMaterialCommonBasis SourceCoordBas
 				}
 
 				// TODO: inconsistent with TransformLocal<TO>World with instancing
-				AddLWCFuncUsage(ELWCFunctionKind::MultiplyVectorMatrix);
+				if (bIsPositionTransform) { AddLWCFuncUsage(ELWCFunctionKind::MultiplyVectorMatrix); }
 				CodeStr = LWCMultiplyMatrix(TEXT("<A>"), TEXT("Get<PREV>WorldToLocal(Parameters)"), AWComponent, true);
 				CodeDerivStr = LWCMultiplyMatrix(TEXT("<A>"), TEXT("Get<PREV>WorldToLocal(Parameters)"), 0, true);
 			}
@@ -10605,14 +10605,14 @@ int32 FHLSLMaterialTranslator::TransformBase(EMaterialCommonBasis SourceCoordBas
 			}
 			else if (DestCoordBasis == MCB_MeshParticle)
 			{
-				AddLWCFuncUsage(ELWCFunctionKind::MultiplyVectorMatrix);
+				if (bIsPositionTransform) { AddLWCFuncUsage(ELWCFunctionKind::MultiplyVectorMatrix); }
 				CodeStr = LWCMultiplyMatrix(TEXT("<A>"), TEXT("GetWorldToParticle(Parameters)"), AWComponent, true);
 				CodeDerivStr = LWCMultiplyMatrix(TEXT("<A>"), TEXT("GetWorldToParticle(Parameters)"), 0, true);
 				bUsesParticleWorldToLocal = true;
 			}
 			else if (DestCoordBasis == MCB_Instance)
 			{
-				AddLWCFuncUsage(ELWCFunctionKind::MultiplyVectorMatrix);
+				if (bIsPositionTransform) { AddLWCFuncUsage(ELWCFunctionKind::MultiplyVectorMatrix); }
 				CodeStr = LWCMultiplyMatrix(TEXT("<A>"), TEXT("GetWorldToInstance(Parameters)"), AWComponent, true);
 				CodeDerivStr = LWCMultiplyMatrix(TEXT("<A>"), TEXT("GetWorldToInstance(Parameters)"), 0, true);
 				bUsesInstanceWorldToLocalPS |= ShaderFrequency == SF_Pixel;
@@ -10648,7 +10648,7 @@ int32 FHLSLMaterialTranslator::TransformBase(EMaterialCommonBasis SourceCoordBas
 		{
 			if (DestCoordBasis == MCB_World)
 			{
-				AddLWCFuncUsage(ELWCFunctionKind::MultiplyVectorMatrix);
+				if (bIsPositionTransform) { AddLWCFuncUsage(ELWCFunctionKind::MultiplyVectorMatrix); }
 				CodeStr = LWCMultiplyMatrix(TEXT("<A>"), TEXT("GetParticleToWorld(Parameters)"), AWComponent, false);
 				CodeDerivStr = LWCMultiplyMatrix(TEXT("<A>"), TEXT("GetParticleToWorld(Parameters)"), 0, false);
 				bUsesParticleLocalToWorld = true;
@@ -10666,7 +10666,7 @@ int32 FHLSLMaterialTranslator::TransformBase(EMaterialCommonBasis SourceCoordBas
 		{
 			if (DestCoordBasis == MCB_World)
 			{
-				AddLWCFuncUsage(ELWCFunctionKind::MultiplyVectorMatrix);
+				if (bIsPositionTransform) { AddLWCFuncUsage(ELWCFunctionKind::MultiplyVectorMatrix); }
 				CodeStr = LWCMultiplyMatrix(TEXT("<A>"), TEXT("Get<PREV>InstanceToWorld(Parameters)"), AWComponent, false);
 				CodeDerivStr = LWCMultiplyMatrix(TEXT("<A>"), TEXT("Get<PREV>InstanceToWorld(Parameters)"), 0, false);
 				bUsesInstanceLocalToWorldPS |= ShaderFrequency == SF_Pixel;
