Fix for invalid ddx/ddy calls in generated material code when using the DDX and DDY nodes in a material that can be executed in a compute shader context other than Nanite shading (such as HeterogeneousVolumes or Volumetric Clouds).
Fixed by using existing DDX() and DDY() wrapper functions which are already set up to return 0 in all shader types except pixel shaders and nanite compute shading compute shaders.

#rnx

diff --git a/Engine/Shaders/Private/MaterialTemplate.ush b/Engine/Shaders/Private/MaterialTemplate.ush
--- a/Engine/Shaders/Private/MaterialTemplate.ush
+++ b/Engine/Shaders/Private/MaterialTemplate.ush
@@ -42,7 +42,7 @@
 
 // Update this GUID to force all material shaders to recompile
 // Merge conflicts on this line should be resolved by generating a new GUID
-#pragma message("UESHADERMETADATA_VERSION 838540DA-BB92-4973-80CB-D9E993EBA40D")
+#pragma message("UESHADERMETADATA_VERSION 5AE81BAA-6F06-4443-AEC8-98E87D189D36")
 
 #if USES_SPEEDTREE
 	#include "/Engine/Private/SpeedTreeCommon.ush"
diff --git a/Engine/Source/Runtime/Engine/Private/Materials/HLSLMaterialTranslator.cpp b/Engine/Source/Runtime/Engine/Private/Materials/HLSLMaterialTranslator.cpp
--- a/Engine/Source/Runtime/Engine/Private/Materials/HLSLMaterialTranslator.cpp
+++ b/Engine/Source/Runtime/Engine/Private/Materials/HLSLMaterialTranslator.cpp
@@ -11030,7 +11030,7 @@ int32 FHLSLMaterialTranslator::PixelNormalWS()
 	return AddInlinedCodeChunk(MCT_Float3,TEXT("Parameters.WorldNormal"));	
 }
 
-int32 FHLSLMaterialTranslator::Derivative(int32 A, const TCHAR* Component)
+int32 FHLSLMaterialTranslator::Derivative(int32 A, EDervativeComponent Component)
 {
 	if (A == INDEX_NONE)
 	{
@@ -11050,10 +11050,29 @@ int32 FHLSLMaterialTranslator::Derivative(int32 A, const TCHAR* Component)
 
 	const FDerivInfo ADerivInfo = GetDerivInfo(A);
 	const EMaterialValueType ResultType = MakeNonLWCType(ADerivInfo.Type);
-	const TCHAR* FunctionName = IsLWCType(ADerivInfo.Type) ? TEXT("WSDd") : TEXT("dd");
-	const TCHAR* FunctionNamePostfix = IsLWCType(ADerivInfo.Type) ? TEXT("Demote") : TEXT("");
+	const bool bIsLWCType = IsLWCType(ADerivInfo.Type);
+	const TCHAR* FunctionName = bIsLWCType ? TEXT("WSDd") : TEXT("DD");
+	const TCHAR* FunctionNamePostfix = bIsLWCType ? TEXT("Demote") : TEXT("");
+	
+	// For non-LWC types, we use the DDX() and DDY() functions defined in Common.ush, which correctly deal with shader types that do not support hardware derivatives.
+	const TCHAR* ComponentLowerCase = TEXT("");
+	const TCHAR* ComponentUpperCase = TEXT("");
+	switch (Component)
+	{
+	case EDervativeComponent::X:
+		ComponentLowerCase = TEXT("x");
+		ComponentUpperCase = TEXT("X");
+		break;
+	case EDervativeComponent::Y:
+		ComponentLowerCase = TEXT("y");
+		ComponentUpperCase = TEXT("Y");
+		break;
+	default:
+		checkNoEntry();
+		break;
+	}
     
-	if (IsLWCType(ADerivInfo.Type))
+	if (bIsLWCType)
 	{
 		AddLWCFuncUsage(ELWCFunctionKind::Other);
 	}
@@ -11070,16 +11089,16 @@ int32 FHLSLMaterialTranslator::Derivative(int32 A, const TCHAR* Component)
     	case MCT_Float4: CastType = TEXT("(float4)"); break;
     }
             
-	const FString FiniteCode = FString::Printf(TEXT("%s%s%s(%s%s)"), FunctionName, Component, FunctionNamePostfix, CastType, *GetParameterCode(A));
+	const FString FiniteCode = FString::Printf(TEXT("%s%s%s(%s%s)"), FunctionName, bIsLWCType ? ComponentLowerCase : ComponentUpperCase, FunctionNamePostfix, CastType, *GetParameterCode(A));
 
 	if (IsAnalyticDerivEnabled() && ADerivInfo.DerivativeStatus == EDerivativeStatus::Valid)
 	{
 		FString ADeriv = GetParameterCodeDeriv(A, CompiledPDV_Analytic);
 
 		FString AnalyticCode = DerivativeAutogen.ConstructDeriv(
-			FString::Printf(TEXT("%s.Dd%s"), *ADeriv, Component),
-			FString::Printf(TEXT("dd%s(%s.Ddx)"), Component, *ADeriv),
-			FString::Printf(TEXT("dd%s(%s.Ddy)"), Component, *ADeriv),
+			FString::Printf(TEXT("%s.Dd%s"), *ADeriv, ComponentLowerCase),
+			FString::Printf(TEXT("DD%s(%s.Ddx)"), ComponentUpperCase, *ADeriv),
+			FString::Printf(TEXT("DD%s(%s.Ddy)"), ComponentUpperCase, *ADeriv),
 			MakeNonLWCType(ADerivInfo.DerivType)
 		);
 
@@ -11093,12 +11112,12 @@ int32 FHLSLMaterialTranslator::Derivative(int32 A, const TCHAR* Component)
 
 int32 FHLSLMaterialTranslator::DDX( int32 A )
 {
-	return Derivative(A, TEXT("x"));
+	return Derivative(A, EDervativeComponent::X);
 }
 
 int32 FHLSLMaterialTranslator::DDY( int32 A )
 {
-	return Derivative(A, TEXT("y"));
+	return Derivative(A, EDervativeComponent::Y);
 }
 
 int32 FHLSLMaterialTranslator::AntialiasedTextureMask(int32 Tex, int32 UV, float Threshold, uint8 Channel)
diff --git a/Engine/Source/Runtime/Engine/Private/Materials/HLSLMaterialTranslator.h b/Engine/Source/Runtime/Engine/Private/Materials/HLSLMaterialTranslator.h
--- a/Engine/Source/Runtime/Engine/Private/Materials/HLSLMaterialTranslator.h
+++ b/Engine/Source/Runtime/Engine/Private/Materials/HLSLMaterialTranslator.h
@@ -1161,7 +1161,12 @@ protected:
 	virtual int32 DDX(int32 A) override;
 	virtual int32 DDY(int32 A) override;
 
-	int32 Derivative(int32 A, const TCHAR* Component);
+	enum class EDervativeComponent
+	{
+		X, 
+		Y,
+	};
+	int32 Derivative(int32 A, EDervativeComponent Component);
 
 	virtual int32 AntialiasedTextureMask(int32 Tex, int32 UV, float Threshold, uint8 Channel) override;
 	virtual int32 DepthOfFieldFunction(int32 Depth, int32 FunctionValueIndex) override;
