AggregateGeom (shape arrays);

Adding the metadata specifier TitleProperty so we know what the name of the shapes are in the shape arrays.

-- Code --
AggregateGeom:
FKAggregateGeom - Added TitleProperty = "Name" to the shape array properties.

review [REDACTED_HANDLE], [REDACTED_HANDLE], [REDACTED_HANDLE]

diff --git a/Engine/Source/Runtime/Engine/Classes/PhysicsEngine/AggregateGeom.h b/Engine/Source/Runtime/Engine/Classes/PhysicsEngine/AggregateGeom.h
--- a/Engine/Source/Runtime/Engine/Classes/PhysicsEngine/AggregateGeom.h
+++ b/Engine/Source/Runtime/Engine/Classes/PhysicsEngine/AggregateGeom.h
@@ -22,25 +22,25 @@ struct FKAggregateGeom
 {
 	GENERATED_USTRUCT_BODY()
 
-	UPROPERTY(EditAnywhere, editfixedsize, Category = "Aggregate Geometry", meta = (DisplayName = "Spheres"))
+	UPROPERTY(EditAnywhere, editfixedsize, Category = "Aggregate Geometry", meta = (DisplayName = "Spheres", TitleProperty = "Name"))
 	TArray<FKSphereElem> SphereElems;
 
-	UPROPERTY(EditAnywhere, editfixedsize, Category = "Aggregate Geometry", meta = (DisplayName = "Boxes"))
+	UPROPERTY(EditAnywhere, editfixedsize, Category = "Aggregate Geometry", meta = (DisplayName = "Boxes", TitleProperty = "Name"))
 	TArray<FKBoxElem> BoxElems;
 
-	UPROPERTY(EditAnywhere, editfixedsize, Category = "Aggregate Geometry", meta = (DisplayName = "Capsules"))
+	UPROPERTY(EditAnywhere, editfixedsize, Category = "Aggregate Geometry", meta = (DisplayName = "Capsules", TitleProperty = "Name"))
 	TArray<FKSphylElem> SphylElems;
 
-	UPROPERTY(EditAnywhere, editfixedsize, Category = "Aggregate Geometry", meta = (DisplayName = "Convex Elements"))
+	UPROPERTY(EditAnywhere, editfixedsize, Category = "Aggregate Geometry", meta = (DisplayName = "Convex Elements", TitleProperty = "Name"))
 	TArray<FKConvexElem> ConvexElems;
 
-	UPROPERTY(EditAnywhere, editfixedsize, Category = "Aggregate Geometry", meta = (DisplayName = "Tapered Capsules"))
+	UPROPERTY(EditAnywhere, editfixedsize, Category = "Aggregate Geometry", meta = (DisplayName = "Tapered Capsules", TitleProperty = "Name"))
 	TArray<FKTaperedCapsuleElem> TaperedCapsuleElems;
 
-	UPROPERTY(EditAnywhere, editfixedsize, Category = "Aggregate Geometry", meta = (DisplayName = "Level Sets"))
+	UPROPERTY(EditAnywhere, editfixedsize, Category = "Aggregate Geometry", meta = (DisplayName = "Level Sets", TitleProperty = "Name"))
 	TArray<FKLevelSetElem> LevelSetElems;
 
-	UPROPERTY(EditAnywhere, editfixedsize, Category = "Aggregate Geometry", meta = (DisplayName = "(Experimental) Skinned Level Sets"), Experimental)
+	UPROPERTY(EditAnywhere, editfixedsize, Category = "Aggregate Geometry", meta = (DisplayName = "(Experimental) Skinned Level Sets", TitleProperty = "Name"), Experimental)
 	TArray<FKSkinnedLevelSetElem> SkinnedLevelSetElems;
 
 	FKAggregateGeom()
