- Better reported reason when moving an actor to an EDL fails because of asset reference restrictions.
- FFindReferencedAssets::BuildAssetList: Added option to only return direct references.
#pregflight 661ecbae473aadcea1ab7bc7

diff --git a/Engine/Source/Editor/UnrealEd/Private/ReferencedAssetsUtils.cpp b/Engine/Source/Editor/UnrealEd/Private/ReferencedAssetsUtils.cpp
--- a/Engine/Source/Editor/UnrealEd/Private/ReferencedAssetsUtils.cpp
+++ b/Engine/Source/Editor/UnrealEd/Private/ReferencedAssetsUtils.cpp
@@ -138,7 +138,7 @@ bool FFindReferencedAssets::ShouldSearchForAssets( const UObject* Object, const
 /**
  * Returns a list of all assets referenced by the specified UObject.
  */
-void FFindReferencedAssets::BuildAssetList(UObject *Object, const TArray<UClass*>& IgnoreClasses, const TArray<UPackage*>& IgnorePackages, TSet<UObject*>& ReferencedAssets, bool bIncludeDefaultRefs)
+void FFindReferencedAssets::BuildAssetList(UObject *Object, const TArray<UClass*>& IgnoreClasses, const TArray<UPackage*>& IgnorePackages, TSet<UObject*>& ReferencedAssets, bool bIncludeDefaultRefs, bool bOnlyDirectReferences)
 {
 	TArray<FReferencedAssets> LocalReferencers;
 
@@ -159,8 +159,8 @@ void FFindReferencedAssets::BuildAssetList(UObject *Object, const TArray<UClass*
 	}
 
 	// Add to the list of referenced assets.
-	FFindAssetsArchive( Object, LocalReferencers.Last().AssetList, NULL, /*MaxRecursion=*/0, /*bIncludeClasses=*/true, bIncludeDefaultRefs );
-
+	const int32 MaxRecursionDepth = bOnlyDirectReferences ? 1 : 0;
+	FFindAssetsArchive( Object, LocalReferencers.Last().AssetList, NULL, MaxRecursionDepth, /*bIncludeClasses=*/true, bIncludeDefaultRefs );
 	ReferencedAssets = LocalReferencers.Last().AssetList;
 }
 
diff --git a/Engine/Source/Editor/UnrealEd/Public/ReferencedAssetsUtils.h b/Engine/Source/Editor/UnrealEd/Public/ReferencedAssetsUtils.h
--- a/Engine/Source/Editor/UnrealEd/Public/ReferencedAssetsUtils.h
+++ b/Engine/Source/Editor/UnrealEd/Public/ReferencedAssetsUtils.h
@@ -157,9 +157,16 @@ public:
 	UNREALED_API static bool ShouldSearchForAssets( const UObject* Object, const TArray<UClass*>& ClassesToIgnore, const TArray<UPackage*>& PackagesToIgnore, bool bIncludeDefaults=false );
 
 	/**
-	 * Returns a list of all assets referenced by the specified UObject.
-	 */
-	UNREALED_API static void BuildAssetList(UObject *Object, const TArray<UClass*>& IgnoreClasses, const TArray<UPackage*>& IgnorePackages, TSet<UObject*>& ReferencedAssets, bool bIncludeDefaultRefs=false);
+	 * Returns a list of assets referenced by the specified UObject.
+	 * @param Object the object used to build the references
+	 * @param IgnoreClasses the list of classes to skip
+	 * @param IgnorePackages the list of packages to skip
+	 * @param bIncludeDefaultRefs specify true to include content referenced through defaults
+	 * @param bOnlyDirectReferences specify true to only include direct references
+	 * 
+	 * @return list of assets referenced by the specified UObject.
+	 */
+	UNREALED_API static void BuildAssetList(UObject *Object, const TArray<UClass*>& IgnoreClasses, const TArray<UPackage*>& IgnorePackages, TSet<UObject*>& ReferencedAssets, bool bIncludeDefaultRefs=false, bool bOnlyDirectReferences=false);
 
 protected:
 	/**
diff --git a/Engine/Source/Runtime/Engine/Private/WorldPartition/DataLayer/ExternalDataLayerHelper.cpp b/Engine/Source/Runtime/Engine/Private/WorldPartition/DataLayer/ExternalDataLayerHelper.cpp
--- a/Engine/Source/Runtime/Engine/Private/WorldPartition/DataLayer/ExternalDataLayerHelper.cpp
+++ b/Engine/Source/Runtime/Engine/Private/WorldPartition/DataLayer/ExternalDataLayerHelper.cpp
@@ -99,7 +99,7 @@ void FExternalDataLayerHelper::GetExternalDataLayerUIDs(const FAssetData& Asset,
 
 namespace UE::Private::ExternalDataLayerHelper
 {
-	static bool ValidateAssetUsingAssetReferenceRestrictions(const UObject* InAsset, const TSet<UObject*>& InReferencedAssets, TSet<FString>& OutInvalidReferenceReasons)
+	static bool ValidateAssetUsingAssetReferenceRestrictions(const UObject* InAsset, const TSet<UObject*>& InReferencedAssets, TMap<FString, TArray<FString>>& OutInvalidReferenceReasons)
 	{
 		if (!InAsset)
 		{
@@ -118,9 +118,11 @@ namespace UE::Private::ExternalDataLayerHelper
 				for (UObject* ReferencedAsset : InReferencedAssets)
 				{
 					FText FailureReason;
-					if (!AssetReferenceFilter->PassesFilter(FAssetData(ReferencedAsset), &FailureReason))
+					FAssetData ReferencedAssetData(ReferencedAsset);
+					if (!AssetReferenceFilter->PassesFilter(ReferencedAssetData, &FailureReason))
 					{
-						OutInvalidReferenceReasons.Add(FailureReason.ToString());
+						const FString MountPoint = FPackageName::GetPackageMountPoint(ReferencedAssetData.PackagePath.ToString()).ToString();
+						OutInvalidReferenceReasons.FindOrAdd(FailureReason.ToString()).Add(FString::Printf(TEXT("%s (Mount Point: %s)"), *ReferencedAssetData.GetObjectPathString(), *MountPoint));
 						++ErrorCount;
 					}
 				}
@@ -193,7 +195,9 @@ bool FExternalDataLayerHelper::CanMoveActorsToExternalDataLayer(const TArray<AAc
 
 		// Gather actor asset references
 		TSet<UObject*> ActorReferencedAssets;
-		FFindReferencedAssets::BuildAssetList(InActor, IgnoreClasses, IgnorePackages, ActorReferencedAssets, true);
+		const bool bIncludeDefaultRefs = false;
+		const bool bOnlyDirectReferences = true;
+		FFindReferencedAssets::BuildAssetList(InActor, IgnoreClasses, IgnorePackages, ActorReferencedAssets, bIncludeDefaultRefs, bOnlyDirectReferences);
 		TArray<UObject*> ReferencedContent;
 		InActor->GetReferencedContentObjects(ReferencedContent);
 		// Remove itself and its data layer assets from the list
@@ -208,19 +212,30 @@ bool FExternalDataLayerHelper::CanMoveActorsToExternalDataLayer(const TArray<AAc
 		}
 
 		// Validate if there are restrictions between the world or the new data layer asset and the actor asset references
-		TSet<FString> InvalidReferenceReasons;
+		TMap<FString, TArray<FString>> InvalidReferences;
 		const UObject* Referencer = NewExternalDataLayerAsset ? (UObject*)NewExternalDataLayerAsset : (UObject*)InActor->GetLevel();
-		if (!UE::Private::ExternalDataLayerHelper::ValidateAssetUsingAssetReferenceRestrictions(Referencer, ActorReferencedAssets, InvalidReferenceReasons))
+		if (!UE::Private::ExternalDataLayerHelper::ValidateAssetUsingAssetReferenceRestrictions(Referencer, ActorReferencedAssets, InvalidReferences))
 		{
-			const FString JoinedReasons = FString::Join(InvalidReferenceReasons, TEXT("\n"));
+			FStringBuilderBase StringBuilder;
+			for (auto& [Reason, InvalidReferenceAssets] : InvalidReferences)
+			{
+				const FString JoinedInvalidReferences = FString::Join(InvalidReferenceAssets, TEXT(", "));
+				StringBuilder.Appendf(TEXT(" - Reason: %s\n - Invalid References: \n"), *Reason);
+				for (const FString& InvalidReference : InvalidReferenceAssets)
+				{
+					StringBuilder.Appendf(TEXT("   - %s\n"), *InvalidReference);
+				}
+			}
+			const FString JoinedReasons = StringBuilder.ToString();
+			
 			if (NewExternalDataLayerAsset)
 			{
-				OutFailureReason = FText::Format(LOCTEXT("CantMoveActorToEDLReferenceRestrictions", "Can't move Actor {0} to External Data Layer {1}. Reason: {2}."), FText::FromString(InActor->GetName()), FText::FromString(NewExternalDataLayerAsset->GetName()), FText::FromString(JoinedReasons));
+				OutFailureReason = FText::Format(LOCTEXT("CantMoveActorToEDLReferenceRestrictions", "Can't move Actor {0} to External Data Layer {1}.\n{2}"), FText::FromString(InActor->GetName()), FText::FromString(NewExternalDataLayerAsset->GetName()), FText::FromString(JoinedReasons));
 			}
 			else
 			{
 				check(OldExternalDataLayerAsset);
-				OutFailureReason = FText::Format(LOCTEXT("CantRemoveEDLFromActorReferenceRestrictions", "Can't remove External Data Layer {0} from Actor {1}. Reason: {2}."), FText::FromString(OldExternalDataLayerAsset->GetName()), FText::FromString(InActor->GetName()), FText::FromString(JoinedReasons));
+				OutFailureReason = FText::Format(LOCTEXT("CantRemoveEDLFromActorReferenceRestrictions", "Can't remove External Data Layer {0} from Actor {1}.\n{2}"), FText::FromString(OldExternalDataLayerAsset->GetName()), FText::FromString(InActor->GetName()), FText::FromString(JoinedReasons));
 			}
 			return false;
 		}
