Null check the world when disabling the input component on an actor.

#rnx

diff --git a/Engine/Source/Runtime/Engine/Private/Actor.cpp b/Engine/Source/Runtime/Engine/Private/Actor.cpp
--- a/Engine/Source/Runtime/Engine/Private/Actor.cpp
+++ b/Engine/Source/Runtime/Engine/Private/Actor.cpp
@@ -4356,9 +4356,9 @@ void AActor::DisableInput(APlayerController* PlayerController)
 		{
 			PlayerController->PopInputComponent(InputComponent);
 		}
-		else
+		else if (UWorld* World = GetWorld())
 		{
-			for (FConstPlayerControllerIterator PCIt = GetWorld()->GetPlayerControllerIterator(); PCIt; ++PCIt)
+			for (FConstPlayerControllerIterator PCIt = World->GetPlayerControllerIterator(); PCIt; ++PCIt)
 			{
 				if (APlayerController* PC = PCIt->Get())
 				{
