Cluster union : Deprecate SetIsAnchored to avoid it to desync from the rigid state

diff --git a/Engine/Source/Runtime/Engine/Classes/PhysicsEngine/ClusterUnionComponent.h b/Engine/Source/Runtime/Engine/Classes/PhysicsEngine/ClusterUnionComponent.h
--- a/Engine/Source/Runtime/Engine/Classes/PhysicsEngine/ClusterUnionComponent.h
+++ b/Engine/Source/Runtime/Engine/Classes/PhysicsEngine/ClusterUnionComponent.h
@@ -231,7 +231,8 @@ public:
 	UFUNCTION(BlueprintCallable, Category = "Cluster Union")
 	ENGINE_API TArray<AActor*> GetActors();
 
-	UFUNCTION(BlueprintCallable, Category = "Cluster Union")
+	UE_DEPRECATED(5.5, "SetIsAnchored is deprecated, Use SetSimulatePhysics instead")
+	UFUNCTION(BlueprintCallable, Category = "Cluster Union", meta = (DeprecatedFunction, DeprecationMessage = "SetIsAnchored is deprecated, Use SetSimulatePhysics instead"))
 	ENGINE_API virtual void SetIsAnchored(bool bIsAnchored);
 
 	ENGINE_API bool IsAnchored() const;
diff --git a/Engine/Source/Runtime/Engine/Private/PhysicsEngine/ClusterUnionComponent.cpp b/Engine/Source/Runtime/Engine/Private/PhysicsEngine/ClusterUnionComponent.cpp
--- a/Engine/Source/Runtime/Engine/Private/PhysicsEngine/ClusterUnionComponent.cpp
+++ b/Engine/Source/Runtime/Engine/Private/PhysicsEngine/ClusterUnionComponent.cpp
@@ -621,6 +621,8 @@ TArray<AActor*> UClusterUnionComponent::GetActors()
 
 void UClusterUnionComponent::SetIsAnchored(bool bIsAnchored)
 {
+	// deprecated - should not be called anymore
+	ensure(false);
 	if (!PhysicsProxy)
 	{
 		return;
@@ -1332,7 +1334,6 @@ void UClusterUnionComponent::OnRep_RigidState()
 		return;
 	}
 
-	PhysicsProxy->SetIsAnchored_External(ReplicatedRigidState.bIsAnchored);
 	SetRigidState(static_cast<Chaos::EObjectStateType>(ReplicatedRigidState.ObjectState));
 }
 
@@ -1384,6 +1385,8 @@ void UClusterUnionComponent::SetRigidState(Chaos::EObjectStateType ObjectState)
 	if (PhysicsProxy)
 	{
 		PhysicsProxy->SetObjectState_External(ObjectState);
+		const bool bIsDynamic = (ObjectState == Chaos::EObjectStateType::Dynamic || ObjectState == Chaos::EObjectStateType::Sleeping);
+		PhysicsProxy->SetIsAnchored_External(!bIsDynamic);
 	}
 }
 
