ContentBundles: logging: Remove the UE_SET_LOG_VERBOSITY of LogContentBundle to verbose. It causes 413k out of 518k lines of log in the cook. We no longer need the extra information it provides on by default, so removing it to reduce the spam.
#rnx

diff --git a/Engine/Source/Runtime/Engine/Private/WorldPartition/ContentBundle/ContentBundleEngineSubsystem.cpp b/Engine/Source/Runtime/Engine/Private/WorldPartition/ContentBundle/ContentBundleEngineSubsystem.cpp
--- a/Engine/Source/Runtime/Engine/Private/WorldPartition/ContentBundle/ContentBundleEngineSubsystem.cpp
+++ b/Engine/Source/Runtime/Engine/Private/WorldPartition/ContentBundle/ContentBundleEngineSubsystem.cpp
@@ -38,11 +38,6 @@ void UContentBundleEngineSubsystem::Deinitialize()
 
 TSharedPtr<FContentBundleClient> UContentBundleEngineSubsystem::RegisterContentBundle(const UContentBundleDescriptor* Descriptor, const FString& ClientDisplayName)
 {
-	if (ContentBundleClients.IsEmpty())
-	{
-		UE_SET_LOG_VERBOSITY(LogContentBundle, Verbose);
-	}
-
 	if (Descriptor != nullptr)
 	{
 		TSharedPtr<FContentBundleClient>* RegisteredClient = FindRegisteredClient(Descriptor);
