[shaders] fix logic for enabling per-shader ddc in cooks to not affect other non-cook commandlets

diff --git a/Engine/Source/Runtime/Engine/Private/ShaderCompiler/ShaderCompiler.cpp b/Engine/Source/Runtime/Engine/Private/ShaderCompiler/ShaderCompiler.cpp
--- a/Engine/Source/Runtime/Engine/Private/ShaderCompiler/ShaderCompiler.cpp
+++ b/Engine/Source/Runtime/Engine/Private/ShaderCompiler/ShaderCompiler.cpp
@@ -250,14 +250,17 @@ static bool IsShaderJobCacheDDCRemotePolicyEnabled()
 bool IsShaderJobCacheDDCEnabled()
 {
 #if WITH_EDITOR
-	static const bool bForceAllowShaderCompilerJobCache = FParse::Param(FCommandLine::Get(), TEXT("forceAllowShaderCompilerJobCache")) ||
-		CVarShaderCompilerPerShaderDDCCook.GetValueOnAnyThread();
+	static const bool bForceAllowShaderCompilerJobCache = FParse::Param(FCommandLine::Get(), TEXT("forceAllowShaderCompilerJobCache"));
+	static const bool bEnablePerShaderDDCCook = IsRunningCookCommandlet() && CVarShaderCompilerPerShaderDDCCook.GetValueOnAnyThread();
+	static const bool bIsNonCookCommandlet = IsRunningCommandlet() && !IsRunningCookCommandlet();
 #else
 	const bool bForceAllowShaderCompilerJobCache = false;
+	const bool bEnablePerShaderDDCCook = false;
+	const bool bIsNonCookCommandlet = false;
 #endif
 
-	// For now we only support the editor and not commandlets like the cooker.
-	if ((GIsEditor || IsRunningGame()) && (!IsRunningCommandlet() || bForceAllowShaderCompilerJobCache))
+	// Enable remote per-shader DDC for editor, game, cooks (if cvar is set), and for other commandlets only if the force flag is set on the cmdline.
+	if ((GIsEditor || IsRunningGame() || bEnablePerShaderDDCCook) && (!bIsNonCookCommandlet || bForceAllowShaderCompilerJobCache))
 	{
 		// job cache itself must be enabled first
 		return GShaderCompilerJobCache && CVarJobCacheDDC.GetValueOnAnyThread();
