[SplineMeshComponent] fixed linear extrapolation by using the positions at MinT and MaxT instead of the Spline StartPos and EndPos respectively

diff --git a/Engine/Source/Runtime/Engine/Private/Components/SplineMeshComponent.cpp b/Engine/Source/Runtime/Engine/Private/Components/SplineMeshComponent.cpp
--- a/Engine/Source/Runtime/Engine/Private/Components/SplineMeshComponent.cpp
+++ b/Engine/Source/Runtime/Engine/Private/Components/SplineMeshComponent.cpp
@@ -1106,13 +1106,13 @@ FTransform USplineMeshComponent::CalcSliceTransformAtSplineOffset(const float Al
 	if (Alpha < MinT)
 	{
 		const FVector3f StartTangent(SplineEvalTangent(SplineParams, MinT));
-		SplinePos = FVector3f(SplineParams.StartPos) + (StartTangent * (Alpha - MinT));
+		SplinePos = SplineEvalPos(SplineParams, MinT) + (StartTangent * (Alpha - MinT));
 		SplineDir = StartTangent.GetSafeNormal();
 	}
 	else if (Alpha > MaxT)
 	{
 		const FVector3f EndTangent(SplineEvalTangent(SplineParams, MaxT));
-		SplinePos = FVector3f(SplineParams.EndPos) + (EndTangent * (Alpha - MaxT));
+		SplinePos = SplineEvalPos(SplineParams, MaxT) + (EndTangent * (Alpha - MaxT));
 		SplineDir = EndTangent.GetSafeNormal();
 	}
 	else
