Level Instance: Property Override prep work for Scene Outliner Column
- Add new ELevelInstanceFlag value: HasEditablePropertyOverrides which gives info about Actor being overriden by an Editable Level Instances (will be used for different icon coloring)
- Property Override editing now tracks the override flags to give clear feedback with column about which actor is currently overriden (tracked through OnObjectPropertyChanged)
- Fix bug where saving an empty override wouldn't reset it on the Level Instance (setting it to null instead of keeping an empty asset)

#rnx

diff --git a/Engine/Source/Runtime/Engine/Classes/GameFramework/Actor.h b/Engine/Source/Runtime/Engine/Classes/GameFramework/Actor.h
--- a/Engine/Source/Runtime/Engine/Classes/GameFramework/Actor.h
+++ b/Engine/Source/Runtime/Engine/Classes/GameFramework/Actor.h
@@ -106,8 +106,12 @@ UENUM()
 enum class ELevelInstanceFlags : uint8
 {
 	None = 0,
+	/** Actor is inside an Level Instance editing hierarchy (used to show post process effect) */
 	IsInEditHierarchy = 1,
+	/** Actor has property overrides applied */
 	HasPropertyOverrides = 2,
+	/** Actor has property overrides from a top level editable level instance */
+	HasEditablePropertyOverrides = 4 
 };
 
 ENUM_CLASS_FLAGS(ELevelInstanceFlags)
@@ -302,15 +306,16 @@ private:
 
 #if WITH_EDITORONLY_DATA
 	/** Whether this actor belongs to a level instance or not and what type of level instance. */
-	UPROPERTY(Transient)
+	UPROPERTY(Transient, NonTransactional)
 	ELevelInstanceType LevelInstanceType;
 
 	/** Flags related to level instances for this actor. */
-	UPROPERTY(Transient)
+	UPROPERTY(Transient, NonTransactional)
 	ELevelInstanceFlags LevelInstanceFlags;
 
 	friend struct FSetActorIsInLevelInstance;
 	friend struct FAddActorLevelInstanceFlags;
+	friend struct FRemoveActorLevelInstanceFlags;
 
 public:
 	UPROPERTY(EditAnywhere, AdvancedDisplay, Category = LevelInstance, meta = (Tooltip = "If checked, this Actor will only get loaded in a main world (persistent level), it will not be loaded through Level Instances."))
@@ -369,12 +374,18 @@ public:
 		return EnumHasAnyFlags(LevelInstanceFlags, ELevelInstanceFlags::IsInEditHierarchy);
 	}
 
-	/** iF true, this actor or one of its components has property overrides applied. */
+	/** If true, this actor or one of its components has property overrides applied. */
 	bool HasLevelInstancePropertyOverrides() const
 	{
 		return EnumHasAnyFlags(LevelInstanceFlags, ELevelInstanceFlags::HasPropertyOverrides);
 	}
 
+	/** If true, this actor or one of its components has property overrides applied from an editable level instance (level instance is in a non readonly level) */
+	bool HasEditableLevelLevelInstancePropertyOverrides() const
+	{
+		return EnumHasAnyFlags(LevelInstanceFlags, ELevelInstanceFlags::HasEditablePropertyOverrides);
+	}
+
 	/** If true, the actor belongs to a level instance. */
 	bool IsInLevelInstance() const
 	{
@@ -4617,6 +4628,19 @@ private:
 	friend class ULevelStreamingLevelInstanceEditor;
 	friend class ULevelStreamingLevelInstanceEditorPropertyOverride;
 };
+
+struct FRemoveActorLevelInstanceFlags
+{
+private:
+	FRemoveActorLevelInstanceFlags(AActor* InActor, ELevelInstanceFlags InFlagsToRemove)
+	{
+		EnumRemoveFlags(InActor->LevelInstanceFlags, InFlagsToRemove);
+	}
+
+	friend class ULevelStreamingLevelInstance;
+	friend class ULevelStreamingLevelInstanceEditor;
+	friend class ULevelStreamingLevelInstanceEditorPropertyOverride;
+};
 #endif
 
 /** Helper function for executing tick functions based on the normal conditions previous found in UActorComponent::ConditionalTick */
diff --git a/Engine/Source/Runtime/Engine/Private/LevelInstance/LevelInstanceEditorPropertyOverrideLevelStreaming.cpp b/Engine/Source/Runtime/Engine/Private/LevelInstance/LevelInstanceEditorPropertyOverrideLevelStreaming.cpp
--- a/Engine/Source/Runtime/Engine/Private/LevelInstance/LevelInstanceEditorPropertyOverrideLevelStreaming.cpp
+++ b/Engine/Source/Runtime/Engine/Private/LevelInstance/LevelInstanceEditorPropertyOverrideLevelStreaming.cpp
@@ -120,8 +120,8 @@ void ULevelStreamingLevelInstanceEditorPropertyOverride::ApplyPropertyOverrides(
 			if(ApplyActorType == EApplyActorType::Archetype || ApplyActorType == EApplyActorType::ActorAndArchetype)
 			{
 				// Gather Archetype Contextual Property Overrides and apply them to the archetype actor (Archetype will get overrides applied up to the Property Edit Owner)
-				TArray<const FActorPropertyOverride*> ArchetypePropertyOverrides;
-				if (LevelInstanceSubsystem->GetLevelInstancePropertyOverridesForActor(Actor, ArchetypeContextContainerID, ArchetypePropertyOverrides))
+				TArray<FLevelInstanceActorPropertyOverride> LevelInstanceArchetypePropertyOverrides;
+				if (LevelInstanceSubsystem->GetLevelInstancePropertyOverridesForActor(Actor, ArchetypeContextContainerID, LevelInstanceArchetypePropertyOverrides))
 				{
 					AActor* ArchetypeActor = CastChecked<AActor>(GetArchetypeForObject(Actor));
 
@@ -135,9 +135,9 @@ void ULevelStreamingLevelInstanceEditorPropertyOverride::ApplyPropertyOverrides(
 					bool bAppliedProperties = false;
 					if (ApplyPropertyOverrideType == EApplyPropertyOverrideType::PreConstructionScript || ApplyPropertyOverrideType == EApplyPropertyOverrideType::PreAndPostConstruction)
 					{
-						for (const FActorPropertyOverride* ActorPropertyOverride : ArchetypePropertyOverrides)
+						for (const FLevelInstanceActorPropertyOverride& LevelInstanceActorPropertyOverride : LevelInstanceArchetypePropertyOverrides)
 						{
-							bAppliedProperties |= ULevelInstancePropertyOverrideAsset::ApplyPropertyOverrides(ActorPropertyOverride, ArchetypeActor, false);
+							bAppliedProperties |= ULevelInstancePropertyOverrideAsset::ApplyPropertyOverrides(LevelInstanceActorPropertyOverride.ActorPropertyOverride, ArchetypeActor, false);
 						}
 					}
 					
@@ -150,9 +150,9 @@ void ULevelStreamingLevelInstanceEditorPropertyOverride::ApplyPropertyOverrides(
 					// If we are applying Post Construction Script overrides do it now
 					if (ApplyPropertyOverrideType == EApplyPropertyOverrideType::PostConstructionScript || ApplyPropertyOverrideType == EApplyPropertyOverrideType::PreAndPostConstruction)
 					{
-						for (const FActorPropertyOverride* ActorPropertyOverride : ArchetypePropertyOverrides)
+						for (const FLevelInstanceActorPropertyOverride& LevelInstanceActorPropertyOverride : LevelInstanceArchetypePropertyOverrides)
 						{
-							ULevelInstancePropertyOverrideAsset::ApplyPropertyOverrides(ActorPropertyOverride, ArchetypeActor, true);
+							ULevelInstancePropertyOverrideAsset::ApplyPropertyOverrides(LevelInstanceActorPropertyOverride.ActorPropertyOverride, ArchetypeActor, true);
 						}
 					}
 
@@ -161,14 +161,17 @@ void ULevelStreamingLevelInstanceEditorPropertyOverride::ApplyPropertyOverrides(
 					{
 						ApplyTransform(ArchetypeActor, LevelTransform, true);
 					}
+
+					// Flag needed so we can track properly in OnObjectPropertyChanged
+					FAddActorLevelInstanceFlags AddFlags(ArchetypeActor, ELevelInstanceFlags::HasPropertyOverrides);
 				}
 			}
 
 			if(ApplyActorType == EApplyActorType::Actor || ApplyActorType == EApplyActorType::ActorAndArchetype)
 			{	
 				// Gather Contextual Property Overrides and apply them to the actor
-				TArray<const FActorPropertyOverride*> ActorPropertyOverrides;
-				if (LevelInstanceSubsystem->GetLevelInstancePropertyOverridesForActor(Actor, ContextContainerID, ActorPropertyOverrides))
+				TArray<FLevelInstanceActorPropertyOverride> LevelInstanceActorPropertyOverrides;
+				if (LevelInstanceSubsystem->GetLevelInstancePropertyOverridesForActor(Actor, ContextContainerID, LevelInstanceActorPropertyOverrides))
 				{
 					// If we have Property Overrides we need to Remove the level transform before applying them in case the Relative transform of the actors was modified
 					if (bInAlreadyAppliedTransformOnActors && Actor->GetRootComponent())
@@ -179,9 +182,9 @@ void ULevelStreamingLevelInstanceEditorPropertyOverride::ApplyPropertyOverrides(
 					bool bAppliedProperties = false;
 					if (ApplyPropertyOverrideType == EApplyPropertyOverrideType::PreConstructionScript || ApplyPropertyOverrideType == EApplyPropertyOverrideType::PreAndPostConstruction)
 					{
-						for (const FActorPropertyOverride* ActorPropertyOverride : ActorPropertyOverrides)
+						for (const FLevelInstanceActorPropertyOverride& LevelInstanceActorPropertyOverride : LevelInstanceActorPropertyOverrides)
 						{
-							bAppliedProperties |= ULevelInstancePropertyOverrideAsset::ApplyPropertyOverrides(ActorPropertyOverride, Actor, false);
+							bAppliedProperties |= ULevelInstancePropertyOverrideAsset::ApplyPropertyOverrides(LevelInstanceActorPropertyOverride.ActorPropertyOverride, Actor, false);
 						}
 					}
 
@@ -193,9 +196,9 @@ void ULevelStreamingLevelInstanceEditorPropertyOverride::ApplyPropertyOverrides(
 
 					if (ApplyPropertyOverrideType == EApplyPropertyOverrideType::PostConstructionScript || ApplyPropertyOverrideType == EApplyPropertyOverrideType::PreAndPostConstruction)
 					{
-						for (const FActorPropertyOverride* ActorPropertyOverride : ActorPropertyOverrides)
+						for (const FLevelInstanceActorPropertyOverride& LevelInstanceActorPropertyOverride : LevelInstanceActorPropertyOverrides)
 						{
-							ULevelInstancePropertyOverrideAsset::ApplyPropertyOverrides(ActorPropertyOverride, Actor, true);
+							ULevelInstancePropertyOverrideAsset::ApplyPropertyOverrides(LevelInstanceActorPropertyOverride.ActorPropertyOverride, Actor, true);
 						}
 					}
 
@@ -207,7 +210,12 @@ void ULevelStreamingLevelInstanceEditorPropertyOverride::ApplyPropertyOverrides(
 					}
 
 					// Flag actor as being overriden
-					FAddActorLevelInstanceFlags AddFlags(Actor, ELevelInstanceFlags::HasPropertyOverrides);
+					ELevelInstanceFlags FlagsToAdd = ELevelInstanceFlags::HasPropertyOverrides;
+					if (LevelInstanceSubsystem->HasEditableLevelInstancePropertyOverrides(LevelInstanceActorPropertyOverrides))
+					{
+						EnumAddFlags(FlagsToAdd, ELevelInstanceFlags::HasEditablePropertyOverrides);
+					}
+					FAddActorLevelInstanceFlags AddFlags(Actor, FlagsToAdd);
 				}
 			}
 		}
@@ -228,6 +236,32 @@ void ULevelStreamingLevelInstanceEditorPropertyOverride::OnLoadedActorsAddedToLe
 	ApplyPropertyOverrides(InActors, bAlreadyAppliedTransformOnActors, EApplyPropertyOverrideType::PostConstructionScript);
 }
 
+void ULevelStreamingLevelInstanceEditorPropertyOverride::OnObjectPropertyChanged(UObject* Object, struct FPropertyChangedEvent& Event)
+{
+	if (Event.ChangeType != EPropertyChangeType::Interactive)
+	{
+		AActor* Actor = Object->IsA<AActor>() ? Cast<AActor>(Object) : Object->GetTypedOuter<AActor>();
+		if (Actor && Actor->GetLevel() == GetLoadedLevel())
+		{
+			FActorPropertyOverride ActorOverride;
+			if (!ULevelInstancePropertyOverrideAsset::SerializeActorPropertyOverrides(this, Actor, /*bForReset=*/true, ActorOverride))
+			{
+				ELevelInstanceFlags FlagsToRemove = ELevelInstanceFlags::HasEditablePropertyOverrides;
+				if (AActor* ArchetypeActor = Cast<AActor>(GetArchetypeForObject(Actor)); ArchetypeActor && !ArchetypeActor->HasLevelInstancePropertyOverrides())
+				{
+					EnumAddFlags(FlagsToRemove, ELevelInstanceFlags::HasPropertyOverrides);
+				}
+				FRemoveActorLevelInstanceFlags RemoveFlags(Actor, FlagsToRemove);
+			}
+			else
+			{
+				ELevelInstanceFlags FlagsToAdd = ELevelInstanceFlags::HasPropertyOverrides | ELevelInstanceFlags::HasEditablePropertyOverrides;
+				FAddActorLevelInstanceFlags AddFlags(Actor, FlagsToAdd);
+			}
+		}
+	}
+}
+
 void ULevelStreamingLevelInstanceEditorPropertyOverride::OnPreInitializeContainerInstance(UActorDescContainerInstance::FInitializeParams& InInitParams, UActorDescContainerInstance* InContainerInstance)
 {
 	// Apply override container
@@ -290,6 +324,8 @@ void ULevelStreamingLevelInstanceEditorPropertyOverride::Unload(ULevelStreamingL
 		ULevelInstanceSubsystem* LevelInstanceSubsystem = LevelStreaming->GetWorld()->GetSubsystem<ULevelInstanceSubsystem>();
 		check(LevelInstanceSubsystem);
 			
+		FCoreUObjectDelegates::OnObjectPropertyChanged.RemoveAll(LevelStreaming);
+
 		LoadedLevel->OnLoadedActorAddedToLevelPreEvent.RemoveAll(LevelStreaming);
 		LoadedLevel->OnLoadedActorAddedToLevelPostEvent.RemoveAll(LevelStreaming);
 	
@@ -334,6 +370,8 @@ void ULevelStreamingLevelInstanceEditorPropertyOverride::OnCurrentStateChanged(E
 		// For now this class doesn't support partial loading, but if at some point it does this is needed to apply the property overrides when new actors get loaded in
 		Level->OnLoadedActorAddedToLevelPostEvent.AddUObject(this, &ULevelStreamingLevelInstanceEditorPropertyOverride::OnLoadedActorsAddedToLevelPostEvent);
 
+		FCoreUObjectDelegates::OnObjectPropertyChanged.AddUObject(this, &ULevelStreamingLevelInstanceEditorPropertyOverride::OnObjectPropertyChanged);
+
 		// Push editing state to child actors
 		for (AActor* Actor : Level->Actors)
 		{
diff --git a/Engine/Source/Runtime/Engine/Private/LevelInstance/LevelInstanceEditorPropertyOverrideLevelStreaming.h b/Engine/Source/Runtime/Engine/Private/LevelInstance/LevelInstanceEditorPropertyOverrideLevelStreaming.h
--- a/Engine/Source/Runtime/Engine/Private/LevelInstance/LevelInstanceEditorPropertyOverrideLevelStreaming.h
+++ b/Engine/Source/Runtime/Engine/Private/LevelInstance/LevelInstanceEditorPropertyOverrideLevelStreaming.h
@@ -68,6 +68,7 @@ private:
 	void OnActorReplacedEvent(FWorldPartitionActorDescInstance* InActorDescInstance);
 	void OnPreInitializeContainerInstance(UActorDescContainerInstance::FInitializeParams& InInitParams, UActorDescContainerInstance* InContainerInstance);
 	virtual void OnCurrentStateChanged(ELevelStreamingState InPrevState, ELevelStreamingState InNewState) override;
+	void OnObjectPropertyChanged(UObject* Object, struct FPropertyChangedEvent& Event);
 
 	static UWorld* LoadArchetypeWorld(const FString& InWorldPackageName, const FString& InSuffix);
 	static void UnloadArchetypeWorld(UWorld* InWorld);
diff --git a/Engine/Source/Runtime/Engine/Private/LevelInstance/LevelInstanceLevelStreaming.cpp b/Engine/Source/Runtime/Engine/Private/LevelInstance/LevelInstanceLevelStreaming.cpp
--- a/Engine/Source/Runtime/Engine/Private/LevelInstance/LevelInstanceLevelStreaming.cpp
+++ b/Engine/Source/Runtime/Engine/Private/LevelInstance/LevelInstanceLevelStreaming.cpp
@@ -172,8 +172,8 @@ void ULevelStreamingLevelInstance::ApplyPropertyOverrides(const TArray<AActor*>&
 				if (IsValid(Actor))
 				{
 					// Gather Contextual Property Overrides and apply them
-					TArray<const FActorPropertyOverride*> ActorPropertyOverrides;
-					if (LevelInstanceSubsystem->GetLevelInstancePropertyOverridesForActor(Actor, ContextContainerID, ActorPropertyOverrides))
+					TArray<FLevelInstanceActorPropertyOverride> LevelInstanceActorPropertyOverrides;
+					if (LevelInstanceSubsystem->GetLevelInstancePropertyOverridesForActor(Actor, ContextContainerID, LevelInstanceActorPropertyOverrides))
 					{
 						// If we have Property Overrides we need to Remove the level transform before applying them in case the Relative transform of the actors was modified
 						if (bInAlreadyAppliedTransformOnActors && Actor->GetRootComponent())
@@ -184,9 +184,9 @@ void ULevelStreamingLevelInstance::ApplyPropertyOverrides(const TArray<AActor*>&
 						bool bAppliedProperties = false;
 						if (InApplyPropertyOverrideType == EApplyPropertieOverrideType::PreConstruction || InApplyPropertyOverrideType == EApplyPropertieOverrideType::PreAndPostConstruction)
 						{
-							for (const FActorPropertyOverride* ActorPropertyOverride : ActorPropertyOverrides)
+							for (const FLevelInstanceActorPropertyOverride& LevelInstanceActorPropertyOverride : LevelInstanceActorPropertyOverrides)
 							{
-								bAppliedProperties |= ULevelInstancePropertyOverrideAsset::ApplyPropertyOverrides(ActorPropertyOverride, Actor, false);
+								bAppliedProperties |= ULevelInstancePropertyOverrideAsset::ApplyPropertyOverrides(LevelInstanceActorPropertyOverride.ActorPropertyOverride, Actor, false);
 							}
 						}
 
@@ -198,9 +198,9 @@ void ULevelStreamingLevelInstance::ApplyPropertyOverrides(const TArray<AActor*>&
 
 						if (InApplyPropertyOverrideType == EApplyPropertieOverrideType::PostConstruction || InApplyPropertyOverrideType == EApplyPropertieOverrideType::PreAndPostConstruction)
 						{
-							for (const FActorPropertyOverride* ActorPropertyOverride : ActorPropertyOverrides)
+							for (const FLevelInstanceActorPropertyOverride& LevelInstanceActorPropertyOverride : LevelInstanceActorPropertyOverrides)
 							{
-								ULevelInstancePropertyOverrideAsset::ApplyPropertyOverrides(ActorPropertyOverride, Actor, true);
+								ULevelInstancePropertyOverrideAsset::ApplyPropertyOverrides(LevelInstanceActorPropertyOverride.ActorPropertyOverride, Actor, true);
 							}
 						}
 						
@@ -212,7 +212,12 @@ void ULevelStreamingLevelInstance::ApplyPropertyOverrides(const TArray<AActor*>&
 						}
 
 						// Flag actor as being overriden
-						FAddActorLevelInstanceFlags AddFlags(Actor, ELevelInstanceFlags::HasPropertyOverrides);
+						ELevelInstanceFlags FlagsToAdd = ELevelInstanceFlags::HasPropertyOverrides;
+						if (LevelInstanceSubsystem->HasEditableLevelInstancePropertyOverrides(LevelInstanceActorPropertyOverrides))
+						{
+							EnumAddFlags(FlagsToAdd, ELevelInstanceFlags::HasEditablePropertyOverrides);
+						}
+						FAddActorLevelInstanceFlags AddFlags(Actor, FlagsToAdd);
 					}
 				}
 			}
diff --git a/Engine/Source/Runtime/Engine/Private/LevelInstance/LevelInstancePropertyOverrideAsset.cpp b/Engine/Source/Runtime/Engine/Private/LevelInstance/LevelInstancePropertyOverrideAsset.cpp
--- a/Engine/Source/Runtime/Engine/Private/LevelInstance/LevelInstancePropertyOverrideAsset.cpp
+++ b/Engine/Source/Runtime/Engine/Private/LevelInstance/LevelInstancePropertyOverrideAsset.cpp
@@ -46,62 +46,62 @@ struct FLevelInstancePropertyOverrideUtils
 		static int32 SizeOfEmptyArchive = GetSizeOfEmptyArchive();
 		return SizeOfEmptyArchive;
 	}
+};
 
-	static bool SerializeActorPropertyOverrides(ULevelStreamingLevelInstanceEditorPropertyOverride* InLevelStreaming, AActor* InActor, bool bForReset, FActorPropertyOverride& OutActorPropertyOverrides)
-	{
-		AActor* ActorArchetype = Cast<AActor>(InLevelStreaming->GetArchetypeForObject(InActor));
-		check(ActorArchetype->GetTypedOuter<ULevel>() == InLevelStreaming->GetArchetypeLevel());
+bool ULevelInstancePropertyOverrideAsset::SerializeActorPropertyOverrides(ULevelStreamingLevelInstanceEditorPropertyOverride* InLevelStreaming, AActor* InActor, bool bForReset, FActorPropertyOverride& OutActorPropertyOverrides)
+{
+	AActor* ActorArchetype = Cast<AActor>(InLevelStreaming->GetArchetypeForObject(InActor));
+	check(ActorArchetype->GetTypedOuter<ULevel>() == InLevelStreaming->GetArchetypeLevel());
 
-		// Gather sub objects
-		TArray<UObject*> Objects;
-		GetObjectsWithOuter(InActor, Objects, true, RF_Transient, EInternalObjectFlags::Garbage);
+	// Gather sub objects
+	TArray<UObject*> Objects;
+	GetObjectsWithOuter(InActor, Objects, true, RF_Transient, EInternalObjectFlags::Garbage);
 
-		Objects.Add(InActor);
+	Objects.Add(InActor);
 
-		// Reset table
-		OutActorPropertyOverrides.ReferenceTable = FPropertyOverrideReferenceTable();
+	// Reset table
+	OutActorPropertyOverrides.ReferenceTable = FPropertyOverrideReferenceTable();
 
-		TArray<uint8> Payload;
-		// Serialize SubObjects
-		const FString ActorName = InActor->GetName();
-		for (UObject* Object : Objects)
+	TArray<uint8> Payload;
+	// Serialize SubObjects
+	const FString ActorName = InActor->GetName();
+	for (UObject* Object : Objects)
+	{
+		Payload.Reset();
+		FWorldPartitionPropertyOverrideWriter Writer(Payload);
+		FWorldPartitionPropertyOverrideArchive Archive(Writer, OutActorPropertyOverrides.ReferenceTable);
+		UObject* Archetype = InLevelStreaming->GetArchetypeForObject(Object);
+		if (Archetype->GetTypedOuter<ULevel>() == InLevelStreaming->GetArchetypeLevel())
 		{
-			Payload.Reset();
-			FWorldPartitionPropertyOverrideWriter Writer(Payload);
-			FWorldPartitionPropertyOverrideArchive Archive(Writer, OutActorPropertyOverrides.ReferenceTable);
-			UObject* Archetype = InLevelStreaming->GetArchetypeForObject(Object);
-			if (Archetype->GetTypedOuter<ULevel>() == InLevelStreaming->GetArchetypeLevel())
-			{
-				uint8* ToSerialize = bForReset ? (uint8*)Archetype : (uint8*)Object;
-				uint8* Default = bForReset ? (uint8*)Object : (uint8*)Archetype;
-				Object->GetClass()->SerializeTaggedProperties(Archive, ToSerialize, Object->GetClass(), Default);
+			uint8* ToSerialize = bForReset ? (uint8*)Archetype : (uint8*)Object;
+			uint8* Default = bForReset ? (uint8*)Object : (uint8*)Archetype;
+			Object->GetClass()->SerializeTaggedProperties(Archive, ToSerialize, Object->GetClass(), Default);
 
-				if (Payload.Num() != GetEmptyArchiveSize())
-				{
-					const FString ObjectSubPathString = FSoftObjectPath(Object).GetSubPathString();
-					const int32 IndexOfActorName = ObjectSubPathString.Find(ActorName);
-					const FString SubObjectPathString = ObjectSubPathString.Mid(IndexOfActorName + ActorName.Len() + 1);
-
-					FSubObjectPropertyOverride& SubObjectOverride = OutActorPropertyOverrides.SubObjectOverrides.Add(SubObjectPathString);
-					SubObjectOverride.SerializedTaggedProperties = MoveTemp(Payload);
-				}
-			}
-			else
+			if (Payload.Num() != FLevelInstancePropertyOverrideUtils::GetEmptyArchiveSize())
 			{
-				UE_LOG(LogLevelInstance, Warning, TEXT("Failed to find Property Override Archetype for: %s"), *Object->GetPathName());
+				const FString ObjectSubPathString = FSoftObjectPath(Object).GetSubPathString();
+				const int32 IndexOfActorName = ObjectSubPathString.Find(ActorName);
+				const FString SubObjectPathString = ObjectSubPathString.Mid(IndexOfActorName + ActorName.Len() + 1);
+
+				FSubObjectPropertyOverride& SubObjectOverride = OutActorPropertyOverrides.SubObjectOverrides.Add(SubObjectPathString);
+				SubObjectOverride.SerializedTaggedProperties = MoveTemp(Payload);
 			}
 		}
-
-		if (!OutActorPropertyOverrides.SubObjectOverrides.IsEmpty())
+		else
 		{
-			// Cache actor so that we can serialize its ActorDesc
-			OutActorPropertyOverrides.Actor = InActor;
-			return true;
+			UE_LOG(LogLevelInstance, Warning, TEXT("Failed to find Property Override Archetype for: %s"), *Object->GetPathName());
 		}
+	}
 
-		return false;
+	if (!OutActorPropertyOverrides.SubObjectOverrides.IsEmpty())
+	{
+		// Cache actor so that we can serialize its ActorDesc
+		OutActorPropertyOverrides.Actor = InActor;
+		return true;
 	}
-};
+
+	return false;
+}
 
 TSoftObjectPtr<ULevelInstancePropertyOverrideAsset> ULevelInstancePropertyOverrideAsset::GetSourceAssetPtr() const
 {
@@ -182,7 +182,7 @@ void ULevelInstancePropertyOverrideAsset::ResetPropertyOverridesForActor(ULevelS
 
 	// Serialize Reset
 	FActorPropertyOverride ActorOverride;
-	if (FLevelInstancePropertyOverrideUtils::SerializeActorPropertyOverrides(InLevelStreaming, InActor, /*bForReset=*/true, ActorOverride))
+	if (SerializeActorPropertyOverrides(InLevelStreaming, InActor, /*bForReset=*/true, ActorOverride))
 	{
 		// Apply Reset
 		ApplyPropertyOverrides(&ActorOverride, InActor, false);
@@ -244,7 +244,7 @@ void ULevelInstancePropertyOverrideAsset::SerializePropertyOverrides(ILevelInsta
 							
 					// Serialize Overrides
 					FActorPropertyOverride OutActorPropertyOverride;
-					if (FLevelInstancePropertyOverrideUtils::SerializeActorPropertyOverrides(InLevelStreaming, Actor, /*bForReset=*/false, OutActorPropertyOverride))
+					if (SerializeActorPropertyOverrides(InLevelStreaming, Actor, /*bForReset=*/false, OutActorPropertyOverride))
 					{
 						FContainerPropertyOverride& ContainerOverride = PropertyOverridesPerContainer.FindOrAdd(ContainerPath);
 						FActorPropertyOverride& ActorOverride = ContainerOverride.ActorOverrides.Add(Actor->GetActorGuid(), MoveTemp(OutActorPropertyOverride));
diff --git a/Engine/Source/Runtime/Engine/Private/LevelInstance/LevelInstanceSubsystem.cpp b/Engine/Source/Runtime/Engine/Private/LevelInstance/LevelInstanceSubsystem.cpp
--- a/Engine/Source/Runtime/Engine/Private/LevelInstance/LevelInstanceSubsystem.cpp
+++ b/Engine/Source/Runtime/Engine/Private/LevelInstance/LevelInstanceSubsystem.cpp
@@ -3282,7 +3282,27 @@ const ILevelInstanceInterface* ULevelInstanceSubsystem::GetLevelInstanceProperty
 	return GetLevelInstancePropertyOverridesEditOwner(const_cast<ILevelInstanceInterface*>(LevelInstance));
 }
 
-bool ULevelInstanceSubsystem::GetLevelInstancePropertyOverridesForActor(const AActor* Actor, FActorContainerID PropertyOverrideContext, TArray<const FActorPropertyOverride*>& OutPropertyOverrides) const
+bool ULevelInstanceSubsystem::HasEditableLevelInstancePropertyOverrides(TArray<FLevelInstanceActorPropertyOverride>& InPropertyOverrides) const
+{
+	// Return true if one of the property overrides is from an editable level instance
+	for (const FLevelInstanceActorPropertyOverride& LevelInstanceActorPropertyOverride : InPropertyOverrides)
+	{
+		if (LevelInstanceActorPropertyOverride.LevelInstanceID.IsValid())
+		{
+			if (AActor* LevelInstanceActor = Cast<AActor>(GetLevelInstance(LevelInstanceActorPropertyOverride.LevelInstanceID)))
+			{
+				if (!LevelInstanceActor->IsInLevelInstance() || LevelInstanceActor->IsInEditLevelInstance())
+				{
+					return true;
+				}
+			}
+		}
+	}
+
+	return false;
+}
+
+bool ULevelInstanceSubsystem::GetLevelInstancePropertyOverridesForActor(const AActor* Actor, FActorContainerID PropertyOverrideContext, TArray<FLevelInstanceActorPropertyOverride>& OutPropertyOverrides) const
 {
 	if (ILevelInstanceInterface* OwningLevelInstance = GetParentLevelInstance(Actor))
 	{
@@ -3368,10 +3388,10 @@ bool ULevelInstanceSubsystem::FPropertyOverrideEdit::Save(ILevelInstanceInterfac
 			
 	// Serialize current edit overrides
 	NewPropertyOverride->SerializePropertyOverrides(LevelInstanceOverrideOwner, LevelStreaming);
-			
-	// Set new property override before saving
-	LevelInstanceOverrideOwner->SetPropertyOverrideAsset(NewPropertyOverride);
-
+	
+	// Reset to nullptr if override is empty
+	LevelInstanceOverrideOwner->SetPropertyOverrideAsset(!NewPropertyOverride->GetPropertyOverridesPerContainer().IsEmpty() ? NewPropertyOverride : nullptr);
+	
 	// This will be used when initializing the ActorDesc to distinguish between saving the LevelInstance actor normally or through Property override save
 	TGuardValue<bool> SavingGuard(NewPropertyOverride->bSavingOverrideEdit, true);
 	if (FEditorFileUtils::PromptForCheckoutAndSave({ OverrideOwnerActor->GetPackage() }, /*bCheckDirty*/false, /*bPromptToSave*/false) != FEditorFileUtils::PR_Success)
@@ -3380,7 +3400,7 @@ bool ULevelInstanceSubsystem::FPropertyOverrideEdit::Save(ILevelInstanceInterfac
 		LevelInstanceOverrideOwner->SetPropertyOverrideAsset(ExistingPropertyOverride);
 		OverrideOwnerActor->GetPackage()->SetDirtyFlag(bWasDirty);
 		return false;
-	}	
+	}
 
 	return true;
 }
diff --git a/Engine/Source/Runtime/Engine/Private/WorldPartition/LevelInstance/LevelInstanceContainerInstance.cpp b/Engine/Source/Runtime/Engine/Private/WorldPartition/LevelInstance/LevelInstanceContainerInstance.cpp
--- a/Engine/Source/Runtime/Engine/Private/WorldPartition/LevelInstance/LevelInstanceContainerInstance.cpp
+++ b/Engine/Source/Runtime/Engine/Private/WorldPartition/LevelInstance/LevelInstanceContainerInstance.cpp
@@ -10,7 +10,7 @@
 #include "WorldPartition/LevelInstance/LevelInstancePropertyOverrideDesc.h"
 #include "LevelInstance/LevelInstancePropertyOverrideAsset.h"
 #include "LevelInstance/LevelInstanceSettings.h"
-
+#include "LevelInstance/LevelInstanceInterface.h"
 
 void ULevelInstanceContainerInstance::Initialize(const FInitializeParams& InParams)
 {
@@ -167,7 +167,7 @@ void ULevelInstanceContainerInstance::GetPropertyOverridesForActor(const FActorC
 	}
 }
 
-void ULevelInstanceContainerInstance::GetPropertyOverridesForActor(const FActorContainerID& InContainerID, const FActorContainerID& InContextContainerID, const FGuid& InActorGuid, TArray<const FActorPropertyOverride*>& OutPropertyOverrides) const
+void ULevelInstanceContainerInstance::GetPropertyOverridesForActor(const FActorContainerID& InContainerID, const FActorContainerID& InContextContainerID, const FGuid& InActorGuid, TArray<FLevelInstanceActorPropertyOverride>& OutPropertyOverrides) const
 {
 	if (!ULevelInstanceSettings::Get()->IsPropertyOverrideEnabled())
 	{
@@ -190,7 +190,16 @@ void ULevelInstanceContainerInstance::GetPropertyOverridesForActor(const FActorC
 			{
 				if (const FActorPropertyOverride* FoundActorOverride = FoundContainerOverride->ActorOverrides.Find(InActorGuid))
 				{
-					OutPropertyOverrides.Add(FoundActorOverride);
+					if (const UActorDescContainerInstance* ParentContainerInstancePtr = ParentContainerInstance.Get())
+					{
+						if (const FWorldPartitionActorDescInstance* ParentActorDescInstance = ParentContainerInstance->GetActorDescInstance(GetContainerActorGuid()))
+						{
+							if (const ILevelInstanceInterface* ParentLevelInstance = Cast<ILevelInstanceInterface>(ParentActorDescInstance->GetActor()))
+							{
+								OutPropertyOverrides.Add(FLevelInstanceActorPropertyOverride(ParentLevelInstance->GetLevelInstanceID(), FoundActorOverride));
+							}
+						}
+					}
 				}
 			}
 		}
diff --git a/Engine/Source/Runtime/Engine/Public/LevelInstance/LevelInstancePropertyOverrideAsset.h b/Engine/Source/Runtime/Engine/Public/LevelInstance/LevelInstancePropertyOverrideAsset.h
--- a/Engine/Source/Runtime/Engine/Public/LevelInstance/LevelInstancePropertyOverrideAsset.h
+++ b/Engine/Source/Runtime/Engine/Public/LevelInstance/LevelInstancePropertyOverrideAsset.h
@@ -8,6 +8,7 @@
 #include "UObject/Object.h"
 #include "WorldPartition/WorldPartitionActorContainerID.h"
 #include "WorldPartition/WorldPartitionPropertyOverride.h"
+#include "LevelInstance/LevelInstanceTypes.h"
 
 #include "LevelInstancePropertyOverrideAsset.generated.h"
 
@@ -16,6 +17,16 @@ class ILevelInstanceInterface;
 class ULevelStreamingLevelInstanceEditorPropertyOverride;
 class FLevelInstancePropertyOverrideDesc;
 
+struct FLevelInstanceActorPropertyOverride
+{
+	FLevelInstanceActorPropertyOverride(const FLevelInstanceID& InLevelInstanceID, const FActorPropertyOverride* InActorPropertyOverride)
+		: LevelInstanceID(InLevelInstanceID)
+		, ActorPropertyOverride(InActorPropertyOverride) {}
+
+	FLevelInstanceID LevelInstanceID;
+	const FActorPropertyOverride* ActorPropertyOverride = nullptr;
+};
+
 UCLASS(MinimalAPI, NotBlueprintable)
 class ULevelInstancePropertyOverrideAsset : public UWorldPartitionPropertyOverride
 {
@@ -32,6 +43,8 @@ private:
 	using UWorldPartitionPropertyOverride::ApplyPropertyOverrides;
 	// End UWorldPartitionPropertyOverride Interface
 
+	static bool SerializeActorPropertyOverrides(ULevelStreamingLevelInstanceEditorPropertyOverride* InLevelStreaming, AActor* InActor, bool bForReset, FActorPropertyOverride& OutActorPropertyOverrides);
+
 	void Initialize(const TSoftObjectPtr<UWorld> InWorldAsset) { WorldAsset = InWorldAsset; }
 	void SerializePropertyOverrides(ILevelInstanceInterface* InLevelInstanceOverrideOwner, ULevelStreamingLevelInstanceEditorPropertyOverride* InLevelStreamingInterface);
 	void ResetPropertyOverridesForActor(ULevelStreamingLevelInstanceEditorPropertyOverride* InLevelStreamingInterface, AActor* InActor);
diff --git a/Engine/Source/Runtime/Engine/Public/LevelInstance/LevelInstanceSubsystem.h b/Engine/Source/Runtime/Engine/Public/LevelInstance/LevelInstanceSubsystem.h
--- a/Engine/Source/Runtime/Engine/Public/LevelInstance/LevelInstanceSubsystem.h
+++ b/Engine/Source/Runtime/Engine/Public/LevelInstance/LevelInstanceSubsystem.h
@@ -261,7 +261,8 @@ private:
 	};
 		
 	ENGINE_API FActorContainerID GetLevelInstancePropertyOverridesContext(ILevelInstanceInterface* LevelInstance) const;
-	ENGINE_API bool GetLevelInstancePropertyOverridesForActor(const AActor* Actor, FActorContainerID PropertyOverrideContext, TArray<const FActorPropertyOverride*>& OutPropertyOverrides) const;
+	ENGINE_API bool HasEditableLevelInstancePropertyOverrides(TArray<FLevelInstanceActorPropertyOverride>& InPropertyOverrides) const;
+	ENGINE_API bool GetLevelInstancePropertyOverridesForActor(const AActor* Actor, FActorContainerID PropertyOverrideContext, TArray<FLevelInstanceActorPropertyOverride>& OutPropertyOverrides) const;
 	
 	ENGINE_API ILevelInstanceInterface* GetLevelInstancePropertyOverridesEditOwner(ILevelInstanceInterface* LevelInstance) const;
 	ENGINE_API const ILevelInstanceInterface* GetLevelInstancePropertyOverridesEditOwner(const ILevelInstanceInterface* LevelInstance) const;
diff --git a/Engine/Source/Runtime/Engine/Public/WorldPartition/LevelInstance/LevelInstanceContainerInstance.h b/Engine/Source/Runtime/Engine/Public/WorldPartition/LevelInstance/LevelInstanceContainerInstance.h
--- a/Engine/Source/Runtime/Engine/Public/WorldPartition/LevelInstance/LevelInstanceContainerInstance.h
+++ b/Engine/Source/Runtime/Engine/Public/WorldPartition/LevelInstance/LevelInstanceContainerInstance.h
@@ -32,7 +32,7 @@ protected:
 	virtual void Uninitialize();
 
 	virtual void GetPropertyOverridesForActor(const FActorContainerID& InContainerID, const FGuid& InActorGuid, TArray<FWorldPartitionRuntimeCellPropertyOverride>& OutPropertyOverrides) const override;
-	ENGINE_API void GetPropertyOverridesForActor(const FActorContainerID& InContainerID, const FActorContainerID& InContextContainerID, const FGuid& InActorGuid, TArray<const FActorPropertyOverride*>& OutPropertyOverrides) const;
+	ENGINE_API void GetPropertyOverridesForActor(const FActorContainerID& InContainerID, const FActorContainerID& InContextContainerID, const FGuid& InActorGuid, TArray<FLevelInstanceActorPropertyOverride>& OutPropertyOverrides) const;
 
 	virtual void RegisterContainer(const FInitializeParams& InParams) override;
 	virtual void UnregisterContainer() override;
