Fix for serialized MIDs not resolving in packaged games due to non-unique names being serialized via BPs.

diff --git a/Engine/Source/Runtime/Engine/Private/Materials/MaterialInstanceDynamic.cpp b/Engine/Source/Runtime/Engine/Private/Materials/MaterialInstanceDynamic.cpp
--- a/Engine/Source/Runtime/Engine/Private/Materials/MaterialInstanceDynamic.cpp
+++ b/Engine/Source/Runtime/Engine/Private/Materials/MaterialInstanceDynamic.cpp
@@ -78,7 +78,7 @@ UMaterialInstanceDynamic* UMaterialInstanceDynamic::Create(UMaterialInterface* P
 
 	if (Name == NAME_None && ParentMaterial)
 	{
-		Name = FName(FString("MID_") + ParentMaterial->GetName());
+		Name = MakeUniqueObjectName(InOuter, UMaterialInstanceDynamic::StaticClass(), FName(FString("MID_") + ParentMaterial->GetName()));
 	}
 
 	if (Name != NAME_None)
