[NetPhysics] Default replications legacy flow in GT to hardsnap in PT if async physics is enabled.
Also imlpement settings for this in the NetworkPhysicsSettingsComponent and implemented settings for correcting connected bodies for the other replication modes.

diff --git a/Engine/Source/Runtime/Engine/Private/PhysicsEngine/NetworkPhysicsComponent.cpp b/Engine/Source/Runtime/Engine/Private/PhysicsEngine/NetworkPhysicsComponent.cpp
--- a/Engine/Source/Runtime/Engine/Private/PhysicsEngine/NetworkPhysicsComponent.cpp
+++ b/Engine/Source/Runtime/Engine/Private/PhysicsEngine/NetworkPhysicsComponent.cpp
@@ -25,9 +25,9 @@ namespace PhysicsReplicationCVars
 		static bool bAllowRewindToClosestState = true;
 		static FAutoConsoleVariableRef CVarResimAllowRewindToClosestState(TEXT("np2.Resim.AllowRewindToClosestState"), bAllowRewindToClosestState, TEXT("When rewinding to a specific frame, if the client doens't have state data for that frame, use closest data available. Only affects the first rewind frame, when FPBDRigidsEvolution is set to Reset."));
 		static bool bCompareStateToTriggerRewind = false;
-		static FAutoConsoleVariableRef CVarResimCompareStateToTriggerRewind(TEXT("np2.Resim.CompareStateToTriggerRewind"), bCompareStateToTriggerRewind, TEXT("If we should cache local players custom state struct in rewind history and compare the predicted state with incoming server state to trigger resimulations if they differ, comparison done through FNetworkPhysicsData::CompareData."));
+		static FAutoConsoleVariableRef CVarResimCompareStateToTriggerRewind(TEXT("np2.Resim.CompareStateToTriggerRewind"), bCompareStateToTriggerRewind, TEXT("When true, cache local players custom state struct in rewind history and compare the predicted state with incoming server state to trigger resimulations if they differ, comparison done through FNetworkPhysicsData::CompareData"));
 		static bool bCompareInputToTriggerRewind = false;
-		static FAutoConsoleVariableRef CVarResimCompareInputToTriggerRewind(TEXT("np2.Resim.CompareInputToTriggerRewind"), bCompareInputToTriggerRewind, TEXT("If we should compare local players predicted inputs with incoming server inputs to trigger resimulations if they differ, comparison done through FNetworkPhysicsData::CompareData."));
+		static FAutoConsoleVariableRef CVarResimCompareInputToTriggerRewind(TEXT("np2.Resim.CompareInputToTriggerRewind"), bCompareInputToTriggerRewind, TEXT("When true, compare local players predicted inputs with incoming server inputs to trigger resimulations if they differ, comparison done through FNetworkPhysicsData::CompareData."));
 	}
 }
 
diff --git a/Engine/Source/Runtime/Engine/Private/PhysicsEngine/NetworkPhysicsSettingsComponent.cpp b/Engine/Source/Runtime/Engine/Private/PhysicsEngine/NetworkPhysicsSettingsComponent.cpp
--- a/Engine/Source/Runtime/Engine/Private/PhysicsEngine/NetworkPhysicsSettingsComponent.cpp
+++ b/Engine/Source/Runtime/Engine/Private/PhysicsEngine/NetworkPhysicsSettingsComponent.cpp
@@ -23,6 +23,8 @@ namespace PhysicsReplicationCVars
 }
 
 
+TMap<AActor*, UNetworkPhysicsSettingsComponent*> UNetworkPhysicsSettingsComponent::ObjectToSettings_External = TMap<AActor*, UNetworkPhysicsSettingsComponent*>();
+
 UNetworkPhysicsSettingsComponent::UNetworkPhysicsSettingsComponent()
 {
 	bWantsInitializeComponent = true;
@@ -55,6 +57,7 @@ void UNetworkPhysicsSettingsComponent::InitializeComponent()
 								FNetworkPhysicsSettingsAsyncInput* AsyncInput = NetworkPhysicsSettingsAsync->GetProducerInputData_External();
 								AsyncInput->PhysicsObject = PhysicsObject;
 								AsyncInput->Settings.GeneralSettings = GeneralSettings;
+								AsyncInput->Settings.DefaultReplicationSettings = DefaultReplicationSettings;
 								AsyncInput->Settings.ResimulationSettings = ResimulationSettings;
 								AsyncInput->Settings.PredictiveInterpolationSettings = PredictiveInterpolationSettings;
 							}
@@ -64,6 +67,11 @@ void UNetworkPhysicsSettingsComponent::InitializeComponent()
 			}
 		}
 	}
+
+	if (AActor* Owner = GetOwner())
+	{
+		UNetworkPhysicsSettingsComponent::ObjectToSettings_External.Add(Owner, this);
+	}
 }
 
 void UNetworkPhysicsSettingsComponent::UninitializeComponent()
@@ -82,6 +90,11 @@ void UNetworkPhysicsSettingsComponent::UninitializeComponent()
 		}
 	}
 	NetworkPhysicsSettingsAsync = nullptr;
+
+	if (AActor* Owner = GetOwner())
+	{
+		UNetworkPhysicsSettingsComponent::ObjectToSettings_External.Remove(Owner);
+	}
 }
 
 void UNetworkPhysicsSettingsComponent::BeginPlay()
@@ -100,6 +113,11 @@ void UNetworkPhysicsSettingsComponent::BeginPlay()
 	}
 }
 
+UNetworkPhysicsSettingsComponent* UNetworkPhysicsSettingsComponent::GetSettingsForActor(AActor* Owner)
+{
+	UNetworkPhysicsSettingsComponent** Value = ObjectToSettings_External.Find(Owner);
+	return Value ? *Value : nullptr;
+}
 
 #pragma region // FNetworkPhysicsSettingsComponentAsync
 
diff --git a/Engine/Source/Runtime/Engine/Private/PhysicsEngine/PhysicsReplication.cpp b/Engine/Source/Runtime/Engine/Private/PhysicsEngine/PhysicsReplication.cpp
--- a/Engine/Source/Runtime/Engine/Private/PhysicsEngine/PhysicsReplication.cpp
+++ b/Engine/Source/Runtime/Engine/Private/PhysicsEngine/PhysicsReplication.cpp
@@ -87,13 +87,25 @@ namespace PhysicsReplicationCVars
 	int32 EnableDefaultReplication = 0;
 	static FAutoConsoleVariableRef CVarEnableDefaultReplication(TEXT("np2.EnableDefaultReplication"), EnableDefaultReplication, TEXT("Enable default replication in the networked physics prediction flow."));
 
+	namespace DefaultReplicationCVars
+	{
+		bool bHardsnapLegacyInPT = false;
+		static FAutoConsoleVariableRef CVarDefaultLegacyHardsnapInPT(TEXT("p.DefaultReplication.Legacy.HardsnapInPT"), bHardsnapLegacyInPT, TEXT("If default replication is used and it's running the legacy flow through Game Thread, allow hardsnapping to be performed on Physics Thread if async physics is enabled.."));
+
+		bool bCorrectConnectedBodies = false;
+		static FAutoConsoleVariableRef CVarDefaultCorrectConnectedBodies(TEXT("p.DefaultReplication.CorrectConnectedBodies"), bCorrectConnectedBodies, TEXT("When true, transform corrections will also apply to any connected physics object."));
+
+		bool bCorrectConnectedBodiesFriction = true;
+		static FAutoConsoleVariableRef CVarDefaultCorrectConnectedBodiesFriction(TEXT("p.DefaultReplication.CorrectConnectedBodiesFriction"), bCorrectConnectedBodiesFriction, TEXT("When true, transform correction on any connected physics object will also recalculate their friction."));
+	}
+
 	namespace ResimulationCVars
 	{
 		bool bRuntimeCorrectionEnabled = true;
-		static FAutoConsoleVariableRef CVarResimRuntimeCorrectionEnabled(TEXT("np2.Resim.RuntimeCorrectionEnabled"), bRuntimeCorrectionEnabled, TEXT("Apply runtime corrections while error is small enough not to trigger a resim."));
+		static FAutoConsoleVariableRef CVarResimRuntimeCorrectionEnabled(TEXT("np2.Resim.RuntimeCorrectionEnabled"), bRuntimeCorrectionEnabled, TEXT("Apply positional and rotational runtime corrections while within resim trigger distance."));
 
 		bool bRuntimeVelocityCorrection = false;
-		static FAutoConsoleVariableRef CVarResimRuntimeVelocityCorrection(TEXT("np2.Resim.RuntimeVelocityCorrection"), bRuntimeVelocityCorrection, TEXT("Apply linear and angular velocity corrections in runtime while within resim trigger. Used if RuntimeCorrectionEnabled is true."));
+		static FAutoConsoleVariableRef CVarResimRuntimeVelocityCorrection(TEXT("np2.Resim.RuntimeVelocityCorrection"), bRuntimeVelocityCorrection, TEXT("Apply linear and angular velocity corrections in runtime while within resim trigger distance. Used if RuntimeCorrectionEnabled is true."));
 
 		bool bRuntimeCorrectConnectedBodies = true;
 		static FAutoConsoleVariableRef CVarResimRuntimeCorrectConnectedBodies(TEXT("np2.Resim.RuntimeCorrectConnectedBodies"), bRuntimeCorrectConnectedBodies, TEXT("If true runtime position and rotation correction will also shift transform of any connected physics objects. Used if RuntimeCorrectionEnabled is true."));
@@ -120,28 +132,28 @@ namespace PhysicsReplicationCVars
 	namespace PredictiveInterpolationCVars
 	{
 		float PosCorrectionTimeBase = 0.0f;
-		static FAutoConsoleVariableRef CVarPosCorrectionTimeBase(TEXT("np2.PredictiveInterpolation.PosCorrectionTimeBase"), PosCorrectionTimeBase, TEXT("Base time to correct positional offset over. RTT * PosCorrectionTimeMultiplier are added on top of this."));
+		static FAutoConsoleVariableRef CVarPosCorrectionTimeBase(TEXT("np2.PredictiveInterpolation.PosCorrectionTimeBase"), PosCorrectionTimeBase, TEXT("Base time to correct positional offset over. RoundTripTime * PosCorrectionTimeMultiplier is added on top of this."));
 
 		float PosCorrectionTimeMin = 0.1f;
-		static FAutoConsoleVariableRef CVarPosCorrectionTimeMin(TEXT("np2.PredictiveInterpolation.PosCorrectionTimeMin"), PosCorrectionTimeMin, TEXT("Min time time to correct positional offset over. DeltaSeconds is added on top of this."));
+		static FAutoConsoleVariableRef CVarPosCorrectionTimeMin(TEXT("np2.PredictiveInterpolation.PosCorrectionTimeMin"), PosCorrectionTimeMin, TEXT("Min time to correct positional offset over. DeltaSeconds is added on top of this."));
 
 		float PosCorrectionTimeMultiplier = 1.0f;
-		static FAutoConsoleVariableRef CVarPosCorrectionTimeMultiplier(TEXT("np2.PredictiveInterpolation.PosCorrectionTimeMultiplier"), PosCorrectionTimeMultiplier, TEXT("Multiplier to adjust how much of RTT (network Round Trip Time) to add to positional offset correction."));
+		static FAutoConsoleVariableRef CVarPosCorrectionTimeMultiplier(TEXT("np2.PredictiveInterpolation.PosCorrectionTimeMultiplier"), PosCorrectionTimeMultiplier, TEXT("Multiplier to adjust how much of RoundTripTime to add to positional offset correction."));
 
 		float RotCorrectionTimeBase = 0.0f;
-		static FAutoConsoleVariableRef CVarRotCorrectionTimeBase(TEXT("np2.PredictiveInterpolation.RotCorrectionTimeBase"), RotCorrectionTimeBase, TEXT("Base time to correct positional offset over. RTT * PosCorrectionTimeMultiplier are added on top of this."));
+		static FAutoConsoleVariableRef CVarRotCorrectionTimeBase(TEXT("np2.PredictiveInterpolation.RotCorrectionTimeBase"), RotCorrectionTimeBase, TEXT("Base time to correct rotational offset over. RoundTripTime * RotCorrectionTimeMultiplier is added on top of this."));
 
 		float RotCorrectionTimeMin = 0.1f;
-		static FAutoConsoleVariableRef CVarRotCorrectionTimeMin(TEXT("np2.PredictiveInterpolation.RotCorrectionTimeMin"), RotCorrectionTimeMin, TEXT("Min time time to correct rotational offset over. DeltaSeconds is added on top of this."));
+		static FAutoConsoleVariableRef CVarRotCorrectionTimeMin(TEXT("np2.PredictiveInterpolation.RotCorrectionTimeMin"), RotCorrectionTimeMin, TEXT("Min time to correct rotational offset over. DeltaSeconds is added on top of this."));
 
 		float RotCorrectionTimeMultiplier = 1.0f;
-		static FAutoConsoleVariableRef CVarRotCorrectionTimeMultiplier(TEXT("np2.PredictiveInterpolation.RotCorrectionTimeMultiplier"), RotCorrectionTimeMultiplier, TEXT("Multiplier to adjust how much of RTT (network Round Trip Time) to add to positional offset correction."));
+		static FAutoConsoleVariableRef CVarRotCorrectionTimeMultiplier(TEXT("np2.PredictiveInterpolation.RotCorrectionTimeMultiplier"), RotCorrectionTimeMultiplier, TEXT("Multiplier to adjust how much of RoundTripTime to add to rotational offset correction."));
 
 		float PosInterpolationTimeMultiplier = 1.1f;
-		static FAutoConsoleVariableRef CVarInterpolationTimeMultiplier(TEXT("np2.PredictiveInterpolation.InterpolationTimeMultiplier"), PosInterpolationTimeMultiplier, TEXT("Multiplier to adjust the replication interpolation time which is based on the sendrate of replication data from the server."));
+		static FAutoConsoleVariableRef CVarInterpolationTimeMultiplier(TEXT("np2.PredictiveInterpolation.InterpolationTimeMultiplier"), PosInterpolationTimeMultiplier, TEXT("Multiplier to adjust the interpolation time which is based on the sendrate of state data from the server."));
 		
 		float RotInterpolationTimeMultiplier = 1.25f;
-		static FAutoConsoleVariableRef CVarRotInterpolationTimeMultiplier(TEXT("np2.PredictiveInterpolation.RotInterpolationTimeMultiplier"), RotInterpolationTimeMultiplier, TEXT("Multiplier to adjust the rotational replication interpolation time which is based on the sendrate of replication data from the server."));
+		static FAutoConsoleVariableRef CVarRotInterpolationTimeMultiplier(TEXT("np2.PredictiveInterpolation.RotInterpolationTimeMultiplier"), RotInterpolationTimeMultiplier, TEXT("Multiplier to adjust the rotational interpolation time which is based on the sendrate of state data from the server."));
 		
 		float AverageReceiveIntervalSmoothing = 3.0f;
 		static FAutoConsoleVariableRef CVarAverageReceiveIntervalSmoothing(TEXT("np2.PredictiveInterpolation.AverageReceiveIntervalSmoothing"), AverageReceiveIntervalSmoothing, TEXT("Recommended range: 1.0 - 5.0. Higher value makes the average receive interval adjust itself slower, reducing spikes in InterpolationTime."));
@@ -171,10 +183,10 @@ namespace PhysicsReplicationCVars
 		static FAutoConsoleVariableRef CVarErrorAccAngVelMaxLimit(TEXT("np2.PredictiveInterpolation.ErrorAccAngVelMaxLimit"), ErrorAccAngVelMaxLimit, TEXT("If target angular velocity (in radians) is below this limit we check for desync to trigger softsnap and accumulate time to build up to a hardsnap."));
 		
 		float SoftSnapPosStrength = 0.5f;
-		static FAutoConsoleVariableRef CVarSoftSnapPosStrength(TEXT("np2.PredictiveInterpolation.SoftSnapPosStrength"), SoftSnapPosStrength, TEXT("Value in percent between 0.0 - 1.0 representing how much to softsnap each tick of the remaining distance."));
+		static FAutoConsoleVariableRef CVarSoftSnapPosStrength(TEXT("np2.PredictiveInterpolation.SoftSnapPosStrength"), SoftSnapPosStrength, TEXT("Value in percent between 0.0 - 1.0 representing how much to softsnap each tick of the remaining positional distance."));
 		
 		float SoftSnapRotStrength = 0.5f;
-		static FAutoConsoleVariableRef CVarSoftSnapRotStrength(TEXT("np2.PredictiveInterpolation.SoftSnapRotStrength"), SoftSnapRotStrength, TEXT("Value in percent between 0.0 - 1.0 representing how much to softsnap each tick of the remaining distance."));
+		static FAutoConsoleVariableRef CVarSoftSnapRotStrength(TEXT("np2.PredictiveInterpolation.SoftSnapRotStrength"), SoftSnapRotStrength, TEXT("Value in percent between 0.0 - 1.0 representing how much to softsnap each tick of the remaining rotational distance."));
 
 		bool bSoftSnapToSource = false;
 		static FAutoConsoleVariableRef CVarSoftSnapToSource(TEXT("np2.PredictiveInterpolation.SoftSnapToSource"), bSoftSnapToSource, TEXT("If true, soft snap will be performed towards the source state of the current target instead of the predicted state of the current target."));
@@ -413,6 +425,9 @@ void FPhysicsReplication::OnTick(float DeltaSeconds, TMap<TWeakObjectPtr<UPrimit
 					bool bUpdated = false;
 					if (AActor* OwningActor = PrimComp->GetOwner())
 					{
+						// Update actor replication settings overrides
+						SettingsCurrent = UNetworkPhysicsSettingsComponent::GetSettingsForActor(OwningActor);
+
 						const ENetRole OwnerRole = OwningActor->GetLocalRole();
 						const bool bIsSimulated = OwnerRole == ROLE_SimulatedProxy;
 						const bool bIsReplicatedAutonomous = OwnerRole == ROLE_AutonomousProxy && PrimComp->bReplicatePhysicsToAutonomousProxy;
@@ -695,11 +710,52 @@ bool FPhysicsReplication::ApplyRigidBodyState(float DeltaSeconds, FBodyInstance*
 			PhysicsTarget.AccumulatedErrorSeconds = 0.0f;
 			bRestoredState = true;
 
-			BI->SetBodyTransform(IdealWorldTM, ETeleportType::ResetPhysics, bAutoWake);
+			// Hardsnap in physics thread
+			bool bPTHardSnapSuccess = false;
+			if (PhysicsReplicationAsync != nullptr)
+			{
+				bool bHardsnapLegacyInPT = PhysicsReplicationCVars::DefaultReplicationCVars::bHardsnapLegacyInPT;
+				bool bCorrectConnectedBodies = PhysicsReplicationCVars::DefaultReplicationCVars::bCorrectConnectedBodies;
+				bool bCorrectConnectedBodiesFriction = PhysicsReplicationCVars::DefaultReplicationCVars::bCorrectConnectedBodiesFriction;
+				
+				if (SettingsCurrent.Get())
+				{
+					bHardsnapLegacyInPT = SettingsCurrent.Get()->DefaultReplicationSettings.GetHardsnapDefaultLegacyInPT();
+					bCorrectConnectedBodies = SettingsCurrent.Get()->DefaultReplicationSettings.GetCorrectConnectedBodies();
+					bCorrectConnectedBodiesFriction = SettingsCurrent.Get()->DefaultReplicationSettings.GetCorrectConnectedBodiesFriction();
+				}
+
+				if (bHardsnapLegacyInPT)
+				{
+					if (Chaos::FSingleParticlePhysicsProxy* Proxy = static_cast<Chaos::FSingleParticlePhysicsProxy*>(BI->GetPhysicsActorHandle()))
+					{
+						if (Chaos::FPBDRigidsSolver* Solver = Proxy->GetSolver<Chaos::FPBDRigidsSolver>())
+						{
+							Solver->EnqueueCommandImmediate([Solver, Proxy, IdealWorldTM, NewState, bCorrectConnectedBodies, bCorrectConnectedBodiesFriction]()
+								{
+									Chaos::FRigidBodyHandle_Internal* Handle = Proxy->GetPhysicsThreadAPI();
+								
+									// Set XRVW to hard snap dynamic object and force recalculation of friction
+									Solver->GetEvolution()->ApplyParticleTransformCorrection(Proxy->GetHandle_LowLevel(), IdealWorldTM.GetLocation(), IdealWorldTM.GetRotation(), bCorrectConnectedBodies, bCorrectConnectedBodiesFriction);
+								
+									Handle->SetV(NewState.LinVel);
+									Handle->SetW(FMath::DegreesToRadians(NewState.AngVel));
+								});
+
+							bPTHardSnapSuccess = true;
+						}
+					}
+				}
+			}
+
+			if (!bPTHardSnapSuccess)
+			{
+				BI->SetBodyTransform(IdealWorldTM, ETeleportType::ResetPhysics, bAutoWake);
 
-			// Set the new velocities
-			BI->SetLinearVelocity(NewState.LinVel, false, bAutoWake);
-			BI->SetAngularVelocityInRadians(FMath::DegreesToRadians(NewState.AngVel), false, bAutoWake);
+				// Set the new velocities
+				BI->SetLinearVelocity(NewState.LinVel, false, bAutoWake);
+				BI->SetAngularVelocityInRadians(FMath::DegreesToRadians(NewState.AngVel), false, bAutoWake);
+			}
 		}
 		else
 		{
@@ -1415,8 +1471,12 @@ bool FPhysicsReplicationAsync::DefaultReplication(Chaos::FPBDRigidParticleHandle
 			// Too much error so just snap state here and be done with it
 			Target.AccumulatedErrorSeconds = 0.0f;
 			bRestoredState = true;
-			Handle->SetX(TargetPos);
-			Handle->SetR(TargetQuat);
+
+			// Set XRVW to hard snap dynamic object and force recalculation of friction
+			const bool bCorrectConnectedBodies = SettingsCurrent.DefaultReplicationSettings.GetCorrectConnectedBodies();
+			const bool bCorrectConnectedBodiesFriction = SettingsCurrent.DefaultReplicationSettings.GetCorrectConnectedBodiesFriction();
+			RigidsSolver->GetEvolution()->ApplyParticleTransformCorrection(Handle, TargetPos, TargetQuat, bCorrectConnectedBodies, bCorrectConnectedBodiesFriction);
+
 			Handle->SetV(NewState.LinVel);
 			Handle->SetW(FMath::DegreesToRadians(NewState.AngVel));
 		}
@@ -1625,7 +1685,7 @@ bool FPhysicsReplicationAsync::PredictiveInterpolation(Chaos::FPBDRigidParticleH
 		else 
 		{
 			// Set XRVW to hard snap dynamic object and force recalculation of friction
-			const bool bCorrectConnectedBodies = PhysicsReplicationCVars::PredictiveInterpolationCVars::bCorrectConnectedBodies;
+			const bool bCorrectConnectedBodies = SettingsCurrent.PredictiveInterpolationSettings.GetCorrectConnectedBodies();
 			RigidsSolver->GetEvolution()->ApplyParticleTransformCorrection(Handle, Target.PrevPosTarget, Target.PrevRotTarget, bCorrectConnectedBodies, /*bInRecalculateFrictionOnConnectedBodies*/ true);
 			Handle->SetV(TargetLinVel);
 			Handle->SetW(TargetAngVel);
@@ -1823,8 +1883,8 @@ bool FPhysicsReplicationAsync::PredictiveInterpolation(Chaos::FPBDRigidParticleH
 		// Apply correction as a transform shift
 		if (!PhysicsReplicationCVars::PredictiveInterpolationCVars::bCorrectionAsVelocity)
 		{	
-			const bool bCorrectConnectedBodies = PhysicsReplicationCVars::PredictiveInterpolationCVars::bCorrectConnectedBodies;
-			const bool bCorrectConnectedBodiesFriction = PhysicsReplicationCVars::PredictiveInterpolationCVars::bCorrectConnectedBodiesFriction;
+			const bool bCorrectConnectedBodies = SettingsCurrent.PredictiveInterpolationSettings.GetCorrectConnectedBodies();
+			const bool bCorrectConnectedBodiesFriction = SettingsCurrent.PredictiveInterpolationSettings.GetCorrectConnectedBodiesFriction();
 			RigidsSolver->GetEvolution()->ApplyParticleTransformCorrection(Handle, CorrectionX, CorrectionR, bCorrectConnectedBodies, bCorrectConnectedBodiesFriction);
 		}
 
@@ -1839,8 +1899,8 @@ bool FPhysicsReplicationAsync::PredictiveInterpolation(Chaos::FPBDRigidParticleH
 				FMath::Clamp(SettingsCurrent.PredictiveInterpolationSettings.GetSoftSnapRotStrength(), 0.0f, 1.0f));
 
 			// Apply correction as a transform shift
-			const bool bCorrectConnectedBodies = PhysicsReplicationCVars::PredictiveInterpolationCVars::bCorrectConnectedBodies;
-			const bool bCorrectConnectedBodiesFriction = PhysicsReplicationCVars::PredictiveInterpolationCVars::bCorrectConnectedBodiesFriction;
+			const bool bCorrectConnectedBodies = SettingsCurrent.PredictiveInterpolationSettings.GetCorrectConnectedBodies();
+			const bool bCorrectConnectedBodiesFriction = SettingsCurrent.PredictiveInterpolationSettings.GetCorrectConnectedBodiesFriction();
 			RigidsSolver->GetEvolution()->ApplyParticleTransformCorrection(Handle, SoftSnapPos, SoftSnapRot, bCorrectConnectedBodies, bCorrectConnectedBodiesFriction);
 		}
 	}
@@ -1982,7 +2042,7 @@ bool FPhysicsReplicationAsync::ResimulationReplication(Chaos::FPBDRigidParticleH
 				}
 #endif
 				// Apply correction to position and rotation
-				RigidsSolver->GetEvolution()->ApplyParticleTransformCorrection(Handle, CorrectedX, CorrectedR, PhysicsReplicationCVars::ResimulationCVars::bRuntimeCorrectConnectedBodies);
+				RigidsSolver->GetEvolution()->ApplyParticleTransformCorrection(Handle, CorrectedX, CorrectedR, SettingsCurrent.ResimulationSettings.GetRuntimeCorrectConnectedBodies());
 			}
 
 			// Keep target for NumPredictedFrames time to perform runtime corrections with until a new target is received
diff --git a/Engine/Source/Runtime/Engine/Public/Physics/NetworkPhysicsSettingsComponent.h b/Engine/Source/Runtime/Engine/Public/Physics/NetworkPhysicsSettingsComponent.h
--- a/Engine/Source/Runtime/Engine/Public/Physics/NetworkPhysicsSettingsComponent.h
+++ b/Engine/Source/Runtime/Engine/Public/Physics/NetworkPhysicsSettingsComponent.h
@@ -16,6 +16,7 @@
 class FNetworkPhysicsSettingsComponentAsync;
 
 //  Alias
+using FDefaultReplicationSettings		= FNetworkPhysicsSettingsDefaultReplication;
 using FPredictiveInterpolationSettings	= FNetworkPhysicsSettingsPredictiveInterpolation;
 using FResimulationSettings				= FNetworkPhysicsSettingsResimulation;
 /*
@@ -25,6 +26,13 @@ using FRenderInterpolationSettings		= FNetworkPhysicsSettingsRenderInterpolation
 
 namespace PhysicsReplicationCVars
 {
+	namespace DefaultReplicationCVars
+	{
+		extern bool bHardsnapLegacyInPT;
+		extern bool bCorrectConnectedBodies;
+		extern bool bCorrectConnectedBodiesFriction;
+	}
+	
 	namespace PredictiveInterpolationCVars
 	{
 		extern float PosCorrectionTimeBase;
@@ -41,6 +49,8 @@ namespace PhysicsReplicationCVars
 		extern bool bSkipVelocityRepOnPosEarlyOut;
 		extern bool bPostResimWaitForUpdate;
 		extern bool bDisableSoftSnap;
+		extern bool bCorrectConnectedBodies;
+		extern bool bCorrectConnectedBodiesFriction;
 	}
 
 	namespace ResimulationCVars
@@ -51,6 +61,7 @@ namespace PhysicsReplicationCVars
 		extern float RotStabilityMultiplier;
 		extern float VelStabilityMultiplier;
 		extern float AngVelStabilityMultiplier;
+		extern bool bRuntimeCorrectConnectedBodies;
 	}
 }
 
@@ -62,114 +73,155 @@ struct FNetworkPhysicsSettings
 
 	UPROPERTY(config, EditDefaultsOnly, Category = "Overrides", Meta = (InlineEditConditionToggle))
 	uint32 bOverrideSimProxyRepMode : 1;
-	// Override the EPhysicsReplicationMode if Actor is ENetRole::ROLE_SimulatedProxy and has EPhysicsReplicationMode::Resimulation set
+	// Override the EPhysicsReplicationMode for Actors with ENetRole::ROLE_SimulatedProxy.
 	UPROPERTY(config, EditDefaultsOnly, Category = "Overrides", Meta = (EditCondition = "bOverrideSimProxyRepMode"))
 	EPhysicsReplicationMode SimProxyRepMode = EPhysicsReplicationMode::PredictiveInterpolation;
 };
 
+USTRUCT()
+struct FNetworkPhysicsSettingsDefaultReplication
+{
+	GENERATED_BODY()
+
+	UPROPERTY(config, EditDefaultsOnly, Category = "Overrides", Meta = (InlineEditConditionToggle))
+	uint32 bOverrideDefaultLegacyHardsnapInPT : 1;
+	// Overrides CVar: p.DefaultReplication.Legacy.HardsnapInPT -- If default replication is used and it's running the legacy flow through Game Thread, allow hardsnapping to be performed on Physics Thread if async physics is enabled.
+	UPROPERTY(config, EditDefaultsOnly, Category = "Overrides", Meta = (EditCondition = "bOverrideDefaultLegacyHardsnapInPT"))
+	bool bHardsnapInPhysicsThread = PhysicsReplicationCVars::DefaultReplicationCVars::bHardsnapLegacyInPT;
+	bool GetHardsnapDefaultLegacyInPT() { return bOverrideDefaultLegacyHardsnapInPT ? bHardsnapInPhysicsThread : PhysicsReplicationCVars::DefaultReplicationCVars::bHardsnapLegacyInPT; }
+	
+	UPROPERTY(config, EditDefaultsOnly, Category = "Overrides", Meta = (InlineEditConditionToggle))
+	uint32 bOverrideCorrectConnectedBodies : 1;
+	// Overrides CVar: p.DefaultReplication.CorrectConnectedBodies -- When true, transform corrections will also apply to any connected physics object.
+	UPROPERTY(config, EditDefaultsOnly, Category = "Overrides", Meta = (EditCondition = "bOverrideCorrectConnectedBodies"))
+	bool bCorrectConnectedBodies = PhysicsReplicationCVars::DefaultReplicationCVars::bCorrectConnectedBodies;
+	bool GetCorrectConnectedBodies() { return bOverrideCorrectConnectedBodies ? bCorrectConnectedBodies : PhysicsReplicationCVars::DefaultReplicationCVars::bCorrectConnectedBodies; }
+	
+	UPROPERTY(config, EditDefaultsOnly, Category = "Overrides", Meta = (InlineEditConditionToggle))
+	uint32 bOverrideCorrectConnectedBodiesFriction : 1;
+	// Overrides CVar: p.DefaultReplication.CorrectConnectedBodiesFriction -- When true, transform correction on any connected physics object will also recalculate their friction.
+	UPROPERTY(config, EditDefaultsOnly, Category = "Overrides", Meta = (EditCondition = "bOverrideCorrectConnectedBodiesFriction"))
+	bool bCorrectConnectedBodiesFriction = PhysicsReplicationCVars::DefaultReplicationCVars::bCorrectConnectedBodiesFriction;
+	bool GetCorrectConnectedBodiesFriction() { return bOverrideCorrectConnectedBodiesFriction ? bCorrectConnectedBodiesFriction : PhysicsReplicationCVars::DefaultReplicationCVars::bCorrectConnectedBodiesFriction; }
+};
 
 USTRUCT()
 struct FNetworkPhysicsSettingsPredictiveInterpolation
 {
 	GENERATED_BODY()
 
-	// np2.PredictiveInterpolation.PosCorrectionTimeBase
 	UPROPERTY(config, EditDefaultsOnly, Category = "Overrides", Meta = (InlineEditConditionToggle))
 	uint32 bOverridePosCorrectionTimeBase : 1;
+	// Overrides CVar: np2.PredictiveInterpolation.PosCorrectionTimeBase -- Base time to correct positional offset over. RoundTripTime * PosCorrectionTimeMultiplier is added on top of this.
 	UPROPERTY(config, EditDefaultsOnly, Category = "Overrides", Meta = (EditCondition = "bOverridePosCorrectionTimeBase"))
 	float PosCorrectionTimeBase = PhysicsReplicationCVars::PredictiveInterpolationCVars::PosCorrectionTimeBase;
 	float GetPosCorrectionTimeBase() { return bOverridePosCorrectionTimeBase ? PosCorrectionTimeBase : PhysicsReplicationCVars::PredictiveInterpolationCVars::PosCorrectionTimeBase; }
 
-	// np2.PredictiveInterpolation.PosCorrectionTimeMin
 	UPROPERTY(config, EditDefaultsOnly, Category = "Overrides", Meta = (InlineEditConditionToggle))
 	uint32 bOverridePosCorrectionTimeMin : 1;
+	// Overrides CVar: np2.PredictiveInterpolation.PosCorrectionTimeMin -- Min time to correct positional offset over. DeltaSeconds is added on top of this.
 	UPROPERTY(config, EditDefaultsOnly, Category = "Overrides", Meta = (EditCondition = "bOverridePosCorrectionTimeMin"))
 	float PosCorrectionTimeMin = PhysicsReplicationCVars::PredictiveInterpolationCVars::PosCorrectionTimeMin;
 	float GetPosCorrectionTimeMin() { return bOverridePosCorrectionTimeMin ? PosCorrectionTimeMin : PhysicsReplicationCVars::PredictiveInterpolationCVars::PosCorrectionTimeMin; }
 
-	// np2.PredictiveInterpolation.PosCorrectionTimeMultiplier
 	UPROPERTY(config, EditDefaultsOnly, Category = "Overrides", Meta = (InlineEditConditionToggle))
 	uint32 bOverridePosCorrectionTimeMultiplier : 1;
+	// Overrides CVar: np2.PredictiveInterpolation.PosCorrectionTimeMultiplier -- Multiplier to adjust how much of RoundTripTime to add to positional offset correction.
 	UPROPERTY(config, EditDefaultsOnly, Category = "Overrides", Meta = (EditCondition = "bOverridePosCorrectionTimeMultiplier"))
 	float PosCorrectionTimeMultiplier = PhysicsReplicationCVars::PredictiveInterpolationCVars::PosCorrectionTimeMultiplier;
 	float GetPosCorrectionTimeMultiplier() { return bOverridePosCorrectionTimeMultiplier ? PosCorrectionTimeMultiplier : PhysicsReplicationCVars::PredictiveInterpolationCVars::PosCorrectionTimeMultiplier; }
 
-	// np2.PredictiveInterpolation.RotCorrectionTimeBase
 	UPROPERTY(config, EditDefaultsOnly, Category = "Overrides", Meta = (InlineEditConditionToggle))
 	uint32 bOverrideRotCorrectionTimeBase : 1;
+	// Overrides CVar: np2.PredictiveInterpolation.RotCorrectionTimeBase -- Base time to correct rotational offset over. RoundTripTime * RotCorrectionTimeMultiplier is added on top of this.
 	UPROPERTY(config, EditDefaultsOnly, Category = "Overrides", Meta = (EditCondition = "bOverrideRotCorrectionTimeBase"))
 	float RotCorrectionTimeBase = PhysicsReplicationCVars::PredictiveInterpolationCVars::RotCorrectionTimeBase;
 	float GetRotCorrectionTimeBase() { return bOverrideRotCorrectionTimeBase ? RotCorrectionTimeBase : PhysicsReplicationCVars::PredictiveInterpolationCVars::RotCorrectionTimeBase; }
 
-	// np2.PredictiveInterpolation.RotCorrectionTimeMin
 	UPROPERTY(config, EditDefaultsOnly, Category = "Overrides", Meta = (InlineEditConditionToggle))
 	uint32 bOverrideRotCorrectionTimeMin : 1;
+	// Overrides CVar: np2.PredictiveInterpolation.RotCorrectionTimeMin -- Min time to correct rotational offset over. DeltaSeconds is added on top of this.
 	UPROPERTY(config, EditDefaultsOnly, Category = "Overrides", Meta = (EditCondition = "bOverrideRotCorrectionTimeMin"))
 	float RotCorrectionTimeMin = PhysicsReplicationCVars::PredictiveInterpolationCVars::RotCorrectionTimeMin;
 	float GetRotCorrectionTimeMin() { return bOverrideRotCorrectionTimeMin ? RotCorrectionTimeMin : PhysicsReplicationCVars::PredictiveInterpolationCVars::RotCorrectionTimeMin; }
 
-	// np2.PredictiveInterpolation.RotCorrectionTimeMultiplier
 	UPROPERTY(config, EditDefaultsOnly, Category = "Overrides", Meta = (InlineEditConditionToggle))
 	uint32 bOverrideRotCorrectionTimeMultiplier : 1;
+	// Overrides CVar: np2.PredictiveInterpolation.RotCorrectionTimeMultiplier -- Multiplier to adjust how much of RoundTripTime to add to rotational offset correction.
 	UPROPERTY(config, EditDefaultsOnly, Category = "Overrides", Meta = (EditCondition = "bOverrideRotCorrectionTimeMultiplier"))
 	float RotCorrectionTimeMultiplier = PhysicsReplicationCVars::PredictiveInterpolationCVars::RotCorrectionTimeMultiplier;
 	float GetRotCorrectionTimeMultiplier() { return bOverrideRotCorrectionTimeMultiplier ? RotCorrectionTimeMultiplier : PhysicsReplicationCVars::PredictiveInterpolationCVars::RotCorrectionTimeMultiplier; }
 
-	// np2.PredictiveInterpolation.InterpolationTimeMultiplier
 	UPROPERTY(config, EditDefaultsOnly, Category = "Overrides", Meta = (InlineEditConditionToggle))
 	uint32 bOverridePosInterpolationTimeMultiplier : 1;
+	// Overrides CVar: np2.PredictiveInterpolation.InterpolationTimeMultiplier -- Multiplier to adjust the interpolation time which is based on the sendrate of state data from the server.
 	UPROPERTY(config, EditDefaultsOnly, Category = "Overrides", Meta = (EditCondition = "bOverridePosInterpolationTimeMultiplier"))
 	float PosInterpolationTimeMultiplier = PhysicsReplicationCVars::PredictiveInterpolationCVars::PosInterpolationTimeMultiplier;
 	float GetPosInterpolationTimeMultiplier() { return bOverridePosInterpolationTimeMultiplier ? PosInterpolationTimeMultiplier : PhysicsReplicationCVars::PredictiveInterpolationCVars::PosInterpolationTimeMultiplier; }
 
-	// np2.PredictiveInterpolation.RotInterpolationTimeMultiplier
 	UPROPERTY(config, EditDefaultsOnly, Category = "Overrides", Meta = (InlineEditConditionToggle))
 	uint32 bOverrideRotInterpolationTimeMultiplier : 1;
+	// Overrides CVar: np2.PredictiveInterpolation.RotInterpolationTimeMultiplier -- Multiplier to adjust the rotational interpolation time which is based on the sendrate of state data from the server.
 	UPROPERTY(config, EditDefaultsOnly, Category = "Overrides", Meta = (EditCondition = "bOverrideRotInterpolationTimeMultiplier"))
 	float RotInterpolationTimeMultiplier = PhysicsReplicationCVars::PredictiveInterpolationCVars::RotInterpolationTimeMultiplier;
 	float GetRotInterpolationTimeMultiplier() { return bOverrideRotInterpolationTimeMultiplier ? RotInterpolationTimeMultiplier : PhysicsReplicationCVars::PredictiveInterpolationCVars::RotInterpolationTimeMultiplier; }
 
-	// np2.PredictiveInterpolation.SoftSnapPosStrength
 	UPROPERTY(config, EditDefaultsOnly, Category = "Overrides", Meta = (InlineEditConditionToggle))
 	uint32 bOverrideSoftSnapPosStrength : 1;
+	// Overrides CVar: np2.PredictiveInterpolation.SoftSnapPosStrength -- Value in percent between 0.0 - 1.0 representing how much to softsnap each tick of the remaining positional distance.
 	UPROPERTY(config, EditDefaultsOnly, Category = "Overrides", Meta = (EditCondition = "bOverrideSoftSnapPosStrength"))
 	float SoftSnapPosStrength = PhysicsReplicationCVars::PredictiveInterpolationCVars::SoftSnapPosStrength;
 	float GetSoftSnapPosStrength() { return bOverrideSoftSnapPosStrength ? SoftSnapPosStrength : PhysicsReplicationCVars::PredictiveInterpolationCVars::SoftSnapPosStrength; }
 
-	// np2.PredictiveInterpolation.SoftSnapRotStrength
 	UPROPERTY(config, EditDefaultsOnly, Category = "Overrides", Meta = (InlineEditConditionToggle))
 	uint32 bOverrideSoftSnapRotStrength : 1;
+	// Overrides CVar: np2.PredictiveInterpolation.SoftSnapRotStrength -- Value in percent between 0.0 - 1.0 representing how much to softsnap each tick of the remaining rotational distance.
 	UPROPERTY(config, EditDefaultsOnly, Category = "Overrides", Meta = (EditCondition = "bOverrideSoftSnapRotStrength"))
 	float SoftSnapRotStrength = PhysicsReplicationCVars::PredictiveInterpolationCVars::SoftSnapRotStrength;
 	float GetSoftSnapRotStrength() { return bOverrideSoftSnapRotStrength ? SoftSnapRotStrength : PhysicsReplicationCVars::PredictiveInterpolationCVars::SoftSnapRotStrength; }
 
-	// np2.PredictiveInterpolation.SoftSnapToSource
 	UPROPERTY(config, EditDefaultsOnly, Category = "Overrides", Meta = (InlineEditConditionToggle))
 	uint32 bOverrideSoftSnapToSource : 1;
+	// Overrides CVar: np2.PredictiveInterpolation.SoftSnapToSource -- If true, softsnap will be performed towards the source state of the current target instead of the predicted state of the current target.
 	UPROPERTY(config, EditDefaultsOnly, Category = "Overrides", Meta = (EditCondition = "bOverrideSoftSnapToSource"))
 	bool bSoftSnapToSource = PhysicsReplicationCVars::PredictiveInterpolationCVars::bSoftSnapToSource;
 	bool GetSoftSnapToSource() { return bOverrideSoftSnapToSource ? bSoftSnapToSource : PhysicsReplicationCVars::PredictiveInterpolationCVars::bSoftSnapToSource; }
 
-	// np2.PredictiveInterpolation.SkipVelocityRepOnPosEarlyOut
+	UPROPERTY(config, EditDefaultsOnly, Category = "Overrides", Meta = (InlineEditConditionToggle))
+	uint32 bOverrideDisableSoftSnap : 1;
+	// Overrides CVar: np2.PredictiveInterpolation.DisableSoftSnap -- When true, predictive interpolation will not use softsnap to correct the replication with when velocity fails. Hardsnap will still eventually kick in if replication can't reach the target.
+	UPROPERTY(config, EditDefaultsOnly, Category = "Overrides", Meta = (EditCondition = "bOverrideDisableSoftSnap"))
+	bool bDisableSoftSnap = PhysicsReplicationCVars::PredictiveInterpolationCVars::bDisableSoftSnap;
+	bool GetDisableSoftSnap() { return bOverrideDisableSoftSnap ? bDisableSoftSnap : PhysicsReplicationCVars::PredictiveInterpolationCVars::bDisableSoftSnap; }
+
 	UPROPERTY(config, EditDefaultsOnly, Category = "Overrides", Meta = (InlineEditConditionToggle))
 	uint32 bOverrideSkipVelocityRepOnPosEarlyOut : 1;
+	// Overrides CVar: np2.PredictiveInterpolation.SkipVelocityRepOnPosEarlyOut -- If true, don't run linear velocity replication if position can early out but angular can't early out.
 	UPROPERTY(config, EditDefaultsOnly, Category = "Overrides", Meta = (EditCondition = "bOverrideSkipVelocityRepOnPosEarlyOut"))
 	bool bSkipVelocityRepOnPosEarlyOut = PhysicsReplicationCVars::PredictiveInterpolationCVars::bSkipVelocityRepOnPosEarlyOut;
 	bool GetSkipVelocityRepOnPosEarlyOut() { return bOverrideSkipVelocityRepOnPosEarlyOut ? bSkipVelocityRepOnPosEarlyOut : PhysicsReplicationCVars::PredictiveInterpolationCVars::bSkipVelocityRepOnPosEarlyOut; }
 
-	// np2.PredictiveInterpolation.PostResimWaitForUpdate
 	UPROPERTY(config, EditDefaultsOnly, Category = "Overrides", Meta = (InlineEditConditionToggle))
 	uint32 bOverridePostResimWaitForUpdate : 1;
+	// Overrides CVar: np2.PredictiveInterpolation.PostResimWaitForUpdate -- After a resimulation, wait for replicated states that correspond to post-resim state before processing replication again.
 	UPROPERTY(config, EditDefaultsOnly, Category = "Overrides", Meta = (EditCondition = "bOverridePostResimWaitForUpdate"))
 	bool bPostResimWaitForUpdate = PhysicsReplicationCVars::PredictiveInterpolationCVars::bPostResimWaitForUpdate;
 	bool GetPostResimWaitForUpdate() { return bOverridePostResimWaitForUpdate ? bPostResimWaitForUpdate : PhysicsReplicationCVars::PredictiveInterpolationCVars::bPostResimWaitForUpdate; }
 
-	// np2.PredictiveInterpolation.DisableSoftSnap
 	UPROPERTY(config, EditDefaultsOnly, Category = "Overrides", Meta = (InlineEditConditionToggle))
-	uint32 bOverrideDisableSoftSnap : 1;
-	UPROPERTY(config, EditDefaultsOnly, Category = "Overrides", Meta = (EditCondition = "bOverrideDisableSoftSnap"))
-	bool bDisableSoftSnap = PhysicsReplicationCVars::PredictiveInterpolationCVars::bDisableSoftSnap;
-	bool GetDisableSoftSnap() { return bOverrideDisableSoftSnap ? bDisableSoftSnap : PhysicsReplicationCVars::PredictiveInterpolationCVars::bDisableSoftSnap; }
+	uint32 bOverrideCorrectConnectedBodies : 1;
+	// Overrides CVar: np2.PredictiveInterpolation.CorrectConnectedBodies -- When true, transform corrections will also apply to any connected physics object.
+	UPROPERTY(config, EditDefaultsOnly, Category = "Overrides", Meta = (EditCondition = "bOverrideCorrectConnectedBodies"))
+	bool bCorrectConnectedBodies = PhysicsReplicationCVars::PredictiveInterpolationCVars::bCorrectConnectedBodies;
+	bool GetCorrectConnectedBodies() { return bOverrideCorrectConnectedBodies ? bCorrectConnectedBodies : PhysicsReplicationCVars::PredictiveInterpolationCVars::bCorrectConnectedBodies; }
+
+	UPROPERTY(config, EditDefaultsOnly, Category = "Overrides", Meta = (InlineEditConditionToggle))
+	uint32 bOverrideCorrectConnectedBodiesFriction : 1;
+	// Overrides CVar: np2.PredictiveInterpolation.CorrectConnectedBodiesFriction -- When true, transform correction on any connected physics object will also recalculate their friction.
+	UPROPERTY(config, EditDefaultsOnly, Category = "Overrides", Meta = (EditCondition = "bOverrideCorrectConnectedBodiesFriction"))
+	bool bCorrectConnectedBodiesFriction = PhysicsReplicationCVars::PredictiveInterpolationCVars::bCorrectConnectedBodiesFriction;
+	bool GetCorrectConnectedBodiesFriction() { return bOverrideCorrectConnectedBodiesFriction ? bCorrectConnectedBodiesFriction : PhysicsReplicationCVars::PredictiveInterpolationCVars::bCorrectConnectedBodiesFriction; }
+
 };
 
 USTRUCT()
@@ -177,81 +229,87 @@ struct FNetworkPhysicsSettingsResimulation
 {
 	GENERATED_BODY()
 
-	// Override how many inputs to synchronize each sync to cover packet loss.
+	UPROPERTY(config, EditDefaultsOnly, Category = "Overrides", Meta = (InlineEditConditionToggle))
+	uint32 bOverrideResimulationErrorThreshold : 1;
+	// Overrides Project Settings -> Physics -> Replication -> Physics Prediction -> Resimulation Error Threshold -- Distance that the object is allowed to desync from the server before triggering a resimulation, within this distance runtime correction can be performed if RuntimeCorrectionEnabled is true.
+	UPROPERTY(config, EditDefaultsOnly, Category = "Overrides", Meta = (EditCondition = "bOverrideResimulationErrorThreshold"))
+	uint32 ResimulationErrorThreshold = 10;
+	uint32 GetResimulationErrorThreshold(uint32 DefaultValue) { return bOverrideRotStabilityMultiplier ? ResimulationErrorThreshold : DefaultValue; }
+
 	UPROPERTY(config, EditDefaultsOnly, Category = "Overrides", Meta = (InlineEditConditionToggle))
 	uint32 bOverrideRedundantInputs : 1;
+	// Override how many inputs to synchronize each sync to cover packet loss.
 	UPROPERTY(config, EditDefaultsOnly, Category = "Overrides", Meta = (EditCondition = "bOverrideRedundantInputs"))
 	uint8 RedundantInputs = 3;
 
-	// Override how many states to synchronize each sync to cover packet loss.
 	UPROPERTY(config, EditDefaultsOnly, Category = "Overrides", Meta = (InlineEditConditionToggle))
 	uint32 bOverrideRedundantStates : 1;
+	// Override how many states to synchronize each sync to cover packet loss.
 	UPROPERTY(config, EditDefaultsOnly, Category = "Overrides", Meta = (EditCondition = "bOverrideRedundantStates"))
 	uint8 RedundantStates = 1;
 
-	// np2.Resim.RuntimeCorrectionEnabled
 	UPROPERTY(config, EditDefaultsOnly, Category = "Overrides", Meta = (InlineEditConditionToggle))
 	uint32 bOverrideRuntimeCorrectionEnabled : 1;
+	// Overrides CVar: np2.Resim.RuntimeCorrectionEnabled -- Apply positional and rotational runtime corrections while within resim trigger distance.
 	UPROPERTY(config, EditDefaultsOnly, Category = "Overrides", Meta = (EditCondition = "bOverrideRuntimeCorrectionEnabled"))
 	bool bRuntimeCorrectionEnabled = PhysicsReplicationCVars::ResimulationCVars::bRuntimeCorrectionEnabled;
 	bool GetRuntimeCorrectionEnabled() { return bOverrideRuntimeCorrectionEnabled ? bRuntimeCorrectionEnabled : PhysicsReplicationCVars::ResimulationCVars::bRuntimeCorrectionEnabled; }
 
-	// np2.Resim.RuntimeVelocityCorrection
 	UPROPERTY(config, EditDefaultsOnly, Category = "Overrides", Meta = (InlineEditConditionToggle))
 	uint32 bOverrideRuntimeVelocityCorrection : 1;
+	// Overrides CVar: np2.Resim.RuntimeVelocityCorrection -- Apply linear and angular velocity corrections in runtime while within resim trigger distance. Used if RuntimeCorrectionEnabled is true.
 	UPROPERTY(config, EditDefaultsOnly, Category = "Overrides", Meta = (EditCondition = "bOverrideRuntimeVelocityCorrection"))
 	bool bRuntimeVelocityCorrection = PhysicsReplicationCVars::ResimulationCVars::bRuntimeVelocityCorrection;
 	bool GetRuntimeVelocityCorrectionEnabled() { return bOverrideRuntimeVelocityCorrection ? bRuntimeVelocityCorrection : PhysicsReplicationCVars::ResimulationCVars::bRuntimeVelocityCorrection; }
 
-	// np2.Resim.PosStabilityMultiplier
+	UPROPERTY(config, EditDefaultsOnly, Category = "Overrides", Meta = (InlineEditConditionToggle))
+	uint32 bOverrideRuntimeCorrectConnectedBodies : 1;
+	// Overrides CVar: np2.Resim.RuntimeCorrectConnectedBodies -- If true runtime position and rotation correction will also shift transform of any connected physics objects. Used if RuntimeCorrectionEnabled is true.
+	UPROPERTY(config, EditDefaultsOnly, Category = "Overrides", Meta = (EditCondition = "bOverrideRuntimeCorrectConnectedBodies"))
+	bool bRuntimeCorrectConnectedBodies = PhysicsReplicationCVars::ResimulationCVars::bRuntimeCorrectConnectedBodies;
+	bool GetRuntimeCorrectConnectedBodies() { return bOverrideRuntimeCorrectConnectedBodies ? bRuntimeCorrectConnectedBodies : PhysicsReplicationCVars::ResimulationCVars::bRuntimeCorrectConnectedBodies; }
+
 	UPROPERTY(config, EditDefaultsOnly, Category = "Overrides", Meta = (InlineEditConditionToggle))
 	uint32 bOverridePosStabilityMultiplier : 1;
+	// Overrides CVar: np2.Resim.PosStabilityMultiplier -- Recommended range between 0.0-1.0. Lower value means more stable positional corrections.
 	UPROPERTY(config, EditDefaultsOnly, Category = "Overrides", Meta = (EditCondition = "bOverridePosStabilityMultiplier"))
 	float PosStabilityMultiplier = PhysicsReplicationCVars::ResimulationCVars::PosStabilityMultiplier;
 	float GetPosStabilityMultiplier() { return bOverridePosStabilityMultiplier ? PosStabilityMultiplier : PhysicsReplicationCVars::ResimulationCVars::PosStabilityMultiplier; }
 
-	// np2.Resim.RotStabilityMultiplier
 	UPROPERTY(config, EditDefaultsOnly, Category = "Overrides", Meta = (InlineEditConditionToggle))
 	uint32 bOverrideRotStabilityMultiplier : 1;
+	// Overrides CVar: np2.Resim.RotStabilityMultiplier -- Recommended range between 0.0-1.0. Lower value means more stable rotational corrections.
 	UPROPERTY(config, EditDefaultsOnly, Category = "Overrides", Meta = (EditCondition = "bOverrideRotStabilityMultiplier"))
 	float RotStabilityMultiplier = PhysicsReplicationCVars::ResimulationCVars::RotStabilityMultiplier;
 	float GetRotStabilityMultiplier() { return bOverrideRotStabilityMultiplier ? RotStabilityMultiplier : PhysicsReplicationCVars::ResimulationCVars::RotStabilityMultiplier; }
 
-	// np2.Resim.VelStabilityMultiplier
 	UPROPERTY(config, EditDefaultsOnly, Category = "Overrides", Meta = (InlineEditConditionToggle))
 	uint32 bOverrideVelStabilityMultiplier : 1;
+	// Overrides CVar: np2.Resim.VelStabilityMultiplier -- Recommended range between 0.0-1.0. Lower value means more stable linear velocity corrections.
 	UPROPERTY(config, EditDefaultsOnly, Category = "Overrides", Meta = (EditCondition = "bOverrideVelStabilityMultiplier"))
 	float VelStabilityMultiplier = PhysicsReplicationCVars::ResimulationCVars::VelStabilityMultiplier;
 	float GetVelStabilityMultiplier() { return bOverrideVelStabilityMultiplier ? VelStabilityMultiplier : PhysicsReplicationCVars::ResimulationCVars::VelStabilityMultiplier; }
 
-	// np2.Resim.AngVelStabilityMultiplier
 	UPROPERTY(config, EditDefaultsOnly, Category = "Overrides", Meta = (InlineEditConditionToggle))
 	uint32 bOverrideAngVelStabilityMultiplier : 1;
+	// Overrides CVar: np2.Resim.AngVelStabilityMultiplier -- Recommended range between 0.0-1.0. Lower value means more stable angular velocity corrections.
 	UPROPERTY(config, EditDefaultsOnly, Category = "Overrides", Meta = (EditCondition = "bOverrideAngVelStabilityMultiplier"))
 	float AngVelStabilityMultiplier = PhysicsReplicationCVars::ResimulationCVars::AngVelStabilityMultiplier;
 	float GetAngVelStabilityMultiplier() { return bOverrideAngVelStabilityMultiplier ? AngVelStabilityMultiplier : PhysicsReplicationCVars::ResimulationCVars::AngVelStabilityMultiplier; }
 
-	// Project Settings -> Physics -> Replication -> Physics Prediction -> Resimulation Error Threshold
-	UPROPERTY(config, EditDefaultsOnly, Category = "Overrides", Meta = (InlineEditConditionToggle))
-	uint32 bOverrideResimulationErrorThreshold : 1;
-	UPROPERTY(config, EditDefaultsOnly, Category = "Overrides", Meta = (EditCondition = "bOverrideResimulationErrorThreshold"))
-	uint32 ResimulationErrorThreshold = 10;
-	uint32 GetResimulationErrorThreshold(uint32 DefaultValue) { return bOverrideRotStabilityMultiplier ? ResimulationErrorThreshold : DefaultValue; }
-	
-	// np2.Resim.CompareStateToTriggerRewind
 	UPROPERTY(config, EditDefaultsOnly, Category = "Overrides", Meta = (InlineEditConditionToggle))
 	uint32 bOverrideCompareStateToTriggerRewind : 1;
+	// Overrides CVar: np2.Resim.CompareStateToTriggerRewind -- When true, cache local players custom state struct in rewind history and compare the predicted state with incoming server state to trigger resimulations if they differ, comparison done through FNetworkPhysicsData::CompareData.
 	UPROPERTY(config, EditDefaultsOnly, Category = "Overrides", Meta = (EditCondition = "bOverrideCompareStateToTriggerRewind"))
 	bool bCompareStateToTriggerRewind = false;
 	bool GetCompareStateToTriggerRewind(bool DefaultValue) { return bOverrideCompareStateToTriggerRewind ? bCompareStateToTriggerRewind : DefaultValue; }
 
-	// np2.Resim.CompareInputToTriggerRewind
 	UPROPERTY(config, EditDefaultsOnly, Category = "Overrides", Meta = (InlineEditConditionToggle))
 	uint32 bOverrideCompareInputToTriggerRewind : 1;
+	// Overrides CVar: np2.Resim.CompareInputToTriggerRewind -- When true, compare local players predicted inputs with incoming server inputs to trigger resimulations if they differ, comparison done through FNetworkPhysicsData::CompareData.
 	UPROPERTY(config, EditDefaultsOnly, Category = "Overrides", Meta = (EditCondition = "bOverrideCompareInputToTriggerRewind"))
 	bool bCompareInputToTriggerRewind = false;
 	bool GetCompareInputToTriggerRewind(bool DefaultValue) { return bOverrideCompareInputToTriggerRewind ? bCompareInputToTriggerRewind : DefaultValue; }
-
 };
 
 /*
@@ -286,6 +344,9 @@ public:
 	UPROPERTY(EditDefaultsOnly, Category = "Networked Physics Settings")
 	FNetworkPhysicsSettings GeneralSettings;
 	
+	UPROPERTY(EditDefaultsOnly, Category = "Networked Physics Settings")
+	FNetworkPhysicsSettingsDefaultReplication DefaultReplicationSettings;
+
 	UPROPERTY(EditDefaultsOnly, Category = "Networked Physics Settings")
 	FNetworkPhysicsSettingsPredictiveInterpolation PredictiveInterpolationSettings;
 	
@@ -302,6 +363,12 @@ public:
 
 private:
 	FNetworkPhysicsSettingsComponentAsync* NetworkPhysicsSettingsAsync;
+
+	// Game Thread map of settings component per actor
+	static TMap<AActor*, UNetworkPhysicsSettingsComponent*> ObjectToSettings_External;
+
+public:
+	static UNetworkPhysicsSettingsComponent* GetSettingsForActor(AActor* Owner);
 };
 
 
@@ -314,6 +381,7 @@ private:
 struct FNetworkPhysicsSettingsAsync
 {
 	FNetworkPhysicsSettings GeneralSettings;
+	FNetworkPhysicsSettingsDefaultReplication DefaultReplicationSettings;
 	FNetworkPhysicsSettingsResimulation ResimulationSettings;
 	FNetworkPhysicsSettingsPredictiveInterpolation PredictiveInterpolationSettings;
 };
diff --git a/Engine/Source/Runtime/Engine/Public/PhysicsReplication.h b/Engine/Source/Runtime/Engine/Public/PhysicsReplication.h
--- a/Engine/Source/Runtime/Engine/Public/PhysicsReplication.h
+++ b/Engine/Source/Runtime/Engine/Public/PhysicsReplication.h
@@ -332,6 +332,8 @@ private:
 	TMap<TWeakObjectPtr<UPrimitiveComponent>, FReplicatedPhysicsTarget> ComponentToTargets_DEPRECATED; // This collection is keeping the legacy flow working until fully deprecated in a future release
 	TArray<FReplicatedPhysicsTarget> ReplicatedTargetsQueue;
 	FPhysScene* PhysScene;
+	TWeakObjectPtr<UNetworkPhysicsSettingsComponent> SettingsCurrent;
+
 
 	FPhysicsReplicationAsync* PhysicsReplicationAsync;
 	FPhysicsReplicationAsyncInput* AsyncInput;	//async data being written into before we push into callback
