CustomRenderPass: Added option to render scene color with/without alpha, in addition to the existing scene color with depth render mode.

Having access to alpha can be useful for compositing applications.

diff --git a/Engine/Source/Runtime/Engine/Private/Rendering/CustomRenderPass.cpp b/Engine/Source/Runtime/Engine/Private/Rendering/CustomRenderPass.cpp
--- a/Engine/Source/Runtime/Engine/Private/Rendering/CustomRenderPass.cpp
+++ b/Engine/Source/Runtime/Engine/Private/Rendering/CustomRenderPass.cpp
@@ -63,6 +63,10 @@ ESceneCaptureSource FCustomRenderPassBase::GetSceneCaptureSource() const
 		return SCS_DeviceDepth;
 	else if (RenderOutput == ERenderOutput::SceneColorAndDepth)
 		return SCS_SceneColorSceneDepth;
+	else if (RenderOutput == ERenderOutput::SceneColorAndAlpha)
+		return SCS_SceneColorHDR;
+	else if (RenderOutput == ERenderOutput::SceneColorNoAlpha)
+		return SCS_SceneColorHDRNoAlpha;
 	else
 		return SCS_MAX;
 }
diff --git a/Engine/Source/Runtime/Engine/Public/Rendering/CustomRenderPass.h b/Engine/Source/Runtime/Engine/Public/Rendering/CustomRenderPass.h
--- a/Engine/Source/Runtime/Engine/Public/Rendering/CustomRenderPass.h
+++ b/Engine/Source/Runtime/Engine/Public/Rendering/CustomRenderPass.h
@@ -67,7 +67,9 @@ public:
 		SceneDepth,
 		DeviceDepth,
 		/** Used with ERenderMode::DepthAndBasePass. */
-		SceneColorAndDepth
+		SceneColorAndDepth,
+		SceneColorAndAlpha, // The alpha is inverted opacity (throughput).
+		SceneColorNoAlpha,
 	};
 
 	enum class ERenderCaptureType
