- Added world delegate OnWorldPreBeginPlay triggered before calling DispatchBeginPlay on all world actors (when world's bBegunPlay switches to true).
- World Partition now listens to this delegate and calls BlockTillLevelStreamingCompleted to make sure all requested levels are streamed in before DispatchBeginPlay is called on all actors.
- One example this change fixes is a line trace done on the server (with server streaming disabled) in the BeginPlay of an always loaded actor not colliding with another actor flagged as spatially loaded.
- To limit the impact of this change, the BlockTillLevelStreamingCompleted is only called for partitioned worlds.

diff --git a/Engine/Source/Runtime/Engine/Classes/Engine/World.h b/Engine/Source/Runtime/Engine/Classes/Engine/World.h
--- a/Engine/Source/Runtime/Engine/Classes/Engine/World.h
+++ b/Engine/Source/Runtime/Engine/Classes/Engine/World.h
@@ -2497,6 +2497,10 @@ public:
 	DECLARE_MULTICAST_DELEGATE(FOnMatchStarting);
 	FOnMatchStarting OnWorldMatchStarting;
 
+	/** Called by WorldSettings's NotifyBeginPlay before calling DispatchBeginPlay on all world actors when world's bBegunPlay switches to true. */
+	DECLARE_MULTICAST_DELEGATE(FOnWorldPreBeginPlay);
+	FOnWorldPreBeginPlay OnWorldPreBeginPlay;
+
 	/** Returns true if gameplay has already started, false otherwise. */
 	bool HasBegunPlay() const;
 
diff --git a/Engine/Source/Runtime/Engine/Private/WorldPartition/WorldPartition.cpp b/Engine/Source/Runtime/Engine/Private/WorldPartition/WorldPartition.cpp
--- a/Engine/Source/Runtime/Engine/Private/WorldPartition/WorldPartition.cpp
+++ b/Engine/Source/Runtime/Engine/Private/WorldPartition/WorldPartition.cpp
@@ -932,6 +932,7 @@ void UWorldPartition::RegisterDelegates()
 	{
 		if (IsMainWorldPartition())
 		{
+			World->OnWorldPreBeginPlay.AddUObject(this, &UWorldPartition::OnWorldPreBeginPlay);
 			World->OnWorldMatchStarting.AddUObject(this, &UWorldPartition::OnWorldMatchStarting);
 
 #if !UE_BUILD_SHIPPING
@@ -977,6 +978,7 @@ void UWorldPartition::UnregisterDelegates()
 	{
 		if (IsMainWorldPartition())
 		{
+			World->OnWorldPreBeginPlay.RemoveAll(this);
 			World->OnWorldMatchStarting.RemoveAll(this);
 
 #if !UE_BUILD_SHIPPING
@@ -996,10 +998,19 @@ void UWorldPartition::GetOnScreenMessages(FCoreDelegates::FSeverityMessageMap& O
 }
 #endif
 
+void UWorldPartition::OnWorldPreBeginPlay()
+{
+	check(GetWorld()->IsGameWorld());
+	// Wait for any level streaming to complete before DispatchBeginPlay is called on all world actors 
+	// (when the world's bBegunPlay switches to true).
+	GetWorld()->BlockTillLevelStreamingCompleted();
+}
+
 void UWorldPartition::OnWorldMatchStarting()
 {
 	check(GetWorld()->IsGameWorld());
-	// Wait for any level streaming to complete
+	// Wait for any level streaming to complete 
+	// (in case any level streaming was requested by actor's DispatchBeginPlay)
 	GetWorld()->BlockTillLevelStreamingCompleted();
 }
 
diff --git a/Engine/Source/Runtime/Engine/Private/WorldSettings.cpp b/Engine/Source/Runtime/Engine/Private/WorldSettings.cpp
--- a/Engine/Source/Runtime/Engine/Private/WorldSettings.cpp
+++ b/Engine/Source/Runtime/Engine/Private/WorldSettings.cpp
@@ -328,6 +328,8 @@ void AWorldSettings::NotifyBeginPlay()
 	UWorld* World = GetWorld();
 	if (!World->GetBegunPlay())
 	{
+		World->OnWorldPreBeginPlay.Broadcast();
+
 		for (FActorIterator It(World); It; ++It)
 		{
 			SCOPE_CYCLE_COUNTER(STAT_ActorBeginPlay);
diff --git a/Engine/Source/Runtime/Engine/Public/WorldPartition/WorldPartition.h b/Engine/Source/Runtime/Engine/Public/WorldPartition/WorldPartition.h
--- a/Engine/Source/Runtime/Engine/Public/WorldPartition/WorldPartition.h
+++ b/Engine/Source/Runtime/Engine/Public/WorldPartition/WorldPartition.h
@@ -670,6 +670,7 @@ private:
 	static ENGINE_API FAutoConsoleVariableRef CVarUseMakingInvisibleTransactionRequests;
 
 	ENGINE_API void OnWorldMatchStarting();
+	ENGINE_API void OnWorldPreBeginPlay();
 	
 #if WITH_EDITOR
 	ENGINE_API void OnLevelActorDeleted(AActor* Actor);
