[Navmesh] Fix for brush navmesh generation issues. The Brush's navmesh generation was broken on property changes and transforming in the viewport in certain settings.

- When the Brush is treated as BSP, its navmesh update is not through component registration but Level's BSP Rebuild. This caused navmesh data out-of-sync when we changes a property that can impact its navmesh generation(i.e. CanEverAffectNavigation and CollisionSettings). Now we force a BSP Rebuild when registering the Component if applicable.

-  Filtered out Nav-Irrelevant Brush's Geometry data when exporting BSP triangles.

diff --git a/Engine/Source/Editor/UnrealEd/Private/EditorExporters.cpp b/Engine/Source/Editor/UnrealEd/Private/EditorExporters.cpp
--- a/Engine/Source/Editor/UnrealEd/Private/EditorExporters.cpp
+++ b/Engine/Source/Editor/UnrealEd/Private/EditorExporters.cpp
@@ -2321,6 +2321,12 @@ void UEditorEngine::RebuildStaticNavigableGeometry(ULevel* Level)
 			{
 				FBspNode* Node = &Model->Nodes[i];
 				FBspSurf& Surf = Model->Surfs[Node->iSurf];
+				
+				const bool bIsNavRelevant = Surf.Actor && Surf.Actor->GetBrushComponent() && Surf.Actor->GetBrushComponent()->IsNavigationRelevant();
+				if (!bIsNavRelevant)
+				{
+					continue;
+				}
 
 				const FVector& TextureBase = (FVector)Model->Points[Surf.pBase];
 				const FVector& TextureX = (FVector)Model->Vectors[Surf.vTextureU];
diff --git a/Engine/Source/Runtime/Engine/Classes/Components/BrushComponent.h b/Engine/Source/Runtime/Engine/Classes/Components/BrushComponent.h
--- a/Engine/Source/Runtime/Engine/Classes/Components/BrushComponent.h
+++ b/Engine/Source/Runtime/Engine/Classes/Components/BrushComponent.h
@@ -39,6 +39,12 @@ class UBrushComponent : public UPrimitiveComponent
 	virtual void PostLoad() override;
 	//~ End UObject Interface
 
+#if WITH_EDITOR
+	//~Begin UActorComponent Interface
+	ENGINE_API virtual void OnRegister() override;
+	//~End UActorComponent Interface
+#endif //WITH_EDITOR
+
 	//~ Begin USceneComponent Interface
 	virtual FBoxSphereBounds CalcBounds(const FTransform& LocalToWorld) const override;
 	virtual bool ShouldCollideWhenPlacing() const override { return true; }
diff --git a/Engine/Source/Runtime/Engine/Classes/Engine/Brush.h b/Engine/Source/Runtime/Engine/Classes/Engine/Brush.h
--- a/Engine/Source/Runtime/Engine/Classes/Engine/Brush.h
+++ b/Engine/Source/Runtime/Engine/Classes/Engine/Brush.h
@@ -190,7 +190,7 @@ public:
 	ENGINE_API virtual void CopyPosRotScaleFrom( ABrush* Other );
 
 	static void SetSuppressBSPRegeneration(bool bSuppress) { bSuppressBSPRegeneration = bSuppress; }
-
+	static bool GetSuppressBSPRegeneration() { return bSuppressBSPRegeneration; }
 private:
 
 	/** An array to keep track of all the levels that need rebuilding. This is checked via NeedsRebuild() in the editor tick and triggers a csg rebuild. */
diff --git a/Engine/Source/Runtime/Engine/Private/Components/BrushComponent.cpp b/Engine/Source/Runtime/Engine/Private/Components/BrushComponent.cpp
--- a/Engine/Source/Runtime/Engine/Private/Components/BrushComponent.cpp
+++ b/Engine/Source/Runtime/Engine/Private/Components/BrushComponent.cpp
@@ -25,6 +25,10 @@
 #include "PhysicsEngine/BodySetup.h"
 #include "DynamicMeshBuilder.h"
 
+#if WITH_EDITOR
+#include "Editor.h"
+#endif //WITH_EDITOR
+
 #include UE_INLINE_GENERATED_CPP_BY_NAME(BrushComponent)
 
 DEFINE_LOG_CATEGORY_STATIC(LogBrushComponent, Log, All);
@@ -514,6 +518,28 @@ void UBrushComponent::PostLoad()
 #endif
 }
 
+#if WITH_EDITOR
+void UBrushComponent::OnRegister()
+{
+	Super::OnRegister();
+
+	// BSP can only be rebuilt during a transaction
+	if (!ABrush::GetSuppressBSPRegeneration() && GEditor && GUndo)
+	{
+		if (ABrush* BrushOwner = Cast<ABrush>(GetOwner()))
+		{
+			const bool bIsBrushOwnerBSP = !BrushOwner->IsVolumeBrush() && BrushOwner->IsStaticBrush();
+			// Not checking NavRelavency is intentional to cover cases when it changes from being Nav Relevant to Nav Irrelevant,
+			// we need to clear the previously registered BSP Navigation Data.
+			if (bIsBrushOwnerBSP)
+			{
+				GEditor->RebuildAlteredBSP();
+			}
+		}
+	}
+}
+#endif //WITH_EDITOR
+
 ESceneDepthPriorityGroup UBrushComponent::GetStaticDepthPriorityGroup() const
 {
 	ABrush* BrushOwner = Cast<ABrush>(GetOwner());
