Vectorized some calculations in FRenderAssetInstanceState::ConditionalUpdateBounds and removed redundant struct copy to a stack

diff --git a/Engine/Source/Runtime/Engine/Private/Streaming/TextureInstanceState.cpp b/Engine/Source/Runtime/Engine/Private/Streaming/TextureInstanceState.cpp
--- a/Engine/Source/Runtime/Engine/Private/Streaming/TextureInstanceState.cpp
+++ b/Engine/Source/Runtime/Engine/Private/Streaming/TextureInstanceState.cpp
@@ -550,15 +550,13 @@ void FRenderAssetInstanceState::UpdateBounds(const UPrimitiveComponent* Componen
 	int32* ComponentLink = ComponentMap.Find(Component);
 	if (ComponentLink)
 	{
-		FBoxSphereBounds Bounds = Component->Bounds;
-		Bounds.SphereRadius = Component->GetStreamingScale();
 		int32 ElementIndex = *ComponentLink;
 		while (ElementIndex != INDEX_NONE)
 		{
 			const FElement& Element = Elements[ElementIndex];
 			if (Element.BoundsIndex != INDEX_NONE)
 			{
-				Bounds4[Element.BoundsIndex / 4].FullUpdate(Element.BoundsIndex % 4, Bounds, Component->GetLastRenderTimeOnScreen());
+				Bounds4[Element.BoundsIndex / 4].FullUpdate(Element.BoundsIndex % 4, Component->Bounds.Origin, Component->Bounds.BoxExtent, Component->GetStreamingScale(), Component->GetLastRenderTimeOnScreen());
 			}
 			ElementIndex = Element.NextComponentLink;
 		}
@@ -572,9 +570,7 @@ bool FRenderAssetInstanceState::UpdateBounds(int32 BoundIndex)
 	const UPrimitiveComponent* Component = ensure(Bounds4Components.IsValidIndex(BoundIndex)) ? Bounds4Components[BoundIndex] : nullptr;
 	if (Component)
 	{
-		FBoxSphereBounds Bounds = Component->Bounds;
-		Bounds.SphereRadius = Component->GetStreamingScale();
-		Bounds4[BoundIndex / 4].FullUpdate(BoundIndex % 4, Bounds, Component->GetLastRenderTimeOnScreen());
+		Bounds4[BoundIndex / 4].FullUpdate(BoundIndex % 4, Component->Bounds.Origin, Component->Bounds.BoxExtent, Component->GetStreamingScale(), Component->GetLastRenderTimeOnScreen());
 		return true;
 	}
 	else
@@ -590,29 +586,34 @@ bool FRenderAssetInstanceState::ConditionalUpdateBounds(int32 BoundIndex)
 	const UPrimitiveComponent* Component = ensure(Bounds4Components.IsValidIndex(BoundIndex)) ? Bounds4Components[BoundIndex] : nullptr;
 	if (Component)
 	{
+		const FBoxSphereBounds& Bounds = Component->Bounds;
+
 		if (Component->Mobility != EComponentMobility::Static)
 		{
-			FBoxSphereBounds Bounds = Component->Bounds;
-
 			// Check if the bound is coherent as it could be updated while we read it (from async task).
 			// We don't have to check the position, as if it was partially updated, this should be ok (interp)
-			const float RadiusSquared = FMath::Square<float>(Bounds.SphereRadius);
-			const float XSquared = FMath::Square<float>(Bounds.BoxExtent.X);
-			const float YSquared = FMath::Square<float>(Bounds.BoxExtent.Y);
-			const float ZSquared = FMath::Square<float>(Bounds.BoxExtent.Z);
+			static_assert((offsetof(FBoxSphereBounds, BoxExtent) + sizeof(Bounds.BoxExtent)) == (offsetof(FBoxSphereBounds, SphereRadius)), "Memory layout for FBoxSphereBounds has changed");
+			enum
+			{
+				X = 0,
+				Y,
+				Z,
+				R
+			};
+
+			VectorRegister4Float XYZRData = MakeVectorRegisterFloatFromDouble(VectorLoad(&Bounds.BoxExtent.X));		//X,Y,Z,Radius
+			XYZRData = VectorMultiply(XYZRData, XYZRData);
+			AlignedFloat4 XYZRSquared(XYZRData);
 
-			if (0.5f * FMath::Min3<float>(XSquared, YSquared, ZSquared) <= RadiusSquared && RadiusSquared <= 2.f * (XSquared + YSquared + ZSquared))
+			if (0.5f * FMath::Min3<float>(XYZRSquared[X], XYZRSquared[Y], XYZRSquared[Z]) <= XYZRSquared[R] && XYZRSquared[R] <= 2.f * (XYZRSquared[X] + XYZRSquared[Y] + XYZRSquared[Z]))
 			{
-				Bounds.SphereRadius = Component->GetStreamingScale();
-				Bounds4[BoundIndex / 4].FullUpdate(BoundIndex % 4, Bounds, Component->GetLastRenderTimeOnScreen());
+				Bounds4[BoundIndex / 4].FullUpdate(BoundIndex % 4, Bounds.Origin, Bounds.BoxExtent, Component->GetStreamingScale(), Component->GetLastRenderTimeOnScreen());
 				return true;
 			}
 		}
 		else // Otherwise we assume it is guarantied to be good.
 		{
-			FBoxSphereBounds Bounds = Component->Bounds;
-			Bounds.SphereRadius = Component->GetStreamingScale();
-			Bounds4[BoundIndex / 4].FullUpdate(BoundIndex % 4, Bounds, Component->GetLastRenderTimeOnScreen());
+			Bounds4[BoundIndex / 4].FullUpdate(BoundIndex % 4, Bounds.Origin, Bounds.BoxExtent, Component->GetStreamingScale(), Component->GetLastRenderTimeOnScreen());
 			return true;
 		}
 	}
diff --git a/Engine/Source/Runtime/Engine/Private/Streaming/TextureInstanceView.cpp b/Engine/Source/Runtime/Engine/Private/Streaming/TextureInstanceView.cpp
--- a/Engine/Source/Runtime/Engine/Private/Streaming/TextureInstanceView.cpp
+++ b/Engine/Source/Runtime/Engine/Private/Streaming/TextureInstanceView.cpp
@@ -63,20 +63,20 @@ void FRenderAssetInstanceView::FBounds4::UnpackBounds(int32 Index, const UPrimit
 }
 
 /** Dynamic Path, this needs to reset all members since the dynamic data is rebuilt from scratch every update (the previous data is given to the async task) */
-void FRenderAssetInstanceView::FBounds4::FullUpdate(int32 Index, const FBoxSphereBounds& Bounds, float InLastRenderTime)
+void FRenderAssetInstanceView::FBounds4::FullUpdate(int32 Index, const FVector& NewOrigin, const FVector& NewBoxExtent, float NewSphereRadius, float InLastRenderTime)
 {
 	check(Index >= 0 && Index < 4);
 
-	OriginX.Component(Index) = Bounds.Origin.X;
-	OriginY.Component(Index) = Bounds.Origin.Y;
-	OriginZ.Component(Index) = Bounds.Origin.Z;
-	RangeOriginX.Component(Index) = Bounds.Origin.X;
-	RangeOriginY.Component(Index) = Bounds.Origin.Y;
-	RangeOriginZ.Component(Index) = Bounds.Origin.Z;
-	ExtentX.Component(Index) = Bounds.BoxExtent.X;
-	ExtentY.Component(Index) = Bounds.BoxExtent.Y;
-	ExtentZ.Component(Index) = Bounds.BoxExtent.Z;
-	RadiusOrComponentScale.Component(Index) = Bounds.SphereRadius;
+	OriginX.Component(Index) = NewOrigin.X;
+	OriginY.Component(Index) = NewOrigin.Y;
+	OriginZ.Component(Index) = NewOrigin.Z;
+	RangeOriginX.Component(Index) = NewOrigin.X;
+	RangeOriginY.Component(Index) = NewOrigin.Y;
+	RangeOriginZ.Component(Index) = NewOrigin.Z;
+	ExtentX.Component(Index) = NewBoxExtent.X;
+	ExtentY.Component(Index) = NewBoxExtent.Y;
+	ExtentZ.Component(Index) = NewBoxExtent.Z;
+	RadiusOrComponentScale.Component(Index) = NewSphereRadius;
 	PackedRelativeBox[Index] = PackedRelativeBox_Identity;
 	MinDistanceSq.Component(Index) = 0;
 	MinRangeSq.Component(Index) = 0;
diff --git a/Engine/Source/Runtime/Engine/Private/Streaming/TextureInstanceView.h b/Engine/Source/Runtime/Engine/Private/Streaming/TextureInstanceView.h
--- a/Engine/Source/Runtime/Engine/Private/Streaming/TextureInstanceView.h
+++ b/Engine/Source/Runtime/Engine/Private/Streaming/TextureInstanceView.h
@@ -72,7 +72,7 @@ public:
 
 		void Set(int32 Index, const FBoxSphereBounds& Bounds, uint32 InPackedRelativeBox, float LastRenderTime, const FVector& RangeOrigin, float MinDistanceSq, float MinRangeSq, float MaxRangeSq);
 		void UnpackBounds(int32 Index, const UPrimitiveComponent* Component);
-		void FullUpdate(int32 Index, const FBoxSphereBounds& Bounds, float LastRenderTime);
+		void FullUpdate(int32 Index, const FVector& NewOrigin, const FVector& NewBoxExtent, float NewSphereRadius, float LastRenderTime);
 		FORCEINLINE void UpdateLastRenderTime(int32 Index, float LastRenderTime);
 		FORCEINLINE void UpdateMaxDrawDistanceSquared(int32 Index, float InMaxRangeSq);
 
