Added "Cameras" ShowFlag

diff --git a/Engine/Source/Editor/EditorStyle/Private/SlateEditorStyle.cpp b/Engine/Source/Editor/EditorStyle/Private/SlateEditorStyle.cpp
--- a/Engine/Source/Editor/EditorStyle/Private/SlateEditorStyle.cpp
+++ b/Engine/Source/Editor/EditorStyle/Private/SlateEditorStyle.cpp
@@ -5181,6 +5181,7 @@ void FSlateEditorStyle::FStyle::SetupLevelEditorStyle()
 		Set( "ShowFlagsMenu.StaticMeshes", new IMAGE_BRUSH( "Icons/icon_ShowStaticMeshes_16x", Icon16x16 ) );
 		Set( "ShowFlagsMenu.Translucency", new IMAGE_BRUSH( "Icons/icon_ShowTranslucency_16x", Icon16x16 ) );
 		Set( "ShowFlagsMenu.WidgetComponents", new IMAGE_BRUSH( "UMG/Designer_16x", Icon16x16 ) );
+		Set( "ShowFlagsMenu.Cameras", new IMAGE_BRUSH( "Icons/AssetIcons/CameraActor_16x", Icon16x16 ) );
 	}
 #endif // WITH_EDITOR || (IS_PROGRAM && WITH_UNREAL_DEVELOPER_TOOLS)
 
diff --git a/Engine/Source/Editor/EditorStyle/Private/StarshipStyle.cpp b/Engine/Source/Editor/EditorStyle/Private/StarshipStyle.cpp
--- a/Engine/Source/Editor/EditorStyle/Private/StarshipStyle.cpp
+++ b/Engine/Source/Editor/EditorStyle/Private/StarshipStyle.cpp
@@ -4757,6 +4757,7 @@ void FStarshipEditorStyle::FStyle::SetupLevelEditorStyle()
 		Set( "ShowFlagsMenu.StaticMeshes", new IMAGE_BRUSH_SVG( "Starship/Common/StaticMesh", Icon16x16 ) );
 		Set( "ShowFlagsMenu.Translucency", new IMAGE_BRUSH_SVG( "Starship/Common/Transparency", Icon16x16 ) );
 		Set( "ShowFlagsMenu.WidgetComponents", new IMAGE_BRUSH_SVG( "Starship/Common/WidgetComponents", Icon16x16 ) );
+		Set( "ShowFlagsMenu.Cameras", new IMAGE_BRUSH_SVG( "Starship/AssetIcons/Camera_16", Icon16x16 ) );
 
 		Set("ShowFlagsMenu.SubMenu.PostProcessing", new IMAGE_BRUSH_SVG("Starship/Common/PostProcessing", Icon16x16));
 		Set("ShowFlagsMenu.SubMenu.LightTypes", new IMAGE_BRUSH_SVG("Starship/Common/LightTypes", Icon16x16));
diff --git a/Engine/Source/Editor/UnrealEd/Private/EditorShowFlags.cpp b/Engine/Source/Editor/UnrealEd/Private/EditorShowFlags.cpp
--- a/Engine/Source/Editor/UnrealEd/Private/EditorShowFlags.cpp
+++ b/Engine/Source/Editor/UnrealEd/Private/EditorShowFlags.cpp
@@ -61,6 +61,7 @@ TArray<FShowFlagData>& GetShowFlagMenuItems()
 			EngineFlagsChords.Add("StaticMeshes", FInputChord() );
 			EngineFlagsChords.Add("Landscape", FInputChord(EKeys::L, EModifierKey::Alt) );
 			EngineFlagsChords.Add("Volumes", FInputChord(EKeys::O, EModifierKey::Alt) );
+			EngineFlagsChords.Add("Cameras", FInputChord() );
 
 			struct FIterSink
 			{
diff --git a/Engine/Source/Runtime/Engine/Classes/Camera/CameraComponent.h b/Engine/Source/Runtime/Engine/Classes/Camera/CameraComponent.h
--- a/Engine/Source/Runtime/Engine/Classes/Camera/CameraComponent.h
+++ b/Engine/Source/Runtime/Engine/Classes/Camera/CameraComponent.h
@@ -6,6 +6,7 @@
 #include "UObject/ObjectMacros.h"
 #include "UObject/ScriptInterface.h"
 #include "Components/SceneComponent.h"
+#include "Components/StaticMeshComponent.h"
 #include "Engine/BlendableInterface.h"
 #include "Engine/Scene.h"
 #include "Camera/CameraTypes.h"
@@ -13,6 +14,16 @@
 
 class UStaticMesh;
 
+/** Allows a custom SceneProxy to interact with camera specific show flags. */
+UCLASS()
+class UCameraProxyMeshComponent : public UStaticMeshComponent
+{
+	GENERATED_BODY()
+
+public:
+	virtual FPrimitiveSceneProxy* CreateSceneProxy() override;
+};
+
 /**
   * Represents a camera viewpoint and settings, such as projection type, field of view, and post-process overrides.
   * The default behavior for an actor used as the camera view target is to look for an attached camera component and use its location, rotation, and settings.
diff --git a/Engine/Source/Runtime/Engine/Private/Camera/CameraComponent.cpp b/Engine/Source/Runtime/Engine/Private/Camera/CameraComponent.cpp
--- a/Engine/Source/Runtime/Engine/Private/Camera/CameraComponent.cpp
+++ b/Engine/Source/Runtime/Engine/Private/Camera/CameraComponent.cpp
@@ -19,6 +19,7 @@
 #include "IXRTrackingSystem.h"
 #include "IXRCamera.h"
 #include "Math/UnitConversion.h"
+#include "StaticMeshSceneProxy.h"
 #include "UObject/FortniteMainBranchObjectVersion.h"
 #include "UObject/UE5ReleaseStreamObjectVersion.h"
 #include "UObject/UnrealType.h"
@@ -27,6 +28,33 @@
 
 #define LOCTEXT_NAMESPACE "CameraComponent"
 
+FPrimitiveSceneProxy* UCameraProxyMeshComponent::CreateSceneProxy()
+{
+	class FCameraProxyMeshProxy : public FStaticMeshSceneProxy
+	{
+	public:
+		FCameraProxyMeshProxy(UCameraProxyMeshComponent* Component)
+			: FStaticMeshSceneProxy(Component, false)
+		{
+		}
+
+		virtual FPrimitiveViewRelevance GetViewRelevance(const FSceneView* View) const override
+		{
+			if (View->Family->EngineShowFlags.Cameras)
+			{
+				return FStaticMeshSceneProxy::GetViewRelevance(View);
+			}
+
+			FPrimitiveViewRelevance Result;
+			Result.bDrawRelevance = false;
+
+			return Result;
+		}
+	};
+
+	return new FCameraProxyMeshProxy(this);
+}
+
 //////////////////////////////////////////////////////////////////////////
 // UCameraComponent
 
@@ -129,7 +157,7 @@ void UCameraComponent::OnRegister()
 	{
 		if (ProxyMeshComponent == nullptr)
 		{
-			ProxyMeshComponent = NewObject<UStaticMeshComponent>(MyOwner, NAME_None, RF_Transactional | RF_TextExportTransient);
+			ProxyMeshComponent = NewObject<UCameraProxyMeshComponent>(MyOwner, NAME_None, RF_Transactional | RF_TextExportTransient);
 			ProxyMeshComponent->SetupAttachment(this);
 			ProxyMeshComponent->SetIsVisualizationComponent(true);
 			ProxyMeshComponent->SetStaticMesh(CameraMesh);
diff --git a/Engine/Source/Runtime/Engine/Public/ShowFlagsValues.inl b/Engine/Source/Runtime/Engine/Public/ShowFlagsValues.inl
--- a/Engine/Source/Runtime/Engine/Public/ShowFlagsValues.inl
+++ b/Engine/Source/Runtime/Engine/Public/ShowFlagsValues.inl
@@ -452,5 +452,8 @@ SHOWFLAG_ALWAYS_ACCESSIBLE(ShaderPrint, SFG_Developer, NSLOCTEXT("UnrealEd", "Sh
 /** If true, scene depth is copied to the scene capture output depth texture. If false, other depth sources like editor primitives depth are still copied. */
 SHOWFLAG_ALWAYS_ACCESSIBLE(SceneCaptureCopySceneDepth, SFG_Hidden, NSLOCTEXT("UnrealEd", "SceneCaptureCopySceneDepth", "Copy Scene Depth To Scene Capture"))
 
+/** If false, camera's are not shown in the viewport (but are enabled and available for use). */
+SHOWFLAG_ALWAYS_ACCESSIBLE(Cameras, SFG_Normal, NSLOCTEXT("UnrealEd", "CamerasSF", "Cameras"))
+
 #undef SHOWFLAG_ALWAYS_ACCESSIBLE
 #undef SHOWFLAG_FIXED_IN_SHIPPING
