Added MinimalAPI specifier to UCameraProxyMeshComponent

diff --git a/Engine/Source/Runtime/Engine/Classes/Camera/CameraComponent.h b/Engine/Source/Runtime/Engine/Classes/Camera/CameraComponent.h
--- a/Engine/Source/Runtime/Engine/Classes/Camera/CameraComponent.h
+++ b/Engine/Source/Runtime/Engine/Classes/Camera/CameraComponent.h
@@ -15,7 +15,7 @@
 class UStaticMesh;
 
 /** Allows a custom SceneProxy to interact with camera specific show flags. */
-UCLASS()
+UCLASS(MinimalAPI)
 class UCameraProxyMeshComponent : public UStaticMeshComponent
 {
 	GENERATED_BODY()
