Fix native bitfield-bool properties writing to incorrect memory when exposed as pins on linked anim graphs

diff --git a/Engine/Source/Runtime/Engine/Private/Animation/AnimNode_CustomProperty.cpp b/Engine/Source/Runtime/Engine/Private/Animation/AnimNode_CustomProperty.cpp
--- a/Engine/Source/Runtime/Engine/Private/Animation/AnimNode_CustomProperty.cpp
+++ b/Engine/Source/Runtime/Engine/Private/Animation/AnimNode_CustomProperty.cpp
@@ -44,7 +44,20 @@ void FAnimNode_CustomProperty::PropagateInputProperties(const UObject* InSourceI
 					const uint8* SrcPtr = CallerProperty->ContainerPtrToValuePtr<uint8>(InSourceInstance);
 					uint8* DestPtr = SubProperty->ContainerPtrToValuePtr<uint8>(TargetInstance);
 
-					CallerProperty->CopyCompleteValue(DestPtr, SrcPtr);
+					FBoolProperty* CallerBoolProperty = CastField<FBoolProperty>(CallerProperty);
+					FBoolProperty* SubBoolProperty = CastField<FBoolProperty>(SubProperty);
+					if (CallerBoolProperty && SubBoolProperty)
+					{
+						// Bools may be native bitfields, in which case we cant call CopyCompleteValue as 
+						// there is insufficent information to correctly copy the relevant bits. Calling CopyCompleteValue
+						// in this case can potentially overwrite adjacent memory.
+						const bool bValue = CallerBoolProperty->GetPropertyValue_InContainer(InSourceInstance);
+						SubBoolProperty->SetPropertyValue_InContainer(TargetInstance, bValue);
+					}
+					else
+					{
+						CallerProperty->CopyCompleteValue(DestPtr, SrcPtr);
+					}
 				}
 			}
 		}
