Added pointer check to allow changing scale on body instances that aren't in a scene.

diff --git a/Engine/Source/Runtime/Engine/Private/PhysicsEngine/BodyInstance.cpp b/Engine/Source/Runtime/Engine/Private/PhysicsEngine/BodyInstance.cpp
--- a/Engine/Source/Runtime/Engine/Private/PhysicsEngine/BodyInstance.cpp
+++ b/Engine/Source/Runtime/Engine/Private/PhysicsEngine/BodyInstance.cpp
@@ -2365,8 +2365,10 @@ bool FBodyInstance::UpdateBodyScale(const FVector& InScale3D, bool bForceUpdate)
 	{
 		Scale3D = UpdatedScale3D;
 
-		FPhysScene_Chaos& Scene = *GetPhysicsScene();
-		Scene.UpdateActorInAccelerationStructure(ActorHandle);
+		if (FPhysScene_Chaos* Scene = GetPhysicsScene())
+		{
+			Scene->UpdateActorInAccelerationStructure(ActorHandle);
+		}
 
 		// update mass if required
 		if (bUpdateMassWhenScaleChanges)
