Fixed inverted condition in UMaterialInterface::SubmitRemainingJobsForWorld

diff --git a/Engine/Source/Runtime/Engine/Private/Materials/MaterialInterface.cpp b/Engine/Source/Runtime/Engine/Private/Materials/MaterialInterface.cpp
--- a/Engine/Source/Runtime/Engine/Private/Materials/MaterialInterface.cpp
+++ b/Engine/Source/Runtime/Engine/Private/Materials/MaterialInterface.cpp
@@ -630,7 +630,7 @@ void UMaterialInterface::SubmitRemainingJobsForWorld(UWorld* World, EMaterialSha
 
 	for (IPrimitiveComponent* PrimitiveComponentInterface : ObjectCacheScope.GetContext().GetPrimitiveComponents())
 	{
-		if (World && PrimitiveComponentInterface->GetWorld() == World)
+		if (World && PrimitiveComponentInterface->GetWorld() != World)
 		{
 			continue;
 		}
