Activate vectorized linear solver
when 3 linear lock, 3 rotations soft limits or 3 rotation motor are activated, then the vectorization will be applied.

diff --git a/Engine/Source/Runtime/Engine/Private/PhysicsEngine/ImmediatePhysics/ImmediatePhysicsChaos/ImmediatePhysicsSimulation_Chaos.cpp b/Engine/Source/Runtime/Engine/Private/PhysicsEngine/ImmediatePhysics/ImmediatePhysicsChaos/ImmediatePhysicsSimulation_Chaos.cpp
--- a/Engine/Source/Runtime/Engine/Private/PhysicsEngine/ImmediatePhysics/ImmediatePhysicsChaos/ImmediatePhysicsSimulation_Chaos.cpp
+++ b/Engine/Source/Runtime/Engine/Private/PhysicsEngine/ImmediatePhysics/ImmediatePhysicsChaos/ImmediatePhysicsSimulation_Chaos.cpp
@@ -156,7 +156,7 @@ FAutoConsoleVariableRef  CVarChaosImmPhysInertiaConditioningMaxInvInertiaCompone
 int32 ChaosImmediate_Joint_UseLinearSolver = -1;
 FAutoConsoleVariableRef CVarChaosImmPhysJointUseCachedSolver(TEXT("p.Chaos.ImmPhys.Joint.UseLinearSolver"), ChaosImmediate_Joint_UseLinearSolver, TEXT("Force use of linear or non-linear joint solver. (-1 to use PhysicsAsset setting)"));
 
-bool bChaosImmediate_UseSimdForLinearSolver = false;
+bool bChaosImmediate_UseSimdForLinearSolver = true;
 FAutoConsoleVariableRef  CVarChaosImmPhysUseSimdForLinearSolver(TEXT("p.Chaos.ImmPhys.UseSimd"), bChaosImmediate_UseSimdForLinearSolver, TEXT("Enable/Disable SIMD on the linear joint solver"));
 
 //
