Added a check to catch a mismatch between the mesh deformer being null at the time a skeletal mesh object is created and later valid when the dynamic data is updated.

diff --git a/Engine/Source/Runtime/Engine/Private/SkeletalRenderGPUSkin.cpp b/Engine/Source/Runtime/Engine/Private/SkeletalRenderGPUSkin.cpp
--- a/Engine/Source/Runtime/Engine/Private/SkeletalRenderGPUSkin.cpp
+++ b/Engine/Source/Runtime/Engine/Private/SkeletalRenderGPUSkin.cpp
@@ -2469,6 +2469,11 @@ void FDynamicSkelMeshObjectDataGPUSkin::InitDynamicSkelMeshObjectDataGPUSkin(
 
 	GPUSkinTechnique = InMeshObject->GetGPUSkinTechnique(LODIndex);
 
+	checkf(GPUSkinTechnique == ESkeletalMeshGPUSkinTechnique::MeshDeformer || !InMeshComponent->GetMeshDeformerInstanceForLOD(LODIndex),
+		TEXT("Skeletal mesh %s, LOD %d is not set to use the mesh deformer skin technique, but the component deformer instance is valid. ")
+		TEXT("This means a mesh deformer was added but the skeletal mesh object was not recreated."),
+		*InMeshComponent->GetName(), LODIndex);
+
 	if (!IsSkeletalMeshClothBlendEnabled())
 	{
 		ClothBlendWeight = 0.f;
