Flush denormals to zero for anim jobs

diff --git a/Engine/Source/Runtime/Engine/Private/Components/SkeletalMeshComponent.cpp b/Engine/Source/Runtime/Engine/Private/Components/SkeletalMeshComponent.cpp
--- a/Engine/Source/Runtime/Engine/Private/Components/SkeletalMeshComponent.cpp
+++ b/Engine/Source/Runtime/Engine/Private/Components/SkeletalMeshComponent.cpp
@@ -208,6 +208,8 @@ public:
 	{
 		if (USkeletalMeshComponent* Comp = SkeletalMeshComponent.Get())
 		{
+			FScopedFTZFloatMode FTZ;
+
 			FScopeCycleCounterUObject ContextScope(Comp);
 #if !UE_BUILD_TEST && !UE_BUILD_SHIPPING
 			float Stall = CVarStallParallelAnimation.GetValueOnAnyThread();
@@ -267,6 +269,8 @@ public:
 
 		if (USkeletalMeshComponent* Comp = SkeletalMeshComponent.Get())
 		{
+			FScopedFTZFloatMode FTZ;
+
 			FScopeCycleCounterUObject ComponentScope(Comp);
 			FScopeCycleCounterUObject MeshScope(Comp->GetSkeletalMeshAsset());
 
