Skeletal Mesh GPU Skin:  Speculative fix for crash in CreatePassthroughVertexFactory.  Merge lambdas that generate VertexFactories and PassthroughVertexFactories, so they run in the same task.  The latter is dependent on data from the former, and this makes it impossible for the logic to not run together and sequentially.  If this doesn't solve the bug, it at least further narrows it down.

#rnx

diff --git a/Engine/Source/Runtime/Engine/Private/SkeletalRenderGPUSkin.cpp b/Engine/Source/Runtime/Engine/Private/SkeletalRenderGPUSkin.cpp
--- a/Engine/Source/Runtime/Engine/Private/SkeletalRenderGPUSkin.cpp
+++ b/Engine/Source/Runtime/Engine/Private/SkeletalRenderGPUSkin.cpp
@@ -1800,13 +1800,45 @@ static FPSOPrecacheVertexFactoryData GetVertexFactoryData(FSkeletalMeshLODRender
 	return FPSOPrecacheVertexFactoryData(VertexFactoryType, VertexElements);
 }
 
+static FGPUSkinPassthroughVertexFactory* AllocatePassthroughVertexFactory(
+	TArray<TUniquePtr<FGPUSkinPassthroughVertexFactory>>* PassthroughVertexFactories,
+	ERHIFeatureLevel::Type InFeatureLevel,
+	FGPUSkinPassthroughVertexFactory::EVertexAttributeFlags VertexAttributeMask)
+{
+	FGPUSkinPassthroughVertexFactory* NewPassthroughVertexFactory = nullptr;
+
+	if (PassthroughVertexFactories)
+	{
+		NewPassthroughVertexFactory = new FGPUSkinPassthroughVertexFactory(InFeatureLevel, VertexAttributeMask);
+		PassthroughVertexFactories->Add(TUniquePtr<FGPUSkinPassthroughVertexFactory>(NewPassthroughVertexFactory));
+	}
+
+	return NewPassthroughVertexFactory;
+}
+
+static void InitPassthroughVertexFactory_RenderThread(
+	FGPUSkinPassthroughVertexFactory* NewPassthroughVertexFactory,
+	FGPUBaseSkinVertexFactory* SourceVertexFactory,
+	FRHICommandList& RHICmdList)
+{
+	if (NewPassthroughVertexFactory)
+	{
+		FLocalVertexFactory::FDataType Data;
+		SourceVertexFactory->CopyDataTypeForLocalVertexFactory(Data);
+		NewPassthroughVertexFactory->SetData(RHICmdList, Data);
+		NewPassthroughVertexFactory->InitResource(RHICmdList);
+	}
+}
+
 /**
  * Creates a vertex factory entry for the given type and initialize it on the render thread
  */
-static FGPUBaseSkinVertexFactory* CreateVertexFactory(
+static void CreateVertexFactory(
 	TArray<TUniquePtr<FGPUBaseSkinVertexFactory>>& VertexFactories,
+	TArray<TUniquePtr<FGPUSkinPassthroughVertexFactory>>* PassthroughVertexFactories,
 	const FSkeletalMeshObjectGPUSkin::FVertexFactoryBuffers& VertexBuffers,
 	ERHIFeatureLevel::Type FeatureLevel,
+	FGPUSkinPassthroughVertexFactory::EVertexAttributeFlags VertexAttributeMask,
 	uint32 BaseVertexIndex,
 	bool bUsedForPassthroughVertexFactory)
 {
@@ -1822,21 +1854,24 @@ static FGPUBaseSkinVertexFactory* CreateVertexFactory(
 	}
 	VertexFactories.Add(TUniquePtr<FGPUBaseSkinVertexFactory>(VertexFactory));
 
+	// Allocate optional passthrough vertex factory, if PassthroughVertexFactories is non-null
+	FGPUSkinPassthroughVertexFactory* NewPassthroughVertexFactory = AllocatePassthroughVertexFactory(PassthroughVertexFactories, FeatureLevel, VertexAttributeMask);
+
 	// Setup the update data for enqueue
 	FDynamicUpdateVertexFactoryData VertexUpdateData(VertexFactory, VertexBuffers);
 
 	// update vertex factory components and sync it
 	ENQUEUE_RENDER_COMMAND(InitGPUSkinVertexFactory)(UE::RenderCommandPipe::SkeletalMesh,
-		[VertexUpdateData](FRHICommandList& RHICmdList)
+		[VertexUpdateData, NewPassthroughVertexFactory](FRHICommandList& RHICmdList)
 		{
 			FGPUSkinDataType Data;
 			InitGPUSkinVertexFactoryComponents(&Data, VertexUpdateData.VertexBuffers, VertexUpdateData.VertexFactory);
 			VertexUpdateData.VertexFactory->SetData(RHICmdList, &Data);
 			VertexUpdateData.VertexFactory->InitResource(RHICmdList);
+
+			InitPassthroughVertexFactory_RenderThread(NewPassthroughVertexFactory, VertexUpdateData.VertexFactory, RHICmdList);
 		}
 	);
-
-	return VertexFactory;
 }
 
 void UpdateVertexFactory(
@@ -1867,27 +1902,6 @@ void UpdateVertexFactory(
 	}
 }
 
-static void CreatePassthroughVertexFactory(
-	ERHIFeatureLevel::Type InFeatureLevel,
-	TArray<TUniquePtr<FGPUSkinPassthroughVertexFactory>>& PassthroughVertexFactories,
-	FGPUBaseSkinVertexFactory* SourceVertexFactory,
-	FGPUSkinPassthroughVertexFactory::EVertexAttributeFlags VertexAttributeMask)
-{
-	FGPUSkinPassthroughVertexFactory* NewPassthroughVertexFactory = new FGPUSkinPassthroughVertexFactory(InFeatureLevel, VertexAttributeMask);
-	PassthroughVertexFactories.Add(TUniquePtr<FGPUSkinPassthroughVertexFactory>(NewPassthroughVertexFactory));
-
-	// update vertex factory components and sync it
-	ENQUEUE_RENDER_COMMAND(InitPassthroughGPUSkinVertexFactory)(UE::RenderCommandPipe::SkeletalMesh,
-		[NewPassthroughVertexFactory, SourceVertexFactory](FRHICommandList& RHICmdList)
-		{
-			FLocalVertexFactory::FDataType Data;
-			SourceVertexFactory->CopyDataTypeForLocalVertexFactory(Data);
-			NewPassthroughVertexFactory->SetData(RHICmdList, Data);
-			NewPassthroughVertexFactory->InitResource(RHICmdList);
-		}
-	);
-}
-
 // APEX cloth
 
 static FPSOPrecacheVertexFactoryData GetVertexFactoryDataCloth(FSkeletalMeshLODRenderData& LODRenderData, FGPUSkinDataType& GPUSkinDataType, ERHIFeatureLevel::Type FeatureLevel)
@@ -1910,10 +1924,12 @@ static FPSOPrecacheVertexFactoryData GetVertexFactoryDataCloth(FSkeletalMeshLODR
 /**
  * Creates a vertex factory entry for the given type and initialize it on the render thread
  */
-static FGPUBaseSkinVertexFactory* CreateVertexFactoryCloth(
+static void CreateVertexFactoryCloth(
 	TArray<TUniquePtr<FGPUBaseSkinVertexFactory>>& VertexFactories,
+	TArray<TUniquePtr<FGPUSkinPassthroughVertexFactory>>* PassthroughVertexFactories,
 	const FSkeletalMeshObjectGPUSkin::FVertexFactoryBuffers& VertexBuffers,
 	ERHIFeatureLevel::Type FeatureLevel,
+	FGPUSkinPassthroughVertexFactory::EVertexAttributeFlags VertexAttributeMask,
 	uint32 BaseVertexIndex,
 	uint32 NumInfluencesPerVertex,
 	bool bUsedForPassthroughVertexFactory)
@@ -1930,22 +1946,25 @@ static FGPUBaseSkinVertexFactory* CreateVertexFactoryCloth(
 	}
 	VertexFactories.Add(TUniquePtr<FGPUBaseSkinVertexFactory>(VertexFactory));
 
+	// Allocate optional passthrough vertex factory, if PassthroughVertexFactories is non-null
+	FGPUSkinPassthroughVertexFactory* NewPassthroughVertexFactory = AllocatePassthroughVertexFactory(PassthroughVertexFactories, FeatureLevel, VertexAttributeMask);
+
 	// Setup the update data for enqueue
 	FDynamicUpdateVertexFactoryData VertexUpdateData(VertexFactory, VertexBuffers);
 
 	// update vertex factory components and sync it
 	ENQUEUE_RENDER_COMMAND(InitGPUSkinAPEXClothVertexFactory)(UE::RenderCommandPipe::SkeletalMesh,
-		[VertexUpdateData](FRHICommandList& RHICmdList)
+		[VertexUpdateData, NewPassthroughVertexFactory](FRHICommandList& RHICmdList)
 		{
 			FGPUSkinAPEXClothDataType Data;
 			InitGPUSkinVertexFactoryComponents(&Data, VertexUpdateData.VertexBuffers, VertexUpdateData.VertexFactory);
 			InitAPEXClothVertexFactoryComponents(&Data, VertexUpdateData.VertexBuffers);
 			VertexUpdateData.VertexFactory->SetData(RHICmdList, &Data);
 			VertexUpdateData.VertexFactory->InitResource(RHICmdList);
+
+			InitPassthroughVertexFactory_RenderThread(NewPassthroughVertexFactory, VertexUpdateData.VertexFactory, RHICmdList);
 		}
 	);
-
-	return VertexFactory;
 }
 
 void FSkeletalMeshObjectGPUSkin::GetUsedVertexFactoryData(
@@ -2037,15 +2056,16 @@ void FSkeletalMeshObjectGPUSkin::FVertexFactoryData::InitVertexFactories(
 
 	VertexFactories.Empty(Sections.Num());
 
+	// Optionally create passthrough VFs
+	TArray<TUniquePtr<FGPUSkinPassthroughVertexFactory>>* Passthroughs = nullptr;
 	if (bCreatePassthroughVFs)
 	{
 		PassthroughVertexFactories.Empty(Sections.Num());
+		Passthroughs = &PassthroughVertexFactories;
 	}
 
 	for (const FSkelMeshRenderSection& Section : Sections)
 	{
-		FGPUBaseSkinVertexFactory* VertexFactory;
-
 		if (Section.HasClothingData() && bClothEnabled)
 		{
 			constexpr int32 ClothLODBias = 0;
@@ -2054,16 +2074,11 @@ void FSkeletalMeshObjectGPUSkin::FVertexFactoryData::InitVertexFactories(
 			// NumInfluencesPerVertex should be a whole integer
 			check(NumClothWeights % NumPositionVertices == 0);
 			const uint32 NumInfluencesPerVertex = NumClothWeights / NumPositionVertices;
-			VertexFactory = CreateVertexFactoryCloth(VertexFactories, VertexBuffers, InFeatureLevel, Section.BaseVertexIndex, NumInfluencesPerVertex, bUsedForPassthroughVertexFactory);
+			CreateVertexFactoryCloth(VertexFactories, Passthroughs, VertexBuffers, InFeatureLevel, VertexAttributeMask, Section.BaseVertexIndex, NumInfluencesPerVertex, bUsedForPassthroughVertexFactory);
 		}
 		else
 		{
-			VertexFactory = CreateVertexFactory(VertexFactories, VertexBuffers, InFeatureLevel, Section.BaseVertexIndex, bUsedForPassthroughVertexFactory);
-		}
-
-		if (bCreatePassthroughVFs)
-		{
-			CreatePassthroughVertexFactory(InFeatureLevel, PassthroughVertexFactories, VertexFactory, VertexAttributeMask);
+			CreateVertexFactory(VertexFactories, Passthroughs, VertexBuffers, InFeatureLevel, VertexAttributeMask, Section.BaseVertexIndex, bUsedForPassthroughVertexFactory);
 		}
 	}
 }
@@ -2555,4 +2570,4 @@ bool FDynamicSkelMeshObjectDataGPUSkin::UpdateClothSimulationData(USkinnedMeshCo
 		return true;
 	}
 	return false;
-}
\ No newline at end of file
+}
