Adds SetIsTemporarilyHiddenInEditor at UPrimitiveComponent level.

diff --git a/Engine/Source/Runtime/Engine/Classes/Components/PrimitiveComponent.h b/Engine/Source/Runtime/Engine/Classes/Components/PrimitiveComponent.h
--- a/Engine/Source/Runtime/Engine/Classes/Components/PrimitiveComponent.h
+++ b/Engine/Source/Runtime/Engine/Classes/Components/PrimitiveComponent.h
@@ -362,6 +362,12 @@ private:
 	UPROPERTY(EditAnywhere, BlueprintGetter = GetGenerateOverlapEvents, BlueprintSetter = SetGenerateOverlapEvents, Category = Collision)
 	uint8 bGenerateOverlapEvents : 1;
 
+#if WITH_EDITORONLY_DATA
+	/** Whether this component is temporarily hidden within the editor. */
+	UPROPERTY(Transient)
+	uint8 bHiddenEdTemporary:1;
+#endif // WITH_EDITORONLY_DATA
+
 public:
 	/**
 	 * If true, this component will generate individual overlaps for each overlapping physics body if it is a multi-body component. When false, this component will
@@ -2563,6 +2569,18 @@ public:
 	/** Returns mask that represents in which views this primitive is hidden */
 	ENGINE_API virtual uint64 GetHiddenEditorViews() const;
 
+	/**
+	 * Returns whether or not this component was explicitly hidden in the editor for the duration of the current editor session
+	 * @param bIncludeParent - Whether to recurse up actor hierarchy or not
+	 */
+	ENGINE_API bool IsTemporarilyHiddenInEditor(const bool bIncludeParent) const;
+
+	/**
+ 	 * Explicitly sets whether or not this component is hidden in the editor for the duration of the current editor session
+ 	 * @param bIsHidden	True if the component is hidden
+ 	 */
+	ENGINE_API void SetIsTemporarilyHiddenInEditor(const bool bInIsHidden);
+
 	/** Sets whether this component is being moved by the editor so the renderer can render velocities for it, even when Static. */
 	ENGINE_API void SetIsBeingMovedByEditor(bool bNewIsBeingMoved);
 
diff --git a/Engine/Source/Runtime/Engine/Private/Components/PrimitiveComponent.cpp b/Engine/Source/Runtime/Engine/Private/Components/PrimitiveComponent.cpp
--- a/Engine/Source/Runtime/Engine/Private/Components/PrimitiveComponent.cpp
+++ b/Engine/Source/Runtime/Engine/Private/Components/PrimitiveComponent.cpp
@@ -362,6 +362,9 @@ UPrimitiveComponent::UPrimitiveComponent(const FObjectInitializer& ObjectInitial
 	bIsBeingMovedByEditor = false;
 
 	SetGenerateOverlapEvents(true);
+#if WITH_EDITORONLY_DATA
+	bHiddenEdTemporary = false;
+#endif // WITH_EDITORONLY_DATA 
 	bMultiBodyOverlap = false;
 	bReturnMaterialOnMove = false;
 	bCanEverAffectNavigation = false;
@@ -1674,11 +1677,15 @@ bool UPrimitiveComponent::ShouldComponentAddToScene() const
 #if WITH_EDITOR
 	AActor* Owner = GetOwner();
 	const bool bIsHiddenInEditor = GIsEditor && Owner && Owner->IsHiddenEd();
+
+	constexpr bool bIncludeParent = false;
+	const bool bIsTemporarilyHiddenInEditor = IsTemporarilyHiddenInEditor(bIncludeParent);
 #else
 	const bool bIsHiddenInEditor = false;
+	const bool bIsTemporarilyHiddenInEditor = false;
 #endif
 
-	return bSceneAdd && (ShouldRender() || (bCastHiddenShadow && !bIsHiddenInEditor) || bAffectIndirectLightingWhileHidden || bRayTracingFarField);
+	return bSceneAdd && !bIsTemporarilyHiddenInEditor && (ShouldRender() || (bCastHiddenShadow && !bIsHiddenInEditor) || bAffectIndirectLightingWhileHidden || bRayTracingFarField);
 }
 
 bool UPrimitiveComponent::ShouldCreatePhysicsState() const
@@ -1953,6 +1960,30 @@ uint64 UPrimitiveComponent::GetHiddenEditorViews() const
 	return OwnerActor ? OwnerActor->HiddenEditorViews : 0;
 }
 
+bool UPrimitiveComponent::IsTemporarilyHiddenInEditor(const bool bIncludeParent) const
+{
+	if (bHiddenEdTemporary)
+	{
+		return true;
+	}
+
+	if (bIncludeParent)
+	{
+		return GetOwner() && GetOwner()->IsTemporarilyHiddenInEditor(true);
+	}
+
+	return false;
+}
+
+void UPrimitiveComponent::SetIsTemporarilyHiddenInEditor(const bool bInIsHidden)
+{
+	if (bHiddenEdTemporary != bInIsHidden)
+	{
+		bHiddenEdTemporary = bInIsHidden;
+		MarkRenderStateDirty();
+	}
+}
+
 void UPrimitiveComponent::SetIsBeingMovedByEditor(bool bIsBeingMoved)
 {
 	bIsBeingMovedByEditor = bIsBeingMoved;
