Use Manifold by default on new Physics Asset

diff --git a/Engine/Source/Runtime/Core/Public/UObject/FortniteMainBranchObjectVersion.h b/Engine/Source/Runtime/Core/Public/UObject/FortniteMainBranchObjectVersion.h
--- a/Engine/Source/Runtime/Core/Public/UObject/FortniteMainBranchObjectVersion.h
+++ b/Engine/Source/Runtime/Core/Public/UObject/FortniteMainBranchObjectVersion.h
@@ -441,6 +441,9 @@ struct FFortniteMainBranchObjectVersion
 		// Add support for anytype in dataflow
 		DataflowAnyTypeSupport,
 
+		// Adding a new flag in RBAN Solver Setting to be able to use manifolds
+		PhysicsAssetUseManifoldFlags,
+
 		// -----<new versions can be added above this line>-------------------------------------------------
 		VersionPlusOne,
 		LatestVersion = VersionPlusOne - 1
diff --git a/Engine/Source/Runtime/Engine/Private/PhysicsEngine/PhysicsAsset.cpp b/Engine/Source/Runtime/Engine/Private/PhysicsEngine/PhysicsAsset.cpp
--- a/Engine/Source/Runtime/Engine/Private/PhysicsEngine/PhysicsAsset.cpp
+++ b/Engine/Source/Runtime/Engine/Private/PhysicsEngine/PhysicsAsset.cpp
@@ -38,7 +38,7 @@ FPhysicsAssetSolverSettings::FPhysicsAssetSolverSettings()
 	, MaxDepenetrationVelocity(0.0f)
 	, FixedTimeStep(0.0f)
 	, bUseLinearJointSolver(true)
-	, bUseManifolds(false)
+	, bUseManifolds(true)
 {
 }
 
@@ -141,6 +141,7 @@ void UPhysicsAsset::DeclareConstructClasses(TArray<FTopLevelAssetPath>& OutConst
 void UPhysicsAsset::Serialize(FArchive& Ar)
 {
 	Ar.UsingCustomVersion(FFortniteSeasonBranchObjectVersion::GUID);
+	Ar.UsingCustomVersion(FFortniteMainBranchObjectVersion::GUID);
 
 	Super::Serialize(Ar);
 	Ar << CollisionDisableTable;
@@ -166,6 +167,12 @@ void UPhysicsAsset::Serialize(FArchive& Ar)
 		SolverSettings.CullDistance = 1.0f;
 	}
 
+	const bool bShouldNotUseManifoldOnOldAsset = (Ar.CustomVer(FFortniteMainBranchObjectVersion::GUID) < FFortniteMainBranchObjectVersion::PhysicsAssetUseManifoldFlags);
+	if (bShouldNotUseManifoldOnOldAsset)
+	{
+		SolverSettings.bUseManifolds = false;
+	}
+
 	Ar.UsingCustomVersion(FFrameworkObjectVersion::GUID);
 	Ar.UsingCustomVersion(FReleaseObjectVersion::GUID);
 }
