Audio Stream Cache logging
#tests PIE tests with CVar debugging tools

diff --git a/Engine/Source/Runtime/Engine/Private/AudioStreamingCache.cpp b/Engine/Source/Runtime/Engine/Private/AudioStreamingCache.cpp
--- a/Engine/Source/Runtime/Engine/Private/AudioStreamingCache.cpp
+++ b/Engine/Source/Runtime/Engine/Private/AudioStreamingCache.cpp
@@ -1808,7 +1808,12 @@ void FAudioChunkCache::KickOffAsyncLoad(FCacheElement* CacheElement, const FChun
 
 			CacheElement->ChunkDataSize = ChunkDataSize;
 			CacheElement->bIsLoaded = true;
-
+			if (bRequestFailed)
+			{
+				FMemory::Memzero(CacheElement->ChunkData, CacheElement->ChunkDataSize);
+				UE_LOG(LogAudio, Warning, TEXT("FAudioChunkCache::KickOffAsyncLoad -> DDCTask.OnLoadComplete: Request Failed. ChunkIdx: %d; SoundWave: %s"), InKey.ChunkIndex, *InKey.SoundWaveName.ToString());
+			}
+			
 #if DEBUG_STREAM_CACHE
 			CacheElement->DebugInfo.TimeToLoad = FPlatformTime::ToMilliseconds64(FPlatformTime::Cycles64() - CacheElement->DebugInfo.TimeLoadStarted);
 #endif
@@ -1862,7 +1867,12 @@ void FAudioChunkCache::KickOffAsyncLoad(FCacheElement* CacheElement, const FChun
 			CacheElement->Key = InKey;
 			CacheElement->ChunkDataSize = ChunkDataSize;
 			CacheElement->bIsLoaded = true;
-
+			if (bWasCancelled)
+			{
+				FMemory::Memzero(CacheElement->ChunkData, CacheElement->ChunkDataSize);
+				UE_LOG(LogAudio, Warning, TEXT("FAudioChunkCache::KickOffAsyncLoad -> AsyncFileCallBack: Request Cancelled. ChunkIdx: %d; SoundWave: %s"), InKey.ChunkIndex, *InKey.SoundWaveName.ToString());
+			}
+			
 #if DEBUG_STREAM_CACHE
 			CacheElement->DebugInfo.TimeToLoad = FPlatformTime::ToMilliseconds64(FPlatformTime::Cycles64() - CacheElement->DebugInfo.TimeLoadStarted);
 #endif
