Texture2D : nop comment only

diff --git a/Engine/Source/Runtime/Engine/Private/Texture2D.cpp b/Engine/Source/Runtime/Engine/Private/Texture2D.cpp
--- a/Engine/Source/Runtime/Engine/Private/Texture2D.cpp
+++ b/Engine/Source/Runtime/Engine/Private/Texture2D.cpp
@@ -711,6 +711,8 @@ void UTexture2D::GetAssetRegistryTags(FAssetRegistryTagsContext Context) const
 	//	it does not tell you if the texture actually has non-opaque alpha
 	Context.AddTag( FAssetRegistryTag("HasAlphaChannel", HasAlphaChannel() ? TEXT("True") : TEXT("False"), FAssetRegistryTag::TT_Alphabetical) );
 	Context.AddTag( FAssetRegistryTag("Format", GPixelFormats[GetPixelFormat()].Name, FAssetRegistryTag::TT_Alphabetical) );
+	
+	// @todo Oodle : this should be in UTexture:: , not 2D::
 #if WITH_EDITORONLY_DATA
 	const FString PowerOfTwoModeStr = StaticEnum<ETexturePowerOfTwoSetting::Type>()->GetNameStringByValue(static_cast<int64>(PowerOfTwoMode));
 	Context.AddTag(FAssetRegistryTag("PowerOfTwoMode", *PowerOfTwoModeStr, FAssetRegistryTag::TT_Alphabetical));
