CIS fix for DLL export macro usage.

diff --git a/Engine/Source/Runtime/Engine/Classes/Sound/QuartzInterfaces.h b/Engine/Source/Runtime/Engine/Classes/Sound/QuartzInterfaces.h
--- a/Engine/Source/Runtime/Engine/Classes/Sound/QuartzInterfaces.h
+++ b/Engine/Source/Runtime/Engine/Classes/Sound/QuartzInterfaces.h
@@ -30,7 +30,7 @@ namespace Audio
 	}; // struct FQuartzMetronomeDelegateData
 
 	// Struct used to queue events to be sent to the Audio Render thread closer to their start time
-	ENGINE_API struct FQuartzQueueCommandData : public FQuartzCrossThreadMessage
+	struct ENGINE_API FQuartzQueueCommandData : public FQuartzCrossThreadMessage
 	{
 		FAudioComponentCommandInfo AudioComponentCommandInfo;
 		FName ClockName;
