Ensure that state machines cached pose properly detects when state is stale

diff --git a/Engine/Source/Runtime/Engine/Private/Animation/AnimNode_StateMachine.cpp b/Engine/Source/Runtime/Engine/Private/Animation/AnimNode_StateMachine.cpp
--- a/Engine/Source/Runtime/Engine/Private/Animation/AnimNode_StateMachine.cpp
+++ b/Engine/Source/Runtime/Engine/Private/Animation/AnimNode_StateMachine.cpp
@@ -299,10 +299,18 @@ void FAnimNode_StateMachine::ConditionallyCacheBonesForState(int32 StateIndex, F
 {
 	// Only call CacheBones when needed.
 	check(StateCacheBoneCounters.IsValidIndex(StateIndex));
-	if (!StateCacheBoneCounters[StateIndex].IsSynchronized_Counter(Context.AnimInstanceProxy->GetCachedBonesCounter()))
+
+	const FGraphTraversalCounter& ProxyCachedBonesCounter = Context.AnimInstanceProxy->GetCachedBonesCounter();
+
+	// Check both the frame and counter
+	// Multiple anim instances can be linked together and their LODs can be independent
+	// This means that a sub-instance can have its LOD change, causing the value we cache here to run ahead
+	// of the one from the main anim instance. When the LOD of the main instance will change, we'll reset
+	// all the counters to match it which can cause them to run backwards
+	if (!StateCacheBoneCounters[StateIndex].IsSynchronized_All(ProxyCachedBonesCounter))
 	{
 		// keep track of states that have had CacheBones called on.
-		StateCacheBoneCounters[StateIndex].SynchronizeWith(Context.AnimInstanceProxy->GetCachedBonesCounter());
+		StateCacheBoneCounters[StateIndex].SynchronizeWith(ProxyCachedBonesCounter);
 
 		FAnimationCacheBonesContext CacheBoneContext(Context.AnimInstanceProxy);
 		StatePoseLinks[StateIndex].CacheBones(CacheBoneContext);
