[CoreGlobals]
 - Add accessors for GPlayInEditorID and deprecate the global

diff --git a/Engine/Source/Editor/UnrealEd/Private/Editor.cpp b/Engine/Source/Editor/UnrealEd/Private/Editor.cpp
--- a/Engine/Source/Editor/UnrealEd/Private/Editor.cpp
+++ b/Engine/Source/Editor/UnrealEd/Private/Editor.cpp
@@ -1037,7 +1037,7 @@ UWorld* SetPlayInEditorWorld( UWorld* PlayInEditorWorld )
 
 	if (FWorldContext* WorldContext = GEngine->GetWorldContextFromWorld(PlayInEditorWorld))
 	{
-		GPlayInEditorID = WorldContext->PIEInstance;
+		UE::SetPlayInEditorID(WorldContext->PIEInstance);
 		UpdatePlayInEditorWorldDebugString(WorldContext);
 	}
 
@@ -1055,7 +1055,7 @@ void RestoreEditorWorld( UWorld* EditorWorld )
 	check(GIsPlayInEditorWorld);
 	GIsPlayInEditorWorld = false;
 	GWorld = EditorWorld;
-	GPlayInEditorID = INDEX_NONE;
+	UE::SetPlayInEditorID(INDEX_NONE);
 	UpdatePlayInEditorWorldDebugString(nullptr);
 }
 
diff --git a/Engine/Source/Editor/UnrealEd/Private/EditorEngine.cpp b/Engine/Source/Editor/UnrealEd/Private/EditorEngine.cpp
--- a/Engine/Source/Editor/UnrealEd/Private/EditorEngine.cpp
+++ b/Engine/Source/Editor/UnrealEd/Private/EditorEngine.cpp
@@ -2092,8 +2092,8 @@ void UEditorEngine::Tick( float DeltaSeconds, bool bIdleMode )
 		bFirstTick = false;
 	}
 
-	ensure(GPlayInEditorID == INDEX_NONE);
-	GPlayInEditorID = INDEX_NONE;
+	ensure(UE::GetPlayInEditorID() == INDEX_NONE);
+	UE::SetPlayInEditorID(INDEX_NONE);
 
 	// Clean up any game viewports that may have been closed during the level tick (eg by Kismet).
 	CleanupGameViewport();
diff --git a/Engine/Source/Editor/UnrealEd/Private/PlayLevel.cpp b/Engine/Source/Editor/UnrealEd/Private/PlayLevel.cpp
--- a/Engine/Source/Editor/UnrealEd/Private/PlayLevel.cpp
+++ b/Engine/Source/Editor/UnrealEd/Private/PlayLevel.cpp
@@ -2252,7 +2252,7 @@ void UEditorEngine::OnScriptExecutionStart(const FBlueprintContextTracker& Conte
 	if (IsInGameThread() && GWorld)
 	{
 		// See if we should create a world switcher, which is true if we don't have one and our PIE info is missing
-		if (!FunctionStackWorldSwitcher && (!GIsPlayInEditorWorld || GPlayInEditorID == -1))
+		if (!FunctionStackWorldSwitcher && (!GIsPlayInEditorWorld || UE::GetPlayInEditorID() == -1))
 		{
 			check(FunctionStackWorldSwitcherTag == -1);
 			UWorld* ContextWorld = GetWorldFromContextObject(ContextObject, EGetWorldErrorMode::ReturnNull);
@@ -2318,7 +2318,7 @@ UWorld* UEditorEngine::CreatePIEWorldByDuplication(FWorldContext &WorldContext,
 	PlayWorldPackage->SetSavedHash( InPackage->GetSavedHash() );
 	PlayWorldPackage->MarkAsFullyLoaded();
 
-	// check(GPlayInEditorID == -1 || GPlayInEditorID == WorldContext.PIEInstance);
+	// check(UE::GetPlayInEditorID() == -1 || UE::GetPlayInEditorID() == WorldContext.PIEInstance);
 	// Currently GPlayInEditorID is not correctly reset after map loading, so it's not safe to assert here
 	FTemporaryPlayInEditorIDOverride IDHelper(WorldContext.PIEInstance);
 
diff --git a/Engine/Source/Runtime/Core/Private/Misc/CoreGlobals.cpp b/Engine/Source/Runtime/Core/Private/Misc/CoreGlobals.cpp
--- a/Engine/Source/Runtime/Core/Private/Misc/CoreGlobals.cpp
+++ b/Engine/Source/Runtime/Core/Private/Misc/CoreGlobals.cpp
@@ -390,8 +390,10 @@ bool					GIsEditorLoadingPackage			= false;
 bool					GIsCookerLoadingPackage			= false;
 /** Whether GWorld points to the play in editor world														*/
 bool					GIsPlayInEditorWorld			= false;
+PRAGMA_DISABLE_DEPRECATION_WARNINGS
 /** Unique ID for multiple PIE instances running in one process												*/
 FPlayInEditorID			GPlayInEditorID;
+PRAGMA_ENABLE_DEPRECATION_WARNINGS
 /** Whether or not PIE was attempting to play from PlayerStart												*/
 bool					GIsPIEUsingPlayerStart			= false;
 /** true if the runtime needs textures to be powers of two													*/
@@ -861,18 +863,32 @@ namespace UE::Core::Private
 	}
 }
 
+namespace UE
+{
+	int32 GetPlayInEditorID()
+	{
+		int32 Value = PRIVATE_GetGPlayInEditorID();
+		checkf(Value != -2, TEXT("GPlayInEditorID has not been properly forwarded by the loading-thread."));
+		return Value;
+	}
+	void SetPlayInEditorID(int32 InPlayInEditorID)
+	{
+		PRIVATE_SetGPlayInEditorID(InPlayInEditorID);
+	}
+}
+
+PRAGMA_DISABLE_DEPRECATION_WARNINGS
 FPlayInEditorID& FPlayInEditorID::operator= (int32 InOther)
 {
-	PRIVATE_SetGPlayInEditorID(InOther);
+	UE::SetPlayInEditorID(InOther);
 	return *this;
 }
 
 FPlayInEditorID::operator int32() const
 {
-	int32 Value = PRIVATE_GetGPlayInEditorID();
-	checkf(Value != -2, TEXT("GPlayInEditorID has not been properly forwarded by the loading-thread."));
-	return Value;
+	return UE::GetPlayInEditorID();
 }
+PRAGMA_ENABLE_DEPRECATION_WARNINGS
 
 #undef LOCTEXT_NAMESPACE
 
diff --git a/Engine/Source/Runtime/Core/Public/CoreGlobals.h b/Engine/Source/Runtime/Core/Public/CoreGlobals.h
--- a/Engine/Source/Runtime/Core/Public/CoreGlobals.h
+++ b/Engine/Source/Runtime/Core/Public/CoreGlobals.h
@@ -488,14 +488,24 @@ extern CORE_API bool GIsCookerLoadingPackage;
 /** Whether GWorld points to the play in editor world */
 extern CORE_API bool GIsPlayInEditorWorld;
 
-class FPlayInEditorID
+class UE_DEPRECATED(5.5, "This will be removed.") FPlayInEditorID
 {
 public:
 	CORE_API FPlayInEditorID& operator= (int32 InOther);
 	CORE_API operator int32() const;
 };
+
+PRAGMA_DISABLE_DEPRECATION_WARNINGS
 /** In the editor, this is set to the specific world context PIEInstance that is being currently processed */
+UE_DEPRECATED(5.5, "Use UE::GetPlayInEditorID() and UE::SetPlayInEditorID() instead.")
 extern CORE_API FPlayInEditorID GPlayInEditorID;
+PRAGMA_ENABLE_DEPRECATION_WARNINGS
+
+namespace UE
+{
+	CORE_API int32 GetPlayInEditorID();
+	CORE_API void  SetPlayInEditorID(int32 InPlayInEditorID);
+}
 
 /** Whether or not PIE was attempting to play from PlayerStart */
 UE_DEPRECATED(4.25, "This variable is no longer set. Use !GEditor->GetPlayInEditorSessionInfo()->OriginalRequestParams.HasPlayWorldPlacement() instead.")
diff --git a/Engine/Source/Runtime/CoreUObject/Private/UObject/LazyObjectPtr.cpp b/Engine/Source/Runtime/CoreUObject/Private/UObject/LazyObjectPtr.cpp
--- a/Engine/Source/Runtime/CoreUObject/Private/UObject/LazyObjectPtr.cpp
+++ b/Engine/Source/Runtime/CoreUObject/Private/UObject/LazyObjectPtr.cpp
@@ -109,7 +109,7 @@ void FLazyObjectPtr::PossiblySerializeObjectGuid(UObject *Object, FStructuredArc
 		{
 			if (UnderlyingArchive.GetPortFlags() & PPF_DuplicateForPIE)
 			{
-				Guid = RemapGuid(Guid, GPlayInEditorID);
+				Guid = RemapGuid(Guid, UE::GetPlayInEditorID());
 			}
 
 			GuidSlot.GetValue() << Guid;
@@ -186,7 +186,8 @@ void FLazyObjectPtr::PossiblySerializeObjectGuid(UObject *Object, FStructuredArc
 
 void FLazyObjectPtr::ResetPIEFixups()
 {
-	check(GPlayInEditorID != -1);
-	check(GPlayInEditorID < MAX_PIE_INSTANCES);
-	PIEGuidMap[GPlayInEditorID].Reset();
+	const int32 PlayInEditorID = UE::GetPlayInEditorID();
+	check(PlayInEditorID != -1);
+	check(PlayInEditorID < MAX_PIE_INSTANCES);
+	PIEGuidMap[PlayInEditorID].Reset();
 }
\ No newline at end of file
diff --git a/Engine/Source/Runtime/CoreUObject/Private/UObject/SoftObjectPath.cpp b/Engine/Source/Runtime/CoreUObject/Private/UObject/SoftObjectPath.cpp
--- a/Engine/Source/Runtime/CoreUObject/Private/UObject/SoftObjectPath.cpp
+++ b/Engine/Source/Runtime/CoreUObject/Private/UObject/SoftObjectPath.cpp
@@ -524,7 +524,7 @@ UObject* FSoftObjectPath::TryLoad(FUObjectSerializeContext* InLoadContext) const
 		{
 			FString PathString = ToString();
 #if WITH_EDITOR
-			if (GPlayInEditorID != INDEX_NONE)
+			if (UE::GetPlayInEditorID() != INDEX_NONE)
 			{
 				// If we are in PIE and this hasn't already been fixed up, we need to fixup at resolution time. We cannot modify the path as it may be somewhere like a blueprint CDO
 				FSoftObjectPath FixupObjectPath = *this;
@@ -564,7 +564,7 @@ int32 FSoftObjectPath::LoadAsync(FLoadSoftObjectPathAsyncDelegate InCompletionDe
 	FSoftObjectPath RequestedPath = *this;
 	FSoftObjectPath PathToLoad = RequestedPath;
 #if WITH_EDITOR
-	if (GPlayInEditorID != INDEX_NONE)
+	if (UE::GetPlayInEditorID() != INDEX_NONE)
 	{
 		// @todo: This logic may need updating to handle level instances properly and we may want to handle other fixups like CoreRedirects before requesting
 		PathToLoad.FixupForPIE();
@@ -598,7 +598,7 @@ UObject* FSoftObjectPath::ResolveObject() const
 	}
 
 #if WITH_EDITOR
-	if (GPlayInEditorID != INDEX_NONE)
+	if (UE::GetPlayInEditorID() != INDEX_NONE)
 	{
 		// If we are in PIE and this hasn't already been fixed up, we need to fixup at resolution time. We cannot modify the path as it may be somewhere like a blueprint CDO
 		FSoftObjectPath FixupObjectPath = *this;
@@ -708,7 +708,7 @@ bool FSoftObjectPath::FixupForPIE(int32 InPIEInstance, TFunctionRef<void(int32,
 
 bool FSoftObjectPath::FixupForPIE(TFunctionRef<void(int32, FSoftObjectPath&)> InPreFixupForPIECustomFunction)
 {
-	return FixupForPIE(GPlayInEditorID, InPreFixupForPIECustomFunction);
+	return FixupForPIE(UE::GetPlayInEditorID(), InPreFixupForPIECustomFunction);
 }
 
 bool FSoftObjectPath::FixupCoreRedirects()
diff --git a/Engine/Source/Runtime/CoreUObject/Public/UObject/LazyObjectPtr.h b/Engine/Source/Runtime/CoreUObject/Public/UObject/LazyObjectPtr.h
--- a/Engine/Source/Runtime/CoreUObject/Public/UObject/LazyObjectPtr.h
+++ b/Engine/Source/Runtime/CoreUObject/Public/UObject/LazyObjectPtr.h
@@ -53,7 +53,7 @@ struct FUniqueObjectGuid
 	COREUOBJECT_API void FromString(const FString& From);
 
 	/** Fixes up this UniqueObjectID to add or remove the PIE prefix depending on what is currently active */
-	COREUOBJECT_API FUniqueObjectGuid FixupForPIE(int32 PlayInEditorID = GPlayInEditorID) const;
+	COREUOBJECT_API FUniqueObjectGuid FixupForPIE(int32 PlayInEditorID = UE::GetPlayInEditorID()) const;
 
 	/**
 	 * Attempts to find a currently loaded object that matches this object ID
diff --git a/Engine/Source/Runtime/CoreUObject/Public/UObject/SoftObjectPtr.h b/Engine/Source/Runtime/CoreUObject/Public/UObject/SoftObjectPtr.h
--- a/Engine/Source/Runtime/CoreUObject/Public/UObject/SoftObjectPtr.h
+++ b/Engine/Source/Runtime/CoreUObject/Public/UObject/SoftObjectPtr.h
@@ -92,7 +92,7 @@ public:
 	/** Overridden to deal with PIE lookups */
 	FORCEINLINE UObject* Get() const
 	{
-		if (GPlayInEditorID != INDEX_NONE)
+		if (UE::GetPlayInEditorID() != INDEX_NONE)
 		{
 			// Cannot use or set the cached value in PIE as it may affect other PIE instances or the editor
 			TWeakObjectPtr<UObject> Result = GetUniqueID().ResolveObject();
diff --git a/Engine/Source/Runtime/Engine/Private/KismetSystemLibrary.cpp b/Engine/Source/Runtime/Engine/Private/KismetSystemLibrary.cpp
--- a/Engine/Source/Runtime/Engine/Private/KismetSystemLibrary.cpp
+++ b/Engine/Source/Runtime/Engine/Private/KismetSystemLibrary.cpp
@@ -392,7 +392,7 @@ void UKismetSystemLibrary::PrintString(const UObject* WorldContextObject, const
 				case NM_Client:
 					// GPlayInEditorID 0 is always the server, so 1 will be first client.
 					// You want to keep this logic in sync with GeneratePIEViewportWindowTitle and UpdatePlayInEditorWorldDebugString
-					Prefix = FString::Printf(TEXT("Client %d: "), (int32)GPlayInEditorID);
+					Prefix = FString::Printf(TEXT("Client %d: "), UE::GetPlayInEditorID());
 					break;
 				case NM_DedicatedServer:
 				case NM_ListenServer:
diff --git a/Engine/Source/Runtime/Engine/Private/ShaderCompiler/ShaderCompiler.cpp b/Engine/Source/Runtime/Engine/Private/ShaderCompiler/ShaderCompiler.cpp
--- a/Engine/Source/Runtime/Engine/Private/ShaderCompiler/ShaderCompiler.cpp
+++ b/Engine/Source/Runtime/Engine/Private/ShaderCompiler/ShaderCompiler.cpp
@@ -6250,7 +6250,7 @@ void FShaderCompilingManager::BlockOnShaderMapCompletion(const TArray<int32>& Sh
 			}
 		}
 
-		FScopedSlowTask SlowTask(NumJobs, FText::Format(LOCTEXT("BlockOnShaderMapCompletion", "Compiling Shaders ({0})"), NumJobs), GIsEditor && !IsRunningCommandlet() && GPlayInEditorID == INDEX_NONE);
+		FScopedSlowTask SlowTask(NumJobs, FText::Format(LOCTEXT("BlockOnShaderMapCompletion", "Compiling Shaders ({0})"), NumJobs), GIsEditor && !IsRunningCommandlet() && UE::GetPlayInEditorID() == INDEX_NONE);
 		if (NumJobs > 0)
 		{
 			SlowTask.MakeDialogDelayed(1.0f);
diff --git a/Engine/Source/Runtime/Engine/Private/UnrealEngine.cpp b/Engine/Source/Runtime/Engine/Private/UnrealEngine.cpp
--- a/Engine/Source/Runtime/Engine/Private/UnrealEngine.cpp
+++ b/Engine/Source/Runtime/Engine/Private/UnrealEngine.cpp
@@ -642,7 +642,7 @@ ENGINE_API void UpdatePlayInEditorWorldDebugString(const FWorldContext* WorldCon
 
 	if (WorldContext == nullptr)
 	{
-		ensure(GPlayInEditorID == INDEX_NONE);
+		ensure(UE::GetPlayInEditorID() == INDEX_NONE);
 		GPlayInEditorContextString = NSLOCTEXT("Engine", "PlayWorldIsNotActive", "Not in a play world").ToString();
 	}
 	else
@@ -653,7 +653,7 @@ ENGINE_API void UpdatePlayInEditorWorldDebugString(const FWorldContext* WorldCon
 }
 
 FTemporaryPlayInEditorIDOverride::FTemporaryPlayInEditorIDOverride(int32 NewOverrideID)
-	: PreviousID(GPlayInEditorID)
+	: PreviousID(UE::GetPlayInEditorID())
 {
 	SetID(NewOverrideID);
 }
@@ -665,10 +665,10 @@ FTemporaryPlayInEditorIDOverride::~FTemporaryPlayInEditorIDOverride()
 
 void FTemporaryPlayInEditorIDOverride::SetID(int32 NewID)
 {
-	if (GPlayInEditorID != NewID)
+	if (UE::GetPlayInEditorID() != NewID)
 	{
-		GPlayInEditorID = NewID;
-		UpdatePlayInEditorWorldDebugString(GEngine->GetWorldContextFromPIEInstance(GPlayInEditorID));
+		UE::SetPlayInEditorID(NewID);
+		UpdatePlayInEditorWorldDebugString(GEngine->GetWorldContextFromPIEInstance(NewID));
 	}
 }
 
@@ -15346,7 +15346,7 @@ bool UEngine::LoadMap( FWorldContext& WorldContext, FURL URL, class UPendingNetG
 
 					// We are loading a new world for this context, so clear out PIE fixups that might be lingering.
 					// (note we dont want to do this in DuplicateWorldForPIE, since that is also called on streaming worlds.
-					GPlayInEditorID = WorldContext.PIEInstance;
+					UE::SetPlayInEditorID(WorldContext.PIEInstance);
 					UpdatePlayInEditorWorldDebugString(&WorldContext);
 					FLazyObjectPtr::ResetPIEFixups();
 
@@ -16245,14 +16245,14 @@ UWorld* UEngine::GetCurrentPlayWorld(UWorld* PossiblePlayWorld) const
 			if (WorldContext.WorldType == EWorldType::PIE)
 			{
 				// This is a PIE world, and PIE instance is either not set or matches this world
-				if (GPlayInEditorID == -1 || GPlayInEditorID == WorldContext.PIEInstance)
+				if (UE::GetPlayInEditorID() == -1 || UE::GetPlayInEditorID() == WorldContext.PIEInstance)
 				{
 					return WorldContext.World();
 				}
 				else
 				{
 					// If you see this warning, game code is traversing PIE world boundaries in an unsafe way. That should be fixed or GPlayInEditorID needs to be set properly
-					UE_LOG(LogEngine, Warning, TEXT("GetCurrentPlayWorld failed with ambiguous PIE world! GPlayInEditorID %d does not match %s"), (int32)GPlayInEditorID, *PossiblePlayWorld->GetPathName());
+					UE_LOG(LogEngine, Warning, TEXT("GetCurrentPlayWorld failed with ambiguous PIE world! GPlayInEditorID %d does not match %s"), UE::GetPlayInEditorID(), *PossiblePlayWorld->GetPathName());
 				}
 			}
 #endif
@@ -16274,7 +16274,7 @@ UWorld* UEngine::GetCurrentPlayWorld(UWorld* PossiblePlayWorld) const
 		}
 #if WITH_EDITOR
 		// This is a PIE world, PIE instance is set, and it matches this world
-		else if (WorldContext.WorldType == EWorldType::PIE && GPlayInEditorID != -1 && GPlayInEditorID == WorldContext.PIEInstance)
+		else if (WorldContext.WorldType == EWorldType::PIE && UE::GetPlayInEditorID() != -1 && UE::GetPlayInEditorID() == WorldContext.PIEInstance)
 		{
 			BestWorld = WorldContext.World();
 		}
diff --git a/Engine/Source/Runtime/Experimental/Iris/Core/Private/Iris/ReplicationSystem/ObjectReferenceCache.cpp b/Engine/Source/Runtime/Experimental/Iris/Core/Private/Iris/ReplicationSystem/ObjectReferenceCache.cpp
--- a/Engine/Source/Runtime/Experimental/Iris/Core/Private/Iris/ReplicationSystem/ObjectReferenceCache.cpp
+++ b/Engine/Source/Runtime/Experimental/Iris/Core/Private/Iris/ReplicationSystem/ObjectReferenceCache.cpp
@@ -213,7 +213,7 @@ bool FObjectReferenceCache::ShouldIgnoreWhenMissing(FNetRefHandle RefHandle) con
 
 bool FObjectReferenceCache::RenamePathForPie(uint32 ConnectionId, FString& Str, bool bReading)
 {
-	return GPlayInEditorID != -1 ? ReplicationBridge->RemapPathForPIE(ConnectionId, Str, bReading) : false;
+	return UE::GetPlayInEditorID() != -1 ? ReplicationBridge->RemapPathForPIE(ConnectionId, Str, bReading) : false;
 }
 
 FNetRefHandle FObjectReferenceCache::CreateObjectReferenceHandle(const UObject* Object)
diff --git a/Engine/Source/Runtime/MovieScene/Private/MovieSceneSequencePlayer.cpp b/Engine/Source/Runtime/MovieScene/Private/MovieSceneSequencePlayer.cpp
--- a/Engine/Source/Runtime/MovieScene/Private/MovieSceneSequencePlayer.cpp
+++ b/Engine/Source/Runtime/MovieScene/Private/MovieSceneSequencePlayer.cpp
@@ -1434,7 +1434,7 @@ FString UMovieSceneSequencePlayer::GetSequenceName(bool bAddClientInfo) const
 			AActor* Actor = GetTypedOuter<AActor>();
 			if (Actor && Actor->GetWorld() && Actor->GetWorld()->GetNetMode() == NM_Client)
 			{
-				SequenceName += FString::Printf(TEXT(" (client %d)"), GPlayInEditorID - 1);
+				SequenceName += FString::Printf(TEXT(" (client %d)"), UE::GetPlayInEditorID() - 1);
 			}
 		}
 		return SequenceName;
diff --git a/Engine/Source/Runtime/MovieSceneCapture/Private/UserDefinedCaptureProtocol.cpp b/Engine/Source/Runtime/MovieSceneCapture/Private/UserDefinedCaptureProtocol.cpp
--- a/Engine/Source/Runtime/MovieSceneCapture/Private/UserDefinedCaptureProtocol.cpp
+++ b/Engine/Source/Runtime/MovieSceneCapture/Private/UserDefinedCaptureProtocol.cpp
@@ -138,7 +138,7 @@ bool UUserDefinedCaptureProtocol::SetupImpl()
 		World = nullptr;
 	}
 
-	int32 PreviousPlayInEditorID = GPlayInEditorID;
+	int32 PreviousPlayInEditorID = UE::GetPlayInEditorID();
 
 	if (World)
 	{
@@ -146,7 +146,7 @@ bool UUserDefinedCaptureProtocol::SetupImpl()
 		{
 			if (World == Context.World())
 			{
-				GPlayInEditorID = Context.PIEInstance;
+				UE::SetPlayInEditorID(Context.PIEInstance);
 			}
 		}
 	}
@@ -155,7 +155,7 @@ bool UUserDefinedCaptureProtocol::SetupImpl()
 	FinalPixelsFrameGrabber.Reset(new FFrameGrabber(InitSettings->SceneViewport.ToSharedRef(), InitSettings->DesiredSize, PF_B8G8R8A8, 3));
 	const bool bSuccess = OnSetup();
 
-	GPlayInEditorID = PreviousPlayInEditorID;
+	UE::SetPlayInEditorID(PreviousPlayInEditorID);
 
 	return bSuccess;
 }
