[NetPhysics] Implement ENGINE_API in NetworkPhysicsSettingsComponent

diff --git a/Engine/Source/Runtime/Engine/Public/Physics/NetworkPhysicsSettingsComponent.h b/Engine/Source/Runtime/Engine/Public/Physics/NetworkPhysicsSettingsComponent.h
--- a/Engine/Source/Runtime/Engine/Public/Physics/NetworkPhysicsSettingsComponent.h
+++ b/Engine/Source/Runtime/Engine/Public/Physics/NetworkPhysicsSettingsComponent.h
@@ -326,20 +326,21 @@ struct FNetworkPhysicsSettingsRenderInterpolation
 };
 */
 
-UCLASS(BlueprintType, meta = (BlueprintSpawnableComponent))
+/** Settings Component for network replicated physics actors
+* Overrides default settings, CVar settings and project settings. */
+UCLASS(BlueprintType, MinimalAPI, meta = (BlueprintSpawnableComponent))
 class UNetworkPhysicsSettingsComponent : public UActorComponent
 {
 	GENERATED_BODY()
 
 public:
-	UNetworkPhysicsSettingsComponent();
+	ENGINE_API UNetworkPhysicsSettingsComponent();
 
 	virtual void InitializeComponent() override;
 	virtual void UninitializeComponent() override;
 
 	virtual void BeginPlay() override;
 
-
 public:
 	UPROPERTY(EditDefaultsOnly, Category = "Networked Physics Settings")
 	FNetworkPhysicsSettings GeneralSettings;
