Cooker: Add the ability to declare CookRuntimeReferences from UObject::Presave, for object classes that are culled on the targetplatform due to e.g. IsNeededForClientLoad=false but still want to add their references to the cook for that platform.
#rnx

diff --git a/Engine/Source/Editor/UnrealEd/Private/CookOnTheFlyServer.cpp b/Engine/Source/Editor/UnrealEd/Private/CookOnTheFlyServer.cpp
--- a/Engine/Source/Editor/UnrealEd/Private/CookOnTheFlyServer.cpp
+++ b/Engine/Source/Editor/UnrealEd/Private/CookOnTheFlyServer.cpp
@@ -6698,10 +6698,10 @@ void FSaveCookedPackageContext::CalculatePlatformRuntimeDependencies()
 	{
 		FName PackageFName = Package->GetFName();
 		TConstArrayView<const ITargetPlatform*> ReachablePlatforms(&TargetPlatform, 1);
-		for (const TArray<FName>& DependencyNames : { SavePackageResult.ImportPackages, SavePackageResult.SoftPackageReferences })
+		for (const TArray<FName>* DependencyNames : { &SavePackageResult.ImportPackages, &SavePackageResult.SoftPackageReferences })
 		{
-			bool bHard = &DependencyNames == &SavePackageResult.ImportPackages;
-			for (FName DependencyName : DependencyNames)
+			bool bHard = DependencyNames == &SavePackageResult.ImportPackages;
+			for (FName DependencyName : *DependencyNames)
 			{
 				UE::Cook::FPackageData* DependencyData = COTFS.PackageDatas->TryAddPackageDataByPackageName(DependencyName);
 				if (DependencyData)
diff --git a/Engine/Source/Runtime/CoreUObject/Internal/UObject/ImportExportCollector.h b/Engine/Source/Runtime/CoreUObject/Internal/UObject/ImportExportCollector.h
new file mode 100644
--- /dev/null
+++ b/Engine/Source/Runtime/CoreUObject/Internal/UObject/ImportExportCollector.h
@@ -0,0 +1,71 @@
+// Copyright Epic Games, Inc. All Rights Reserved.
+
+#include "CoreTypes.h"
+
+#if WITH_EDITORONLY_DATA
+#include "Containers/Map.h"
+#include "Containers/RingBuffer.h"
+#include "Containers/Set.h"
+#include "Serialization/ArchiveUObject.h"
+#include "UObject/Package.h"
+#include "UObject/NameTypes.h"
+#include "UObject/SoftObjectPath.h"
+
+/** An Archive that records all of the imported packages from a tree of exports. */
+class FImportExportCollector : public FArchiveUObject
+{
+public:
+	COREUOBJECT_API explicit FImportExportCollector(UPackage* InRootPackage);
+
+	/**
+	 * Mark that a given export (e.g. the export that is doing the collecting) should not be explored
+	 * if encountered again. Prevents infinite recursion when the collector is constructed and called during
+	 * Serialize.
+	 */
+	COREUOBJECT_API void AddExportToIgnore(UObject* Export);
+	/**
+	 * Serialize the given object, following its object references to find other imports and exports,
+	 * and recursively serialize any new exports that it references.
+	 */
+	COREUOBJECT_API void SerializeObjectAndReferencedExports(UObject* RootObject);
+	/** Restore the collector to empty. */
+	COREUOBJECT_API void Reset();
+	const TSet<UObject*>& GetExports() const;
+	const TMap<FSoftObjectPath, ESoftObjectPathCollectType>& GetImports() const;
+	const TMap<FName, ESoftObjectPathCollectType>& GetImportedPackages() const;
+
+	COREUOBJECT_API virtual FArchive& operator<<(UObject*& Obj) override;
+	COREUOBJECT_API virtual FArchive& operator<<(FSoftObjectPath& Value) override;
+
+private:
+	void AddImport(const FSoftObjectPath& Path, ESoftObjectPathCollectType CollectType);
+	ESoftObjectPathCollectType Union(ESoftObjectPathCollectType A, ESoftObjectPathCollectType B);
+
+	TSet<UObject*> Exports;
+	TRingBuffer<UObject*> ExportsExploreQueue;
+	TMap<FSoftObjectPath, ESoftObjectPathCollectType> Imports;
+	TMap<FName, ESoftObjectPathCollectType> ImportedPackages;
+	UPackage* RootPackage;
+	FName RootPackageName;
+};
+
+///////////////////////////////////////////////////////
+// Inline implementations
+///////////////////////////////////////////////////////
+
+inline const TSet<UObject*>& FImportExportCollector::GetExports() const
+{
+	return Exports;
+}
+
+inline const TMap<FSoftObjectPath, ESoftObjectPathCollectType>& FImportExportCollector::GetImports() const
+{
+	return Imports;
+}
+
+inline const TMap<FName, ESoftObjectPathCollectType>& FImportExportCollector::GetImportedPackages() const
+{
+	return ImportedPackages;
+}
+
+#endif // WITH_EDITORONLY_DATA
diff --git a/Engine/Source/Runtime/CoreUObject/Private/UObject/ImportExportCollector.cpp b/Engine/Source/Runtime/CoreUObject/Private/UObject/ImportExportCollector.cpp
new file mode 100644
--- /dev/null
+++ b/Engine/Source/Runtime/CoreUObject/Private/UObject/ImportExportCollector.cpp
@@ -0,0 +1,101 @@
+// Copyright Epic Games, Inc. All Rights Reserved.
+
+#include "UObject/ImportExportCollector.h"
+
+#if WITH_EDITORONLY_DATA
+
+FImportExportCollector::FImportExportCollector(UPackage* InRootPackage)
+	: RootPackage(InRootPackage)
+	, RootPackageName(InRootPackage->GetFName())
+{
+	ArIsObjectReferenceCollector = true;
+	ArIsModifyingWeakAndStrongReferences = true;
+	SetIsSaving(true);
+	SetIsPersistent(true);
+}
+
+void FImportExportCollector::Reset()
+{
+	Exports.Reset();
+	Imports.Reset();
+}
+
+void FImportExportCollector::AddExportToIgnore(UObject* Export)
+{
+	Exports.Add(Export);
+}
+
+void FImportExportCollector::SerializeObjectAndReferencedExports(UObject* RootObject)
+{
+	*this << RootObject;
+	while (!ExportsExploreQueue.IsEmpty())
+	{
+		UObject* Export = ExportsExploreQueue.PopFrontValue();
+		Export->Serialize(*this);
+	}
+}
+
+FArchive& FImportExportCollector::operator<<(UObject*& Obj)
+{
+	if (!Obj)
+	{
+		return *this;
+	}
+	UPackage* Package = Obj->GetPackage();
+	if (!Package)
+	{
+		return *this;
+	}
+	if (Package != RootPackage)
+	{
+		AddImport(FSoftObjectPath(Obj), ESoftObjectPathCollectType::AlwaysCollect);
+		return *this;
+	}
+
+	bool bAlreadyExists;
+	Exports.Add(Obj, &bAlreadyExists);
+	if (bAlreadyExists)
+	{
+		return *this;
+	}
+	ExportsExploreQueue.Add(Obj);
+	return *this;
+}
+
+FArchive& FImportExportCollector::operator<<(FSoftObjectPath& Value)
+{
+	FName CurrentPackage;
+	FName PropertyName;
+	ESoftObjectPathCollectType CollectType;
+	ESoftObjectPathSerializeType SerializeType;
+	FSoftObjectPathThreadContext& ThreadContext = FSoftObjectPathThreadContext::Get();
+	ThreadContext.GetSerializationOptions(CurrentPackage, PropertyName, CollectType, SerializeType, this);
+
+	if (CollectType != ESoftObjectPathCollectType::NeverCollect && CollectType != ESoftObjectPathCollectType::NonPackage)
+	{
+		FName PackageName = Value.GetLongPackageFName();
+		if (PackageName != RootPackageName && !PackageName.IsNone())
+		{
+			AddImport(Value, CollectType);
+		}
+	}
+	return *this;
+}
+
+void FImportExportCollector::AddImport(const FSoftObjectPath& Path, ESoftObjectPathCollectType CollectType)
+{
+	ESoftObjectPathCollectType& ExistingImport = Imports.FindOrAdd(
+		Path, ESoftObjectPathCollectType::EditorOnlyCollect);
+	ExistingImport = Union(ExistingImport, CollectType);
+
+	ESoftObjectPathCollectType& ExistingPackage = ImportedPackages.FindOrAdd(
+		Path.GetLongPackageFName(), ESoftObjectPathCollectType::EditorOnlyCollect);
+	ExistingPackage = Union(ExistingPackage, CollectType);
+}
+
+ESoftObjectPathCollectType FImportExportCollector::Union(ESoftObjectPathCollectType A, ESoftObjectPathCollectType B)
+{
+	return static_cast<ESoftObjectPathCollectType>(FMath::Max(static_cast<int>(A), static_cast<int>(B)));
+}
+
+#endif // WITH_EDITORONLY_DATA
diff --git a/Engine/Source/Runtime/CoreUObject/Private/UObject/SavePackage/SaveContext.cpp b/Engine/Source/Runtime/CoreUObject/Private/UObject/SavePackage/SaveContext.cpp
--- a/Engine/Source/Runtime/CoreUObject/Private/UObject/SavePackage/SaveContext.cpp
+++ b/Engine/Source/Runtime/CoreUObject/Private/UObject/SavePackage/SaveContext.cpp
@@ -152,7 +152,15 @@ FSavePackageResultStruct FSaveContext::GetFinalResult()
 	TSet<FName>& SoftPackageReferenceList = GetSoftPackageReferenceList();
 	ResultData.SoftPackageReferences = SoftPackageReferenceList.Array();
 #if WITH_EDITOR
-	ResultData.CookDependencies = MoveTemp(ObjectSaveContext.CookDependencies);
+	for (const FSoftObjectPath& RuntimeDependency : ObjectSaveContext.CookRuntimeDependencies)
+	{
+		FName PackageDependency = RuntimeDependency.GetLongPackageFName();
+		if (!PackageDependency.IsNone())
+		{
+			ResultData.SoftPackageReferences.Add(PackageDependency);
+		}
+	}
+	ResultData.CookDependencies = MoveTemp(ObjectSaveContext.CookBuildDependencies);
 #endif
 
 	return ResultData;
diff --git a/Engine/Source/Runtime/CoreUObject/Private/UObject/SavePackage/SavePackageUtilities.cpp b/Engine/Source/Runtime/CoreUObject/Private/UObject/SavePackage/SavePackageUtilities.cpp
--- a/Engine/Source/Runtime/CoreUObject/Private/UObject/SavePackage/SavePackageUtilities.cpp
+++ b/Engine/Source/Runtime/CoreUObject/Private/UObject/SavePackage/SavePackageUtilities.cpp
@@ -31,6 +31,7 @@
 #include "UObject/AssetRegistryTagsContext.h"
 #include "UObject/Class.h"
 #include "UObject/GCScopeLock.h"
+#include "UObject/ImportExportCollector.h"
 #include "UObject/Linker.h"
 #include "UObject/LinkerLoad.h"
 #include "UObject/LinkerSave.h"
@@ -764,9 +765,47 @@ FObjectSaveContextData::FObjectSaveContextData(UPackage* Package, const ITargetP
 }
 
 #if WITH_EDITOR
-void FObjectPreSaveContext::AddCookDependency(UE::Cook::FCookDependency CookDependency)
+void FObjectPreSaveContext::AddCookBuildDependency(UE::Cook::FCookDependency BuildDependency)
 {
-	Data.CookDependencies.Add(MoveTemp(CookDependency));
+	Data.CookBuildDependencies.Add(MoveTemp(BuildDependency));
+}
+void FObjectPreSaveContext::AddCookRuntimeDependency(FSoftObjectPath RuntimeDependency)
+{
+	Data.CookRuntimeDependencies.Add(MoveTemp(RuntimeDependency));
+}
+void FObjectPreSaveContext::HarvestCookRuntimeDependencies(UObject* HarvestReferencesFrom)
+{
+	if (!HarvestReferencesFrom)
+	{
+		return;
+	}
+	if (!GetTargetPlatform())
+	{
+		return;
+	}
+
+	UPackage* PackageBeingSaved = nullptr; // We don't have a pointer for this, so set it to null
+	FArchiveCookContext CookContext(PackageBeingSaved, Data.CookType, Data.CookingDLC, GetTargetPlatform());
+	FArchiveCookData CookData(*GetTargetPlatform(), CookContext);
+	FImportExportCollector Collector(HarvestReferencesFrom->GetPackage());
+	Collector.SetCookData(&CookData);
+	Collector.SerializeObjectAndReferencedExports(HarvestReferencesFrom);
+	for (const TPair<FName, ESoftObjectPathCollectType>& Pair : Collector.GetImportedPackages())
+	{
+		if (Pair.Value != ESoftObjectPathCollectType::AlwaysCollect)
+		{
+			continue;
+		}
+		FName PackageName = Pair.Key;
+		if (FPackageName::IsScriptPackage(WriteToString<256>(PackageName)))
+		{
+			// Ignore native imports; we don't need to mark them for cooking
+			continue;
+		}
+		FSoftObjectPath PackageSoftPath(PackageName, NAME_None, FString());
+		AddCookRuntimeDependency(PackageSoftPath);
+	}
+
 }
 #endif
 
diff --git a/Engine/Source/Runtime/CoreUObject/Public/UObject/ObjectSaveContext.h b/Engine/Source/Runtime/CoreUObject/Public/UObject/ObjectSaveContext.h
--- a/Engine/Source/Runtime/CoreUObject/Public/UObject/ObjectSaveContext.h
+++ b/Engine/Source/Runtime/CoreUObject/Public/UObject/ObjectSaveContext.h
@@ -10,6 +10,7 @@
 class FPackagePath;
 class ITargetPlatform;
 class UPackage;
+struct FSoftObjectPath;
 #if WITH_EDITOR
 namespace UE::Cook { class FCookDependency; }
 #endif
@@ -74,7 +75,8 @@ struct FObjectSaveContextData
 
 	// Collection variables that are written but not read during the PreSave/PostSave functions
 #if WITH_EDITOR
-	TArray<UE::Cook::FCookDependency> CookDependencies;
+	TArray<UE::Cook::FCookDependency> CookBuildDependencies;
+	TArray<FSoftObjectPath> CookRuntimeDependencies;
 #endif
 
 	// Per-object Output variables; writable from PreSave functions, readable from PostSave functions
@@ -131,10 +133,18 @@ public:
 
 #if WITH_EDITOR
 	/**
-	 * Add the given FCookDependency to the record for the package being cook-saved. Iterative cooks will
+	 * Add the given FCookDependency to the build dependencies for the package being cook-saved. Iterative cooks will
 	 * invalidate the package and recook it if the CookDependency changes.
 	 */
-	COREUOBJECT_API void AddCookDependency(UE::Cook::FCookDependency CookDependency);
+	COREUOBJECT_API void AddCookBuildDependency(UE::Cook::FCookDependency BuildDependency);
+	/**
+	 * Add the given UObject's package as a runtime dependency for the package being cook-saved. It will
+	 * be cooked for the current platform even if it is not otherwise referenced, and even if the object
+	 * being PreSaved does not end up saving into the package's exports for the current platform.
+	 */
+	COREUOBJECT_API void AddCookRuntimeDependency(FSoftObjectPath Dependency);
+	/** Serialize an object to find all packages that it references, and AddCookRuntimeDependency for each one. */
+	COREUOBJECT_API void HarvestCookRuntimeDependencies(UObject* HarvestReferencesFrom);
 #endif
 
 	/**
diff --git a/Engine/Source/Runtime/Engine/Private/BlueprintGeneratedClass.cpp b/Engine/Source/Runtime/Engine/Private/BlueprintGeneratedClass.cpp
--- a/Engine/Source/Runtime/Engine/Private/BlueprintGeneratedClass.cpp
+++ b/Engine/Source/Runtime/Engine/Private/BlueprintGeneratedClass.cpp
@@ -25,6 +25,7 @@
 #include "Misc/ConfigCacheIni.h"
 #include "Net/Core/PushModel/PushModel.h"
 #include "UObject/FortniteMainBranchObjectVersion.h"
+#include "UObject/ImportExportCollector.h"
 #include "UObject/Package.h" // IWYU pragma: keep
 #include "UObject/PrimaryAssetId.h"
 #include "UObject/SparseClassDataUtils.h"
@@ -2507,47 +2508,6 @@ ENGINE_API int32 IncrementUberGraphSerialNumber()
 }
 #endif//VALIDATE_UBER_GRAPH_PERSISTENT_FRAME
 
-
-#if WITH_EDITORONLY_DATA
-/** An Archive that records all of the imported packages from a tree of exports. */
-class FImportExportCollector : public FArchiveUObject
-{
-public:
-	explicit FImportExportCollector(UPackage* InRootPackage);
-
-	/**
-	 * Mark that a given export (e.g. the export that is doing the collecting) should not be explored
-	 * if encountered again. Prevents infinite recursion when the collector is constructed and called during
-	 * Serialize.
-	 */
-	void AddExportToIgnore(UObject* Export);
-	/**
-	 * Serialize the given object, following its object references to find other imports and exports,
-	 * and recursively serialize any new exports that it references.
-	 */
-	void SerializeObjectAndReferencedExports(UObject* RootObject);
-	/** Restore the collector to empty. */
-	void Reset();
-	const TSet<UObject*>& GetExports() const { return Exports; }
-	const TMap<FSoftObjectPath, ESoftObjectPathCollectType>& GetImports() const { return Imports; }
-	const TMap<FName, ESoftObjectPathCollectType>& GetImportedPackages() const { return ImportedPackages; }
-
-	virtual FArchive& operator<<(UObject*& Obj) override;
-	virtual FArchive& operator<<(FSoftObjectPath& Value) override;
-
-private:
-	void AddImport(const FSoftObjectPath& Path, ESoftObjectPathCollectType CollectType);
-	ESoftObjectPathCollectType Union(ESoftObjectPathCollectType A, ESoftObjectPathCollectType B);
-
-	TSet<UObject*> Exports;
-	TRingBuffer<UObject*> ExportsExploreQueue;
-	TMap<FSoftObjectPath, ESoftObjectPathCollectType> Imports;
-	TMap<FName, ESoftObjectPathCollectType> ImportedPackages;
-	UPackage* RootPackage;
-	FName RootPackageName;
-};
-#endif
-
 void UBlueprintGeneratedClass::Serialize(FArchive& Ar)
 {
 #if VALIDATE_UBER_GRAPH_PERSISTENT_FRAME
@@ -2982,100 +2942,6 @@ void UBlueprintGeneratedClass::PurgeCookedMetaData()
 	}
 }
 
-FImportExportCollector::FImportExportCollector(UPackage* InRootPackage)
-	: RootPackage(InRootPackage)
-	, RootPackageName(InRootPackage->GetFName())
-{
-	ArIsObjectReferenceCollector = true;
-	ArIsModifyingWeakAndStrongReferences = true;
-	SetIsSaving(true);
-	SetIsPersistent(true);
-}
-
-void FImportExportCollector::Reset()
-{
-	Exports.Reset();
-	Imports.Reset();
-}
-
-void FImportExportCollector::AddExportToIgnore(UObject* Export)
-{
-	Exports.Add(Export);
-}
-
-void FImportExportCollector::SerializeObjectAndReferencedExports(UObject* RootObject)
-{
-	*this << RootObject;
-	while (!ExportsExploreQueue.IsEmpty())
-	{
-		UObject* Export = ExportsExploreQueue.PopFrontValue();
-		Export->Serialize(*this);
-	}
-}
-
-FArchive& FImportExportCollector::operator<<(UObject*& Obj)
-{
-	if (!Obj)
-	{
-		return *this;
-	}
-	UPackage* Package = Obj->GetPackage();
-	if (!Package)
-	{
-		return *this;
-	}
-	if (Package != RootPackage)
-	{
-		AddImport(FSoftObjectPath(Obj), ESoftObjectPathCollectType::AlwaysCollect);
-		return *this;
-	}
-
-	bool bAlreadyExists;
-	Exports.Add(Obj, &bAlreadyExists);
-	if (bAlreadyExists)
-	{
-		return *this;
-	}
-	ExportsExploreQueue.Add(Obj);
-	return *this;
-}
-
-FArchive& FImportExportCollector::operator<<(FSoftObjectPath& Value)
-{
-	FName CurrentPackage;
-	FName PropertyName;
-	ESoftObjectPathCollectType CollectType;
-	ESoftObjectPathSerializeType SerializeType;
-	FSoftObjectPathThreadContext& ThreadContext = FSoftObjectPathThreadContext::Get();
-	ThreadContext.GetSerializationOptions(CurrentPackage, PropertyName, CollectType, SerializeType, this);
-
-	if (CollectType != ESoftObjectPathCollectType::NeverCollect && CollectType != ESoftObjectPathCollectType::NonPackage)
-	{
-		FName PackageName = Value.GetLongPackageFName();
-		if (PackageName != RootPackageName && !PackageName.IsNone())
-		{
-			AddImport(Value, CollectType);
-		}
-	}
-	return *this;
-}
-
-void FImportExportCollector::AddImport(const FSoftObjectPath& Path, ESoftObjectPathCollectType CollectType)
-{
-	ESoftObjectPathCollectType& ExistingImport = Imports.FindOrAdd(
-		Path, ESoftObjectPathCollectType::EditorOnlyCollect);
-	ExistingImport = Union(ExistingImport, CollectType);
-
-	ESoftObjectPathCollectType& ExistingPackage = ImportedPackages.FindOrAdd(
-		Path.GetLongPackageFName(), ESoftObjectPathCollectType::EditorOnlyCollect);
-	ExistingPackage = Union(ExistingPackage, CollectType);
-}
-
-ESoftObjectPathCollectType FImportExportCollector::Union(ESoftObjectPathCollectType A, ESoftObjectPathCollectType B)
-{
-	return static_cast<ESoftObjectPathCollectType>(FMath::Max(static_cast<int>(A), static_cast<int>(B)));
-}
-
 #endif //if WITH_EDITORONLY_DATA
 
 FBlueprintDebugData::FBlueprintDebugData() = default;
diff --git a/Engine/Source/Runtime/Engine/Private/Materials/MaterialShader.cpp b/Engine/Source/Runtime/Engine/Private/Materials/MaterialShader.cpp
--- a/Engine/Source/Runtime/Engine/Private/Materials/MaterialShader.cpp
+++ b/Engine/Source/Runtime/Engine/Private/Materials/MaterialShader.cpp
@@ -518,7 +518,7 @@ void RecordMaterialDependenciesForCook(FObjectPreSaveContext ObjectSaveContext,
 	Writer << HashMaterialDependenciesForCookArgsVersion;
 	Writer << RecordedResources;
 
-	ObjectSaveContext.AddCookDependency(
+	ObjectSaveContext.AddCookBuildDependency(
 		UE::Cook::FCookDependency::Function(
 			UE_COOK_DEPENDENCY_FUNCTION_CALL(HashMaterialDependenciesForCook), Writer.Save()));
 }
