Add default member values to FSceneViewProjectionData to avoid non-deterministic errors when a value hasn't been set manually.

diff --git a/Engine/Source/Runtime/Engine/Public/SceneView.h b/Engine/Source/Runtime/Engine/Public/SceneView.h
--- a/Engine/Source/Runtime/Engine/Public/SceneView.h
+++ b/Engine/Source/Runtime/Engine/Public/SceneView.h
@@ -56,23 +56,23 @@ class FRenderTarget;
 struct FSceneViewProjectionData
 {
 	/** The view origin. */
-	FVector ViewOrigin;
+	FVector ViewOrigin = FVector::ZeroVector;
 
 	/** Rotation matrix transforming from world space to view space. */
-	FMatrix ViewRotationMatrix;
+	FMatrix ViewRotationMatrix = FMatrix::Identity;
 
 	/** UE projection matrix projects such that clip space Z=1 is the near plane, and Z=0 is the infinite far plane. */
-	FMatrix ProjectionMatrix;
+	FMatrix ProjectionMatrix = FMatrix::Identity;
 
 	//The unconstrained (no aspect ratio bars applied) view rectangle (also unscaled)
-	FIntRect ViewRect;
+	FIntRect ViewRect = FIntRect(0,0,0,0);
 
 	//The vector (including distance) from the camera to it's viewtarget, if set. Primarily only used for Ortho views.
-	FVector CameraToViewTarget;
+	FVector CameraToViewTarget = FVector::ZeroVector;
 
 protected:
 	// The constrained view rectangle (identical to UnconstrainedUnscaledViewRect if aspect ratio is not constrained)
-	FIntRect ConstrainedViewRect;
+	FIntRect ConstrainedViewRect = FIntRect(0,0,0,0);
 
 public:
 	void SetViewRectangle(const FIntRect& InViewRect)
