Fix UseManifold cvar,
It should be still able to enable or disable globally the use Manifold. By default is set to -1 and doesn't override default behavior

diff --git a/Engine/Source/Runtime/Engine/Private/PhysicsEngine/ImmediatePhysics/ImmediatePhysicsChaos/ImmediatePhysicsSimulation_Chaos.cpp b/Engine/Source/Runtime/Engine/Private/PhysicsEngine/ImmediatePhysics/ImmediatePhysicsChaos/ImmediatePhysicsSimulation_Chaos.cpp
--- a/Engine/Source/Runtime/Engine/Private/PhysicsEngine/ImmediatePhysics/ImmediatePhysicsChaos/ImmediatePhysicsSimulation_Chaos.cpp
+++ b/Engine/Source/Runtime/Engine/Private/PhysicsEngine/ImmediatePhysics/ImmediatePhysicsChaos/ImmediatePhysicsSimulation_Chaos.cpp
@@ -78,7 +78,7 @@ Chaos::FRealSingle ChaosImmediate_Collision_MaxDepenetrationVelocity = -1.0f;
 Chaos::FRealSingle ChaosImmediate_Collision_RestitutionThresholdMultiplier = 1.5f;
 int32 ChaosImmediate_Collision_RestitutionEnabled = true;
 int32 ChaosImmediate_Collision_DeferNarrowPhase = 1;
-int32 ChaosImmediate_Collision_UseManifolds = 0;
+int32 ChaosImmediate_Collision_UseManifolds = -1;
 FAutoConsoleVariableRef CVarChaosImmPhysCollisionDisable(TEXT("p.Chaos.ImmPhys.Collision.Enabled"), ChaosImmediate_Collision_Enabled, TEXT("Enable/Disable collisions in Immediate Physics."));
 FAutoConsoleVariableRef CVarChaosImmPhysCollisionPriority(TEXT("p.Chaos.ImmPhys.Collision.Priority"), ChaosImmediate_Collision_Priority, TEXT("Set the Collision constraint sort order (Joints have priority 0)"));
 FAutoConsoleVariableRef CVarChaosImmPhysCollisionCullDistance(TEXT("p.Chaos.ImmPhys.Collision.CullDistance"), ChaosImmediate_Collision_CullDistance, TEXT("Set the collision CullDistance (if >= 0)"));
@@ -86,7 +86,7 @@ FAutoConsoleVariableRef CVarChaosImmPhysCollisionMaxDepenetrationVelocity(TEXT("
 FAutoConsoleVariableRef CVarChaosImmPhysCollisionRestitutionThresholdMultiplier(TEXT("p.Chaos.ImmPhys.Collision.RestitutionThresholdMultiplier"), ChaosImmediate_Collision_RestitutionThresholdMultiplier, TEXT("Collision Restitution Threshold (Acceleration) = Multiplier * Gravity"));
 FAutoConsoleVariableRef CVarChaosImmPhysCollisionRestitutionEnabled(TEXT("p.Chaos.ImmPhys.Collision.RestitutionEnabled"), ChaosImmediate_Collision_RestitutionEnabled, TEXT("Collision Restitution Enable/Disable"));
 FAutoConsoleVariableRef CVarChaosImmPhysCollisionDeferNarrowPhase(TEXT("p.Chaos.ImmPhys.Collision.DeferNarrowPhase"), ChaosImmediate_Collision_DeferNarrowPhase, TEXT("[Legacy Solver] Create contacts for all broadphase pairs, perform NarrowPhase later."));
-FAutoConsoleVariableRef CVarChaosImmPhysCollisionUseManifolds(TEXT("p.Chaos.ImmPhys.Collision.UseManifolds"), ChaosImmediate_Collision_UseManifolds, TEXT("[Legacy Solver] Enable/Disable use of manifoldes in collision."));
+FAutoConsoleVariableRef CVarChaosImmPhysCollisionUseManifolds(TEXT("p.Chaos.ImmPhys.Collision.UseManifolds"), ChaosImmediate_Collision_UseManifolds, TEXT("Override Enable/Disable use of manifoldes in collision (if >= 0)."));
 
 Chaos::FRealSingle ChaosImmediate_Joint_SwingTwistAngleTolerance = 1.0e-6f;
 Chaos::FRealSingle ChaosImmediate_Joint_PositionTolerance = 0.025f;
@@ -782,7 +782,7 @@ namespace ImmediatePhysics_Chaos
 		SimSpaceSettings.EulerAlpha = ChaosImmediate_Evolution_SimSpaceEulerAlpha;
 	}
 
-	void FSimulation::SetSolverSettings(const FReal FixedDt, const FReal CullDistance, const FReal MaxDepenetrationVelocity, const int32 UseLinearJointSolver, const int32 PositionIts, const int32 VelocityIts, const int32 ProjectionIts, const bool bUseManifoilds)
+	void FSimulation::SetSolverSettings(const FReal FixedDt, const FReal CullDistance, const FReal MaxDepenetrationVelocity, const int32 UseLinearJointSolver, const int32 PositionIts, const int32 VelocityIts, const int32 ProjectionIts, const int32 bUseManifoilds)
 	{
 		if (FixedDt >= FReal(0))
 		{
@@ -820,7 +820,10 @@ namespace ImmediatePhysics_Chaos
 			Implementation->Joints.SetUseSimd(bChaosImmediate_UseSimdForLinearSolver); // RBAN solver can use safely the SIMD version
 		}
 
-		Implementation->Collisions.SetAllowManifolds(bUseManifoilds || (ChaosImmediate_Collision_UseManifolds != 0));
+		if (bUseManifoilds >= 0)
+		{
+			Implementation->Collisions.SetAllowManifolds(bUseManifoilds != 0);
+		}
 	}
 
 	void FSimulation::DebugDraw()
@@ -873,7 +876,7 @@ namespace ImmediatePhysics_Chaos
 				ChaosImmediate_Evolution_PositionIterations,
 				ChaosImmediate_Evolution_VelocityIterations,
 				ChaosImmediate_Evolution_ProjectionIterations,
-				ChaosImmediate_Collision_UseManifolds != 0);
+				ChaosImmediate_Collision_UseManifolds);
 
 			FPBDJointSolverSettings JointsSettings = Implementation->Joints.GetSettings();
 			JointsSettings.SwingTwistAngleTolerance = ChaosImmediate_Joint_SwingTwistAngleTolerance;
diff --git a/Engine/Source/Runtime/Engine/Public/Physics/ImmediatePhysics/ImmediatePhysicsChaos/ImmediatePhysicsSimulation_Chaos.h b/Engine/Source/Runtime/Engine/Public/Physics/ImmediatePhysics/ImmediatePhysicsChaos/ImmediatePhysicsSimulation_Chaos.h
--- a/Engine/Source/Runtime/Engine/Public/Physics/ImmediatePhysics/ImmediatePhysicsChaos/ImmediatePhysicsSimulation_Chaos.h
+++ b/Engine/Source/Runtime/Engine/Public/Physics/ImmediatePhysics/ImmediatePhysicsChaos/ImmediatePhysicsSimulation_Chaos.h
@@ -107,7 +107,7 @@ namespace ImmediatePhysics_Chaos
 			const int32 PositionIts,
 			const int32 VelocityIts,
 			const int32 ProjectionIts,
-			const bool bUseManifolds);
+			const int32 bUseManifolds);
 
 		/** Explicit debug draw path if the use case needs it to happen at a point outside of the simulation **/
 		ENGINE_API void DebugDraw();
