Skeletal Mesh: Remove the mesh info check.

The mesh info data isn't necessarily contiguous, there are often extraneous points in the FBX not covered by the part information.

#rnx

diff --git a/Engine/Source/Runtime/Engine/Private/Rendering/SkeletalMeshLODImporterData.cpp b/Engine/Source/Runtime/Engine/Private/Rendering/SkeletalMeshLODImporterData.cpp
--- a/Engine/Source/Runtime/Engine/Private/Rendering/SkeletalMeshLODImporterData.cpp
+++ b/Engine/Source/Runtime/Engine/Private/Rendering/SkeletalMeshLODImporterData.cpp
@@ -2930,51 +2930,20 @@ FSkeletalMeshImportData FSkeletalMeshImportData::CreateFromMeshDescription(const
 	InMeshDescription.VertexAttributes().ForEachByType<FVector3f>(FCreateAndCopyAttributeValues<FVector3f>(SkelMeshImportData, InMeshDescription.Vertices()));
 	InMeshDescription.VertexAttributes().ForEachByType<FVector4f>(FCreateAndCopyAttributeValues<FVector4f>(SkelMeshImportData, InMeshDescription.Vertices()));
 
-	// Copy MeshInfos back in, if any, and only if they're valid.
+	// Add in the geometry part data as mesh info. We don't particularly care about validity of the data, leave that for the user.
+	// FIXME: This should be stored as a polygroup on the mesh description object instead.
 	if (MeshAttributes.HasSourceGeometryParts())
 	{
 		FSkeletalMeshAttributes::FSourceGeometryPartNameConstRef NameAttribute = MeshAttributes.GetSourceGeometryPartNames();
 		FSkeletalMeshAttributes::FSourceGeometryPartVertexOffsetAndCountConstRef VertexAndCountAttribute = MeshAttributes.GetSourceGeometryPartVertexOffsetAndCounts();
 
-		// Ensure that the counts + offsets add up to exactly the vertices we have.
-		TArray<SkeletalMeshImportData::FMeshInfo> MeshInfos;
 		for (FSourceGeometryPartID SourceGeometryPartID: MeshAttributes.SourceGeometryParts().GetElementIDs())
 		{
 			SkeletalMeshImportData::FMeshInfo Info;
 			Info.Name = NameAttribute.Get(SourceGeometryPartID);
 			Info.NumVertices = VertexAndCountAttribute.Get(SourceGeometryPartID)[1];
 			Info.StartImportedVertex = VertexAndCountAttribute.Get(SourceGeometryPartID)[0]; 
-			MeshInfos.Add(Info);
-		}
-
-		if (!MeshInfos.IsEmpty())
-		{
-			MeshInfos.Sort([](const SkeletalMeshImportData::FMeshInfo& A, const SkeletalMeshImportData::FMeshInfo& B)
-			{
-				return A.StartImportedVertex < B.StartImportedVertex;
-			});
-
-			bool bValid = true;
-			int32 VertexIndex = 0;
-			for (int32 Index = 0; Index < MeshInfos.Num(); Index++)
-			{
-				if (VertexIndex != MeshInfos[Index].StartImportedVertex)
-				{
-					bValid = false;
-					break;
-				}
-
-				VertexIndex += MeshInfos[Index].NumVertices;
-			}
-			if (VertexIndex != SkelMeshImportData.Points.Num())
-			{
-				bValid = false;
-			}
-
-			if (bValid)
-			{
-				SkelMeshImportData.MeshInfos = MoveTemp(MeshInfos);
-			}
+			SkelMeshImportData.MeshInfos.Add(Info);
 		}
 	}
 
