Small optimization in PIE : APackedLevelActor::IsHiddenEd was causing performance degradation of AddToWorld (Called by multiple call sites like ShouldComponentAddToScene and ShouldRender)

diff --git a/Engine/Source/Runtime/Engine/Private/PackedLevelActor/PackedLevelActor.cpp b/Engine/Source/Runtime/Engine/Private/PackedLevelActor/PackedLevelActor.cpp
--- a/Engine/Source/Runtime/Engine/Private/PackedLevelActor/PackedLevelActor.cpp
+++ b/Engine/Source/Runtime/Engine/Private/PackedLevelActor/PackedLevelActor.cpp
@@ -265,6 +265,10 @@ void APackedLevelActor::OnCommit(bool bChanged)
 
 bool APackedLevelActor::IsHiddenEd() const
 {
+	if (UWorld* World = GetWorld(); World && World->IsGameWorld())
+	{
+		return Super::IsHiddenEd();
+	}
 	return Super::IsHiddenEd() || IsEditing() || HasChildEdit() || ShouldLoadForPacking();
 }
 
