[Actor Coloration] Make the flag transient so it doesn't stays on when launching the editor.
This is for 5.4.1, for 5.5 we will revamp the actor coloration UI and save the last active handler.

diff --git a/Engine/Source/Runtime/Engine/Public/ShowFlagsValues.inl b/Engine/Source/Runtime/Engine/Public/ShowFlagsValues.inl
--- a/Engine/Source/Runtime/Engine/Public/ShowFlagsValues.inl
+++ b/Engine/Source/Runtime/Engine/Public/ShowFlagsValues.inl
@@ -276,7 +276,7 @@ SHOWFLAG_FIXED_IN_SHIPPING(0, Volumes, SFG_Advanced, NSLOCTEXT("UnrealEd", "Volu
 /** if this is a game viewport, needed? */
 SHOWFLAG_ALWAYS_ACCESSIBLE(Game, SFG_Hidden, NSLOCTEXT("UnrealEd", "GameSF", "Game"))
 /** Render objects with colors based on what the actors coloring handlers provides */
-SHOWFLAG_FIXED_IN_SHIPPING(0, ActorColoration, SFG_Hidden, NSLOCTEXT("UnrealEd", "ActorColorationSF", "Actor Coloration"))
+SHOWFLAG_FIXED_IN_SHIPPING(0, ActorColoration, SFG_Transient, NSLOCTEXT("UnrealEd", "ActorColorationSF", "Actor Coloration"))
 /** Draws BSP brushes (in game or editor textured triangles usually with lightmaps), for now SHOWFLAG_ALWAYS_ACCESSIBLE because it's exposed in SceneCapture */
 SHOWFLAG_ALWAYS_ACCESSIBLE(BSP, SFG_Normal, NSLOCTEXT("UnrealEd", "BSPSF", "BSP"))
 /** Collision drawing */
